UnrealBuildTool: More automated code cleanup

diff --git a/Engine/Source/Programs/UnrealBuildTool/Modes/UnrealHeaderToolMode.cs b/Engine/Source/Programs/UnrealBuildTool/Modes/UnrealHeaderToolMode.cs
--- a/Engine/Source/Programs/UnrealBuildTool/Modes/UnrealHeaderToolMode.cs
+++ b/Engine/Source/Programs/UnrealBuildTool/Modes/UnrealHeaderToolMode.cs
@@ -139,58 +139,58 @@ namespace UnrealBuildTool.Modes
 
 		#region IUhtConfig Implementation
 		/// <inheritdoc/>
-		public EGeneratedCodeVersion DefaultGeneratedCodeVersion => this._defaultGeneratedCodeVersion;
+		public EGeneratedCodeVersion DefaultGeneratedCodeVersion => _defaultGeneratedCodeVersion;
 
 		/// <inheritdoc/>
-		public UhtIssueBehavior EngineNativePointerMemberBehavior => this._engineNativePointerMemberBehavior;
+		public UhtIssueBehavior EngineNativePointerMemberBehavior => _engineNativePointerMemberBehavior;
 
 		/// <inheritdoc/>
-		public UhtIssueBehavior EngineObjectPtrMemberBehavior => this._engineObjectPtrMemberBehavior;
+		public UhtIssueBehavior EngineObjectPtrMemberBehavior => _engineObjectPtrMemberBehavior;
 
 		/// <inheritdoc/>
-		public UhtIssueBehavior EnginePluginNativePointerMemberBehavior => this._enginePluginNativePointerMemberBehavior;
+		public UhtIssueBehavior EnginePluginNativePointerMemberBehavior => _enginePluginNativePointerMemberBehavior;
 
 		/// <inheritdoc/>
-		public UhtIssueBehavior EnginePluginObjectPtrMemberBehavior => this._enginePluginObjectPtrMemberBehavior;
+		public UhtIssueBehavior EnginePluginObjectPtrMemberBehavior => _enginePluginObjectPtrMemberBehavior;
 
 		/// <inheritdoc/>
-		public UhtIssueBehavior NonEngineNativePointerMemberBehavior => this._nonEngineNativePointerMemberBehavior;
+		public UhtIssueBehavior NonEngineNativePointerMemberBehavior => _nonEngineNativePointerMemberBehavior;
 
 		/// <inheritdoc/>
-		public UhtIssueBehavior NonEngineObjectPtrMemberBehavior => this._nonEngineObjectPtrMemberBehavior;
+		public UhtIssueBehavior NonEngineObjectPtrMemberBehavior => _nonEngineObjectPtrMemberBehavior;
 
 		/// <summary>
 		/// If true, UObject properties are enabled in RigVM
 		/// </summary>
-		public bool AreRigVMUObjectPropertiesEnabled => this._areRigVMUObjectPropertiesEnabled;
+		public bool AreRigVMUObjectPropertiesEnabled => _areRigVMUObjectPropertiesEnabled;
 
 		/// <summary>
 		/// If true, UInterface properties are enabled in RigVM
 		/// </summary>
-		public bool AreRigVMUInterfaceProeprtiesEnabled => this._areRigVMUInterfaceProeprtiesEnabled;
+		public bool AreRigVMUInterfaceProeprtiesEnabled => _areRigVMUInterfaceProeprtiesEnabled;
 
 		/// <summary>
 		/// If true, deprecation warnings should be shown
 		/// </summary>
-		public bool ShowDeprecations => this._showDeprecations;
+		public bool ShowDeprecations => _showDeprecations;
 
 		/// <inheritdoc/>
-		public UhtIssueBehavior EngineMissingGeneratedHeaderIncludeBehavior => this._engineMissingGeneratedHeaderIncludeBehavior;
+		public UhtIssueBehavior EngineMissingGeneratedHeaderIncludeBehavior => _engineMissingGeneratedHeaderIncludeBehavior;
 
 		/// <inheritdoc/>
-		public UhtIssueBehavior NonEngineMissingGeneratedHeaderIncludeBehavior => this._nonEngineMissingGeneratedHeaderIncludeBehavior;
+		public UhtIssueBehavior NonEngineMissingGeneratedHeaderIncludeBehavior => _nonEngineMissingGeneratedHeaderIncludeBehavior;
 
 		/// <inheritdoc/>
-		public UhtIssueBehavior EngineEnumUnderlyingTypeNotSet => this._engineEnumUnderlyingTypeNotSet;
+		public UhtIssueBehavior EngineEnumUnderlyingTypeNotSet => _engineEnumUnderlyingTypeNotSet;
 
 		/// <inheritdoc/>
-		public UhtIssueBehavior NonEngineEnumUnderlyingTypeNotSet => this._nonEngineEnumUnderlyingTypeNotSet;
+		public UhtIssueBehavior NonEngineEnumUnderlyingTypeNotSet => _nonEngineEnumUnderlyingTypeNotSet;
 
 		/// <inheritdoc/>
-		public IReadOnlyDictionary<string, UhtDocumentationPolicy> DocumentationPolicies => this._documentationPolicies;
+		public IReadOnlyDictionary<string, UhtDocumentationPolicy> DocumentationPolicies => _documentationPolicies;
 
 		/// <inheritdoc/>
-		public string DefaultDocumentationPolicy => this._defaultDocumentationPolicy;
+		public string DefaultDocumentationPolicy => _defaultDocumentationPolicy;
 
 		/// <inheritdoc/>
 		public void RedirectTypeIdentifier(ref UhtToken Token)
@@ -200,7 +200,7 @@ namespace UnrealBuildTool.Modes
 				throw new Exception("Attempt to redirect type identifier when the token isn't an identifier.");
 			}
 
-			if (this._typeRedirectMap.TryGetValue(Token.Value, out StringView Redirect))
+			if (_typeRedirectMap.TryGetValue(Token.Value, out StringView Redirect))
 			{
 				Token.Value = Redirect;
 			}
@@ -209,7 +209,7 @@ namespace UnrealBuildTool.Modes
 		/// <inheritdoc/>
 		public bool RedirectMetaDataKey(string Key, out string NewKey)
 		{
-			if (this._metaDataRedirectMap.TryGetValue(Key, out string? Redirect))
+			if (_metaDataRedirectMap.TryGetValue(Key, out string? Redirect))
 			{
 				NewKey = Redirect;
 				return Key != NewKey;
@@ -224,21 +224,21 @@ namespace UnrealBuildTool.Modes
 		/// <inheritdoc/>
 		public bool IsValidUnits(StringView Units)
 		{
-			return this._units.Contains(Units);
+			return _units.Contains(Units);
 		}
 
 		/// <inheritdoc/>
 		public bool IsStructWithTPrefix(StringView Name)
 		{
-			return this._structsWithTPrefix.Contains(Name);
+			return _structsWithTPrefix.Contains(Name);
 		}
 
 		/// <inheritdoc/>
 		public int FindDelegateParameterCount(StringView DelegateMacro)
 		{
-			for (int Index = 0, Count = this._delegateParameterCountStrings.Count; Index < Count; ++Index)
+			for (int Index = 0, Count = _delegateParameterCountStrings.Count; Index < Count; ++Index)
 			{
-				if (DelegateMacro.Span.Contains(this._delegateParameterCountStrings[Index].Span, StringComparison.Ordinal))
+				if (DelegateMacro.Span.Contains(_delegateParameterCountStrings[Index].Span, StringComparison.Ordinal))
 				{
 					return Index;
 				}
@@ -249,13 +249,13 @@ namespace UnrealBuildTool.Modes
 		/// <inheritdoc/>
 		public StringView GetDelegateParameterCountString(int Index)
 		{
-			return Index >= 0 ? this._delegateParameterCountStrings[Index] : "";
+			return Index >= 0 ? _delegateParameterCountStrings[Index] : "";
 		}
 
 		/// <inheritdoc/>
 		public bool IsExporterEnabled(string Name)
 		{
-			this._ini.GetBool("UnrealHeaderTool", Name, out bool Value);
+			_ini.GetBool("UnrealHeaderTool", Name, out bool Value);
 			return Value;
 		}
 		#endregion
@@ -267,28 +267,28 @@ namespace UnrealBuildTool.Modes
 		public UhtConfigImpl(CommandLineArguments Args)
 		{
 			DirectoryReference ConfigDirectory = DirectoryReference.Combine(Unreal.EngineDirectory, "Programs", "UnrealHeaderTool");
-			this._ini = ConfigCache.ReadHierarchy(ConfigHierarchyType.Engine, ConfigDirectory, BuildHostPlatform.Current.Platform, "", Args.GetRawArray());
-
-			this._typeRedirectMap = GetRedirectsStringView("UnrealHeaderTool", "TypeRedirects", "OldType", "NewType");
-			this._metaDataRedirectMap = GetRedirectsString("CoreUObject.Metadata", "MetadataRedirects", "OldKey", "NewKey");
-			this._structsWithNoPrefix = GetHashSet("UnrealHeaderTool", "StructsWithNoPrefix", StringViewComparer.Ordinal);
-			this._structsWithTPrefix = GetHashSet("UnrealHeaderTool", "StructsWithTPrefix", StringViewComparer.Ordinal);
-			this._units = GetHashSet("UnrealHeaderTool", "Units", StringViewComparer.OrdinalIgnoreCase);
-			this._delegateParameterCountStrings = GetList("UnrealHeaderTool", "DelegateParameterCountStrings");
-			this._defaultGeneratedCodeVersion = GetGeneratedCodeVersion("UnrealHeaderTool", "DefaultGeneratedCodeVersion", EGeneratedCodeVersion.V1);
-			this._engineNativePointerMemberBehavior = GetIssueBehavior("UnrealHeaderTool", "EngineNativePointerMemberBehavior", UhtIssueBehavior.AllowSilently);
-			this._engineObjectPtrMemberBehavior = GetIssueBehavior("UnrealHeaderTool", "EngineObjectPtrMemberBehavior", UhtIssueBehavior.AllowSilently);
-			this._enginePluginNativePointerMemberBehavior = GetIssueBehavior("UnrealHeaderTool", "EnginePluginNativePointerMemberBehavior", UhtIssueBehavior.AllowSilently);
-			this._enginePluginObjectPtrMemberBehavior = GetIssueBehavior("UnrealHeaderTool", "EnginePluginObjectPtrMemberBehavior", UhtIssueBehavior.AllowSilently);
-			this._nonEngineNativePointerMemberBehavior = GetIssueBehavior("UnrealHeaderTool", "NonEngineNativePointerMemberBehavior", UhtIssueBehavior.AllowSilently);
-			this._nonEngineObjectPtrMemberBehavior = GetIssueBehavior("UnrealHeaderTool", "NonEngineObjectPtrMemberBehavior", UhtIssueBehavior.AllowSilently);
-			this._areRigVMUObjectPropertiesEnabled = GetBoolean("UnrealHeaderTool", "AreRigVMUObjectPropertiesEnabled", false);
-			this._areRigVMUInterfaceProeprtiesEnabled = GetBoolean("UnrealHeaderTool", "AreRigVMUInterfaceProeprtiesEnabled", false);
-			this._showDeprecations = GetBoolean("UnrealHeaderTool", "ShowDeprecations", true);
-			this._engineMissingGeneratedHeaderIncludeBehavior = GetIssueBehavior("UnrealHeaderTool", "EngineMissingGeneratedHeaderIncludeBehavior", UhtIssueBehavior.AllowSilently);
-			this._nonEngineMissingGeneratedHeaderIncludeBehavior = GetIssueBehavior("UnrealHeaderTool", "NonEngineMissingGeneratedHeaderIncludeBehavior", UhtIssueBehavior.AllowSilently);
-			this._engineEnumUnderlyingTypeNotSet = GetIssueBehavior("UnrealHeaderTool", "EngineEnumUnderlyingTypeNotSet", UhtIssueBehavior.AllowSilently);
-			this._nonEngineEnumUnderlyingTypeNotSet = GetIssueBehavior("UnrealHeaderTool", "NonEngineEnumUnderlyingTypeNotSet", UhtIssueBehavior.AllowSilently);
+			_ini = ConfigCache.ReadHierarchy(ConfigHierarchyType.Engine, ConfigDirectory, BuildHostPlatform.Current.Platform, "", Args.GetRawArray());
+
+			_typeRedirectMap = GetRedirectsStringView("UnrealHeaderTool", "TypeRedirects", "OldType", "NewType");
+			_metaDataRedirectMap = GetRedirectsString("CoreUObject.Metadata", "MetadataRedirects", "OldKey", "NewKey");
+			_structsWithNoPrefix = GetHashSet("UnrealHeaderTool", "StructsWithNoPrefix", StringViewComparer.Ordinal);
+			_structsWithTPrefix = GetHashSet("UnrealHeaderTool", "StructsWithTPrefix", StringViewComparer.Ordinal);
+			_units = GetHashSet("UnrealHeaderTool", "Units", StringViewComparer.OrdinalIgnoreCase);
+			_delegateParameterCountStrings = GetList("UnrealHeaderTool", "DelegateParameterCountStrings");
+			_defaultGeneratedCodeVersion = GetGeneratedCodeVersion("UnrealHeaderTool", "DefaultGeneratedCodeVersion", EGeneratedCodeVersion.V1);
+			_engineNativePointerMemberBehavior = GetIssueBehavior("UnrealHeaderTool", "EngineNativePointerMemberBehavior", UhtIssueBehavior.AllowSilently);
+			_engineObjectPtrMemberBehavior = GetIssueBehavior("UnrealHeaderTool", "EngineObjectPtrMemberBehavior", UhtIssueBehavior.AllowSilently);
+			_enginePluginNativePointerMemberBehavior = GetIssueBehavior("UnrealHeaderTool", "EnginePluginNativePointerMemberBehavior", UhtIssueBehavior.AllowSilently);
+			_enginePluginObjectPtrMemberBehavior = GetIssueBehavior("UnrealHeaderTool", "EnginePluginObjectPtrMemberBehavior", UhtIssueBehavior.AllowSilently);
+			_nonEngineNativePointerMemberBehavior = GetIssueBehavior("UnrealHeaderTool", "NonEngineNativePointerMemberBehavior", UhtIssueBehavior.AllowSilently);
+			_nonEngineObjectPtrMemberBehavior = GetIssueBehavior("UnrealHeaderTool", "NonEngineObjectPtrMemberBehavior", UhtIssueBehavior.AllowSilently);
+			_areRigVMUObjectPropertiesEnabled = GetBoolean("UnrealHeaderTool", "AreRigVMUObjectPropertiesEnabled", false);
+			_areRigVMUInterfaceProeprtiesEnabled = GetBoolean("UnrealHeaderTool", "AreRigVMUInterfaceProeprtiesEnabled", false);
+			_showDeprecations = GetBoolean("UnrealHeaderTool", "ShowDeprecations", true);
+			_engineMissingGeneratedHeaderIncludeBehavior = GetIssueBehavior("UnrealHeaderTool", "EngineMissingGeneratedHeaderIncludeBehavior", UhtIssueBehavior.AllowSilently);
+			_nonEngineMissingGeneratedHeaderIncludeBehavior = GetIssueBehavior("UnrealHeaderTool", "NonEngineMissingGeneratedHeaderIncludeBehavior", UhtIssueBehavior.AllowSilently);
+			_engineEnumUnderlyingTypeNotSet = GetIssueBehavior("UnrealHeaderTool", "EngineEnumUnderlyingTypeNotSet", UhtIssueBehavior.AllowSilently);
+			_nonEngineEnumUnderlyingTypeNotSet = GetIssueBehavior("UnrealHeaderTool", "NonEngineEnumUnderlyingTypeNotSet", UhtIssueBehavior.AllowSilently);
 
 			GetDocumentationPolicies("UnrealHeaderTool", "DocumentationPolicies");
 			_defaultDocumentationPolicy = GetString("UnrealHeaderTool", "DefaultDocumentationPolicy", "");
@@ -296,7 +296,7 @@ namespace UnrealBuildTool.Modes
 
 		private bool GetBoolean(string SectionName, string KeyName, bool bDefault)
 		{
-			if (this._ini.TryGetValue(SectionName, KeyName, out bool value))
+			if (_ini.TryGetValue(SectionName, KeyName, out bool value))
 			{
 				return value;
 			}
@@ -305,7 +305,7 @@ namespace UnrealBuildTool.Modes
 
 		private string GetString(string SectionName, string KeyName, string Default)
 		{
-			if (this._ini.TryGetValue(SectionName, KeyName, out string? value))
+			if (_ini.TryGetValue(SectionName, KeyName, out string? value))
 			{
 				return value;
 			}
@@ -314,11 +314,11 @@ namespace UnrealBuildTool.Modes
 
 		private UhtIssueBehavior GetIssueBehavior(string SectionName, string KeyName, UhtIssueBehavior Default)
 		{
-			if (this._ini.TryGetValue(SectionName, KeyName, out string? BehaviorStr))
+			if (_ini.TryGetValue(SectionName, KeyName, out string? BehaviorStr))
 			{
 				if (!Enum.TryParse(BehaviorStr, out UhtIssueBehavior Value))
 				{
-					throw new Exception(string.Format("Unrecognized issue behavior '{0}'", BehaviorStr));
+					throw new Exception(String.Format("Unrecognized issue behavior '{0}'", BehaviorStr));
 				}
 				return Value;
 			}
@@ -327,11 +327,11 @@ namespace UnrealBuildTool.Modes
 
 		private EGeneratedCodeVersion GetGeneratedCodeVersion(string SectionName, string KeyName, EGeneratedCodeVersion Default)
 		{
-			if (this._ini.TryGetValue(SectionName, KeyName, out string? BehaviorStr))
+			if (_ini.TryGetValue(SectionName, KeyName, out string? BehaviorStr))
 			{
 				if (!Enum.TryParse(BehaviorStr, out EGeneratedCodeVersion Value))
 				{
-					throw new Exception(string.Format("Unrecognized generated code version '{0}'", BehaviorStr));
+					throw new Exception(String.Format("Unrecognized generated code version '{0}'", BehaviorStr));
 				}
 				return Value;
 			}
@@ -342,7 +342,7 @@ namespace UnrealBuildTool.Modes
 		{
 			Dictionary<StringView, StringView> Redirects = new();
 
-			if (this._ini.TryGetValues(Section, Key, out IReadOnlyList<string>? StringList))
+			if (_ini.TryGetValues(Section, Key, out IReadOnlyList<string>? StringList))
 			{
 				foreach (string Line in StringList)
 				{
@@ -350,11 +350,11 @@ namespace UnrealBuildTool.Modes
 					{
 						if (!Properties.TryGetValue(OldKeyName, out string? OldKey))
 						{
-							throw new Exception(string.Format("Unable to get the {0} from the {1} value", OldKeyName, Key));
+							throw new Exception(String.Format("Unable to get the {0} from the {1} value", OldKeyName, Key));
 						}
 						if (!Properties.TryGetValue(NewKeyName, out string? NewKey))
 						{
-							throw new Exception(string.Format("Unable to get the {0} from the {1} value", NewKeyName, Key));
+							throw new Exception(String.Format("Unable to get the {0} from the {1} value", NewKeyName, Key));
 						}
 						Redirects.Add(OldKey, NewKey);
 					}
@@ -367,7 +367,7 @@ namespace UnrealBuildTool.Modes
 		{
 			Dictionary<string, string> Redirects = new();
 
-			if (this._ini.TryGetValues(Section, Key, out IReadOnlyList<string>? StringList))
+			if (_ini.TryGetValues(Section, Key, out IReadOnlyList<string>? StringList))
 			{
 				foreach (string Line in StringList)
 				{
@@ -375,11 +375,11 @@ namespace UnrealBuildTool.Modes
 					{
 						if (!Properties.TryGetValue(OldKeyName, out string? OldKey))
 						{
-							throw new Exception(string.Format("Unable to get the {0} from the {1} value", OldKeyName, Key));
+							throw new Exception(String.Format("Unable to get the {0} from the {1} value", OldKeyName, Key));
 						}
 						if (!Properties.TryGetValue(NewKeyName, out string? NewKey))
 						{
-							throw new Exception(string.Format("Unable to get the {0} from the {1} value", NewKeyName, Key));
+							throw new Exception(String.Format("Unable to get the {0} from the {1} value", NewKeyName, Key));
 						}
 						Redirects.Add(OldKey, NewKey);
 					}
@@ -392,7 +392,7 @@ namespace UnrealBuildTool.Modes
 		{
 			List<StringView> List = new();
 
-			if (this._ini.TryGetValues(Section, Key, out IReadOnlyList<string>? StringList))
+			if (_ini.TryGetValues(Section, Key, out IReadOnlyList<string>? StringList))
 			{
 				foreach (string Value in StringList)
 				{
@@ -406,7 +406,7 @@ namespace UnrealBuildTool.Modes
 		{
 			HashSet<StringView> Set = new(Comparer);
 
-			if (this._ini.TryGetValues(Section, Key, out IReadOnlyList<string>? StringList))
+			if (_ini.TryGetValues(Section, Key, out IReadOnlyList<string>? StringList))
 			{
 				foreach (string Value in StringList)
 				{
@@ -427,7 +427,7 @@ namespace UnrealBuildTool.Modes
 				FloatRangesRequired = true,
 			};
 
-			if (this._ini.TryGetValues(Section, Key, out IReadOnlyList<string>? StringList))
+			if (_ini.TryGetValues(Section, Key, out IReadOnlyList<string>? StringList))
 			{
 				foreach (string Value in StringList)
 				{
@@ -567,26 +567,26 @@ namespace UnrealBuildTool.Modes
 		public UhtTestFileManager(string RootDirectory)
 		{
 			this.RootDirectory = RootDirectory;
-			this.InnerManager = new UhtStdFileManager();
+			InnerManager = new UhtStdFileManager();
 		}
 
 		/// <inheritdoc/>
 		public string GetFullFilePath(string FilePath)
 		{
-			if (this.RootDirectory == null)
+			if (RootDirectory == null)
 			{
 				return FilePath;
 			}
 			else
 			{
-				return Path.Combine(this.RootDirectory, FilePath);
+				return Path.Combine(RootDirectory, FilePath);
 			}
 		}
 
 		/// <inheritdoc/>
 		public bool ReadSource(string FilePath, out UhtSourceFragment Fragment)
 		{
-			if (this.SourceFragments.TryGetValue(FilePath, out Fragment))
+			if (SourceFragments.TryGetValue(FilePath, out Fragment))
 			{
 				return true;
 			}
@@ -611,9 +611,9 @@ namespace UnrealBuildTool.Modes
 		/// <inheritdoc/>
 		public bool WriteOutput(string FilePath, ReadOnlySpan<char> Contents)
 		{
-			lock (this.Outputs)
+			lock (Outputs)
 			{
-				this.Outputs.Add(FilePath, Contents.ToString());
+				Outputs.Add(FilePath, Contents.ToString());
 			}
 			return true;
 		}
@@ -621,12 +621,12 @@ namespace UnrealBuildTool.Modes
 		/// <inheritdoc/>
 		public bool RenameOutput(string OldFilePath, string NewFilePath)
 		{
-			lock (this.Outputs)
+			lock (Outputs)
 			{
-				if (this.Outputs.TryGetValue(OldFilePath, out string? Content))
+				if (Outputs.TryGetValue(OldFilePath, out string? Content))
 				{
-					this.Outputs.Remove(OldFilePath);
-					this.Outputs.Add(NewFilePath, Content);
+					Outputs.Remove(OldFilePath);
+					Outputs.Add(NewFilePath, Content);
 				}
 			}
 			return true;
@@ -642,7 +642,7 @@ namespace UnrealBuildTool.Modes
 		/// <param name="Data">The data associated with the path</param>
 		public void AddSourceFragment(UhtSourceFile SourceFile, string FilePath, int LineNumber, StringView Data)
 		{
-			this.SourceFragments.Add(FilePath, new UhtSourceFragment { SourceFile = SourceFile, FilePath = FilePath, LineNumber = LineNumber, Data = Data });
+			SourceFragments.Add(FilePath, new UhtSourceFragment { SourceFile = SourceFile, FilePath = FilePath, LineNumber = LineNumber, Data = Data });
 		}
 	}
 
@@ -1233,7 +1233,7 @@ namespace UnrealBuildTool.Modes
 					FileManager = new UhtStdFileManager(),
 					EngineDirectory = Unreal.EngineDirectory.FullName,
 					ProjectFile = ProjectFile,
-					ProjectDirectory = string.IsNullOrEmpty(ProjectPath) ? null : ProjectPath,
+					ProjectDirectory = String.IsNullOrEmpty(ProjectPath) ? null : ProjectPath,
 					ReferenceDirectory = FileReference.Combine(EngineProgramSavedDirectory, "UnrealHeaderTool", "Saved", "ReferenceExports").FullName,
 					VerifyDirectory = FileReference.Combine(EngineProgramSavedDirectory, "UnrealHeaderTool", "Saved", "VerifyExports").FullName,
 					WarningsAsErrors = Options.bWarningsAsErrors,
diff --git a/Engine/Source/Programs/UnrealBuildTool/Platform/Android/AndroidAARHandler.cs b/Engine/Source/Programs/UnrealBuildTool/Platform/Android/AndroidAARHandler.cs
--- a/Engine/Source/Programs/UnrealBuildTool/Platform/Android/AndroidAARHandler.cs
+++ b/Engine/Source/Programs/UnrealBuildTool/Platform/Android/AndroidAARHandler.cs
@@ -175,10 +175,10 @@ namespace UnrealBuildTool
 			try
 			{
 				string[] Sections = VersionString.Split(".".ToCharArray());
-				Value |= (Sections.Length > 0) ? (uint.Parse(Sections[0]) << 24) : 0;
-				Value |= (Sections.Length > 1) ? (uint.Parse(Sections[1]) << 16) : 0;
-				Value |= (Sections.Length > 2) ? (uint.Parse(Sections[2]) << 8) : 0;
-				Value |= (Sections.Length > 3) ? uint.Parse(Sections[3]) : 0;
+				Value |= (Sections.Length > 0) ? (UInt32.Parse(Sections[0]) << 24) : 0;
+				Value |= (Sections.Length > 1) ? (UInt32.Parse(Sections[1]) << 16) : 0;
+				Value |= (Sections.Length > 2) ? (UInt32.Parse(Sections[2]) << 8) : 0;
+				Value |= (Sections.Length > 3) ? UInt32.Parse(Sections[3]) : 0;
 			}
 			catch (Exception)
 			{
diff --git a/Engine/Source/Programs/UnrealBuildTool/Platform/Android/AndroidPlatformSDK.cs b/Engine/Source/Programs/UnrealBuildTool/Platform/Android/AndroidPlatformSDK.cs
--- a/Engine/Source/Programs/UnrealBuildTool/Platform/Android/AndroidPlatformSDK.cs
+++ b/Engine/Source/Programs/UnrealBuildTool/Platform/Android/AndroidPlatformSDK.cs
@@ -93,7 +93,7 @@ namespace UnrealBuildTool
 				if (EqualsIndex > 0)
 				{
 					string[] RevisionParts = RevisionString.Substring(EqualsIndex + 1).Trim().Split('.');
-					int RevisionMinor = int.Parse(RevisionParts.Length > 1 ? RevisionParts[1] : "0");
+					int RevisionMinor = Int32.Parse(RevisionParts.Length > 1 ? RevisionParts[1] : "0");
 					char RevisionLetter = Convert.ToChar('a' + RevisionMinor);
 					NDKToolchainVersion = "r" + RevisionParts[0] + (RevisionMinor > 0 ? Char.ToString(RevisionLetter) : "");
 				}
@@ -116,7 +116,7 @@ namespace UnrealBuildTool
 		public override bool TryConvertVersionToInt(string? StringValue, out UInt64 OutValue, string? Hint)
 		{
 			// convert r<num>[letter] to hex
-			if (!string.IsNullOrEmpty(StringValue))
+			if (!String.IsNullOrEmpty(StringValue))
 			{
 				Match Result = Regex.Match(StringValue, @"^r(\d*)([a-z])?");
 
@@ -128,7 +128,7 @@ namespace UnrealBuildTool
 					{
 						RevisionNumber = (Result.Groups[2].Value[0] - 'a') + 1;
 					}
-					string VersionString = string.Format("{0}{1:00}{2:00}", Result.Groups[1], RevisionNumber, 0);
+					string VersionString = String.Format("{0}{1:00}{2:00}", Result.Groups[1], RevisionNumber, 0);
 					return UInt64.TryParse(VersionString, out OutValue);
 				}
 			}
@@ -228,7 +228,7 @@ namespace UnrealBuildTool
 				string path;
 				foreach (KeyValuePair<string, string> kvp in EnvVarNames)
 				{
-					if (GetPath(configCacheIni, "/Script/AndroidPlatformEditor.AndroidSDKSettings", kvp.Value, out path) && !string.IsNullOrEmpty(path))
+					if (GetPath(configCacheIni, "/Script/AndroidPlatformEditor.AndroidSDKSettings", kvp.Value, out path) && !String.IsNullOrEmpty(path))
 					{
 						AndroidEnv.Add(kvp.Key, path);
 					}
@@ -326,7 +326,7 @@ namespace UnrealBuildTool
 					// JAVA_VERSION="17.0.6"
 					int VersionStartIndex = LineContents.IndexOf("\"");
 					int VersionStopIndex = LineContents.IndexOf(".");
-					int.TryParse(LineContents.Substring(VersionStartIndex, VersionStopIndex - VersionStartIndex), out JavaVersion);
+					Int32.TryParse(LineContents.Substring(VersionStartIndex, VersionStopIndex - VersionStartIndex), out JavaVersion);
 					break;
 				}
 			}
diff --git a/Engine/Source/Programs/UnrealBuildTool/Platform/Android/AndroidToolChain.cs b/Engine/Source/Programs/UnrealBuildTool/Platform/Android/AndroidToolChain.cs
--- a/Engine/Source/Programs/UnrealBuildTool/Platform/Android/AndroidToolChain.cs
+++ b/Engine/Source/Programs/UnrealBuildTool/Platform/Android/AndroidToolChain.cs
@@ -62,30 +62,30 @@ namespace UnrealBuildTool
 		protected static string? AndroidClangBuild;
 
 		// the "-android" suffix paths here are vcpkg triplets for the android platform
-		static private Dictionary<UnrealArch, string[]> AllArchNames = new() {
+		private static Dictionary<UnrealArch, string[]> AllArchNames = new() {
 			{ UnrealArch.Arm64, new string[] { "arm64", "arm64-v8a", "arm64-android" } },
 			{ UnrealArch.X64,   new string[] { "x64", "x86_64", "x64-android" } }
 		};
 
 		// architecture paths to use for filtering include and lib paths
-		static private Dictionary<UnrealArch, string[]> AllFilterArchNames = new() {
+		private static Dictionary<UnrealArch, string[]> AllFilterArchNames = new() {
 			{ UnrealArch.Arm64, new string[] { "arm64", "arm64-v8a", "arm64-android" } },
 			{ UnrealArch.X64,   new string[] { "x64", "x86_64", "x64-android" } },
 			// using Default as a placeholder to remove old folders for arches we no longer support, but licensees may have in their Build rules that we need to strip out
 			{ UnrealArch.Deprecated  , new string[] { "armv7", "armeabi-v7a", "arm-android", "x86", "x86-android" } }
 		};
 
-		static private Dictionary<UnrealArch, string[]> LibrariesToSkip = new() {
+		private static Dictionary<UnrealArch, string[]> LibrariesToSkip = new() {
 			{ UnrealArch.Arm64, new string[] { "nvToolsExt", "nvToolsExtStub", "vorbisenc", } },
 			{ UnrealArch.X64,   new string[] { "nvToolsExt", "nvToolsExtStub", "oculus", "OVRPlugin", "vrapi", "ovrkernel", "systemutils", "openglloader", "ovrplatformloader", "gpg", "vorbisenc", } }
 		};
 
-		static private Dictionary<UnrealArch, string[]> ModulesToSkip = new() {
+		private static Dictionary<UnrealArch, string[]> ModulesToSkip = new() {
 			{ UnrealArch.Arm64, new string[] {  } },
 			{ UnrealArch.X64,   new string[] { "OnlineSubsystemOculus", "OculusHMD", "OculusMR", "OnlineSubsystemGooglePlay" } }
 		};
 
-		static private Dictionary<UnrealArch, string[]> GeneratedModulesToSkip = new() {
+		private static Dictionary<UnrealArch, string[]> GeneratedModulesToSkip = new() {
 			{ UnrealArch.Arm64, new string[] {  } },
 			{ UnrealArch.X64,   new string[] { "OculusEntitlementCallbackProxy", "OculusCreateSessionCallbackProxy", "OculusFindSessionsCallbackProxy", "OculusIdentityCallbackProxy", "OculusNetConnection", "OculusNetDriver", "OnlineSubsystemOculus_init" } }
 		};
@@ -106,7 +106,7 @@ namespace UnrealBuildTool
 
 		public string GetClangVersionString()
 		{
-			return string.Format("{0}.{1}.{2}", ClangVersionMajor, ClangVersionMinor, ClangVersionPatch);
+			return String.Format("{0}.{1}.{2}", ClangVersionMajor, ClangVersionMinor, ClangVersionPatch);
 		}
 
 		public bool IsNewNDKModel()
@@ -186,7 +186,7 @@ namespace UnrealBuildTool
 				}
 				string[] VersionFile = File.ReadAllLines(VersionFilename);
 				string[] VersionParts = VersionFile[0].Split('.');
-				SetClangVersion(int.Parse(VersionParts[0]), (VersionParts.Length > 1) ? int.Parse(VersionParts[1]) : 0, (VersionParts.Length > 2) ? int.Parse(VersionParts[2]) : 0);
+				SetClangVersion(Int32.Parse(VersionParts[0]), (VersionParts.Length > 1) ? Int32.Parse(VersionParts[1]) : 0, (VersionParts.Length > 2) ? Int32.Parse(VersionParts[2]) : 0);
 			}
 			else
 			{
@@ -245,7 +245,7 @@ namespace UnrealBuildTool
 			return new ClangToolChainInfo(FileReference.FromString(ClangPath)!, FileReference.FromString(ArPathArm64)!, Logger);
 		}
 
-		static public string GetGLESVersion(bool bBuildForES31)
+		public static string GetGLESVersion(bool bBuildForES31)
 		{
 			string GLESversion = "0x00030000";
 
@@ -304,10 +304,12 @@ namespace UnrealBuildTool
 			if (NDKVersion.Contains("-"))
 			{
 				int Version;
-				if (int.TryParse(NDKVersion.Substring(NDKVersion.LastIndexOf('-') + 1), out Version))
+				if (Int32.TryParse(NDKVersion.Substring(NDKVersion.LastIndexOf('-') + 1), out Version))
 				{
 					if (Version > NDKVersionInt)
+					{
 						NDKVersionInt = Version;
+					}
 				}
 			}
 			return NDKVersionInt;
@@ -372,7 +374,7 @@ namespace UnrealBuildTool
 			if (ApiString.Contains("-"))
 			{
 				int Version;
-				if (int.TryParse(ApiString.Substring(ApiString.LastIndexOf('-') + 1), out Version))
+				if (Int32.TryParse(ApiString.Substring(ApiString.LastIndexOf('-') + 1), out Version))
 				{
 					return (Version >= MinPlatform && Version <= MaxPlatform);
 				}
@@ -535,7 +537,7 @@ namespace UnrealBuildTool
 				// apparently there can be hashing conflicts with PGO which can result in:
 				// 'Function control flow change detected (hash mismatch)' warnings. 
 				Arguments.Add("-Wno-backend-plugin");
-				Arguments.Add(string.Format(" -fprofile-use=\"{0}.profdata\"", Path.Combine(CompileEnvironment.PGODirectory!, CompileEnvironment.PGOFilenamePrefix!).Replace("\\", "/")));
+				Arguments.Add(String.Format(" -fprofile-use=\"{0}.profdata\"", Path.Combine(CompileEnvironment.PGODirectory!, CompileEnvironment.PGOFilenamePrefix!).Replace("\\", "/")));
 			}
 			else if (CompileEnvironment.bPGOProfile)
 			{
@@ -782,7 +784,7 @@ namespace UnrealBuildTool
 				Result += " -Wno-profile-instr-out-of-date";
 				Result += " -Wno-profile-instr-unprofiled";
 
-				Result += string.Format(" -fprofile-use=\"{0}.profdata\"", Path.Combine(LinkEnvironment.PGODirectory!, LinkEnvironment.PGOFilenamePrefix!).Replace("\\", "/"));
+				Result += String.Format(" -fprofile-use=\"{0}.profdata\"", Path.Combine(LinkEnvironment.PGODirectory!, LinkEnvironment.PGOFilenamePrefix!).Replace("\\", "/"));
 			}
 			else if (LinkEnvironment.bPGOProfile)
 			{
@@ -944,12 +946,12 @@ namespace UnrealBuildTool
 			return false;
 		}
 
-		static private string GetNativeGluePath()
+		private static string GetNativeGluePath()
 		{
 			return Environment.GetEnvironmentVariable("NDKROOT") + "/sources/android/native_app_glue/android_native_app_glue.c";
 		}
 
-		static private string GetCpuFeaturesPath()
+		private static string GetCpuFeaturesPath()
 		{
 			return Environment.GetEnvironmentVariable("NDKROOT") + "/sources/android/cpufeatures/cpu-features.c";
 		}
@@ -1008,7 +1010,7 @@ namespace UnrealBuildTool
 		protected static string? ArPathx64;
 		protected static string? ReadElfPath;
 
-		static public string GetStripExecutablePath(UnrealArch UnrealArch)
+		public static string GetStripExecutablePath(UnrealArch UnrealArch)
 		{
 			string StripPath = ArPathArm64!;
 			if (UnrealArch == UnrealArch.X64)
@@ -1049,7 +1051,7 @@ namespace UnrealBuildTool
 			return base.CompileCPPFiles(CompileEnvironment, InputFiles, OutputDir, ModuleName, Graph);
 		}
 
-		static public string InlineArchName(string Pathname, UnrealArch Arch, bool bUseShortNames = false)
+		public static string InlineArchName(string Pathname, UnrealArch Arch, bool bUseShortNames = false)
 		{
 			string FinalArch = "-" + Arch.ToString().ToLower();
 			if (bUseShortNames)
@@ -1069,7 +1071,7 @@ namespace UnrealBuildTool
 			return Pathname;
 		}
 
-		static public DirectoryReference InlineArchIncludeFolder(DirectoryReference PathRef, UnrealArch Arch)
+		public static DirectoryReference InlineArchIncludeFolder(DirectoryReference PathRef, UnrealArch Arch)
 		{
 			return DirectoryReference.Combine(PathRef, "include", Arch.ToString());
 		}
@@ -1160,7 +1162,7 @@ namespace UnrealBuildTool
 
 				// Add the source file and its included files to the prerequisite item list.
 				CompileAction.PrerequisiteItems.Add(ISPCFile);
-				CompileAction.StatusDescription = string.Format("[{0}] {1}", CompileEnvironment.Architecture, Path.GetFileName(ISPCFile.AbsolutePath));
+				CompileAction.StatusDescription = String.Format("[{0}] {1}", CompileEnvironment.Architecture, Path.GetFileName(ISPCFile.AbsolutePath));
 
 				FileItem ISPCFinalHeaderFile = FileItem.GetItemByFileReference(
 					FileReference.Combine(
@@ -1324,14 +1326,14 @@ namespace UnrealBuildTool
 				{
 					if (IsDirectoryForArch(IncludePath.FullName, CompileEnvironment.Architecture))
 					{
-						Arguments.Add(string.Format(" -I\"{0}\"", IncludePath));
+						Arguments.Add(String.Format(" -I\"{0}\"", IncludePath));
 					}
 				}
 				foreach (DirectoryReference IncludePath in CompileEnvironment.UserIncludePaths)
 				{
 					if (IsDirectoryForArch(IncludePath.FullName, CompileEnvironment.Architecture))
 					{
-						Arguments.Add(string.Format(" -I\"{0}\"", IncludePath));
+						Arguments.Add(String.Format(" -I\"{0}\"", IncludePath));
 					}
 				}
 
@@ -1353,7 +1355,7 @@ namespace UnrealBuildTool
 				// Add the source file and its included files to the prerequisite item list.
 				CompileAction.PrerequisiteItems.Add(ISPCFile);
 
-				CompileAction.StatusDescription = string.Format("[{0}] [{1}]", CompileEnvironment.Architecture, Path.GetFileName(ISPCFile.AbsolutePath));
+				CompileAction.StatusDescription = String.Format("[{0}] [{1}]", CompileEnvironment.Architecture, Path.GetFileName(ISPCFile.AbsolutePath));
 
 				for (int i = 0; i < CompiledISPCObjFiles.Count; i++)
 				{
@@ -1437,7 +1439,7 @@ namespace UnrealBuildTool
 			FileItem OutputFile;
 			OutputFile = FileItem.GetItemByFileReference(LinkEnvironment.OutputFilePaths.Where(x => x.FullName.Contains(LinkEnvironment.Architecture.ToString(), StringComparison.InvariantCultureIgnoreCase)).First());
 			LinkAction.ProducedItems.Add(OutputFile);
-			LinkAction.StatusDescription = string.Format("{0}", Path.GetFileName(OutputFile.AbsolutePath));
+			LinkAction.StatusDescription = String.Format("{0}", Path.GetFileName(OutputFile.AbsolutePath));
 			LinkAction.CommandVersion = AndroidClangBuild!;
 
 			// LinkAction.bPrintDebugInfo = true;
@@ -1445,11 +1447,11 @@ namespace UnrealBuildTool
 			// Add the output file to the command-line.
 			if (LinkEnvironment.bIsBuildingLibrary)
 			{
-				LinkAction.CommandArguments += string.Format(" \"{0}\"", OutputFile.AbsolutePath);
+				LinkAction.CommandArguments += String.Format(" \"{0}\"", OutputFile.AbsolutePath);
 			}
 			else
 			{
-				LinkAction.CommandArguments += string.Format(" -o \"{0}\"", OutputFile.AbsolutePath);
+				LinkAction.CommandArguments += String.Format(" -o \"{0}\"", OutputFile.AbsolutePath);
 			}
 
 			// Add the input files to a response file, and pass the response file on the command-line.
@@ -1465,7 +1467,7 @@ namespace UnrealBuildTool
 				{
 					InputPath = InputFile.Location.FullName;
 				}
-				InputFileNames.Add(string.Format("\"{0}\"", InputPath.Replace('\\', '/')));
+				InputFileNames.Add(String.Format("\"{0}\"", InputPath.Replace('\\', '/')));
 
 				LinkAction.PrerequisiteItems.Add(InputFile);
 			}
@@ -1563,23 +1565,23 @@ namespace UnrealBuildTool
 				// add the library paths to response
 				foreach (string LibaryPath in AdditionalLibraryPaths)
 				{
-					LinkResponseArguments += string.Format(" -L\"{0}\"", LibaryPath);
+					LinkResponseArguments += String.Format(" -L\"{0}\"", LibaryPath);
 				}
 
 				// add libraries in a library group
-				LinkResponseArguments += string.Format(" -Wl,--start-group");
+				LinkResponseArguments += String.Format(" -Wl,--start-group");
 				foreach (string AdditionalLibrary in AdditionalLibraries)
 				{
 					if (AdditionalLibrary.StartsWith("lib"))
 					{
-						LinkResponseArguments += string.Format(" \"-l{0}\"", AdditionalLibrary.Substring(3));
+						LinkResponseArguments += String.Format(" \"-l{0}\"", AdditionalLibrary.Substring(3));
 					}
 					else
 					{
-						LinkResponseArguments += string.Format(" \"{0}\"", AdditionalLibrary);
+						LinkResponseArguments += String.Format(" \"{0}\"", AdditionalLibrary);
 					}
 				}
-				LinkResponseArguments += string.Format(" -Wl,--end-group");
+				LinkResponseArguments += String.Format(" -Wl,--end-group");
 
 				// Write the MAP file to the output directory.
 				if (LinkEnvironment.bCreateMapFile)
@@ -1603,7 +1605,7 @@ namespace UnrealBuildTool
 
 			FileItem ResponseFileItem = Graph.CreateIntermediateTextFile(ResponseFileName, InputFileNames);
 
-			LinkAction.CommandArguments += string.Format(" @\"{0}\"", ResponseFileName);
+			LinkAction.CommandArguments += String.Format(" @\"{0}\"", ResponseFileName);
 			LinkAction.PrerequisiteItems.Add(ResponseFileItem);
 
 			// Fix up the paths in commandline
@@ -1671,7 +1673,7 @@ namespace UnrealBuildTool
 
 		// captures stderr from clang
 		private static string LinkerCommandline = "";
-		static public void OutputReceivedForLinker(Object Sender, DataReceivedEventArgs Line)
+		public static void OutputReceivedForLinker(Object Sender, DataReceivedEventArgs Line)
 		{
 			if ((Line != null) && (Line.Data != null) && (Line.Data.Contains("--sysroot")))
 			{
diff --git a/Engine/Source/Programs/UnrealBuildTool/Platform/Android/UEBuildAndroid.cs b/Engine/Source/Programs/UnrealBuildTool/Platform/Android/UEBuildAndroid.cs
--- a/Engine/Source/Programs/UnrealBuildTool/Platform/Android/UEBuildAndroid.cs
+++ b/Engine/Source/Programs/UnrealBuildTool/Platform/Android/UEBuildAndroid.cs
@@ -185,11 +185,11 @@ namespace UnrealBuildTool
 
 		public override void ValidateTarget(TargetRules Target)
 		{
-			if (!string.IsNullOrWhiteSpace(Environment.GetEnvironmentVariable("CLANG_STATIC_ANALYZER_MODE")))
+			if (!String.IsNullOrWhiteSpace(Environment.GetEnvironmentVariable("CLANG_STATIC_ANALYZER_MODE")))
 			{
 				Target.StaticAnalyzer = StaticAnalyzer.Default;
 				Target.StaticAnalyzerOutputType = (Environment.GetEnvironmentVariable("CLANG_ANALYZER_OUTPUT")?.Contains("html", StringComparison.OrdinalIgnoreCase) == true) ? StaticAnalyzerOutputType.Html : StaticAnalyzerOutputType.Text;
-				Target.StaticAnalyzerMode = string.Equals(Environment.GetEnvironmentVariable("CLANG_STATIC_ANALYZER_MODE"), "shallow") ? StaticAnalyzerMode.Shallow : StaticAnalyzerMode.Deep;
+				Target.StaticAnalyzerMode = String.Equals(Environment.GetEnvironmentVariable("CLANG_STATIC_ANALYZER_MODE"), "shallow") ? StaticAnalyzerMode.Shallow : StaticAnalyzerMode.Deep;
 			}
 			else if (Target.StaticAnalyzer == StaticAnalyzer.Clang)
 			{
@@ -432,7 +432,7 @@ namespace UnrealBuildTool
 			SDK.TryConvertVersionToInt(NDKToolchainVersion, out NDKVersionInt);
 
 			// PLATFORM_ANDROID_NDK_VERSION is in the form 150100, where 15 is major version, 01 is the letter (1 is 'a'), 00 indicates beta revision if letter is 00
-			CompileEnvironment.Definitions.Add(string.Format("PLATFORM_ANDROID_NDK_VERSION={0}", NDKVersionInt));
+			CompileEnvironment.Definitions.Add(String.Format("PLATFORM_ANDROID_NDK_VERSION={0}", NDKVersionInt));
 
 			Logger.LogInformation("NDK toolchain: {Version}, NDK version: {NdkVersion}, ClangVersion: {ClangVersion}", NDKToolchainVersion, NDKVersionInt, ToolChain.GetClangVersionString());
 
@@ -506,7 +506,7 @@ namespace UnrealBuildTool
 				if (Target.bPGOOptimize)
 				{
 					CompileEnvironment.PGODirectory = DirectoryReference.Combine(Target.ProjectFile.Directory, "Platforms", "Android", "Build", "PGO").FullName;
-					CompileEnvironment.PGOFilenamePrefix = string.Format("{0}-Android", Target.Name);
+					CompileEnvironment.PGOFilenamePrefix = String.Format("{0}-Android", Target.Name);
 
 					LinkEnvironment.PGODirectory = CompileEnvironment.PGODirectory;
 					LinkEnvironment.PGOFilenamePrefix = CompileEnvironment.PGOFilenamePrefix;
@@ -548,7 +548,7 @@ namespace UnrealBuildTool
 			};
 		}
 
-		static public ClangToolChainOptions CreateToolChainOptions(AndroidTargetRules TargetRules)
+		public static ClangToolChainOptions CreateToolChainOptions(AndroidTargetRules TargetRules)
 		{
 			ClangToolChainOptions Options = ClangToolChainOptions.None;
 			if (TargetRules.bEnableAddressSanitizer)
diff --git a/Engine/Source/Programs/UnrealBuildTool/Platform/Android/UEDeployAndroid.cs b/Engine/Source/Programs/UnrealBuildTool/Platform/Android/UEDeployAndroid.cs
--- a/Engine/Source/Programs/UnrealBuildTool/Platform/Android/UEDeployAndroid.cs
+++ b/Engine/Source/Programs/UnrealBuildTool/Platform/Android/UEDeployAndroid.cs
@@ -41,7 +41,7 @@ namespace UnrealBuildTool
 		private int MinimumSDKLevelForGradle = 19;
 
 		// Reserved Java keywords not allowed in package names without modification
-		static private string[] JavaReservedKeywords = new string[] {
+		private static string[] JavaReservedKeywords = new string[] {
 			"abstract", "assert", "boolean", "break", "byte", "case", "catch", "char", "class", "const", "continue", "default", "do",
 			"double", "else", "enum", "extends", "final", "finally", "float", "for", "goto", "if", "implements", "import", "instanceof",
 			"int", "interface", "long", "native", "new", "package", "private", "protected", "public", "return", "short", "static",
@@ -214,7 +214,7 @@ namespace UnrealBuildTool
 			if (ApiString.Contains("-"))
 			{
 				int Version;
-				if (int.TryParse(ApiString.Substring(ApiString.LastIndexOf('-') + 1), out Version))
+				if (Int32.TryParse(ApiString.Substring(ApiString.LastIndexOf('-') + 1), out Version))
 				{
 					VersionInt = Version;
 				}
@@ -299,10 +299,10 @@ namespace UnrealBuildTool
 			try
 			{
 				string[] Sections = VersionString.Split(".".ToCharArray());
-				Value |= (Sections.Length > 0) ? (uint.Parse(Sections[0]) << 24) : 0;
-				Value |= (Sections.Length > 1) ? (uint.Parse(Sections[1]) << 16) : 0;
-				Value |= (Sections.Length > 2) ? (uint.Parse(Sections[2]) << 8) : 0;
-				Value |= (Sections.Length > 3) ? uint.Parse(Sections[3]) : 0;
+				Value |= (Sections.Length > 0) ? (UInt32.Parse(Sections[0]) << 24) : 0;
+				Value |= (Sections.Length > 1) ? (UInt32.Parse(Sections[1]) << 16) : 0;
+				Value |= (Sections.Length > 2) ? (UInt32.Parse(Sections[2]) << 8) : 0;
+				Value |= (Sections.Length > 3) ? UInt32.Parse(Sections[3]) : 0;
 			}
 			catch (Exception)
 			{
@@ -981,7 +981,7 @@ namespace UnrealBuildTool
 					}
 
 					int Changelist = 0;
-					if (int.TryParse(EngineChangelist, out Changelist))
+					if (Int32.TryParse(EngineChangelist, out Changelist))
 					{
 						if (Changelist != 0)
 						{
@@ -1015,7 +1015,7 @@ namespace UnrealBuildTool
 
 		public string GetVersionDisplayName(bool bIsEmbedded)
 		{
-			if (string.IsNullOrEmpty(CachedVersionDisplayName))
+			if (String.IsNullOrEmpty(CachedVersionDisplayName))
 			{
 				ConfigHierarchy Ini = GetConfigCacheIni(ConfigHierarchyType.Engine);
 				string VersionDisplayName = "";
@@ -1027,14 +1027,14 @@ namespace UnrealBuildTool
 					Ini.GetBool("/Script/AndroidRuntimeSettings.AndroidRuntimeSettings", "bAppendChangeListToVersionDisplayName", out bAppendChangeListToVersionDisplayName);
 					if (bAppendChangeListToVersionDisplayName)
 					{
-						VersionDisplayName = string.Format("{0}-{1}", VersionDisplayName, EngineChangelist);
+						VersionDisplayName = String.Format("{0}-{1}", VersionDisplayName, EngineChangelist);
 					}
 
 					bool bAppendPlatformToVersionDisplayName = false;
 					Ini.GetBool("/Script/AndroidRuntimeSettings.AndroidRuntimeSettings", "bAppendPlatformToVersionDisplayName", out bAppendPlatformToVersionDisplayName);
 					if (bAppendPlatformToVersionDisplayName)
 					{
-						VersionDisplayName = string.Format("{0}-Android", VersionDisplayName);
+						VersionDisplayName = String.Format("{0}-Android", VersionDisplayName);
 					}
 
 					// append optional text to version name if embedded build
@@ -1691,7 +1691,7 @@ namespace UnrealBuildTool
 			}
 		}
 
-		static private List<FilterOperation>? ActiveStdOutFilter = null;
+		private static List<FilterOperation>? ActiveStdOutFilter = null;
 
 		static List<string> ParseCSVString(string Input)
 		{
@@ -1736,7 +1736,7 @@ namespace UnrealBuildTool
 						Results.Add(WorkString.ToString());
 						WorkString.Clear();
 					}
-					else if (!char.IsWhiteSpace(CurChar))
+					else if (!Char.IsWhiteSpace(CurChar))
 					{
 						WorkString.Append(CurChar);
 					}
@@ -1942,18 +1942,18 @@ namespace UnrealBuildTool
 		{
 			// make the settings string - this will be char by char compared against last time
 			StringBuilder CurrentSettings = new StringBuilder();
-			CurrentSettings.AppendLine(string.Format("NDKROOT={0}", Environment.GetEnvironmentVariable("NDKROOT")));
-			CurrentSettings.AppendLine(string.Format("ANDROID_HOME={0}", Environment.GetEnvironmentVariable("ANDROID_HOME")));
-			CurrentSettings.AppendLine(string.Format("JAVA_HOME={0}", Environment.GetEnvironmentVariable("JAVA_HOME")));
-			CurrentSettings.AppendLine(string.Format("NDKVersion={0}", ToolChain.GetNdkApiLevel()));
-			CurrentSettings.AppendLine(string.Format("SDKVersion={0}", GetSdkApiLevel(ToolChain)));
-			CurrentSettings.AppendLine(string.Format("bForDistribution={0}", bForDistribution));
-			CurrentSettings.AppendLine(string.Format("bMakeSeparateApks={0}", bMakeSeparateApks));
-			CurrentSettings.AppendLine(string.Format("bPackageDataInsideApk={0}", bPackageDataInsideApk));
-			CurrentSettings.AppendLine(string.Format("bDisableVerifyOBBOnStartUp={0}", bDisableVerifyOBBOnStartUp));
-			CurrentSettings.AppendLine(string.Format("bUseExternalFilesDir={0}", bUseExternalFilesDir));
-			CurrentSettings.AppendLine(string.Format("UPLHashCode={0}", UPLHashCode));
-			CurrentSettings.AppendLine(string.Format("TemplatesHashCode={0}", TemplatesHashCode));
+			CurrentSettings.AppendLine(String.Format("NDKROOT={0}", Environment.GetEnvironmentVariable("NDKROOT")));
+			CurrentSettings.AppendLine(String.Format("ANDROID_HOME={0}", Environment.GetEnvironmentVariable("ANDROID_HOME")));
+			CurrentSettings.AppendLine(String.Format("JAVA_HOME={0}", Environment.GetEnvironmentVariable("JAVA_HOME")));
+			CurrentSettings.AppendLine(String.Format("NDKVersion={0}", ToolChain.GetNdkApiLevel()));
+			CurrentSettings.AppendLine(String.Format("SDKVersion={0}", GetSdkApiLevel(ToolChain)));
+			CurrentSettings.AppendLine(String.Format("bForDistribution={0}", bForDistribution));
+			CurrentSettings.AppendLine(String.Format("bMakeSeparateApks={0}", bMakeSeparateApks));
+			CurrentSettings.AppendLine(String.Format("bPackageDataInsideApk={0}", bPackageDataInsideApk));
+			CurrentSettings.AppendLine(String.Format("bDisableVerifyOBBOnStartUp={0}", bDisableVerifyOBBOnStartUp));
+			CurrentSettings.AppendLine(String.Format("bUseExternalFilesDir={0}", bUseExternalFilesDir));
+			CurrentSettings.AppendLine(String.Format("UPLHashCode={0}", UPLHashCode));
+			CurrentSettings.AppendLine(String.Format("TemplatesHashCode={0}", TemplatesHashCode));
 
 			// all AndroidRuntimeSettings ini settings in here
 			ConfigHierarchy Ini = GetConfigCacheIni(ConfigHierarchyType.Engine);
@@ -1972,7 +1972,7 @@ namespace UnrealBuildTool
 					{
 						foreach (string Value in Values)
 						{
-							CurrentSettings.AppendLine(string.Format("{0}={1}", Key, Value));
+							CurrentSettings.AppendLine(String.Format("{0}={1}", Key, Value));
 						}
 					}
 				}
@@ -1993,7 +1993,7 @@ namespace UnrealBuildTool
 					{
 						foreach (string Value in Values)
 						{
-							CurrentSettings.AppendLine(string.Format("{0}={1}", Key, Value));
+							CurrentSettings.AppendLine(String.Format("{0}={1}", Key, Value));
 						}
 					}
 				}
@@ -2023,7 +2023,7 @@ namespace UnrealBuildTool
 					{
 						foreach (string Value in Values)
 						{
-							CurrentSettings.AppendLine(string.Format("{0}={1}", Key, Value));
+							CurrentSettings.AppendLine(String.Format("{0}={1}", Key, Value));
 						}
 					}
 				}
@@ -2068,7 +2068,7 @@ namespace UnrealBuildTool
 						{
 							foreach (string Value in Values)
 							{
-								CurrentSettings.AppendLine(string.Format("{0}:{1}={2}", SectionName, Key, Value));
+								CurrentSettings.AppendLine(String.Format("{0}:{1}={2}", SectionName, Key, Value));
 							}
 						}
 					}
@@ -2722,9 +2722,9 @@ namespace UnrealBuildTool
 					Text.AppendLine("          " + Line);
 				}
 			}
-			Text.AppendLine(string.Format("          android:installLocation=\"{0}\"", InstallLocation));
-			Text.AppendLine(string.Format("          android:versionCode=\"{0}\"", StoreVersion));
-			Text.AppendLine(string.Format("          android:versionName=\"{0}\">", VersionDisplayName));
+			Text.AppendLine(String.Format("          android:installLocation=\"{0}\"", InstallLocation));
+			Text.AppendLine(String.Format("          android:versionCode=\"{0}\"", StoreVersion));
+			Text.AppendLine(String.Format("          android:versionName=\"{0}\">", VersionDisplayName));
 
 			Text.AppendLine("");
 
@@ -2767,7 +2767,7 @@ namespace UnrealBuildTool
 				}
 			}
 			Text.AppendLine("\t             android:hardwareAccelerated=\"true\"");
-			Text.AppendLine(string.Format("\t             android:extractNativeLibs=\"{0}\"", bExtractNativeLibs ? "true" : "false"));
+			Text.AppendLine(String.Format("\t             android:extractNativeLibs=\"{0}\"", bExtractNativeLibs ? "true" : "false"));
 			Text.AppendLine("\t				android:name=\"com.epicgames.unreal.GameApplication\"");
 			if (!bIsForDistribution && SDKLevelInt >= 29 && !bRequestedLegacyExternalStorage)
 			{
@@ -2787,11 +2787,11 @@ namespace UnrealBuildTool
 				{
 					Text.AppendLine("\t\t          android:resizeableActivity=\"false\"");
 				}
-				Text.AppendLine(string.Format("\t\t          android:screenOrientation=\"{0}\"", Orientation));
-				Text.AppendLine(string.Format("\t\t          android:debuggable=\"{0}\">", bIsForDistribution ? "false" : "true"));
+				Text.AppendLine(String.Format("\t\t          android:screenOrientation=\"{0}\"", Orientation));
+				Text.AppendLine(String.Format("\t\t          android:debuggable=\"{0}\">", bIsForDistribution ? "false" : "true"));
 				Text.AppendLine("\t\t\t<intent-filter>");
 				Text.AppendLine("\t\t\t\t<action android:name=\"android.intent.action.MAIN\" />");
-				Text.AppendLine(string.Format("\t\t\t\t<category android:name=\"android.intent.category.LAUNCHER\" />"));
+				Text.AppendLine(String.Format("\t\t\t\t<category android:name=\"android.intent.category.LAUNCHER\" />"));
 				Text.AppendLine("\t\t\t</intent-filter>");
 				Text.AppendLine("\t\t</activity>");
 				Text.AppendLine("\t\t<activity android:name=\"com.epicgames.unreal.GameActivity\"");
@@ -2815,7 +2815,7 @@ namespace UnrealBuildTool
 				Text.AppendLine("\t\t          android:resizeableActivity=\"false\"");
 			}
 			Text.AppendLine("\t\t          android:launchMode=\"singleTask\"");
-			Text.AppendLine(string.Format("\t\t          android:screenOrientation=\"{0}\"", Orientation));
+			Text.AppendLine(String.Format("\t\t          android:screenOrientation=\"{0}\"", Orientation));
 			if (ExtraActivityNodeTags != null)
 			{
 				foreach (string Line in ExtraActivityNodeTags)
@@ -2823,16 +2823,16 @@ namespace UnrealBuildTool
 					Text.AppendLine("\t\t          " + Line);
 				}
 			}
-			Text.AppendLine(string.Format("\t\t          android:debuggable=\"{0}\">", bIsForDistribution ? "false" : "true"));
+			Text.AppendLine(String.Format("\t\t          android:debuggable=\"{0}\">", bIsForDistribution ? "false" : "true"));
 			Text.AppendLine("\t\t\t<meta-data android:name=\"android.app.lib_name\" android:value=\"Unreal\"/>");
 			if (!bShowLaunchImage)
 			{
 				Text.AppendLine("\t\t\t<intent-filter>");
 				Text.AppendLine("\t\t\t\t<action android:name=\"android.intent.action.MAIN\" />");
-				Text.AppendLine(string.Format("\t\t\t\t<category android:name=\"android.intent.category.LAUNCHER\" />"));
+				Text.AppendLine(String.Format("\t\t\t\t<category android:name=\"android.intent.category.LAUNCHER\" />"));
 				Text.AppendLine("\t\t\t</intent-filter>");
 			}
-			if (!string.IsNullOrEmpty(ExtraActivitySettings))
+			if (!String.IsNullOrEmpty(ExtraActivitySettings))
 			{
 				ExtraActivitySettings = ExtraActivitySettings.Replace("\\n", "\n");
 				foreach (string Line in ExtraActivitySettings.Split("\r\n".ToCharArray()))
@@ -2854,7 +2854,7 @@ namespace UnrealBuildTool
 			if (bShowLaunchImage)
 			{
 				Text.AppendLine("\t\t<activity android:name=\".DownloaderActivity\"");
-				Text.AppendLine(string.Format("\t\t          android:screenOrientation=\"{0}\"", Orientation));
+				Text.AppendLine(String.Format("\t\t          android:screenOrientation=\"{0}\"", Orientation));
 				Text.AppendLine(bAddDensity ? "\t\t          android:configChanges=\"mcc|mnc|uiMode|density|screenSize|orientation|keyboardHidden|keyboard\""
 											: "\t\t          android:configChanges=\"mcc|mnc|uiMode|screenSize|orientation|keyboardHidden|keyboard\"");
 				Text.AppendLine("\t\t          android:theme=\"@style/UnrealSplashTheme\" />");
@@ -2878,29 +2878,29 @@ namespace UnrealBuildTool
 				}
 			}
 
-			Text.AppendLine(string.Format("\t\t<meta-data android:name=\"com.epicgames.unreal.GameActivity.EngineVersion\" android:value=\"{0}\"/>", EngineVersion));
-			Text.AppendLine(string.Format("\t\t<meta-data android:name=\"com.epicgames.unreal.GameActivity.EngineBranch\" android:value=\"{0}\"/>", EngineBranch));
-			Text.AppendLine(string.Format("\t\t<meta-data android:name=\"com.epicgames.unreal.GameActivity.ProjectVersion\" android:value=\"{0}\"/>", ProjectVersion));
-			Text.AppendLine(string.Format("\t\t<meta-data android:name=\"com.epicgames.unreal.GameActivity.DepthBufferPreference\" android:value=\"{0}\"/>", ConvertDepthBufferIniValue(DepthBufferPreference)));
-			Text.AppendLine(string.Format("\t\t<meta-data android:name=\"com.epicgames.unreal.GameActivity.bPackageDataInsideApk\" android:value=\"{0}\"/>", bPackageDataInsideApk ? "true" : "false"));
-			Text.AppendLine(string.Format("\t\t<meta-data android:name=\"com.epicgames.unreal.GameActivity.bVerifyOBBOnStartUp\" android:value=\"{0}\"/>", (bIsForDistribution && !bDisableVerifyOBBOnStartUp) ? "true" : "false"));
-			Text.AppendLine(string.Format("\t\t<meta-data android:name=\"com.epicgames.unreal.GameActivity.bShouldHideUI\" android:value=\"{0}\"/>", EnableFullScreen ? "true" : "false"));
-			Text.AppendLine(string.Format("\t\t<meta-data android:name=\"com.epicgames.unreal.GameActivity.ProjectName\" android:value=\"{0}\"/>", ProjectName));
-			Text.AppendLine(string.Format("\t\t<meta-data android:name=\"com.epicgames.unreal.GameActivity.AppType\" android:value=\"{0}\"/>", AppType));
-			Text.AppendLine(string.Format("\t\t<meta-data android:name=\"com.epicgames.unreal.GameActivity.bHasOBBFiles\" android:value=\"{0}\"/>", bHasOBBFiles ? "true" : "false"));
-			Text.AppendLine(string.Format("\t\t<meta-data android:name=\"com.epicgames.unreal.GameActivity.BuildConfiguration\" android:value=\"{0}\"/>", Configuration));
-			Text.AppendLine(string.Format("\t\t<meta-data android:name=\"com.epicgames.unreal.GameActivity.CookedFlavors\" android:value=\"{0}\"/>", CookedFlavors));
-			Text.AppendLine(string.Format("\t\t<meta-data android:name=\"com.epicgames.unreal.GameActivity.bValidateTextureFormats\" android:value=\"{0}\"/>", bValidateTextureFormats ? "true" : "false"));
-			Text.AppendLine(string.Format("\t\t<meta-data android:name=\"com.epicgames.unreal.GameActivity.bUseExternalFilesDir\" android:value=\"{0}\"/>", bUseExternalFilesDir ? "true" : "false"));
-			Text.AppendLine(string.Format("\t\t<meta-data android:name=\"com.epicgames.unreal.GameActivity.bPublicLogFiles\" android:value=\"{0}\"/>", bPublicLogFiles ? "true" : "false"));
-			Text.AppendLine(string.Format("\t\t<meta-data android:name=\"com.epicgames.unreal.GameActivity.bUseDisplayCutout\" android:value=\"{0}\"/>", bUseDisplayCutout ? "true" : "false"));
-			Text.AppendLine(string.Format("\t\t<meta-data android:name=\"com.epicgames.unreal.GameActivity.bAllowIMU\" android:value=\"{0}\"/>", bAllowIMU ? "true" : "false"));
-			Text.AppendLine(string.Format("\t\t<meta-data android:name=\"com.epicgames.unreal.GameActivity.bSupportsVulkan\" android:value=\"{0}\"/>", bSupportsVulkan ? "true" : "false"));
-			Text.AppendLine(string.Format("\t\t<meta-data android:name=\"com.epicgames.unreal.GameActivity.PropagateAlpha\" android:value=\"{0}\"/>", PropagateAlpha));
-			Text.AppendLine(string.Format("\t\t<meta-data android:name=\"com.epicgames.unreal.GameActivity.StartupPermissions\" android:value=\"{0}\"/>", StartupPermissions));
+			Text.AppendLine(String.Format("\t\t<meta-data android:name=\"com.epicgames.unreal.GameActivity.EngineVersion\" android:value=\"{0}\"/>", EngineVersion));
+			Text.AppendLine(String.Format("\t\t<meta-data android:name=\"com.epicgames.unreal.GameActivity.EngineBranch\" android:value=\"{0}\"/>", EngineBranch));
+			Text.AppendLine(String.Format("\t\t<meta-data android:name=\"com.epicgames.unreal.GameActivity.ProjectVersion\" android:value=\"{0}\"/>", ProjectVersion));
+			Text.AppendLine(String.Format("\t\t<meta-data android:name=\"com.epicgames.unreal.GameActivity.DepthBufferPreference\" android:value=\"{0}\"/>", ConvertDepthBufferIniValue(DepthBufferPreference)));
+			Text.AppendLine(String.Format("\t\t<meta-data android:name=\"com.epicgames.unreal.GameActivity.bPackageDataInsideApk\" android:value=\"{0}\"/>", bPackageDataInsideApk ? "true" : "false"));
+			Text.AppendLine(String.Format("\t\t<meta-data android:name=\"com.epicgames.unreal.GameActivity.bVerifyOBBOnStartUp\" android:value=\"{0}\"/>", (bIsForDistribution && !bDisableVerifyOBBOnStartUp) ? "true" : "false"));
+			Text.AppendLine(String.Format("\t\t<meta-data android:name=\"com.epicgames.unreal.GameActivity.bShouldHideUI\" android:value=\"{0}\"/>", EnableFullScreen ? "true" : "false"));
+			Text.AppendLine(String.Format("\t\t<meta-data android:name=\"com.epicgames.unreal.GameActivity.ProjectName\" android:value=\"{0}\"/>", ProjectName));
+			Text.AppendLine(String.Format("\t\t<meta-data android:name=\"com.epicgames.unreal.GameActivity.AppType\" android:value=\"{0}\"/>", AppType));
+			Text.AppendLine(String.Format("\t\t<meta-data android:name=\"com.epicgames.unreal.GameActivity.bHasOBBFiles\" android:value=\"{0}\"/>", bHasOBBFiles ? "true" : "false"));
+			Text.AppendLine(String.Format("\t\t<meta-data android:name=\"com.epicgames.unreal.GameActivity.BuildConfiguration\" android:value=\"{0}\"/>", Configuration));
+			Text.AppendLine(String.Format("\t\t<meta-data android:name=\"com.epicgames.unreal.GameActivity.CookedFlavors\" android:value=\"{0}\"/>", CookedFlavors));
+			Text.AppendLine(String.Format("\t\t<meta-data android:name=\"com.epicgames.unreal.GameActivity.bValidateTextureFormats\" android:value=\"{0}\"/>", bValidateTextureFormats ? "true" : "false"));
+			Text.AppendLine(String.Format("\t\t<meta-data android:name=\"com.epicgames.unreal.GameActivity.bUseExternalFilesDir\" android:value=\"{0}\"/>", bUseExternalFilesDir ? "true" : "false"));
+			Text.AppendLine(String.Format("\t\t<meta-data android:name=\"com.epicgames.unreal.GameActivity.bPublicLogFiles\" android:value=\"{0}\"/>", bPublicLogFiles ? "true" : "false"));
+			Text.AppendLine(String.Format("\t\t<meta-data android:name=\"com.epicgames.unreal.GameActivity.bUseDisplayCutout\" android:value=\"{0}\"/>", bUseDisplayCutout ? "true" : "false"));
+			Text.AppendLine(String.Format("\t\t<meta-data android:name=\"com.epicgames.unreal.GameActivity.bAllowIMU\" android:value=\"{0}\"/>", bAllowIMU ? "true" : "false"));
+			Text.AppendLine(String.Format("\t\t<meta-data android:name=\"com.epicgames.unreal.GameActivity.bSupportsVulkan\" android:value=\"{0}\"/>", bSupportsVulkan ? "true" : "false"));
+			Text.AppendLine(String.Format("\t\t<meta-data android:name=\"com.epicgames.unreal.GameActivity.PropagateAlpha\" android:value=\"{0}\"/>", PropagateAlpha));
+			Text.AppendLine(String.Format("\t\t<meta-data android:name=\"com.epicgames.unreal.GameActivity.StartupPermissions\" android:value=\"{0}\"/>", StartupPermissions));
 			if (bPackageForDaydream)
 			{
-				Text.AppendLine(string.Format("\t\t<meta-data android:name=\"com.epicgames.unreal.GameActivity.bDaydream\" android:value=\"true\"/>"));
+				Text.AppendLine(String.Format("\t\t<meta-data android:name=\"com.epicgames.unreal.GameActivity.bDaydream\" android:value=\"true\"/>"));
 			}
 			Text.AppendLine("\t\t<meta-data android:name=\"com.google.android.gms.games.APP_ID\"");
 			Text.AppendLine("\t\t           android:value=\"@string/app_id\" />");
@@ -2911,7 +2911,7 @@ namespace UnrealBuildTool
 				Text.AppendLine("\t\t<activity android:name=\"com.google.android.gms.ads.AdActivity\"");
 				Text.AppendLine("\t\t          android:configChanges=\"keyboard|keyboardHidden|orientation|screenLayout|uiMode|screenSize|smallestScreenSize\"/>");
 			}
-			if (!string.IsNullOrEmpty(ExtraApplicationSettings))
+			if (!String.IsNullOrEmpty(ExtraApplicationSettings))
 			{
 				ExtraApplicationSettings = ExtraApplicationSettings.Replace("\\n", "\n");
 				foreach (string Line in ExtraApplicationSettings.Split("\r\n".ToCharArray()))
@@ -2972,7 +2972,7 @@ namespace UnrealBuildTool
 
 			// Max supported aspect ratio
 			string MaxAspectRatioString = MaxAspectRatioValue.ToString("f", System.Globalization.CultureInfo.InvariantCulture);
-			Text.AppendLine(string.Format("\t\t<meta-data android:name=\"android.max_aspect\" android:value=\"{0}\" />", MaxAspectRatioString));
+			Text.AppendLine(String.Format("\t\t<meta-data android:name=\"android.max_aspect\" android:value=\"{0}\" />", MaxAspectRatioString));
 
 			Text.AppendLine("\t</application>");
 
@@ -3038,7 +3038,7 @@ namespace UnrealBuildTool
 						string TrimmedPermission = Permission.Trim(' ');
 						if (TrimmedPermission != "")
 						{
-							string PermissionString = string.Format("\t<uses-permission android:name=\"{0}\"/>", TrimmedPermission);
+							string PermissionString = String.Format("\t<uses-permission android:name=\"{0}\"/>", TrimmedPermission);
 							if (!Text.ToString().Contains(PermissionString))
 							{
 								Text.AppendLine(PermissionString);
@@ -3163,12 +3163,16 @@ namespace UnrealBuildTool
 
 					int StartIndex = Line.IndexOf("\"app_id\">");
 					if (StartIndex < 0)
+					{
 						continue;
+					}
 
 					StartIndex += 9;
 					int EndIndex = Line.IndexOf("</string>");
 					if (EndIndex < 0)
+					{
 						continue;
+					}
 
 					string XmlAppId = Line.Substring(StartIndex, EndIndex - StartIndex);
 
@@ -3443,7 +3447,7 @@ namespace UnrealBuildTool
 			if (!String.IsNullOrEmpty(TargetOverride))
 			{
 				int OverrideInt = 0;
-				if (int.TryParse(TargetOverride, out OverrideInt))
+				if (Int32.TryParse(TargetOverride, out OverrideInt))
 				{
 					TargetSDKVersion = OverrideInt;
 				}
@@ -3801,26 +3805,26 @@ namespace UnrealBuildTool
 			GradleProperties.AppendLine("android.injected.testOnly=false");
 			GradleProperties.AppendLine("android.useAndroidX=true");
 			GradleProperties.AppendLine("android.enableJetifier=true");
-			GradleProperties.AppendLine(string.Format("COMPILE_SDK_VERSION={0}", CompileSDKVersion));
-			GradleProperties.AppendLine(string.Format("BUILD_TOOLS_VERSION={0}", BuildToolsVersion));
-			GradleProperties.AppendLine(string.Format("PACKAGE_NAME={0}", PackageName));
-			GradleProperties.AppendLine(string.Format("MIN_SDK_VERSION={0}", MinSDKVersion.ToString()));
-			GradleProperties.AppendLine(string.Format("TARGET_SDK_VERSION={0}", TargetSDKVersion.ToString()));
-			GradleProperties.AppendLine(string.Format("STORE_VERSION={0}", StoreVersion.ToString()));
-			GradleProperties.AppendLine(string.Format("VERSION_DISPLAY_NAME={0}", VersionDisplayName));
+			GradleProperties.AppendLine(String.Format("COMPILE_SDK_VERSION={0}", CompileSDKVersion));
+			GradleProperties.AppendLine(String.Format("BUILD_TOOLS_VERSION={0}", BuildToolsVersion));
+			GradleProperties.AppendLine(String.Format("PACKAGE_NAME={0}", PackageName));
+			GradleProperties.AppendLine(String.Format("MIN_SDK_VERSION={0}", MinSDKVersion.ToString()));
+			GradleProperties.AppendLine(String.Format("TARGET_SDK_VERSION={0}", TargetSDKVersion.ToString()));
+			GradleProperties.AppendLine(String.Format("STORE_VERSION={0}", StoreVersion.ToString()));
+			GradleProperties.AppendLine(String.Format("VERSION_DISPLAY_NAME={0}", VersionDisplayName));
 
 			string NDKPath = Environment.GetEnvironmentVariable("NDKROOT")!.Replace("\\", "/");
 			int NDKVersionIndex = NDKPath.LastIndexOf("/");
 			string NDKVersion = NDKVersionIndex > 0 ? NDKPath.Substring(NDKVersionIndex + 1) : DEFAULT_NDK_VERSION;
-			GradleProperties.AppendLine(string.Format("NDK_VERSION={0}", NDKVersion));
+			GradleProperties.AppendLine(String.Format("NDK_VERSION={0}", NDKVersion));
 
 			if (DestApkName != null)
 			{
-				GradleProperties.AppendLine(string.Format("OUTPUT_PATH={0}", Path.GetDirectoryName(DestApkName)!.Replace("\\", "/")));
-				GradleProperties.AppendLine(string.Format("OUTPUT_FILENAME={0}", Path.GetFileName(DestApkName)));
+				GradleProperties.AppendLine(String.Format("OUTPUT_PATH={0}", Path.GetDirectoryName(DestApkName)!.Replace("\\", "/")));
+				GradleProperties.AppendLine(String.Format("OUTPUT_FILENAME={0}", Path.GetFileName(DestApkName)));
 
 				string BundleFilename = Path.GetFileName(DestApkName).Replace(".apk", ".aab");
-				GradleProperties.AppendLine(string.Format("OUTPUT_BUNDLEFILENAME={0}", BundleFilename));
+				GradleProperties.AppendLine(String.Format("OUTPUT_BUNDLEFILENAME={0}", BundleFilename));
 
 				if (bEnableUniversalAPK)
 				{
@@ -3830,10 +3834,10 @@ namespace UnrealBuildTool
 			}
 
 			int OBBFileIndex = 0;
-			GradleProperties.AppendLine(string.Format("OBB_FILECOUNT={0}", OBBFiles.Count));
+			GradleProperties.AppendLine(String.Format("OBB_FILECOUNT={0}", OBBFiles.Count));
 			foreach (string OBBFile in OBBFiles)
 			{
-				GradleProperties.AppendLine(string.Format("OBB_FILE{0}={1}", OBBFileIndex++, OBBFile.Replace("\\", "/")));
+				GradleProperties.AppendLine(String.Format("OBB_FILE{0}={1}", OBBFileIndex++, OBBFile.Replace("\\", "/")));
 			}
 
 			GradleProperties.AppendLine("ANDROID_TOOLS_BUILD_GRADLE_VERSION={0}", ANDROID_TOOLS_BUILD_GRADLE_VERSION);
@@ -3849,7 +3853,7 @@ namespace UnrealBuildTool
 			ABIGradle.AppendLine("android {");
 			ABIGradle.AppendLine("\tdefaultConfig {");
 			ABIGradle.AppendLine("\t\tndk {");
-			ABIGradle.AppendLine(string.Format("\t\t\tabiFilter \"{0}\"", NDKArch));
+			ABIGradle.AppendLine(String.Format("\t\t\tabiFilter \"{0}\"", NDKArch));
 			ABIGradle.AppendLine("\t\t}");
 			ABIGradle.AppendLine("\t}");
 			ABIGradle.AppendLine("}");
@@ -3910,12 +3914,12 @@ namespace UnrealBuildTool
 				Ini.GetString("/Script/AndroidRuntimeSettings.AndroidRuntimeSettings", "KeyStorePassword", out KeyStorePassword);
 				Ini.GetString("/Script/AndroidRuntimeSettings.AndroidRuntimeSettings", "KeyPassword", out KeyPassword);
 
-				if (string.IsNullOrEmpty(KeyStore) || string.IsNullOrEmpty(KeyAlias) || string.IsNullOrEmpty(KeyStorePassword))
+				if (String.IsNullOrEmpty(KeyStore) || String.IsNullOrEmpty(KeyAlias) || String.IsNullOrEmpty(KeyStorePassword))
 				{
 					throw new BuildException("DistributionSigning settings are not all set. Check the DistributionSettings section in the Android tab of Project Settings");
 				}
 
-				if (string.IsNullOrEmpty(KeyPassword) || KeyPassword == "_sameaskeystore_")
+				if (String.IsNullOrEmpty(KeyPassword) || KeyPassword == "_sameaskeystore_")
 				{
 					KeyPassword = KeyStorePassword;
 				}
@@ -3927,17 +3931,17 @@ namespace UnrealBuildTool
 					throw new BuildException("Keystore file is missing. Check the DistributionSettings section in the Android tab of Project Settings");
 				}
 
-				GradleProperties.AppendLine(string.Format("STORE_FILE={0}", KeyStoreFilename.Replace("\\", "/")));
-				GradleProperties.AppendLine(string.Format("STORE_PASSWORD={0}", KeyStorePassword));
-				GradleProperties.AppendLine(string.Format("KEY_ALIAS={0}", KeyAlias));
-				GradleProperties.AppendLine(string.Format("KEY_PASSWORD={0}", KeyPassword));
+				GradleProperties.AppendLine(String.Format("STORE_FILE={0}", KeyStoreFilename.Replace("\\", "/")));
+				GradleProperties.AppendLine(String.Format("STORE_PASSWORD={0}", KeyStorePassword));
+				GradleProperties.AppendLine(String.Format("KEY_ALIAS={0}", KeyAlias));
+				GradleProperties.AppendLine(String.Format("KEY_PASSWORD={0}", KeyPassword));
 
 				GradleBuildAdditionsContent.AppendLine("\tsigningConfigs {");
 				GradleBuildAdditionsContent.AppendLine("\t\trelease {");
-				GradleBuildAdditionsContent.AppendLine(string.Format("\t\t\tstoreFile file('{0}')", KeyStoreFilename.Replace("\\", "/")));
-				GradleBuildAdditionsContent.AppendLine(string.Format("\t\t\tstorePassword '{0}'", KeyStorePassword));
-				GradleBuildAdditionsContent.AppendLine(string.Format("\t\t\tkeyAlias '{0}'", KeyAlias));
-				GradleBuildAdditionsContent.AppendLine(string.Format("\t\t\tkeyPassword '{0}'", KeyPassword));
+				GradleBuildAdditionsContent.AppendLine(String.Format("\t\t\tstoreFile file('{0}')", KeyStoreFilename.Replace("\\", "/")));
+				GradleBuildAdditionsContent.AppendLine(String.Format("\t\t\tstorePassword '{0}'", KeyStorePassword));
+				GradleBuildAdditionsContent.AppendLine(String.Format("\t\t\tkeyAlias '{0}'", KeyAlias));
+				GradleBuildAdditionsContent.AppendLine(String.Format("\t\t\tkeyPassword '{0}'", KeyPassword));
 				if (bDisableV2Signing)
 				{
 					GradleBuildAdditionsContent.AppendLine("\t\t\tv2SigningEnabled false");
@@ -4270,7 +4274,9 @@ namespace UnrealBuildTool
 					{
 						// Only delete the files in the cleanup set
 						if (!cleanFiles.Contains(Path.GetFileName(filename)))
+						{
 							continue;
+						}
 
 						Logger.LogInformation("Cleaning up file {File}", filename);
 						SafeDeleteFile(filename, false);
@@ -4498,12 +4504,12 @@ namespace UnrealBuildTool
 						Ini.GetString("/Script/AndroidRuntimeSettings.AndroidRuntimeSettings", "KeyStorePassword", out KeyStorePassword);
 						Ini.GetString("/Script/AndroidRuntimeSettings.AndroidRuntimeSettings", "KeyPassword", out KeyPassword);
 
-						if (string.IsNullOrEmpty(KeyStore) || string.IsNullOrEmpty(KeyAlias) || string.IsNullOrEmpty(KeyStorePassword))
+						if (String.IsNullOrEmpty(KeyStore) || String.IsNullOrEmpty(KeyAlias) || String.IsNullOrEmpty(KeyStorePassword))
 						{
 							throw new BuildException("DistributionSigning settings are not all set. Check the DistributionSettings section in the Android tab of Project Settings");
 						}
 
-						if (string.IsNullOrEmpty(KeyPassword) || KeyPassword == "_sameaskeystore_")
+						if (String.IsNullOrEmpty(KeyPassword) || KeyPassword == "_sameaskeystore_")
 						{
 							KeyPassword = KeyStorePassword;
 						}
@@ -4810,7 +4816,7 @@ namespace UnrealBuildTool
 				if (CompileSDKMin != "")
 				{
 					int CompileSDKVersionInt;
-					if (!int.TryParse(CompileSDKVersion, out CompileSDKVersionInt))
+					if (!Int32.TryParse(CompileSDKVersion, out CompileSDKVersionInt))
 					{
 						CompileSDKVersionInt = 23;
 					}
@@ -4821,7 +4827,7 @@ namespace UnrealBuildTool
 					{
 						//string VersionString = CompileLine.Replace("android-", "");
 						int VersionInt;
-						if (int.TryParse(CompileLine, out VersionInt))
+						if (Int32.TryParse(CompileLine, out VersionInt))
 						{
 							if (VersionInt > CompileSDKVersionInt)
 							{
@@ -4884,7 +4890,7 @@ namespace UnrealBuildTool
 				String LocalPropertiesFilename = Path.Combine(UnrealBuildGradlePath, "local.properties");
 				StringBuilder LocalProperties = new StringBuilder();
 				//				LocalProperties.AppendLine(string.Format("ndk.dir={0}", Environment.GetEnvironmentVariable("NDKROOT")!.Replace("\\", "/")));
-				LocalProperties.AppendLine(string.Format("sdk.dir={0}", Environment.GetEnvironmentVariable("ANDROID_HOME")!.Replace("\\", "/")));
+				LocalProperties.AppendLine(String.Format("sdk.dir={0}", Environment.GetEnvironmentVariable("ANDROID_HOME")!.Replace("\\", "/")));
 				File.WriteAllText(LocalPropertiesFilename, LocalProperties.ToString());
 
 				CreateGradlePropertiesFiles(ToolChain, Arch, MinSDKVersion, TargetSDKVersion, CompileSDKVersion, BuildToolsVersion, PackageName, DestApkName, NDKArch,
@@ -4896,7 +4902,7 @@ namespace UnrealBuildTool
 					if (!RuntimePlatform.IsWindows)
 					{
 						// fix permissions for Mac/Linux
-						RunCommandLineProgramWithException(UnrealBuildGradlePath, "/bin/sh", string.Format("-c 'chmod 0755 \"{0}\"'", GradleScriptPath.Replace("'", "'\"'\"'")), Logger, "Fix gradlew permissions");
+						RunCommandLineProgramWithException(UnrealBuildGradlePath, "/bin/sh", String.Format("-c 'chmod 0755 \"{0}\"'", GradleScriptPath.Replace("'", "'\"'\"'")), Logger, "Fix gradlew permissions");
 					}
 					else
 					{
@@ -5231,7 +5237,7 @@ namespace UnrealBuildTool
 					ABIGradle.AppendLine("android {");
 					ABIGradle.AppendLine("\tdefaultConfig {");
 					ABIGradle.AppendLine("\t\tndk {");
-					ABIGradle.AppendLine(string.Format("\t\t\tabiFilters{0}", ABIFilter.Substring(1)));
+					ABIGradle.AppendLine(String.Format("\t\t\tabiFilters{0}", ABIFilter.Substring(1)));
 					ABIGradle.AppendLine("\t\t}");
 					ABIGradle.AppendLine("\t}");
 					ABIGradle.AppendLine("}");
@@ -5263,11 +5269,11 @@ namespace UnrealBuildTool
 					// update gradle.properties to set STORE_VERSION properly, and OUTPUT_BUNDLEFILENAME
 					string GradlePropertiesFilename = Path.Combine(UnrealBuildGradlePath, "gradle.properties");
 					string GradlePropertiesContent = File.ReadAllText(GradlePropertiesFilename);
-					GradlePropertiesContent += string.Format("\nSTORE_VERSION={0}\nOUTPUT_BUNDLEFILENAME={1}\n", BaseStoreVersion,
+					GradlePropertiesContent += String.Format("\nSTORE_VERSION={0}\nOUTPUT_BUNDLEFILENAME={1}\n", BaseStoreVersion,
 						Path.GetFileNameWithoutExtension(OutputPath).Replace("UnrealGame", ProjectName) + ".aab");
 					if (bEnableUniversalAPK)
 					{
-						GradlePropertiesContent += string.Format("OUTPUT_UNIVERSALFILENAME={0}\n",
+						GradlePropertiesContent += String.Format("OUTPUT_UNIVERSALFILENAME={0}\n",
 							Path.GetFileNameWithoutExtension(OutputPath).Replace("UnrealGame", ProjectName) + "_universal.apk");
 					}
 					File.WriteAllText(GradlePropertiesFilename, GradlePropertiesContent);
@@ -5276,7 +5282,7 @@ namespace UnrealBuildTool
 					if (!RuntimePlatform.IsWindows)
 					{
 						// fix permissions for Mac/Linux
-						RunCommandLineProgramWithException(UnrealBuildGradlePath, "/bin/sh", string.Format("-c 'chmod 0755 \"{0}\"'", GradleScriptPath.Replace("'", "'\"'\"'")), Logger, "Fix gradlew permissions");
+						RunCommandLineProgramWithException(UnrealBuildGradlePath, "/bin/sh", String.Format("-c 'chmod 0755 \"{0}\"'", GradleScriptPath.Replace("'", "'\"'\"'")), Logger, "Fix gradlew permissions");
 					}
 					else
 					{
@@ -5338,7 +5344,7 @@ namespace UnrealBuildTool
 						if (!RuntimePlatform.IsWindows)
 						{
 							// fix permissions for Mac/Linux
-							RunCommandLineProgramWithException(UnrealBuildGradlePath, "/bin/sh", string.Format("-c 'chmod 0755 \"{0}\"'", GradleScriptPath.Replace("'", "'\"'\"'")), Logger, "Fix gradlew permissions");
+							RunCommandLineProgramWithException(UnrealBuildGradlePath, "/bin/sh", String.Format("-c 'chmod 0755 \"{0}\"'", GradleScriptPath.Replace("'", "'\"'\"'")), Logger, "Fix gradlew permissions");
 						}
 						else
 						{
@@ -5465,7 +5471,7 @@ namespace UnrealBuildTool
 			PackageInfoSource[0] = PackageName;
 			PackageInfoSource[1] = GetStoreVersion(null).ToString();
 			PackageInfoSource[2] = GetVersionDisplayName(bIsEmbedded);
-			PackageInfoSource[3] = string.Format("name='com.epicgames.unreal.GameActivity.AppType' value='{0}'", InTargetType == TargetType.Game ? "" : InTargetType.ToString());
+			PackageInfoSource[3] = String.Format("name='com.epicgames.unreal.GameActivity.AppType' value='{0}'", InTargetType == TargetType.Game ? "" : InTargetType.ToString());
 
 			Logger.LogInformation("Writing packageInfo pkgName:{PkgName} storeVersion:{StoreVer} versionDisplayName:{Version} to {DestFile}", PackageInfoSource[0], PackageInfoSource[1], PackageInfoSource[2], DestPackageNameFileName);
 
@@ -5694,7 +5700,7 @@ namespace UnrealBuildTool
 			if (TemplateChanged)
 			{
 				// deal with insertions of newlines
-				TemplateSrc = string.Join("\n", TemplateSrc).Split(new[] { "\r\n", "\r", "\n" }, StringSplitOptions.None);
+				TemplateSrc = String.Join("\n", TemplateSrc).Split(new[] { "\r\n", "\r", "\n" }, StringSplitOptions.None);
 			}
 
 			if (TemplateDest == null || TemplateSrc.Length != TemplateDest.Length || !TemplateSrc.SequenceEqual(TemplateDest))
@@ -5745,7 +5751,7 @@ namespace UnrealBuildTool
 			if (TemplateChanged)
 			{
 				// deal with insertions of newlines
-				TemplateSrc = string.Join("\n", TemplateSrc).Split(new[] { "\r\n", "\r", "\n" }, StringSplitOptions.None);
+				TemplateSrc = String.Join("\n", TemplateSrc).Split(new[] { "\r\n", "\r", "\n" }, StringSplitOptions.None);
 			}
 
 			if (TemplateDest == null || TemplateSrc.Length != TemplateDest.Length || !TemplateSrc.SequenceEqual(TemplateDest))
@@ -5869,8 +5875,8 @@ namespace UnrealBuildTool
 			{
 				string RelativePath = Path.GetFileName(LibDir);
 
-				SettingsGradleContent.AppendLine(string.Format("include ':{0}'", RelativePath));
-				ProjectDependencyContent.AppendLine(string.Format("\timplementation project(':{0}')", RelativePath));
+				SettingsGradleContent.AppendLine(String.Format("include ':{0}'", RelativePath));
+				ProjectDependencyContent.AppendLine(String.Format("\timplementation project(':{0}')", RelativePath));
 
 				string GradleProjectPath = Path.Combine(UnrealBuildGradlePath, RelativePath);
 				string GradleProjectMainPath = Path.Combine(GradleProjectPath, "src", "main");
@@ -5971,19 +5977,19 @@ namespace UnrealBuildTool
 
 				if (VersionCode != "")
 				{
-					BuildGradleContent.AppendLine(string.Format("\t\tversionCode {0}", VersionCode));
+					BuildGradleContent.AppendLine(String.Format("\t\tversionCode {0}", VersionCode));
 				}
 				if (VersionName != "")
 				{
-					BuildGradleContent.AppendLine(string.Format("\t\tversionName \"{0}\"", VersionName));
+					BuildGradleContent.AppendLine(String.Format("\t\tversionName \"{0}\"", VersionName));
 				}
 				if (MinSdkVersion != "")
 				{
-					BuildGradleContent.AppendLine(string.Format("\t\tminSdkVersion = {0}", MinSdkVersion));
+					BuildGradleContent.AppendLine(String.Format("\t\tminSdkVersion = {0}", MinSdkVersion));
 				}
 				if (TargetSdkVersion != "")
 				{
-					BuildGradleContent.AppendLine(string.Format("\t\ttargetSdkVersion = {0}", TargetSdkVersion));
+					BuildGradleContent.AppendLine(String.Format("\t\ttargetSdkVersion = {0}", TargetSdkVersion));
 				}
 				BuildGradleContent.AppendLine("\t}");
 				BuildGradleContent.AppendLine("}");
@@ -6041,7 +6047,7 @@ namespace UnrealBuildTool
 			AARImportsContent.AppendLine("dependencies {");
 			foreach (AndroidAARHandler.AndroidAAREntry Dependency in AARHandler.AARList!)
 			{
-				AARImportsContent.AppendLine(string.Format("\timplementation '{0}:{1}:{2}'", Dependency.Filename, Dependency.BaseName, Dependency.Version));
+				AARImportsContent.AppendLine(String.Format("\timplementation '{0}:{1}:{2}'", Dependency.Filename, Dependency.BaseName, Dependency.Version));
 			}
 			AARImportsContent.AppendLine("}");
 
diff --git a/Engine/Source/Programs/UnrealBuildTool/Platform/IOS/IOSCodeSigning.cs b/Engine/Source/Programs/UnrealBuildTool/Platform/IOS/IOSCodeSigning.cs
--- a/Engine/Source/Programs/UnrealBuildTool/Platform/IOS/IOSCodeSigning.cs
+++ b/Engine/Source/Programs/UnrealBuildTool/Platform/IOS/IOSCodeSigning.cs
@@ -104,12 +104,12 @@ namespace UnrealBuildTool
 		}
 
 
-		static private string CertToolData = "";
+		private static string CertToolData = "";
 
 		/// <summary>
 		/// When receiving the cert tool process call
 		/// </summary>
-		static public void OutputReceivedCertToolProcessCall(Object Sender, DataReceivedEventArgs Line)
+		public static void OutputReceivedCertToolProcessCall(Object Sender, DataReceivedEventArgs Line)
 		{
 			if ((Line != null) && !String.IsNullOrEmpty(Line.Data))
 			{
@@ -135,7 +135,7 @@ namespace UnrealBuildTool
 					Process CertTool = new Process();
 					CertTool.StartInfo.FileName = "/usr/bin/security";
 					CertTool.StartInfo.UseShellExecute = false;
-					CertTool.StartInfo.Arguments = string.Format("find-certificate -a -c \"{0}\" -p", SearchPrefix);
+					CertTool.StartInfo.Arguments = String.Format("find-certificate -a -c \"{0}\" -p", SearchPrefix);
 					CertTool.StartInfo.RedirectStandardOutput = true;
 					CertTool.StartInfo.StandardOutputEncoding = Encoding.UTF8;
 					CertTool.OutputDataReceived += new DataReceivedEventHandler(OutputReceivedCertToolProcessCall);
@@ -182,7 +182,7 @@ namespace UnrealBuildTool
 				DateTime Now = DateTime.UtcNow;
 
 				bool bCertTimeIsValid = (EffectiveDate < Now) && (ExpirationDate > Now);
-				string CertLine = string.Format("CERTIFICATE-Name:{0},Validity:{1},StartDate:{2},EndDate:{3}", GetFriendlyNameFromCert(TestCert), bCertTimeIsValid ? "VALID" : "EXPIRED", EffectiveDate.ToString("o"), ExpirationDate.ToString("o"));
+				string CertLine = String.Format("CERTIFICATE-Name:{0},Validity:{1},StartDate:{2},EndDate:{3}", GetFriendlyNameFromCert(TestCert), bCertTimeIsValid ? "VALID" : "EXPIRED", EffectiveDate.ToString("o"), ExpirationDate.ToString("o"));
 				_logger?.LogInformation(CertLine);
 				Certs.Add(CertLine);
 			}
@@ -283,7 +283,7 @@ namespace UnrealBuildTool
 				{
 					Validity = "MANAGED";
 				}
-				if ((string.IsNullOrWhiteSpace(SelectedProvision) || FoundName < 2) && Validity == "VALID" && !bDistribution)
+				if ((String.IsNullOrWhiteSpace(SelectedProvision) || FoundName < 2) && Validity == "VALID" && !bDistribution)
 				{
 					int Prev = FoundName;
 					if (bPassesNameCheck)
@@ -442,11 +442,11 @@ namespace UnrealBuildTool
 
 			string? CFBundleIdentifier = CodeSigningConfig.OverrideBundleName;
 
-			if (string.IsNullOrEmpty(CFBundleIdentifier))
+			if (String.IsNullOrEmpty(CFBundleIdentifier))
 			{
 				// Load Info.plist, which guides nearly everything else
 				string plistFile = CodeSigningConfig.EngineBuildDirectory + "/UnrealGame-Info.plist";
-				if (!string.IsNullOrEmpty(CodeSigningConfig.ProjectFile))
+				if (!String.IsNullOrEmpty(CodeSigningConfig.ProjectFile))
 				{
 					plistFile = Path.GetDirectoryName(CodeSigningConfig.ProjectFile) + "/Intermediate/" + CodeSigningConfig.OSString + "/" + Path.GetFileNameWithoutExtension(CodeSigningConfig.ProjectFile) + "-Info.plist";
 
@@ -541,7 +541,7 @@ namespace UnrealBuildTool
 			MobileProvision? MobileProvision;
 			MobileProvision.CacheMobileProvisions();
 
-			if (!string.IsNullOrEmpty(BundleID))
+			if (!String.IsNullOrEmpty(BundleID))
 			{
 				CodeSigningConfig.OverrideBundleName = BundleID;
 			}
@@ -673,7 +673,7 @@ namespace UnrealBuildTool
 				if (Value is byte[])
 				{
 					ValueElement = Doc.CreateElement("data");
-					ValueElement.InnerText = (Value is byte[] byteArray) ? Convert.ToBase64String(byteArray) : string.Empty;
+					ValueElement.InnerText = (Value is byte[] byteArray) ? Convert.ToBase64String(byteArray) : String.Empty;
 				}
 				else
 				{
@@ -1104,7 +1104,7 @@ namespace UnrealBuildTool
 		/// <param name="BasePath">The path that the file should be in</param>
 		/// <param name="FileName">The filename to check for (with and without SigningPrefix prepended to it)</param>
 		/// <returns>The path to the most desirable file (may still not exist)</returns>
-		static public string FindPrefixedFile(string BasePath, string FileName)
+		public static string FindPrefixedFile(string BasePath, string FileName)
 		{
 			return FindPrefixedFileOrDirectory(BasePath, FileName, false);
 		}
@@ -1115,7 +1115,7 @@ namespace UnrealBuildTool
 		/// GameName.mobileprovision or Distro_GameName.mobileprovision, but falling back to using the first one we find
 		/// if the correctly named one is not found helps the user if they mess up during the process.
 		/// </summary>
-		static public string FindAnyFileWithExtension(string BasePath, string FileExt)
+		public static string FindAnyFileWithExtension(string BasePath, string FileExt)
 		{
 			string[] FileList = Directory.GetFiles(BasePath);
 			foreach (string Filename in FileList)
@@ -1136,7 +1136,7 @@ namespace UnrealBuildTool
 		/// <param name="bIsDirectory">Is the desired name a directory?</param>
 		/// 
 		/// <returns>The path to the most desirable file (may still not exist)</returns>
-		static public string FindPrefixedFileOrDirectory(string BasePath, string Name, bool bIsDirectory)
+		public static string FindPrefixedFileOrDirectory(string BasePath, string Name, bool bIsDirectory)
 		{
 			string PrefixedPath = Path.Combine(BasePath, CodeSigningConfig.SigningPrefix + Name);
 
@@ -1170,7 +1170,7 @@ namespace UnrealBuildTool
 		/// <summary>
 		/// Get the string from the plist
 		/// </summary>
-		static public string GetStringFromPList(string KeyName)
+		public static string GetStringFromPList(string KeyName)
 		{
 			// Open the .plist and read out the specified key
 			string PListAsString;
@@ -1195,7 +1195,7 @@ namespace UnrealBuildTool
 		/// <summary>
 		/// Get the precompile source plist filename
 		/// </summary>
-		static public string GetPrecompileSourcePListFilename()
+		public static string GetPrecompileSourcePListFilename()
 		{
 			// check for one in the project directory
 			string SourceName = FileOperations.FindPrefixedFile(CodeSigningConfig.IntermediateDirectory, CodeSigningConfig.Program_GameName + "-Info.plist");
@@ -1226,7 +1226,7 @@ namespace UnrealBuildTool
 		/**
 		 * Handle grabbing the initial plist
 		 */
-		static public bool GetSourcePList(out string PListSource)
+		public static bool GetSourcePList(out string PListSource)
 		{
 			// Check for a premade one
 			string SourceName = GetPrecompileSourcePListFilename();
@@ -1489,7 +1489,7 @@ namespace UnrealBuildTool
 			// First checking for a contains match and then for a wildcard match
 			for (int Phase = -1; Phase < 3; ++Phase)
 			{
-				if (Phase == -1 && string.IsNullOrEmpty(CodeSigningConfig.ProvisionUUID))
+				if (Phase == -1 && String.IsNullOrEmpty(CodeSigningConfig.ProvisionUUID))
 				{
 					continue;
 				}
@@ -1514,7 +1514,7 @@ namespace UnrealBuildTool
 					}
 
 					// check to see if the platform is the same as what we are looking for
-					if (!string.IsNullOrEmpty(TestProvision.Platform) && TestProvision.Platform != CodeSigningConfig.OSString && !string.IsNullOrEmpty(CodeSigningConfig.OSString))
+					if (!String.IsNullOrEmpty(TestProvision.Platform) && TestProvision.Platform != CodeSigningConfig.OSString && !String.IsNullOrEmpty(CodeSigningConfig.OSString))
 					{
 						_logger?.LogInformation("  Failing platform {0} Config: {1}", TestProvision.Platform, CodeSigningConfig.OSString);
 						continue;
diff --git a/Engine/Source/Programs/UnrealBuildTool/Platform/IOS/IOSExports.cs b/Engine/Source/Programs/UnrealBuildTool/Platform/IOS/IOSExports.cs
--- a/Engine/Source/Programs/UnrealBuildTool/Platform/IOS/IOSExports.cs
+++ b/Engine/Source/Programs/UnrealBuildTool/Platform/IOS/IOSExports.cs
@@ -276,7 +276,7 @@ namespace UnrealBuildTool
 				Text.AppendLine("<plist version=\"1.0\">");
 				Text.AppendLine("<dict>");
 				Text.AppendLine("\t<key>get-task-allow</key>");
-				Text.AppendLine(string.Format("\t<{0}/>", bForDistribution ? "false" : "true"));
+				Text.AppendLine(String.Format("\t<{0}/>", bForDistribution ? "false" : "true"));
 				if (bCloudKitSupported)
 				{
 					if (iCloudContainerIdentifiersXML != "")
@@ -304,7 +304,7 @@ namespace UnrealBuildTool
 					}
 
 					Text.AppendLine("\t<key>com.apple.developer.icloud-container-environment</key>");
-					Text.AppendLine(string.Format("\t<string>{0}</string>", bForDistribution ? "Production" : "Development"));
+					Text.AppendLine(String.Format("\t<string>{0}</string>", bForDistribution ? "Production" : "Development"));
 				}
 
 				bool bRemoteNotificationsSupported = false;
@@ -319,7 +319,7 @@ namespace UnrealBuildTool
 				if (bRemoteNotificationsSupported)
 				{
 					Text.AppendLine("\t<key>aps-environment</key>");
-					Text.AppendLine(string.Format("\t<string>{0}</string>", bForDistribution ? "production" : "development"));
+					Text.AppendLine(String.Format("\t<string>{0}</string>", bForDistribution ? "production" : "development"));
 				}
 
 				// for Sign in with Apple
diff --git a/Engine/Source/Programs/UnrealBuildTool/Platform/IOS/IOSPostBuildSyncMode.cs b/Engine/Source/Programs/UnrealBuildTool/Platform/IOS/IOSPostBuildSyncMode.cs
--- a/Engine/Source/Programs/UnrealBuildTool/Platform/IOS/IOSPostBuildSyncMode.cs
+++ b/Engine/Source/Programs/UnrealBuildTool/Platform/IOS/IOSPostBuildSyncMode.cs
@@ -33,24 +33,24 @@ namespace UnrealBuildTool
 
 		public IOSPostBuildSyncTarget(ReadOnlyTargetRules Target, FileReference OutputPath, DirectoryReference? ProjectIntermediateDirectory, List<string> UPLScripts, VersionNumber SdkVersion, Dictionary<string, DirectoryReference> FrameworkNameToSourceDir)
 		{
-			this.Platform = Target.Platform;
-			this.Configuration = Target.Configuration;
-			this.ProjectFile = Target.ProjectFile;
-			this.TargetName = Target.Name;
-			this.TargetType = Target.Type;
+			Platform = Target.Platform;
+			Configuration = Target.Configuration;
+			ProjectFile = Target.ProjectFile;
+			TargetName = Target.Name;
+			TargetType = Target.Type;
 			this.OutputPath = OutputPath;
 			this.UPLScripts = UPLScripts;
 			this.SdkVersion = SdkVersion;
-			this.bCreateStubIPA = Target.IOSPlatform.bCreateStubIPA;
-			this.bSkipCrashlytics = Target.IOSPlatform.bSkipCrashlytics;
-			this.ProjectDirectory = DirectoryReference.FromFile(Target.ProjectFile) ?? Unreal.EngineDirectory;
+			bCreateStubIPA = Target.IOSPlatform.bCreateStubIPA;
+			bSkipCrashlytics = Target.IOSPlatform.bSkipCrashlytics;
+			ProjectDirectory = DirectoryReference.FromFile(Target.ProjectFile) ?? Unreal.EngineDirectory;
 			this.ProjectIntermediateDirectory = ProjectIntermediateDirectory;
-			this.ImportProvision = Target.IOSPlatform.ImportProvision;
-			this.ImportCertificate = Target.IOSPlatform.ImportCertificate;
-			this.ImportCertificatePassword = Target.IOSPlatform.ImportCertificatePassword;
+			ImportProvision = Target.IOSPlatform.ImportProvision;
+			ImportCertificate = Target.IOSPlatform.ImportCertificate;
+			ImportCertificatePassword = Target.IOSPlatform.ImportCertificatePassword;
 			this.FrameworkNameToSourceDir = FrameworkNameToSourceDir;
-			this.bForDistribution = Target.IOSPlatform.bForDistribution;
-			this.bBuildAsFramework = Target.bShouldCompileAsDLL;
+			bForDistribution = Target.IOSPlatform.bForDistribution;
+			bBuildAsFramework = Target.bShouldCompileAsDLL;
 		}
 	}
 
