UnrealBuildTool: More automated code cleanup

diff --git a/Engine/Source/Programs/UnrealBuildTool/Platform/IOS/IOSToolChain.cs b/Engine/Source/Programs/UnrealBuildTool/Platform/IOS/IOSToolChain.cs
--- a/Engine/Source/Programs/UnrealBuildTool/Platform/IOS/IOSToolChain.cs
+++ b/Engine/Source/Programs/UnrealBuildTool/Platform/IOS/IOSToolChain.cs
@@ -65,7 +65,7 @@ namespace UnrealBuildTool
 			SelectSDK(BaseSDKDir, DevicePlatformName, ref IOSSDKVersion, true, Logger);
 
 			// convert to float for easy comparison
-			IOSSDKVersionFloat = float.Parse(IOSSDKVersion, System.Globalization.CultureInfo.InvariantCulture);
+			IOSSDKVersionFloat = Single.Parse(IOSSDKVersion, System.Globalization.CultureInfo.InvariantCulture);
 		}
 
 		public string GetSDKPath(UnrealArch Architecture)
@@ -457,7 +457,7 @@ namespace UnrealBuildTool
 			{
 				GetLinkArguments_Global(LinkEnvironment, LinkArguments);
 			}
-			string LinkCommandArguments = string.Join(' ', LinkArguments);
+			string LinkCommandArguments = String.Join(' ', LinkArguments);
 			if (LinkEnvironment.bIsBuildingDLL)
 			{
 				// @todo roll this put into GetLinkArguments_Global
@@ -465,7 +465,7 @@ namespace UnrealBuildTool
 				LinkCommandArguments += " -dynamiclib -Xlinker -export_dynamic -Xlinker -no_deduplicate";
 
 				string InstallName = LinkEnvironment.InstallName ?? String.Format("@executable_path/Frameworks/{0}", LinkEnvironment.OutputFilePath.MakeRelativeTo(LinkEnvironment.OutputFilePath.Directory.ParentDirectory!));
-				LinkCommandArguments += string.Format(" -Xlinker -install_name -Xlinker {0}", InstallName);
+				LinkCommandArguments += String.Format(" -Xlinker -install_name -Xlinker {0}", InstallName);
 
 				LinkCommandArguments += " -Xlinker -rpath -Xlinker @executable_path/Frameworks";
 				LinkCommandArguments += " -Xlinker -rpath -Xlinker @loader_path/Frameworks";
@@ -476,13 +476,13 @@ namespace UnrealBuildTool
 				// Add the library paths to the argument list.
 				foreach (DirectoryReference LibraryPath in LinkEnvironment.SystemLibraryPaths)
 				{
-					LinkCommandArguments += string.Format(" -L\\\"{0}\\\"", LibraryPath.FullName);
+					LinkCommandArguments += String.Format(" -L\\\"{0}\\\"", LibraryPath.FullName);
 				}
 
 				// Add the additional libraries to the argument list.
 				foreach (string AdditionalLibrary in LinkEnvironment.SystemLibraries)
 				{
-					LinkCommandArguments += string.Format(" -l\\\"{0}\\\"", AdditionalLibrary);
+					LinkCommandArguments += String.Format(" -l\\\"{0}\\\"", AdditionalLibrary);
 				}
 
 				foreach (FileReference Library in LinkEnvironment.Libraries)
@@ -493,7 +493,7 @@ namespace UnrealBuildTool
 					LinkAction.PrerequisiteItems.Add(LibFile);
 
 					// and add to the commandline
-					LinkCommandArguments += string.Format(" \\\"{0}\\\"", Library.FullName);
+					LinkCommandArguments += String.Format(" \\\"{0}\\\"", Library.FullName);
 				}
 			}
 
@@ -515,7 +515,7 @@ namespace UnrealBuildTool
 			if ((!LinkEnvironment.bIsBuildingLibrary || LinkEnvironment.bIsBuildingDLL) && LinkEnvironment.bCreateMapFile)
 			{
 				FileItem MapFile = FileItem.GetItemByFileReference(new FileReference(OutputFile.Location.FullName + ".map"));
-				LinkCommandArguments += string.Format(" -Wl,-map,\\\"{0}\\\"", MapFile.Location.FullName);
+				LinkCommandArguments += String.Format(" -Wl,-map,\\\"{0}\\\"", MapFile.Location.FullName);
 				LinkAction.ProducedItems.Add(MapFile);
 			}
 
@@ -523,7 +523,7 @@ namespace UnrealBuildTool
 			List<string> InputFileNames = new List<string>();
 			foreach (FileItem InputFile in LinkEnvironment.InputFiles)
 			{
-				string FileNamePath = string.Format("\"{0}\"", InputFile.AbsolutePath);
+				string FileNamePath = String.Format("\"{0}\"", InputFile.AbsolutePath);
 				InputFileNames.Add(FileNamePath);
 				LinkAction.PrerequisiteItems.Add(InputFile);
 			}
@@ -545,7 +545,7 @@ namespace UnrealBuildTool
 				bool bIsUnrealGame = LinkEnvironment.OutputFilePath.FullName.Contains("UnrealGame");
 				FileReference ResponsePath = FileReference.Combine(((!bIsUnrealGame && ProjectFile != null) ? ProjectFile.Directory : Unreal.EngineDirectory), "Intermediate", "Build", LinkEnvironment.Platform.ToString(), "LinkFileList_" + LinkEnvironment.OutputFilePath.GetFileNameWithoutExtension() + ".tmp");
 				Graph.CreateIntermediateTextFile(ResponsePath, InputFileNames);
-				LinkCommandArguments += string.Format(" \\\"@{0}\\\"", ResponsePath.FullName);
+				LinkCommandArguments += String.Format(" \\\"@{0}\\\"", ResponsePath.FullName);
 			}
 
 			// if we are making an LTO build, write the lto file next to build so dsymutil can find it
@@ -561,11 +561,11 @@ namespace UnrealBuildTool
 				{
 					LtoObjectFile = OutputFile.AbsolutePath + ".lto.o";
 				}
-				LinkCommandArguments += string.Format(" -flto -Xlinker -object_path_lto -Xlinker \\\"{0}\\\"", LtoObjectFile);
+				LinkCommandArguments += String.Format(" -flto -Xlinker -object_path_lto -Xlinker \\\"{0}\\\"", LtoObjectFile);
 			}
 
 			// Add the output file to the command-line.
-			LinkCommandArguments += string.Format(" -o \\\"{0}\\\"", OutputFile.AbsolutePath);
+			LinkCommandArguments += String.Format(" -o \\\"{0}\\\"", OutputFile.AbsolutePath);
 
 			// Add the additional arguments specified by the environment.
 			LinkCommandArguments += LinkEnvironment.AdditionalArguments;
@@ -573,7 +573,7 @@ namespace UnrealBuildTool
 			// Only execute linking on the local PC.
 			LinkAction.bCanExecuteRemotely = false;
 
-			LinkAction.StatusDescription = string.Format("{0}", OutputFile.AbsolutePath);
+			LinkAction.StatusDescription = String.Format("{0}", OutputFile.AbsolutePath);
 
 			LinkAction.CommandVersion = Info.ClangVersionString;
 
@@ -586,7 +586,7 @@ namespace UnrealBuildTool
 				// the output file to force the linker to recreate it with symbols again.
 				string linkCommandArguments = "-c \"";
 
-				linkCommandArguments += string.Format("rm -f \\\"{0}\\\";", OutputFile.AbsolutePath);
+				linkCommandArguments += String.Format("rm -f \\\"{0}\\\";", OutputFile.AbsolutePath);
 				linkCommandArguments += "\\\"" + LinkerPath + "\\\" " + LinkCommandArguments + ";";
 
 				linkCommandArguments += "\"";
@@ -596,7 +596,7 @@ namespace UnrealBuildTool
 			else
 			{
 				// This is not a shipping build so no need to delete the output file since symbols will not have been stripped from it.
-				LinkAction.CommandArguments = string.Format("-c \"\\\"{0}\\\" {1}\"", LinkerPath, LinkCommandArguments);
+				LinkAction.CommandArguments = String.Format("-c \"\\\"{0}\\\" {1}\"", LinkerPath, LinkCommandArguments);
 			}
 			return OutputFile;
 		}
@@ -704,7 +704,7 @@ namespace UnrealBuildTool
 			string DsymutilPath = GetDsymutilPath(Logger, out ExtraOptions, bIsForLTOBuild: bIsForLTOBuild);
 			if (ProjectSettings.bGeneratedSYMBundle)
 			{
-				GenDebugAction.CommandArguments = string.Format("-c \"rm -rf \\\"{2}\\\"; \\\"{0}\\\" \\\"{1}\\\" {4} -o \\\"{2}\\\"; cd \\\"{2}/..\\\"; zip -r -y -1 {3}.zip {3}\"",
+				GenDebugAction.CommandArguments = String.Format("-c \"rm -rf \\\"{2}\\\"; \\\"{0}\\\" \\\"{1}\\\" {4} -o \\\"{2}\\\"; cd \\\"{2}/..\\\"; zip -r -y -1 {3}.zip {3}\"",
 					DsymutilPath,
 					Executable.AbsolutePath,
 					OutputFile.AbsolutePath,
@@ -715,7 +715,7 @@ namespace UnrealBuildTool
 			}
 			else
 			{
-				GenDebugAction.CommandArguments = string.Format("-c \"rm -rf \\\"{2}\\\"; \\\"{0}\\\" \\\"{1}\\\" {3} -f -o \\\"{2}\\\"\"",
+				GenDebugAction.CommandArguments = String.Format("-c \"rm -rf \\\"{2}\\\"; \\\"{0}\\\" \\\"{1}\\\" {3} -f -o \\\"{2}\\\"\"",
 						DsymutilPath,
 						Executable.AbsolutePath,
 						OutputFile.AbsolutePath,
@@ -752,7 +752,7 @@ namespace UnrealBuildTool
 			GenDebugAction.WorkingDirectory = DirectoryReference.Combine(Unreal.EngineDirectory, "Binaries", "Mac");
 
 			GenDebugAction.CommandPath = BuildHostPlatform.Current.Shell;
-			GenDebugAction.CommandArguments = string.Format("-c \"rm -rf \\\"{1}\\\"; dwarfdump --uuid \\\"{3}\\\" | cut -d\\  -f2; chmod 777 ./DsymExporter; ./DsymExporter -UUID=$(dwarfdump --uuid \\\"{3}\\\" | cut -d\\  -f2) \\\"{0}\\\" \\\"{2}\\\"\"",
+			GenDebugAction.CommandArguments = String.Format("-c \"rm -rf \\\"{1}\\\"; dwarfdump --uuid \\\"{3}\\\" | cut -d\\  -f2; chmod 777 ./DsymExporter; ./DsymExporter -UUID=$(dwarfdump --uuid \\\"{3}\\\" | cut -d\\  -f2) \\\"{0}\\\" \\\"{2}\\\"\"",
 					DWARFFile.AbsolutePath,
 					OutputFile.AbsolutePath,
 					Path.GetDirectoryName(OutputFile.AbsolutePath),
@@ -784,7 +784,7 @@ namespace UnrealBuildTool
 
 			// set up the directories
 			string ZipWorkingDir = String.Format("Payload/{0}.app/", GameName);
-			string ZipSourceDir = string.Format("{0}/Payload/{1}.app", BinaryPath, GameName);
+			string ZipSourceDir = String.Format("{0}/Payload/{1}.app", BinaryPath, GameName);
 
 			// create the file
 			using (ZipFile Zip = new ZipFile())
@@ -882,7 +882,7 @@ namespace UnrealBuildTool
 		public static DirectoryReference GenerateAssetCatalog(FileReference? ProjectFile, UnrealTargetPlatform Platform, ref bool bUserImagesExist)
 		{
 			string EngineDir = Unreal.EngineDirectory.ToString();
-			string BuildDir = (((ProjectFile != null) ? ProjectFile.Directory.ToString() : (string.IsNullOrEmpty(UnrealBuildTool.GetRemoteIniPath()) ? Unreal.EngineDirectory.ToString() : UnrealBuildTool.GetRemoteIniPath()))) + "/Build/" + (Platform == UnrealTargetPlatform.IOS ? "IOS" : "TVOS");
+			string BuildDir = (((ProjectFile != null) ? ProjectFile.Directory.ToString() : (String.IsNullOrEmpty(UnrealBuildTool.GetRemoteIniPath()) ? Unreal.EngineDirectory.ToString() : UnrealBuildTool.GetRemoteIniPath()))) + "/Build/" + (Platform == UnrealTargetPlatform.IOS ? "IOS" : "TVOS");
 			string IntermediateDir = (((ProjectFile != null) ? ProjectFile.Directory.ToString() : Unreal.EngineDirectory.ToString())) + "/Intermediate/" + (Platform == UnrealTargetPlatform.IOS ? "IOS" : "TVOS");
 
 			bUserImagesExist = false;
@@ -1309,7 +1309,7 @@ namespace UnrealBuildTool
 				Process FabricProcess = new Process();
 				FabricProcess.StartInfo.WorkingDirectory = Path.GetDirectoryName(DsymZip);
 				FabricProcess.StartInfo.FileName = "/bin/sh";
-				FabricProcess.StartInfo.Arguments = string.Format("-c \"chmod 777 \\\"{0}/Fabric.framework/upload-symbols\\\"; \\\"{0}/Fabric.framework/upload-symbols\\\" -a 7a4cebd0324af21696e5e321802c5e26ba541cad -p ios {1}\"",
+				FabricProcess.StartInfo.Arguments = String.Format("-c \"chmod 777 \\\"{0}/Fabric.framework/upload-symbols\\\"; \\\"{0}/Fabric.framework/upload-symbols\\\" -a 7a4cebd0324af21696e5e321802c5e26ba541cad -p ios {1}\"",
 					FabricPath,
 					DsymZip);
 
@@ -1400,8 +1400,8 @@ namespace UnrealBuildTool
 			string CookedDataPath = DirectoryReference.Combine(ProjectDirectory, "Saved", "StagedBuilds", "IOS", "cookeddata").ToString();
 			bool GenerateFrameworkWrapperProject = ProjectFiles.Xcode.XcodeFrameworkWrapperUtils.GetGenerateFrameworkWrapperProject(Target.ProjectDirectory);
 
-			string? ProvisionName = string.Empty;
-			string? TeamUUID = string.Empty;
+			string? ProvisionName = String.Empty;
+			string? TeamUUID = String.Empty;
 			IOSProjectSettings ProjectSettings = ((IOSPlatform)UEBuildPlatform.GetBuildPlatform(UnrealTargetPlatform.IOS)).ReadProjectSettings(Target.ProjectFile);
 
 			IOSProvisioningData Data = ((IOSPlatform)UEBuildPlatform.GetBuildPlatform(UnrealTargetPlatform.IOS)).ReadProvisioningData(ProjectSettings, Target.bForDistribution);
@@ -1447,7 +1447,9 @@ namespace UnrealBuildTool
 
 						WriteMissingEmbeddedPlist(PlistSrcLocation, Target.ProjectDirectory, Target.ProjectFile);
 						if (!FileReference.Exists(PlistSrcLocation))
+						{
 							throw new BuildException("Unable to find plist for output framework ({0})", PlistSrcLocation);
+						}
 					}
 
 					FileReference PlistDstLocation = FileReference.Combine(Target.OutputPath.Directory, "Info.plist");
@@ -1642,7 +1644,7 @@ namespace UnrealBuildTool
 							" -destination generic/platform=" + (Target.Platform == UnrealTargetPlatform.IOS ? "iOS" : "tvOS") +
 								" -derivedDataPath \"" + FrameworkDerivedDataDir + "\"" +
 							" CONFIGURATION_BUILD_DIR=\"" + FrameworkPayloadDirectory + "\"" +
-								(!string.IsNullOrEmpty(TeamUUID) ? " DEVELOPMENT_TEAM=" + TeamUUID : "");
+								(!String.IsNullOrEmpty(TeamUUID) ? " DEVELOPMENT_TEAM=" + TeamUUID : "");
 					}
 					else
 					{
@@ -1653,7 +1655,7 @@ namespace UnrealBuildTool
 							" -scheme '" + SchemeName + "'" +
 								" -sdk " + GetCodesignPlatformName(Target.Platform) +
 							" -destination generic/platform=" + (Target.Platform == UnrealTargetPlatform.IOS ? "iOS" : "tvOS") +
-								(!string.IsNullOrEmpty(TeamUUID) ? " DEVELOPMENT_TEAM=" + TeamUUID : "");
+								(!String.IsNullOrEmpty(TeamUUID) ? " DEVELOPMENT_TEAM=" + TeamUUID : "");
 					}
 
 					CmdLine += String.Format(" CODE_SIGN_IDENTITY='{0}'", SigningCertificate);
@@ -1664,7 +1666,7 @@ namespace UnrealBuildTool
 
 					if (!ProjectSettings.bAutomaticSigning)
 					{
-						CmdLine += (!string.IsNullOrEmpty(MobileProvisionUUID) ? (" PROVISIONING_PROFILE_SPECIFIER=" + MobileProvisionUUID) : "");
+						CmdLine += (!String.IsNullOrEmpty(MobileProvisionUUID) ? (" PROVISIONING_PROFILE_SPECIFIER=" + MobileProvisionUUID) : "");
 					}
 					Writer.WriteLine("/usr/bin/xcrun {0}", CmdLine);
 				}
@@ -1840,11 +1842,11 @@ namespace UnrealBuildTool
 
 			FileItem TargetItem = FileItem.GetItemByPath(TargetPath);
 
-			CopyAction.CommandArguments = string.Format("-c \"cp -f -R \\\"{0}\\\" \\\"{1}\\\"; touch -c \\\"{2}\\\"\"", SourcePath, Path.GetDirectoryName(TargetPath)!.Replace('\\', '/') + "/", TargetPath.Replace('\\', '/'));
+			CopyAction.CommandArguments = String.Format("-c \"cp -f -R \\\"{0}\\\" \\\"{1}\\\"; touch -c \\\"{2}\\\"\"", SourcePath, Path.GetDirectoryName(TargetPath)!.Replace('\\', '/') + "/", TargetPath.Replace('\\', '/'));
 			CopyAction.PrerequisiteItems.Add(Executable);
 			CopyAction.ProducedItems.Add(TargetItem);
 			CopyAction.bShouldOutputStatusDescription = Resource.bShouldLog;
-			CopyAction.StatusDescription = string.Format("Copying {0} to app bundle", Path.GetFileName(Resource.ResourcePath));
+			CopyAction.StatusDescription = String.Format("Copying {0} to app bundle", Path.GetFileName(Resource.ResourcePath));
 			CopyAction.bCanExecuteRemotely = false;
 
 			return TargetItem;
diff --git a/Engine/Source/Programs/UnrealBuildTool/Platform/IOS/MobileProvisionContents.cs b/Engine/Source/Programs/UnrealBuildTool/Platform/IOS/MobileProvisionContents.cs
--- a/Engine/Source/Programs/UnrealBuildTool/Platform/IOS/MobileProvisionContents.cs
+++ b/Engine/Source/Programs/UnrealBuildTool/Platform/IOS/MobileProvisionContents.cs
@@ -197,7 +197,7 @@ namespace UnrealBuildTool
 		// return the outerXML of the node's value
 		public string GetNodeXMLValueByName(string InValue)
 		{
-			XmlNodeList elemList = this.Document.GetElementsByTagName("key");
+			XmlNodeList elemList = Document.GetElementsByTagName("key");
 			for (int i = 0; i < elemList.Count; i++)
 			{
 				XmlNode? Node = elemList[i];
@@ -217,7 +217,7 @@ namespace UnrealBuildTool
 		// return the innerXML of the node's value
 		public string GetNodeValueByName(string InValue)
 		{
-			XmlNodeList elemList = this.Document.GetElementsByTagName("key");
+			XmlNodeList elemList = Document.GetElementsByTagName("key");
 			for (int i = 0; i < elemList.Count; i++)
 			{
 				XmlNode? Node = elemList[i];
diff --git a/Engine/Source/Programs/UnrealBuildTool/Platform/IOS/UEBuildIOS.cs b/Engine/Source/Programs/UnrealBuildTool/Platform/IOS/UEBuildIOS.cs
--- a/Engine/Source/Programs/UnrealBuildTool/Platform/IOS/UEBuildIOS.cs
+++ b/Engine/Source/Programs/UnrealBuildTool/Platform/IOS/UEBuildIOS.cs
@@ -41,7 +41,10 @@ namespace UnrealBuildTool
 		{
 			get
 			{
-				if (AppleToolchainArchitectures.ContainsKey(this)) return AppleToolchainArchitectures[this];
+				if (AppleToolchainArchitectures.ContainsKey(this))
+				{
+					return AppleToolchainArchitectures[this];
+				}
 
 				throw new BuildException($"Unknown architecture {ToString()} passed to UnrealArch.AppleName");
 			}
@@ -190,7 +193,7 @@ namespace UnrealBuildTool
 
 		public string? ImportCertificatePassword => Inner.ImportCertificatePassword;
 
-		public float RuntimeVersion => float.Parse(Inner.ProjectSettings!.RuntimeVersion, System.Globalization.CultureInfo.InvariantCulture);
+		public float RuntimeVersion => Single.Parse(Inner.ProjectSettings!.RuntimeVersion, System.Globalization.CultureInfo.InvariantCulture);
 
 		public bool bEnableAddressSanitizer => Inner.bEnableAddressSanitizer;
 
@@ -415,7 +418,7 @@ namespace UnrealBuildTool
 		{
 			this.ProjectFile = ProjectFile;
 			ConfigCache.ReadSettings(DirectoryReference.FromFile(ProjectFile), Platform, this);
-			if ((ProjectFile == null || string.IsNullOrEmpty(ProjectFile.FullName)) && !string.IsNullOrEmpty(Bundle))
+			if ((ProjectFile == null || String.IsNullOrEmpty(ProjectFile.FullName)) && !String.IsNullOrEmpty(Bundle))
 			{
 				BundleIdentifier = Bundle;
 			}
@@ -451,7 +454,7 @@ namespace UnrealBuildTool
 			FileReference? ProjectFile = ProjectSettings.ProjectFile;
 			CodeSigningConfig.Initialize(ProjectFile, bIsTVOS);
 
-			if (!string.IsNullOrEmpty(SigningCertificate))
+			if (!String.IsNullOrEmpty(SigningCertificate))
 			{
 				List<string> Certs = AppleCodeSign.FindCertificates();
 				List<FileReference> Provisions = AppleCodeSign.FindProvisions(ProjectSettings.BundleIdentifier, bForDistribution, out _);
@@ -462,7 +465,7 @@ namespace UnrealBuildTool
 				bHaveCertificate = true;
 			}
 
-			if (!string.IsNullOrEmpty(MobileProvision))
+			if (!String.IsNullOrEmpty(MobileProvision))
 			{
 				DirectoryReference MobileProvisionDir;
 				if (BuildHostPlatform.Current.Platform == UnrealTargetPlatform.Mac)
@@ -589,7 +592,7 @@ namespace UnrealBuildTool
 					}
 				}
 
-				if (string.IsNullOrEmpty(MobileProvisionUUID) || string.IsNullOrEmpty(TeamUUID))
+				if (String.IsNullOrEmpty(MobileProvisionUUID) || String.IsNullOrEmpty(TeamUUID))
 				{
 					MobileProvision = null;
 					SigningCertificate = null;
@@ -612,7 +615,7 @@ namespace UnrealBuildTool
 					Logger.LogInformation("{LineData}", Line.Data);
 				}
 
-				if (!string.IsNullOrEmpty(SigningCertificate))
+				if (!String.IsNullOrEmpty(SigningCertificate))
 				{
 					if (Line.Data.Contains("CERTIFICATE-") && Line.Data.Contains(SigningCertificate))
 					{
@@ -711,11 +714,11 @@ namespace UnrealBuildTool
 
 		public override void ValidateTarget(TargetRules Target)
 		{
-			if (!string.IsNullOrWhiteSpace(Environment.GetEnvironmentVariable("CLANG_STATIC_ANALYZER_MODE")))
+			if (!String.IsNullOrWhiteSpace(Environment.GetEnvironmentVariable("CLANG_STATIC_ANALYZER_MODE")))
 			{
 				Target.StaticAnalyzer = StaticAnalyzer.Default;
 				Target.StaticAnalyzerOutputType = (Environment.GetEnvironmentVariable("CLANG_ANALYZER_OUTPUT")?.Contains("html", StringComparison.OrdinalIgnoreCase) == true) ? StaticAnalyzerOutputType.Html : StaticAnalyzerOutputType.Text;
-				Target.StaticAnalyzerMode = string.Equals(Environment.GetEnvironmentVariable("CLANG_STATIC_ANALYZER_MODE"), "shallow") ? StaticAnalyzerMode.Shallow : StaticAnalyzerMode.Deep;
+				Target.StaticAnalyzerMode = String.Equals(Environment.GetEnvironmentVariable("CLANG_STATIC_ANALYZER_MODE"), "shallow") ? StaticAnalyzerMode.Shallow : StaticAnalyzerMode.Deep;
 			}
 			else if (Target.StaticAnalyzer == StaticAnalyzer.Clang)
 			{
@@ -778,7 +781,7 @@ namespace UnrealBuildTool
 								Args += LibLoc.FullName;
 								Args += " | grep -m1 -i \\(clang\"";
 								string StdOutResult = Utils.RunLocalProcessAndReturnStdOut("bash", Args);
-								if (string.IsNullOrEmpty(StdOutResult))
+								if (String.IsNullOrEmpty(StdOutResult))
 								{
 									continue;
 								}
@@ -848,7 +851,7 @@ namespace UnrealBuildTool
 			IOSProjectSettings? ProjectSettings = null;
 
 			// Use separate lists to prevent an overridden Bundle id polluting the standard project file. 
-			bool bCacheByBundle = !string.IsNullOrEmpty(Bundle);
+			bool bCacheByBundle = !String.IsNullOrEmpty(Bundle);
 			if (bCacheByBundle)
 			{
 				lock (CachedProjectSettingsByBundle)
diff --git a/Engine/Source/Programs/UnrealBuildTool/Platform/IOS/UEDeployIOS.cs b/Engine/Source/Programs/UnrealBuildTool/Platform/IOS/UEDeployIOS.cs
--- a/Engine/Source/Programs/UnrealBuildTool/Platform/IOS/UEDeployIOS.cs
+++ b/Engine/Source/Programs/UnrealBuildTool/Platform/IOS/UEDeployIOS.cs
@@ -78,12 +78,12 @@ namespace UnrealBuildTool
 				// Parse the version string
 				string[] VersionParts = Parts[0].Split(new char[] { '.' });
 
-				if (!int.TryParse(VersionParts[0], out VersionMajor))
+				if (!Int32.TryParse(VersionParts[0], out VersionMajor))
 				{
 					VersionMajor = 0;
 				}
 
-				if ((VersionParts.Length < 2) || (!int.TryParse(VersionParts[1], out VersionMinor)))
+				if ((VersionParts.Length < 2) || (!Int32.TryParse(VersionParts[1], out VersionMinor)))
 				{
 					VersionMinor = 0;
 				}
@@ -286,7 +286,7 @@ namespace UnrealBuildTool
 			// Write the Google iOS URL Scheme if we need it.
 			string GoogleReversedClientId = "";
 			Ini.GetString("/Script/IOSRuntimeSettings.IOSRuntimeSettings", "GoogleReversedClientId", out GoogleReversedClientId);
-			bEnableGoogleSupport = bEnableGoogleSupport && !string.IsNullOrWhiteSpace(GoogleReversedClientId);
+			bEnableGoogleSupport = bEnableGoogleSupport && !String.IsNullOrWhiteSpace(GoogleReversedClientId);
 
 			// Add remote-notifications as background mode
 			bool bRemoteNotificationsSupported = false;
@@ -326,10 +326,10 @@ namespace UnrealBuildTool
 			Text.AppendLine("\t\t\t<string>com.Epic.Unreal</string>");
 			Text.AppendLine("\t\t\t<key>CFBundleURLSchemes</key>");
 			Text.AppendLine("\t\t\t<array>");
-			Text.AppendLine(string.Format("\t\t\t\t<string>{0}</string>", bIsUnrealGame ? "UnrealGame" : GameName));
+			Text.AppendLine(String.Format("\t\t\t\t<string>{0}</string>", bIsUnrealGame ? "UnrealGame" : GameName));
 			if (bEnableGoogleSupport)
 			{
-				Text.AppendLine(string.Format("\t\t\t\t<string>{0}</string>", GoogleReversedClientId));
+				Text.AppendLine(String.Format("\t\t\t\t<string>{0}</string>", GoogleReversedClientId));
 			}
 			Text.AppendLine("\t\t\t</array>");
 			Text.AppendLine("\t\t</dict>");
@@ -337,7 +337,7 @@ namespace UnrealBuildTool
 			Text.AppendLine("\t<key>UIStatusBarHidden</key>");
 			Text.AppendLine("\t<true/>");
 			Text.AppendLine("\t<key>UIFileSharingEnabled</key>");
-			Text.AppendLine(string.Format("\t<{0}/>", bSupportsITunesFileSharing ? "true" : "false"));
+			Text.AppendLine(String.Format("\t<{0}/>", bSupportsITunesFileSharing ? "true" : "false"));
 			if (bSupportsFilesApp)
 			{
 				Text.AppendLine("\t<key>LSSupportsOpeningDocumentsInPlace</key>");
@@ -355,7 +355,7 @@ namespace UnrealBuildTool
 			Text.AppendLine("\t<array>");
 			foreach (string Line in SupportedOrientations.Split("\r\n".ToCharArray()))
 			{
-				if (!string.IsNullOrWhiteSpace(Line))
+				if (!String.IsNullOrWhiteSpace(Line))
 				{
 					Text.AppendLine(Line);
 				}
@@ -385,15 +385,15 @@ namespace UnrealBuildTool
 			Text.AppendLine("\t<key>ITSAppUsesNonExemptEncryption</key>");
 			Text.AppendLine("\t<false/>");
 			// add location services descriptions if used
-			if (!string.IsNullOrWhiteSpace(LocationAlwaysUsageDescription))
+			if (!String.IsNullOrWhiteSpace(LocationAlwaysUsageDescription))
 			{
 				Text.AppendLine("\t<key>NSLocationAlwaysAndWhenInUseUsageDescription</key>");
-				Text.AppendLine(string.Format("\t<string>{0}</string>", LocationAlwaysUsageDescription));
+				Text.AppendLine(String.Format("\t<string>{0}</string>", LocationAlwaysUsageDescription));
 			}
-			if (!string.IsNullOrWhiteSpace(LocationWhenInUseDescription))
+			if (!String.IsNullOrWhiteSpace(LocationWhenInUseDescription))
 			{
 				Text.AppendLine("\t<key>NSLocationWhenInUseUsageDescription</key>");
-				Text.AppendLine(string.Format("\t<string>{0}</string>", LocationWhenInUseDescription));
+				Text.AppendLine(String.Format("\t<string>{0}</string>", LocationWhenInUseDescription));
 			}
 			// disable HTTPS requirement
 			if (bDisableHTTPS)
@@ -404,12 +404,12 @@ namespace UnrealBuildTool
 				Text.AppendLine("\t\t</dict>");
 			}
 
-			if (!string.IsNullOrEmpty(ExtraData))
+			if (!String.IsNullOrEmpty(ExtraData))
 			{
 				ExtraData = ExtraData.Replace("\\n", "\n");
 				foreach (string Line in ExtraData.Split("\r\n".ToCharArray()))
 				{
-					if (!string.IsNullOrWhiteSpace(Line))
+					if (!String.IsNullOrWhiteSpace(Line))
 					{
 						Text.AppendLine("\t" + Line);
 					}
@@ -471,7 +471,7 @@ namespace UnrealBuildTool
 		{
 			// get the settings from the ini file
 			// plist replacements
-			DirectoryReference? DirRef = bIsUnrealGame ? (!string.IsNullOrEmpty(UnrealBuildTool.GetRemoteIniPath()) ? new DirectoryReference(UnrealBuildTool.GetRemoteIniPath()!) : null) : new DirectoryReference(ProjectDirectory);
+			DirectoryReference? DirRef = bIsUnrealGame ? (!String.IsNullOrEmpty(UnrealBuildTool.GetRemoteIniPath()) ? new DirectoryReference(UnrealBuildTool.GetRemoteIniPath()!) : null) : new DirectoryReference(ProjectDirectory);
 
 			if (!AppleExports.UseModernXcode(ProjectFile))
 			{
@@ -517,7 +517,7 @@ namespace UnrealBuildTool
 
 			// get the settings from the ini file
 			// plist replacements
-			DirectoryReference? DirRef = bIsUnrealGame ? (!string.IsNullOrEmpty(UnrealBuildTool.GetRemoteIniPath()) ? new DirectoryReference(UnrealBuildTool.GetRemoteIniPath()!) : null) : new DirectoryReference(ProjectDirectory);
+			DirectoryReference? DirRef = bIsUnrealGame ? (!String.IsNullOrEmpty(UnrealBuildTool.GetRemoteIniPath()) ? new DirectoryReference(UnrealBuildTool.GetRemoteIniPath()!) : null) : new DirectoryReference(ProjectDirectory);
 			ConfigHierarchy Ini = ConfigCache.ReadHierarchy(ConfigHierarchyType.Engine, DirRef, UnrealTargetPlatform.IOS);
 
 			// orientations
@@ -577,7 +577,7 @@ namespace UnrealBuildTool
 			// bundle identifier
 			string BundleIdentifier;
 			Ini.GetString("/Script/IOSRuntimeSettings.IOSRuntimeSettings", "BundleIdentifier", out BundleIdentifier);
-			if (!string.IsNullOrEmpty(BundleID))
+			if (!String.IsNullOrEmpty(BundleID))
 			{
 				BundleIdentifier = BundleID; // overriding bundle ID
 			}
@@ -612,7 +612,7 @@ namespace UnrealBuildTool
 			// Write the Google iOS URL Scheme if we need it.
 			string GoogleReversedClientId = "";
 			Ini.GetString("/Script/IOSRuntimeSettings.IOSRuntimeSettings", "GoogleReversedClientId", out GoogleReversedClientId);
-			bEnableGoogleSupport = bEnableGoogleSupport && !string.IsNullOrWhiteSpace(GoogleReversedClientId);
+			bEnableGoogleSupport = bEnableGoogleSupport && !String.IsNullOrWhiteSpace(GoogleReversedClientId);
 
 			// Add remote-notifications as background mode
 			bool bRemoteNotificationsSupported = false;
@@ -651,10 +651,10 @@ namespace UnrealBuildTool
 			Text.AppendLine("\t\t\t<string>com.Epic.Unreal</string>");
 			Text.AppendLine("\t\t\t<key>CFBundleURLSchemes</key>");
 			Text.AppendLine("\t\t\t<array>");
-			Text.AppendLine(string.Format("\t\t\t\t<string>{0}</string>", bIsUnrealGame ? "UnrealGame" : GameName));
+			Text.AppendLine(String.Format("\t\t\t\t<string>{0}</string>", bIsUnrealGame ? "UnrealGame" : GameName));
 			if (bEnableGoogleSupport)
 			{
-				Text.AppendLine(string.Format("\t\t\t\t<string>{0}</string>", GoogleReversedClientId));
+				Text.AppendLine(String.Format("\t\t\t\t<string>{0}</string>", GoogleReversedClientId));
 			}
 			Text.AppendLine("\t\t\t</array>");
 			Text.AppendLine("\t\t</dict>");
@@ -662,32 +662,32 @@ namespace UnrealBuildTool
 			Text.AppendLine("\t<key>CFBundleDevelopmentRegion</key>");
 			Text.AppendLine("\t<string>English</string>");
 			Text.AppendLine("\t<key>CFBundleDisplayName</key>");
-			Text.AppendLine(string.Format("\t<string>{0}</string>", EncodeBundleName(BundleDisplayName, ProjectName)));
+			Text.AppendLine(String.Format("\t<string>{0}</string>", EncodeBundleName(BundleDisplayName, ProjectName)));
 			Text.AppendLine("\t<key>CFBundleExecutable</key>");
 			string BundleExecutable = bIsUnrealGame ?
 				(bIsClient ? "UnrealClient" : "UnrealGame") :
 				(bIsClient ? GameName + "Client" : GameName);
-			Text.AppendLine(string.Format("\t<string>{0}</string>", BundleExecutable));
+			Text.AppendLine(String.Format("\t<string>{0}</string>", BundleExecutable));
 			Text.AppendLine("\t<key>CFBundleIdentifier</key>");
-			Text.AppendLine(string.Format("\t<string>{0}</string>", BundleIdentifier.Replace("[PROJECT_NAME]", ProjectName).Replace("_", "")));
+			Text.AppendLine(String.Format("\t<string>{0}</string>", BundleIdentifier.Replace("[PROJECT_NAME]", ProjectName).Replace("_", "")));
 			Text.AppendLine("\t<key>CFBundleInfoDictionaryVersion</key>");
 			Text.AppendLine("\t<string>6.0</string>");
 			Text.AppendLine("\t<key>CFBundleName</key>");
-			Text.AppendLine(string.Format("\t<string>{0}</string>", EncodeBundleName(BundleName, ProjectName)));
+			Text.AppendLine(String.Format("\t<string>{0}</string>", EncodeBundleName(BundleName, ProjectName)));
 			Text.AppendLine("\t<key>CFBundlePackageType</key>");
 			Text.AppendLine("\t<string>APPL</string>");
 			Text.AppendLine("\t<key>CFBundleSignature</key>");
 			Text.AppendLine("\t<string>????</string>");
 			Text.AppendLine("\t<key>CFBundleVersion</key>");
-			Text.AppendLine(string.Format("\t<string>{0}</string>", VersionUtilities.UpdateBundleVersion(OldPListData, InEngineDir)));
+			Text.AppendLine(String.Format("\t<string>{0}</string>", VersionUtilities.UpdateBundleVersion(OldPListData, InEngineDir)));
 			Text.AppendLine("\t<key>CFBundleShortVersionString</key>");
-			Text.AppendLine(string.Format("\t<string>{0}</string>", BundleShortVersion));
+			Text.AppendLine(String.Format("\t<string>{0}</string>", BundleShortVersion));
 			Text.AppendLine("\t<key>LSRequiresIPhoneOS</key>");
 			Text.AppendLine("\t<true/>");
 			Text.AppendLine("\t<key>UIStatusBarHidden</key>");
 			Text.AppendLine("\t<true/>");
 			Text.AppendLine("\t<key>UIFileSharingEnabled</key>");
-			Text.AppendLine(string.Format("\t<{0}/>", bSupportsITunesFileSharing ? "true" : "false"));
+			Text.AppendLine(String.Format("\t<{0}/>", bSupportsITunesFileSharing ? "true" : "false"));
 			if (bSupportsFilesApp)
 			{
 				Text.AppendLine("\t<key>LSSupportsOpeningDocumentsInPlace</key>");
@@ -705,7 +705,7 @@ namespace UnrealBuildTool
 			Text.AppendLine("\t<array>");
 			foreach (string Line in SupportedOrientations.Split("\r\n".ToCharArray()))
 			{
-				if (!string.IsNullOrWhiteSpace(Line))
+				if (!String.IsNullOrWhiteSpace(Line))
 				{
 					Text.AppendLine(Line);
 				}
@@ -715,7 +715,7 @@ namespace UnrealBuildTool
 			Text.AppendLine("\t<array>");
 			foreach (string Line in RequiredCaps.Split("\r\n".ToCharArray()))
 			{
-				if (!string.IsNullOrWhiteSpace(Line))
+				if (!String.IsNullOrWhiteSpace(Line))
 				{
 					Text.AppendLine(Line);
 				}
@@ -786,20 +786,20 @@ namespace UnrealBuildTool
 			Text.AppendLine("\t\t<string>iPhoneOS</string>");
 			Text.AppendLine("\t</array>");
 			Text.AppendLine("\t<key>MinimumOSVersion</key>");
-			Text.AppendLine(string.Format("\t<string>{0}</string>", MinVersion));
+			Text.AppendLine(String.Format("\t<string>{0}</string>", MinVersion));
 			// disable exempt encryption
 			Text.AppendLine("\t<key>ITSAppUsesNonExemptEncryption</key>");
 			Text.AppendLine("\t<false/>");
 			// add location services descriptions if used
-			if (!string.IsNullOrWhiteSpace(LocationAlwaysUsageDescription))
+			if (!String.IsNullOrWhiteSpace(LocationAlwaysUsageDescription))
 			{
 				Text.AppendLine("\t<key>NSLocationAlwaysAndWhenInUseUsageDescription</key>");
-				Text.AppendLine(string.Format("\t<string>{0}</string>", LocationAlwaysUsageDescription));
+				Text.AppendLine(String.Format("\t<string>{0}</string>", LocationAlwaysUsageDescription));
 			}
-			if (!string.IsNullOrWhiteSpace(LocationWhenInUseDescription))
+			if (!String.IsNullOrWhiteSpace(LocationWhenInUseDescription))
 			{
 				Text.AppendLine("\t<key>NSLocationWhenInUseUsageDescription</key>");
-				Text.AppendLine(string.Format("\t<string>{0}</string>", LocationWhenInUseDescription));
+				Text.AppendLine(String.Format("\t<string>{0}</string>", LocationWhenInUseDescription));
 			}
 			// disable HTTPS requirement
 			if (bDisableHTTPS)
@@ -810,12 +810,12 @@ namespace UnrealBuildTool
 				Text.AppendLine("\t\t</dict>");
 			}
 
-			if (!string.IsNullOrEmpty(ExtraData))
+			if (!String.IsNullOrEmpty(ExtraData))
 			{
 				ExtraData = ExtraData.Replace("\\n", "\n");
 				foreach (string Line in ExtraData.Split("\r\n".ToCharArray()))
 				{
-					if (!string.IsNullOrWhiteSpace(Line))
+					if (!String.IsNullOrWhiteSpace(Line))
 					{
 						Text.AppendLine("\t" + Line);
 					}
@@ -843,7 +843,7 @@ namespace UnrealBuildTool
 			}
 
 			// write the iCloud container identifier, if present in the old file
-			if (!string.IsNullOrEmpty(OldPListData))
+			if (!String.IsNullOrEmpty(OldPListData))
 			{
 				int index = OldPListData.IndexOf("ICloudContainerIdentifier");
 				if (index > 0)
@@ -852,7 +852,7 @@ namespace UnrealBuildTool
 					int length = OldPListData.IndexOf("</string>", index) - index;
 					string ICloudContainerIdentifier = OldPListData.Substring(index, length);
 					Text.AppendLine("\t<key>ICloudContainerIdentifier</key>");
-					Text.AppendLine(string.Format("\t<string>{0}</string>", ICloudContainerIdentifier));
+					Text.AppendLine(String.Format("\t<string>{0}</string>", ICloudContainerIdentifier));
 				}
 			}
 
@@ -1036,7 +1036,7 @@ namespace UnrealBuildTool
 				FileInfo DestFileInfo;
 				if (!Directory.Exists(ProvisionDir))
 				{
-					throw new DirectoryNotFoundException(string.Format("Provision Directory {0} not found.", ProvisionDir), null);
+					throw new DirectoryNotFoundException(String.Format("Provision Directory {0} not found.", ProvisionDir), null);
 				}
 
 				string LocalProvisionFolder = Path.Combine(Environment.GetFolderPath(Environment.SpecialFolder.Personal), "Library/MobileDevice/Provisioning Profiles");
@@ -1050,7 +1050,7 @@ namespace UnrealBuildTool
 					}
 					else
 					{
-						throw new DirectoryNotFoundException(string.Format("Local Provision Folder {0} could not be created.", LocalProvisionFolder), null);
+						throw new DirectoryNotFoundException(String.Format("Local Provision Folder {0} could not be created.", LocalProvisionFolder), null);
 					}
 				}
 
@@ -1188,7 +1188,7 @@ namespace UnrealBuildTool
 				CopyAllProvisions(BuildDirectory, Logger);
 				CopyAllProvisions(InEngineDir + "/Build/IOS", Logger);
 				string? ProvisionDirectory = Environment.GetEnvironmentVariable("ProvisionDirectory");
-				if (!string.IsNullOrWhiteSpace(ProvisionDirectory))
+				if (!String.IsNullOrWhiteSpace(ProvisionDirectory))
 				{
 					CopyAllProvisions(ProvisionDirectory, Logger);
 				}
@@ -1361,7 +1361,7 @@ namespace UnrealBuildTool
 			RecursiveFolderCopy(new DirectoryInfo(SourceDirectory), new DirectoryInfo(DestinationDirectory), bOverwrite, Filter);
 		}
 
-		static private void RecursiveFolderCopy(DirectoryInfo SourceFolderInfo, DirectoryInfo DestFolderInfo, bool bOverwrite = false, FilenameFilter? Filter = null)
+		private static void RecursiveFolderCopy(DirectoryInfo SourceFolderInfo, DirectoryInfo DestFolderInfo, bool bOverwrite = false, FilenameFilter? Filter = null)
 		{
 			foreach (FileInfo SourceFileInfo in SourceFolderInfo.GetFiles())
 			{
diff --git a/Engine/Source/Programs/UnrealBuildTool/Platform/Linux/LinuxCommon.cs b/Engine/Source/Programs/UnrealBuildTool/Platform/Linux/LinuxCommon.cs
--- a/Engine/Source/Programs/UnrealBuildTool/Platform/Linux/LinuxCommon.cs
+++ b/Engine/Source/Programs/UnrealBuildTool/Platform/Linux/LinuxCommon.cs
@@ -12,7 +12,7 @@ namespace UnrealBuildTool
 		{
 			Process proc = new Process();
 			proc.StartInfo.FileName = "/bin/sh";
-			proc.StartInfo.Arguments = string.Format("-c 'which {0}'", name);
+			proc.StartInfo.Arguments = System.String.Format("-c 'which {0}'", name);
 			proc.StartInfo.UseShellExecute = false;
 			proc.StartInfo.CreateNoWindow = true;
 			proc.StartInfo.RedirectStandardOutput = true;
@@ -24,7 +24,7 @@ namespace UnrealBuildTool
 			string? path = proc.StandardOutput.ReadLine();
 			Logger.LogDebug("which {Name} result: ({ExitCode}) {Path}", name, proc.ExitCode, path);
 
-			if (proc.ExitCode == 0 && string.IsNullOrEmpty(proc.StandardError.ReadToEnd()))
+			if (proc.ExitCode == 0 && System.String.IsNullOrEmpty(proc.StandardError.ReadToEnd()))
 			{
 				return path;
 			}
@@ -34,7 +34,7 @@ namespace UnrealBuildTool
 		public static string? WhichClang(ILogger Logger)
 		{
 			string? InternalSDKPath = UEBuildPlatform.GetSDK(UnrealTargetPlatform.Linux)?.GetInternalSDKPath();
-			if (!string.IsNullOrEmpty(InternalSDKPath))
+			if (!System.String.IsNullOrEmpty(InternalSDKPath))
 			{
 				return Path.Combine(InternalSDKPath, "bin", "clang++");
 			}
diff --git a/Engine/Source/Programs/UnrealBuildTool/Platform/Linux/LinuxToolChain.cs b/Engine/Source/Programs/UnrealBuildTool/Platform/Linux/LinuxToolChain.cs
--- a/Engine/Source/Programs/UnrealBuildTool/Platform/Linux/LinuxToolChain.cs
+++ b/Engine/Source/Programs/UnrealBuildTool/Platform/Linux/LinuxToolChain.cs
@@ -58,7 +58,7 @@ namespace UnrealBuildTool
 			if (CompilerVersionLessThan(15, 0, 0) || CompilerVersionGreaterOrEqual(16, 0, 0))
 			{
 				throw new BuildException(
-					string.Format("This version of the Unreal Engine can only be compiled with clang 15.0. clang {0} may not build it - please use a different version.",
+					String.Format("This version of the Unreal Engine can only be compiled with clang 15.0. clang {0} may not build it - please use a different version.",
 						Info.ClangVersion)
 					);
 			}
@@ -286,7 +286,7 @@ namespace UnrealBuildTool
 		{
 			// set UE_LINUX_USE_LIBCXX to either 0 or 1. If unset, defaults to 1.
 			string? UseLibcxxEnvVarOverride = Environment.GetEnvironmentVariable("UE_LINUX_USE_LIBCXX");
-			if (string.IsNullOrEmpty(UseLibcxxEnvVarOverride) || UseLibcxxEnvVarOverride == "1")
+			if (String.IsNullOrEmpty(UseLibcxxEnvVarOverride) || UseLibcxxEnvVarOverride == "1")
 			{
 				return true;
 			}
@@ -548,8 +548,8 @@ namespace UnrealBuildTool
 			string? Key = SplitData.ElementAtOrDefault(0);
 			string? Value = SplitData.ElementAtOrDefault(1);
 
-			if (string.IsNullOrEmpty(Key)) { return ""; }
-			if (!string.IsNullOrEmpty(Value))
+			if (String.IsNullOrEmpty(Key)) { return ""; }
+			if (!String.IsNullOrEmpty(Value))
 			{
 				if (!Value.StartsWith("\"") && (Value.Contains(" ") || Value.Contains("$")))
 				{
@@ -562,8 +562,8 @@ namespace UnrealBuildTool
 			}
 
 			return Value == null
-				? string.Format("{0}", Key)
-				: string.Format("{0}={1}", Key, Value);
+				? String.Format("{0}", Key)
+				: String.Format("{0}={1}", Key, Value);
 		}
 
 		protected virtual void GetLinkArguments(LinkEnvironment LinkEnvironment, List<string> Arguments)
@@ -633,12 +633,12 @@ namespace UnrealBuildTool
 					// x64 only replaced the linux folder with arch, while on arm64 its still linux
 					if (LinkEnvironment.Architecture == UnrealArch.Arm64)
 					{
-						Arguments.Add(string.Format("-Wl,-rpath=\"{0}/lib/clang/{1}.{2}.{3}/lib/linux\"",
+						Arguments.Add(String.Format("-Wl,-rpath=\"{0}/lib/clang/{1}.{2}.{3}/lib/linux\"",
 								LinuxInfo.BaseLinuxPath, Info.ClangVersion.Major, Info.ClangVersion.Minor, Info.ClangVersion.Build));
 					}
 					else
 					{
-						Arguments.Add(string.Format("-Wl,-rpath=\"{0}/lib/clang/{1}.{2}.{3}/lib/x86_64-unknown-linux-gnu\"",
+						Arguments.Add(String.Format("-Wl,-rpath=\"{0}/lib/clang/{1}.{2}.{3}/lib/x86_64-unknown-linux-gnu\"",
 								LinuxInfo.BaseLinuxPath, Info.ClangVersion.Major, Info.ClangVersion.Minor, Info.ClangVersion.Build));
 					}
 				}
@@ -710,7 +710,7 @@ namespace UnrealBuildTool
 				Arguments.Add("-Wno-backend-plugin");
 
 				Log.TraceInformationOnce("Enabling Profile Guided Optimization (PGO). Linking will take a while.");
-				Arguments.Add(string.Format("-fprofile-instr-use=\"{0}\"", Path.Combine(LinkEnvironment.PGODirectory!, LinkEnvironment.PGOFilenamePrefix!)));
+				Arguments.Add(String.Format("-fprofile-instr-use=\"{0}\"", Path.Combine(LinkEnvironment.PGODirectory!, LinkEnvironment.PGOFilenamePrefix!)));
 			}
 			else if (LinkEnvironment.bPGOProfile)
 			{
@@ -887,14 +887,14 @@ namespace UnrealBuildTool
 			ArchiveAction.ProducedItems.Add(OutputFile);
 			ArchiveAction.CommandDescription = "Archive";
 			ArchiveAction.StatusDescription = Path.GetFileName(OutputFile.AbsolutePath);
-			ArchiveAction.CommandArguments += string.Format("{1} \"{2}\"", GetArchiveArguments(LinkEnvironment), OutputFile.AbsolutePath);
+			ArchiveAction.CommandArguments += String.Format("{1} \"{2}\"", GetArchiveArguments(LinkEnvironment), OutputFile.AbsolutePath);
 
 			// Add the input files to a response file, and pass the response file on the command-line.
 			List<string> InputFileNames = new List<string>();
 			foreach (FileItem InputFile in LinkEnvironment.InputFiles)
 			{
 				string InputAbsolutePath = InputFile.AbsolutePath.Replace("\\", "/");
-				InputFileNames.Add(string.Format("\"{0}\"", InputAbsolutePath));
+				InputFileNames.Add(String.Format("\"{0}\"", InputAbsolutePath));
 				ArchiveAction.PrerequisiteItems.Add(InputFile);
 			}
 
@@ -905,7 +905,7 @@ namespace UnrealBuildTool
 				FileItem ResponseFileItem = Graph.CreateIntermediateTextFile(ResponsePath, InputFileNames);
 				ArchiveAction.PrerequisiteItems.Add(ResponseFileItem);
 			}
-			ArchiveAction.CommandArguments += string.Format(" @\"{0}\"", ResponsePath.FullName);
+			ArchiveAction.CommandArguments += String.Format(" @\"{0}\"", ResponsePath.FullName);
 
 			// Add the additional arguments specified by the environment.
 			ArchiveAction.CommandArguments += LinkEnvironment.AdditionalArguments;
@@ -948,7 +948,7 @@ namespace UnrealBuildTool
 				Action PostLinkAction = Graph.CreateAction(ActionType.Link);
 				PostLinkAction.WorkingDirectory = Unreal.EngineSourceDirectory;
 				PostLinkAction.CommandPath = ShellBinary;
-				PostLinkAction.StatusDescription = string.Format("{0}", Path.GetFileName(Executable.AbsolutePath));
+				PostLinkAction.StatusDescription = String.Format("{0}", Path.GetFileName(Executable.AbsolutePath));
 				PostLinkAction.CommandDescription = "FixDeps";
 				PostLinkAction.bCanExecuteRemotely = false;
 				PostLinkAction.CommandArguments = ExecuteSwitch;
@@ -1012,7 +1012,7 @@ namespace UnrealBuildTool
 			// Get link arguments.
 			List<string> LinkArguments = new List<string>();
 			GetLinkArguments(LinkEnvironment, LinkArguments);
-			LinkCommandString += " " + string.Join(' ', LinkArguments);
+			LinkCommandString += " " + String.Join(' ', LinkArguments);
 
 			// Tell the action that we're building an import library here and it should conditionally be
 			// ignored as a prerequisite for other actions
@@ -1044,19 +1044,19 @@ namespace UnrealBuildTool
 			LinkAction.StatusDescription = Path.GetFileName(OutputFile.AbsolutePath);
 
 			// Add the output file to the command-line.
-			LinkCommandString += string.Format(" -o \"{0}\"", OutputFile.AbsolutePath);
+			LinkCommandString += String.Format(" -o \"{0}\"", OutputFile.AbsolutePath);
 
 			// Add the input files to a response file, and pass the response file on the command-line.
 			List<string> ResponseLines = new List<string>();
 			foreach (FileItem InputFile in LinkEnvironment.InputFiles)
 			{
-				ResponseLines.Add(string.Format("\"{0}\"", InputFile.AbsolutePath.Replace("\\", "/")));
+				ResponseLines.Add(String.Format("\"{0}\"", InputFile.AbsolutePath.Replace("\\", "/")));
 				LinkAction.PrerequisiteItems.Add(InputFile);
 			}
 
 			if (LinkEnvironment.bIsBuildingDLL)
 			{
-				ResponseLines.Add(string.Format(" -soname=\"{0}\"", OutputFile.Location.GetFileName()));
+				ResponseLines.Add(String.Format(" -soname=\"{0}\"", OutputFile.Location.GetFileName()));
 			}
 
 			// Start with the configured LibraryPaths and also add paths to any libraries that
@@ -1105,7 +1105,7 @@ namespace UnrealBuildTool
 					if (!RPaths.Contains(RelativePath))
 					{
 						RPaths.Add(RelativePath);
-						ResponseLines.Add(string.Format(" -rpath=\"${{ORIGIN}}/{0}\"", RelativePath.Replace('\\', '/')));
+						ResponseLines.Add(String.Format(" -rpath=\"${{ORIGIN}}/{0}\"", RelativePath.Replace('\\', '/')));
 					}
 				}
 			}
@@ -1148,17 +1148,17 @@ namespace UnrealBuildTool
 				if (!RPaths.Contains(RelativePath))
 				{
 					RPaths.Add(RelativePath);
-					ResponseLines.Add(string.Format(" -rpath=\"${{ORIGIN}}/{0}\"", RelativePath.Replace('\\', '/')));
+					ResponseLines.Add(String.Format(" -rpath=\"${{ORIGIN}}/{0}\"", RelativePath.Replace('\\', '/')));
 				}
 			}
 
-			ResponseLines.Add(string.Format(" -rpath-link=\"{0}\"", Path.GetDirectoryName(OutputFile.AbsolutePath)));
+			ResponseLines.Add(String.Format(" -rpath-link=\"{0}\"", Path.GetDirectoryName(OutputFile.AbsolutePath)));
 
 			// Add the library paths to the argument list.
 			foreach (DirectoryReference LibraryPath in AllLibraryPaths)
 			{
 				// use absolute paths because of FixDependencies script again
-				ResponseLines.Add(string.Format(" -L\"{0}\"", LibraryPath.FullName.Replace('\\', '/')));
+				ResponseLines.Add(String.Format(" -L\"{0}\"", LibraryPath.FullName.Replace('\\', '/')));
 			}
 
 			List<string> EngineAndGameLibrariesLinkFlags = new List<string>();
@@ -1171,14 +1171,14 @@ namespace UnrealBuildTool
 			string ExternalLibraries = "";
 
 			// add libraries in a library group
-			ResponseLines.Add(string.Format(" --start-group"));
+			ResponseLines.Add(String.Format(" --start-group"));
 
 			foreach (string AdditionalLibrary in AdditionalLibraries)
 			{
 				if (String.IsNullOrEmpty(Path.GetDirectoryName(AdditionalLibrary)))
 				{
 					// library was passed just like "jemalloc", turn it into -ljemalloc
-					ExternalLibraries += string.Format(" -l{0}", AdditionalLibrary);
+					ExternalLibraries += String.Format(" -l{0}", AdditionalLibrary);
 				}
 				else if (Path.GetExtension(AdditionalLibrary) == ".a")
 				{
@@ -1186,7 +1186,7 @@ namespace UnrealBuildTool
 					string AbsoluteAdditionalLibrary = Path.GetFullPath(AdditionalLibrary);
 					if (AbsoluteAdditionalLibrary.Contains(" "))
 					{
-						AbsoluteAdditionalLibrary = string.Format("\"{0}\"", AbsoluteAdditionalLibrary);
+						AbsoluteAdditionalLibrary = String.Format("\"{0}\"", AbsoluteAdditionalLibrary);
 					}
 					AbsoluteAdditionalLibrary = AbsoluteAdditionalLibrary.Replace('\\', '/');
 
@@ -1219,9 +1219,9 @@ namespace UnrealBuildTool
 						// The library exists, but it is not prefixed with "lib", so force the
 						// linker to find it without that prefix by prepending a colon to
 						// the file name.
-						LibName = string.Format(":{0}", LibraryDependency.Name);
+						LibName = String.Format(":{0}", LibraryDependency.Name);
 					}
-					string LibLinkFlag = string.Format(" -l{0}", LibName);
+					string LibLinkFlag = String.Format(" -l{0}", LibName);
 
 					if (LinkEnvironment.bIsBuildingDLL && LinkEnvironment.bIsCrossReferenced)
 					{
@@ -1233,7 +1233,7 @@ namespace UnrealBuildTool
 						// it is important to add this exactly to the same place where the missing libraries would have been, it will be replaced later
 						if (!ExternalLibraries.Contains("--allow-shlib-undefined"))
 						{
-							ExternalLibraries += string.Format(" -Wl,--allow-shlib-undefined");
+							ExternalLibraries += String.Format(" -Wl,--allow-shlib-undefined");
 						}
 					}
 					else
@@ -1248,7 +1248,7 @@ namespace UnrealBuildTool
 			FileReference ResponseFileName = GetResponseFileName(LinkEnvironment, OutputFile);
 			FileItem ResponseFileItem = Graph.CreateIntermediateTextFile(ResponseFileName, ResponseLines);
 
-			LinkCommandString += string.Format(" -Wl,@\"{0}\"", ResponseFileName);
+			LinkCommandString += String.Format(" -Wl,@\"{0}\"", ResponseFileName);
 			LinkAction.PrerequisiteItems.Add(ResponseFileItem);
 
 			LinkCommandString += " -Wl,--start-group";
@@ -1305,7 +1305,7 @@ namespace UnrealBuildTool
 				LinkCommandString = LinkCommandString.Replace("$'{", "'${");    // fixing $'{ORIGIN}' to be '${ORIGIN}'
 			}
 
-			string LinkScriptName = string.Format((bUseCmdExe ? "Link-{0}.link.bat" : "Link-{0}.link.sh"), OutputFile.Location.GetFileName());
+			string LinkScriptName = String.Format((bUseCmdExe ? "Link-{0}.link.bat" : "Link-{0}.link.sh"), OutputFile.Location.GetFileName());
 			string LinkScriptFullPath = Path.Combine(LinkEnvironment.LocalShadowDirectory!.FullName, LinkScriptName);
 			Logger.LogDebug("Creating link script: {LinkScriptFullPath}", LinkScriptFullPath);
 			Directory.CreateDirectory(Path.GetDirectoryName(LinkScriptFullPath)!);
@@ -1465,7 +1465,7 @@ namespace UnrealBuildTool
 					}
 
 					// create the relinking step
-					string RelinkScriptName = string.Format((bUseCmdExe ? "Relink-{0}.bat" : "Relink-{0}.sh"), OutputFile.Location.GetFileName());
+					string RelinkScriptName = String.Format((bUseCmdExe ? "Relink-{0}.bat" : "Relink-{0}.sh"), OutputFile.Location.GetFileName());
 					string RelinkScriptFullPath = Path.Combine(LinkEnvironment.LocalShadowDirectory.FullName, RelinkScriptName);
 
 					Logger.LogDebug("Creating script: {RelinkScriptFullPath}", RelinkScriptFullPath);
diff --git a/Engine/Source/Programs/UnrealBuildTool/Platform/Linux/UEBuildLinux.cs b/Engine/Source/Programs/UnrealBuildTool/Platform/Linux/UEBuildLinux.cs
--- a/Engine/Source/Programs/UnrealBuildTool/Platform/Linux/UEBuildLinux.cs
+++ b/Engine/Source/Programs/UnrealBuildTool/Platform/Linux/UEBuildLinux.cs
@@ -24,7 +24,10 @@ namespace UnrealBuildTool
 		{
 			get
 			{
-				if (AppleToolchainArchitectures.ContainsKey(this)) return LinuxToolchainArchitectures[this];
+				if (AppleToolchainArchitectures.ContainsKey(this))
+				{
+					return LinuxToolchainArchitectures[this];
+				}
 
 				throw new BuildException($"Unknown architecture {ToString()} passed to UnrealArch.LinuxName");
 			}
@@ -194,11 +197,11 @@ namespace UnrealBuildTool
 
 		public override void ValidateTarget(TargetRules Target)
 		{
-			if (!string.IsNullOrWhiteSpace(Environment.GetEnvironmentVariable("CLANG_STATIC_ANALYZER_MODE")))
+			if (!String.IsNullOrWhiteSpace(Environment.GetEnvironmentVariable("CLANG_STATIC_ANALYZER_MODE")))
 			{
 				Target.StaticAnalyzer = StaticAnalyzer.Default;
 				Target.StaticAnalyzerOutputType = (Environment.GetEnvironmentVariable("CLANG_ANALYZER_OUTPUT")?.Contains("html", StringComparison.OrdinalIgnoreCase) == true) ? StaticAnalyzerOutputType.Html : StaticAnalyzerOutputType.Text;
-				Target.StaticAnalyzerMode = string.Equals(Environment.GetEnvironmentVariable("CLANG_STATIC_ANALYZER_MODE"), "shallow") ? StaticAnalyzerMode.Shallow : StaticAnalyzerMode.Deep;
+				Target.StaticAnalyzerMode = String.Equals(Environment.GetEnvironmentVariable("CLANG_STATIC_ANALYZER_MODE"), "shallow") ? StaticAnalyzerMode.Shallow : StaticAnalyzerMode.Deep;
 			}
 			else if (Target.StaticAnalyzer == StaticAnalyzer.Clang)
 			{
diff --git a/Engine/Source/Programs/UnrealBuildTool/Platform/Mac/MacToolChain.cs b/Engine/Source/Programs/UnrealBuildTool/Platform/Mac/MacToolChain.cs
--- a/Engine/Source/Programs/UnrealBuildTool/Platform/Mac/MacToolChain.cs
+++ b/Engine/Source/Programs/UnrealBuildTool/Platform/Mac/MacToolChain.cs
@@ -57,7 +57,7 @@ namespace UnrealBuildTool
 			{
 				throw new BuildException("Unable to find installed MacOS SDK on remote agent.");
 			}
-			else if (!float.TryParse(MacOSSDKVersion, NumberStyles.AllowDecimalPoint, CultureInfo.InvariantCulture.NumberFormat, out MacOSSDKVersionFloat))
+			else if (!Single.TryParse(MacOSSDKVersion, NumberStyles.AllowDecimalPoint, CultureInfo.InvariantCulture.NumberFormat, out MacOSSDKVersionFloat))
 			{
 				throw new BuildException("Unable to parse installed MacOS version (\"{0}\")", MacOSSDKVersion);
 			}
@@ -120,9 +120,9 @@ namespace UnrealBuildTool
 		{
 			DirectoryReference? ProgramFinder = StartingDir;
 			while (ProgramFinder != null &&
-				string.Compare(ProgramFinder.GetDirectoryName(), "Source", true) != 0 &&
-				string.Compare(ProgramFinder.GetDirectoryName(), "Intermediate", true) != 0 &&
-				string.Compare(ProgramFinder.GetDirectoryName(), "Binaries", true) != 0)
+				String.Compare(ProgramFinder.GetDirectoryName(), "Source", true) != 0 &&
+				String.Compare(ProgramFinder.GetDirectoryName(), "Intermediate", true) != 0 &&
+				String.Compare(ProgramFinder.GetDirectoryName(), "Binaries", true) != 0)
 			{
 				ProgramFinder = ProgramFinder.ParentDirectory;
 			}
@@ -195,15 +195,15 @@ namespace UnrealBuildTool
 		/// <inheritdoc/>
 		protected override void GetCompileArguments_AdditionalArgs(CppCompileEnvironment CompileEnvironment, List<string> Arguments)
 		{
-			if (!string.IsNullOrWhiteSpace(CompileEnvironment.AdditionalArguments))
+			if (!String.IsNullOrWhiteSpace(CompileEnvironment.AdditionalArguments))
 			{
-				string EscapedAdditionalArgs = string.Empty;
+				string EscapedAdditionalArgs = String.Empty;
 				foreach (string AdditionalArg in CompileEnvironment.AdditionalArguments.Split(new char[] { ' ' }, StringSplitOptions.RemoveEmptyEntries))
 				{
 					Match DefinitionMatch = Regex.Match(AdditionalArg, "-D\"?(?<Name>.*)=(?<Value>.*)\"?");
 					if (DefinitionMatch.Success)
 					{
-						EscapedAdditionalArgs += string.Format(" -D{0}=\"{1}\"", DefinitionMatch.Groups["Name"].Value, DefinitionMatch.Groups["Value"].Value);
+						EscapedAdditionalArgs += String.Format(" -D{0}=\"{1}\"", DefinitionMatch.Groups["Name"].Value, DefinitionMatch.Groups["Value"].Value);
 					}
 					else
 					{
@@ -211,7 +211,7 @@ namespace UnrealBuildTool
 					}
 				}
 
-				if (!string.IsNullOrWhiteSpace(EscapedAdditionalArgs))
+				if (!String.IsNullOrWhiteSpace(EscapedAdditionalArgs))
 				{
 					Arguments.Add(EscapedAdditionalArgs);
 				}
@@ -269,7 +269,7 @@ namespace UnrealBuildTool
 		{
 			// Pass through architecture and OS info		
 			Arguments.Add(FormatArchitectureArg(LinkEnvironment.Architectures));
-			Arguments.Add(string.Format("-isysroot \"{0}\"", SDKPath));
+			Arguments.Add(String.Format("-isysroot \"{0}\"", SDKPath));
 			Arguments.Add("-mmacosx-version-min=" + Settings.MacOSVersion);
 			Arguments.Add("-dead_strip");
 
@@ -541,7 +541,7 @@ namespace UnrealBuildTool
 				LipoAction.CommandPath = FileReference.Combine(Info.Clang.Directory, "lipo");
 				LipoAction.CommandDescription = (bBuildImportLibraryOnly ? "LipoStub" : "Lipo");
 				LipoAction.CommandVersion = Info.ClangVersionString;
-				LipoAction.CommandArguments = string.Join(" ", PerArchOutputFiles) + $" -create -output {OutputFileItem.AbsolutePath}";
+				LipoAction.CommandArguments = String.Join(" ", PerArchOutputFiles) + $" -create -output {OutputFileItem.AbsolutePath}";
 				LipoAction.StatusDescription = Path.GetFileName(OutputFileItem.AbsolutePath);
 				LipoAction.bCanExecuteRemotely = false;
 			}
@@ -578,7 +578,7 @@ namespace UnrealBuildTool
 
 				// bundle identifier
 				// plist replacements
-				DirectoryReference? DirRef = (!string.IsNullOrEmpty(UnrealBuildTool.GetRemoteIniPath()) ? new DirectoryReference(UnrealBuildTool.GetRemoteIniPath()!) : (ProjectFile != null ? ProjectFile.Directory : null));
+				DirectoryReference? DirRef = (!String.IsNullOrEmpty(UnrealBuildTool.GetRemoteIniPath()) ? new DirectoryReference(UnrealBuildTool.GetRemoteIniPath()!) : (ProjectFile != null ? ProjectFile.Directory : null));
 				ConfigHierarchy IOSIni = ConfigCache.ReadHierarchy(ConfigHierarchyType.Engine, DirRef, UnrealTargetPlatform.IOS);
 
 				string BundleIdentifier;
@@ -743,7 +743,7 @@ namespace UnrealBuildTool
 				GetLinkArguments_Global(LinkEnvironment, LinkArguments);
 			}
 
-			LinkCommand = VersionArg + " " + string.Join(' ', LinkArguments);
+			LinkCommand = VersionArg + " " + String.Join(' ', LinkArguments);
 
 			// Tell the action that we're building an import library here and it should conditionally be
 			// ignored as a prerequisite for other actions
@@ -807,15 +807,15 @@ namespace UnrealBuildTool
 					}
 					else if (AdditionalLibrary.Contains(".framework/"))
 					{
-						LinkCommand += string.Format(" \"{0}\"", AdditionalLibrary);
+						LinkCommand += String.Format(" \"{0}\"", AdditionalLibrary);
 					}
-					else if (string.IsNullOrEmpty(Path.GetDirectoryName(AdditionalLibrary)) && string.IsNullOrEmpty(Path.GetExtension(AdditionalLibrary)))
+					else if (String.IsNullOrEmpty(Path.GetDirectoryName(AdditionalLibrary)) && String.IsNullOrEmpty(Path.GetExtension(AdditionalLibrary)))
 					{
-						LinkCommand += string.Format(" -l\"{0}\"", AdditionalLibrary);
+						LinkCommand += String.Format(" -l\"{0}\"", AdditionalLibrary);
 					}
 					else
 					{
-						LinkCommand += string.Format(" \"{0}\"", Path.GetFullPath(AdditionalLibrary));
+						LinkCommand += String.Format(" \"{0}\"", Path.GetFullPath(AdditionalLibrary));
 					}
 
 					AddLibraryPathToRPaths(AdditionalLibrary, AbsolutePath, ref RPaths, ref LinkCommand, bIsBuildingAppBundle, Logger);
@@ -823,7 +823,7 @@ namespace UnrealBuildTool
 
 				foreach (string AdditionalLibrary in LinkEnvironment.DelayLoadDLLs)
 				{
-					LinkCommand += string.Format(" -weak_library \"{0}\"", Path.GetFullPath(AdditionalLibrary));
+					LinkCommand += String.Format(" -weak_library \"{0}\"", Path.GetFullPath(AdditionalLibrary));
 
 					AddLibraryPathToRPaths(AdditionalLibrary, AbsolutePath, ref RPaths, ref LinkCommand, bIsBuildingAppBundle, Logger);
 				}
@@ -877,7 +877,7 @@ namespace UnrealBuildTool
 
 					FileReference MapFilePath = new FileReference(MapFileBaseName + ".map");
 					FileItem MapFile = FileItem.GetItemByFileReference(MapFilePath);
-					LinkCommand += string.Format(" -Wl,-map,\"{0}\"", MapFilePath);
+					LinkCommand += String.Format(" -Wl,-map,\"{0}\"", MapFilePath);
 					LinkAction.ProducedItems.Add(MapFile);
 				}
 			}
@@ -888,15 +888,15 @@ namespace UnrealBuildTool
 				string? InstallName = LinkEnvironment.InstallName;
 				if (InstallName == null)
 				{
-					InstallName = string.Format("{0}/{1}", DylibsPath, Path.GetFileName(OutputFile.AbsolutePath).Replace($".dylib_{LinkEnvironment.Architecture}", ".dylib"));
+					InstallName = String.Format("{0}/{1}", DylibsPath, Path.GetFileName(OutputFile.AbsolutePath).Replace($".dylib_{LinkEnvironment.Architecture}", ".dylib"));
 				}
-				LinkCommand += string.Format(" -install_name \"{0}\"", InstallName);
+				LinkCommand += String.Format(" -install_name \"{0}\"", InstallName);
 			}
 
 			List<string> InputFileNames = new List<string>();
 			foreach (FileItem InputFile in InputFiles)
 			{
-				InputFileNames.Add(string.Format("\"{0}\"", NormalizeCommandLinePath(InputFile)));
+				InputFileNames.Add(String.Format("\"{0}\"", NormalizeCommandLinePath(InputFile)));
 				LinkAction.PrerequisiteItems.Add(InputFile);
 			}
 
@@ -906,7 +906,7 @@ namespace UnrealBuildTool
 			}
 
 			// Add the output file to the command-line.
-			LinkCommand += string.Format(" -o \"{0}\"", OutputFile.AbsolutePath);
+			LinkCommand += String.Format(" -o \"{0}\"", OutputFile.AbsolutePath);
 
 			// Add the additional arguments specified by the environment.
 			LinkCommand += LinkEnvironment.AdditionalArguments;
@@ -975,35 +975,35 @@ namespace UnrealBuildTool
 			string UniversalDsymutilScriptPath = FileReference.Combine(Unreal.EngineDirectory, "Build/BatchFiles/Mac/GenerateUniversalDSYM.sh").FullName;
 
 			string ArgumentString = "-c \"";
-			ArgumentString += string.Format("for i in {{1..30}}; ");
-			ArgumentString += string.Format("do if [ -f \\\"{0}\\\" ] ; ", MachOBinary.AbsolutePath);
-			ArgumentString += string.Format("then ");
-			ArgumentString += string.Format("break; ");
-			ArgumentString += string.Format("else ");
-			ArgumentString += string.Format("sleep 1; ");
-			ArgumentString += string.Format("fi; ");
-			ArgumentString += string.Format("done; ");
-
-			ArgumentString += string.Format("if [ ! -f \\\"{1}\\\" ] || [ \\\"{0}\\\" -nt \\\"{1}\\\" ] ; ", MachOBinary.AbsolutePath, OutputFile.AbsolutePath);
-			ArgumentString += string.Format("then ");
-			ArgumentString += string.Format("rm -rf \\\"{0}\\\"; ", OutputFile.AbsolutePath);
+			ArgumentString += String.Format("for i in {{1..30}}; ");
+			ArgumentString += String.Format("do if [ -f \\\"{0}\\\" ] ; ", MachOBinary.AbsolutePath);
+			ArgumentString += String.Format("then ");
+			ArgumentString += String.Format("break; ");
+			ArgumentString += String.Format("else ");
+			ArgumentString += String.Format("sleep 1; ");
+			ArgumentString += String.Format("fi; ");
+			ArgumentString += String.Format("done; ");
+
+			ArgumentString += String.Format("if [ ! -f \\\"{1}\\\" ] || [ \\\"{0}\\\" -nt \\\"{1}\\\" ] ; ", MachOBinary.AbsolutePath, OutputFile.AbsolutePath);
+			ArgumentString += String.Format("then ");
+			ArgumentString += String.Format("rm -rf \\\"{0}\\\"; ", OutputFile.AbsolutePath);
 			// use the new script for monolthic (ie large) targets
 			if (LinkEnvironment.LinkType == TargetLinkType.Monolithic)
 			{
-				ArgumentString += string.Format(" \\\"{0}\\\" \\\"{1}\\\" \\\"{2}\\\"; ",
+				ArgumentString += String.Format(" \\\"{0}\\\" \\\"{1}\\\" \\\"{2}\\\"; ",
 					UniversalDsymutilScriptPath,
 					MachOBinary.AbsolutePath,
 					OutputFile.AbsolutePath);
 			}
 			else
 			{
-				ArgumentString += string.Format(" \\\"{0}\\\" {3} -f \\\"{1}\\\" -o \\\"{2}\\\"; ",
+				ArgumentString += String.Format(" \\\"{0}\\\" {3} -f \\\"{1}\\\" -o \\\"{2}\\\"; ",
 					DsymutilPath,
 					MachOBinary.AbsolutePath,
 					OutputFile.AbsolutePath,
 					ExtraOptions);
 			}
-			ArgumentString += string.Format("fi; ");
+			ArgumentString += String.Format("fi; ");
 			ArgumentString += "\"";
 
 			GenDebugAction.CommandArguments = ArgumentString;
@@ -1042,7 +1042,7 @@ namespace UnrealBuildTool
 			FinalizeAppBundleAction.CommandArguments = "\"" + BundleScript.AbsolutePath + "\"";
 			FinalizeAppBundleAction.PrerequisiteItems.Add(Executable);
 			FinalizeAppBundleAction.ProducedItems.Add(DestFile);
-			FinalizeAppBundleAction.StatusDescription = string.Format("Finalizing app bundle: {0}.app", Path.GetFileName(Executable.AbsolutePath));
+			FinalizeAppBundleAction.StatusDescription = String.Format("Finalizing app bundle: {0}.app", Path.GetFileName(Executable.AbsolutePath));
 			FinalizeAppBundleAction.bCanExecuteRemotely = false;
 
 			return DestFile;
@@ -1061,11 +1061,11 @@ namespace UnrealBuildTool
 
 			FileItem TargetItem = FileItem.GetItemByPath(TargetPath);
 
-			CopyAction.CommandArguments = string.Format("-c \"cp -f -R \\\"{0}\\\" \\\"{1}\\\"; touch -c \\\"{2}\\\"\"", SourcePath, Path.GetDirectoryName(TargetPath)!.Replace('\\', '/') + "/", TargetPath.Replace('\\', '/'));
+			CopyAction.CommandArguments = String.Format("-c \"cp -f -R \\\"{0}\\\" \\\"{1}\\\"; touch -c \\\"{2}\\\"\"", SourcePath, Path.GetDirectoryName(TargetPath)!.Replace('\\', '/') + "/", TargetPath.Replace('\\', '/'));
 			CopyAction.PrerequisiteItems.Add(Executable);
 			CopyAction.ProducedItems.Add(TargetItem);
 			CopyAction.bShouldOutputStatusDescription = Resource.bShouldLog;
-			CopyAction.StatusDescription = string.Format("Copying {0} to app bundle", Path.GetFileName(Resource.ResourcePath));
+			CopyAction.StatusDescription = String.Format("Copying {0} to app bundle", Path.GetFileName(Resource.ResourcePath));
 			CopyAction.bCanExecuteRemotely = false;
 
 			return TargetItem;
@@ -1092,7 +1092,7 @@ namespace UnrealBuildTool
 							break;
 					}
 					string? DSYMExtension = Array.Find(DebugExtensions, element => element == ".dSYM");
-					if (!string.IsNullOrEmpty(DSYMExtension))
+					if (!String.IsNullOrEmpty(DSYMExtension))
 					{
 						string BinaryPath = BuildProductPair.Key.FullName;
 						if (BinaryPath.Contains(".app"))
@@ -1191,7 +1191,7 @@ namespace UnrealBuildTool
 			List<string> OutputFiles = Makefile.OutputItems.Select(Item => Path.ChangeExtension(Item.FullName, ".dSYM")).Distinct().ToList();
 			foreach (FileItem DebugItem in DebugInfoFiles)
 			{
-				if (OutputFiles.Any(Item => string.Equals(Item, DebugItem.FullName, StringComparison.InvariantCultureIgnoreCase)))
+				if (OutputFiles.Any(Item => String.Equals(Item, DebugItem.FullName, StringComparison.InvariantCultureIgnoreCase)))
 				{
 					Makefile.OutputItems.Add(DebugItem);
 				}
diff --git a/Engine/Source/Programs/UnrealBuildTool/Platform/Mac/UEBuildMac.cs b/Engine/Source/Programs/UnrealBuildTool/Platform/Mac/UEBuildMac.cs
--- a/Engine/Source/Programs/UnrealBuildTool/Platform/Mac/UEBuildMac.cs
+++ b/Engine/Source/Programs/UnrealBuildTool/Platform/Mac/UEBuildMac.cs
@@ -264,11 +264,11 @@ namespace UnrealBuildTool
 
 		public override void ValidateTarget(TargetRules Target)
 		{
-			if (!string.IsNullOrWhiteSpace(Environment.GetEnvironmentVariable("CLANG_STATIC_ANALYZER_MODE")))
+			if (!String.IsNullOrWhiteSpace(Environment.GetEnvironmentVariable("CLANG_STATIC_ANALYZER_MODE")))
 			{
 				Target.StaticAnalyzer = StaticAnalyzer.Default;
 				Target.StaticAnalyzerOutputType = (Environment.GetEnvironmentVariable("CLANG_ANALYZER_OUTPUT")?.Contains("html", StringComparison.OrdinalIgnoreCase) == true) ? StaticAnalyzerOutputType.Html : StaticAnalyzerOutputType.Text;
-				Target.StaticAnalyzerMode = string.Equals(Environment.GetEnvironmentVariable("CLANG_STATIC_ANALYZER_MODE"), "shallow") ? StaticAnalyzerMode.Shallow : StaticAnalyzerMode.Deep;
+				Target.StaticAnalyzerMode = String.Equals(Environment.GetEnvironmentVariable("CLANG_STATIC_ANALYZER_MODE"), "shallow") ? StaticAnalyzerMode.Shallow : StaticAnalyzerMode.Deep;
 			}
 			else if (Target.StaticAnalyzer == StaticAnalyzer.Clang)
 			{
diff --git a/Engine/Source/Programs/UnrealBuildTool/Platform/Mac/UEDeployMac.cs b/Engine/Source/Programs/UnrealBuildTool/Platform/Mac/UEDeployMac.cs
--- a/Engine/Source/Programs/UnrealBuildTool/Platform/Mac/UEDeployMac.cs
+++ b/Engine/Source/Programs/UnrealBuildTool/Platform/Mac/UEDeployMac.cs
@@ -41,7 +41,7 @@ namespace UnrealBuildTool
 
 			// bundle identifier
 			// plist replacements
-			DirectoryReference? DirRef = bIsUnrealGame ? (!string.IsNullOrEmpty(UnrealBuildTool.GetRemoteIniPath()) ? new DirectoryReference(UnrealBuildTool.GetRemoteIniPath()!) : null) : new DirectoryReference(ProjectDirectory);
+			DirectoryReference? DirRef = bIsUnrealGame ? (!System.String.IsNullOrEmpty(UnrealBuildTool.GetRemoteIniPath()) ? new DirectoryReference(UnrealBuildTool.GetRemoteIniPath()!) : null) : new DirectoryReference(ProjectDirectory);
 			ConfigHierarchy Ini = ConfigCache.ReadHierarchy(ConfigHierarchyType.Engine, DirRef, UnrealTargetPlatform.IOS);
 
 			string BundleIdentifier;
diff --git a/Engine/Source/Programs/UnrealBuildTool/Platform/TVOS/UEDeployTVOS.cs b/Engine/Source/Programs/UnrealBuildTool/Platform/TVOS/UEDeployTVOS.cs
--- a/Engine/Source/Programs/UnrealBuildTool/Platform/TVOS/UEDeployTVOS.cs
+++ b/Engine/Source/Programs/UnrealBuildTool/Platform/TVOS/UEDeployTVOS.cs
@@ -43,7 +43,7 @@ namespace UnrealBuildTool
 			// get the settings from the ini file
 			// plist replacements
 			// @todo tvos: Are we going to make TVOS specific .ini files?
-			DirectoryReference? DirRef = bIsUnrealGame ? (!string.IsNullOrEmpty(UnrealBuildTool.GetRemoteIniPath()) ? new DirectoryReference(UnrealBuildTool.GetRemoteIniPath()!) : null) : new DirectoryReference(ProjectDirectory);
+			DirectoryReference? DirRef = bIsUnrealGame ? (!String.IsNullOrEmpty(UnrealBuildTool.GetRemoteIniPath()) ? new DirectoryReference(UnrealBuildTool.GetRemoteIniPath()!) : null) : new DirectoryReference(ProjectDirectory);
 			ConfigHierarchy Ini = ConfigCache.ReadHierarchy(ConfigHierarchyType.Engine, DirRef, UnrealTargetPlatform.IOS);
 
 			// bundle display name
@@ -53,7 +53,7 @@ namespace UnrealBuildTool
 			// bundle identifier
 			string BundleIdentifier;
 			Ini.GetString("/Script/IOSRuntimeSettings.IOSRuntimeSettings", "BundleIdentifier", out BundleIdentifier);
-			if (!string.IsNullOrEmpty(BundleID))
+			if (!String.IsNullOrEmpty(BundleID))
 			{
 				BundleIdentifier = BundleID;
 			}
@@ -95,33 +95,33 @@ namespace UnrealBuildTool
 			Text.AppendLine("\t<key>CFBundleDevelopmentRegion</key>");
 			Text.AppendLine("\t<string>en</string>");
 			Text.AppendLine("\t<key>CFBundleDisplayName</key>");
-			Text.AppendLine(string.Format("\t<string>{0}</string>", EncodeBundleName(BundleDisplayName, ProjectName)));
+			Text.AppendLine(String.Format("\t<string>{0}</string>", EncodeBundleName(BundleDisplayName, ProjectName)));
 			Text.AppendLine("\t<key>CFBundleExecutable</key>");
 			string BundleExecutable = bIsUnrealGame ?
 				(bIsClient ? "UnrealClient" : "UnrealGame") :
 				(bIsClient ? GameName + "Client" : GameName);
-			Text.AppendLine(string.Format("\t<string>{0}</string>", BundleExecutable));
+			Text.AppendLine(String.Format("\t<string>{0}</string>", BundleExecutable));
 			Text.AppendLine("\t<key>CFBundleIdentifier</key>");
-			Text.AppendLine(string.Format("\t<string>{0}</string>", BundleIdentifier.Replace("[PROJECT_NAME]", ProjectName).Replace("_", "")));
+			Text.AppendLine(String.Format("\t<string>{0}</string>", BundleIdentifier.Replace("[PROJECT_NAME]", ProjectName).Replace("_", "")));
 			Text.AppendLine("\t<key>CFBundleInfoDictionaryVersion</key>");
 			Text.AppendLine("\t<string>6.0</string>");
 			Text.AppendLine("\t<key>CFBundleName</key>");
-			Text.AppendLine(string.Format("\t<string>{0}</string>", EncodeBundleName(BundleName, ProjectName)));
+			Text.AppendLine(String.Format("\t<string>{0}</string>", EncodeBundleName(BundleName, ProjectName)));
 			Text.AppendLine("\t<key>CFBundlePackageType</key>");
 			Text.AppendLine("\t<string>APPL</string>");
 			Text.AppendLine("\t<key>CFBundleSignature</key>");
 			Text.AppendLine("\t<string>????</string>");
 			Text.AppendLine("\t<key>CFBundleVersion</key>");
-			Text.AppendLine(string.Format("\t<string>{0}</string>", VersionUtilities.UpdateBundleVersion(OldPListData, InEngineDir)));
+			Text.AppendLine(String.Format("\t<string>{0}</string>", VersionUtilities.UpdateBundleVersion(OldPListData, InEngineDir)));
 			Text.AppendLine("\t<key>CFBundleShortVersionString</key>");
-			Text.AppendLine(string.Format("\t<string>{0}</string>", BundleShortVersion));
+			Text.AppendLine(String.Format("\t<string>{0}</string>", BundleShortVersion));
 			Text.AppendLine("\t<key>LSRequiresIPhoneOS</key>");
 			Text.AppendLine("\t<true/>");
 			Text.AppendLine("\t<key>UIRequiredDeviceCapabilities</key>");
 			Text.AppendLine("\t<array>");
 			foreach (string Line in RequiredCaps.Split("\r\n".ToCharArray()))
 			{
-				if (!string.IsNullOrWhiteSpace(Line))
+				if (!String.IsNullOrWhiteSpace(Line))
 				{
 					Text.AppendLine(Line);
 				}
@@ -142,7 +142,7 @@ namespace UnrealBuildTool
 			Text.AppendLine("\t<string>LaunchScreen</string>");
 
 			// write the iCloud container identifier, if present in the old file
-			if (!string.IsNullOrEmpty(OldPListData))
+			if (!String.IsNullOrEmpty(OldPListData))
 			{
 				int index = OldPListData.IndexOf("ICloudContainerIdentifier");
 				if (index > 0)
@@ -151,7 +151,7 @@ namespace UnrealBuildTool
 					int length = OldPListData.IndexOf("</string>", index) - index;
 					string ICloudContainerIdentifier = OldPListData.Substring(index, length);
 					Text.AppendLine("\t<key>ICloudContainerIdentifier</key>");
-					Text.AppendLine(string.Format("\t<string>{0}</string>", ICloudContainerIdentifier));
+					Text.AppendLine(String.Format("\t<string>{0}</string>", ICloudContainerIdentifier));
 				}
 			}
 
diff --git a/Engine/Source/Programs/UnrealBuildTool/Platform/Windows/AppXManifestGeneratorBase.cs b/Engine/Source/Programs/UnrealBuildTool/Platform/Windows/AppXManifestGeneratorBase.cs
--- a/Engine/Source/Programs/UnrealBuildTool/Platform/Windows/AppXManifestGeneratorBase.cs
+++ b/Engine/Source/Programs/UnrealBuildTool/Platform/Windows/AppXManifestGeneratorBase.cs
@@ -15,10 +15,10 @@ namespace UnrealBuildTool
 	/// <summary>
 	/// Base class for VC appx manifest generation
 	/// </summary>
-	abstract public class AppXManifestGeneratorBase
+	public abstract class AppXManifestGeneratorBase
 	{
 		/// config section for platform-specific target settings
-		protected virtual string IniSection_PlatformTargetSettings => string.Format("/Script/{0}PlatformEditor.{0}TargetSettings", Platform.ToString());
+		protected virtual string IniSection_PlatformTargetSettings => String.Format("/Script/{0}PlatformEditor.{0}TargetSettings", Platform.ToString());
 
 		/// config section for general target settings
 		protected virtual string IniSection_GeneralProjectSettings => "/Script/EngineSettings.GeneralProjectSettings";
@@ -84,7 +84,7 @@ namespace UnrealBuildTool
 				{
 					string QuadElement = SplitVersionString[VersionElement];
 					int QuadValue = 0;
-					if (QuadElement.Length == 0 || !int.TryParse(QuadElement, out QuadValue))
+					if (QuadElement.Length == 0 || !Int32.TryParse(QuadElement, out QuadValue))
 					{
 						CompletedVersionString += "0";
 					}
@@ -147,14 +147,20 @@ namespace UnrealBuildTool
 		protected string? ReadIniString(string? Key, string Section, string? DefaultValue = null)
 		{
 			if (Key == null)
+			{
 				return DefaultValue;
+			}
 
 			string Value;
-			if (GameIni!.GetString(Section, Key, out Value) && !string.IsNullOrWhiteSpace(Value))
+			if (GameIni!.GetString(Section, Key, out Value) && !String.IsNullOrWhiteSpace(Value))
+			{
 				return Value;
+			}
 
-			if (EngineIni!.GetString(Section, Key, out Value) && !string.IsNullOrWhiteSpace(Value))
+			if (EngineIni!.GetString(Section, Key, out Value) && !String.IsNullOrWhiteSpace(Value))
+			{
 				return Value;
+			}
 
 			return DefaultValue;
 		}
@@ -178,7 +184,9 @@ namespace UnrealBuildTool
 			var ResultStr = ReadIniString(PlatformKey, IniSection_PlatformTargetSettings, GenericValue);
 
 			if (ResultStr == null)
+			{
 				return DefaultValue;
+			}
 
 			ResultStr = ResultStr.Trim().ToLower();
 
@@ -192,14 +200,16 @@ namespace UnrealBuildTool
 		{
 			var ConfigValue = GetConfigString(PlatformConfigKey, null, null);
 			if (ConfigValue == null)
+			{
 				return DefaultValue;
+			}
 
 			Dictionary<string, string>? Pairs;
 			int R, G, B;
 			if (ConfigHierarchy.TryParse(ConfigValue, out Pairs) &&
-				int.TryParse(Pairs["R"], out R) &&
-				int.TryParse(Pairs["G"], out G) &&
-				int.TryParse(Pairs["B"], out B))
+				Int32.TryParse(Pairs["R"], out R) &&
+				Int32.TryParse(Pairs["G"], out G) &&
+				Int32.TryParse(Pairs["B"], out B))
 			{
 				return "#" + R.ToString("X2") + G.ToString("X2") + B.ToString("X2");
 			}
@@ -243,14 +253,14 @@ namespace UnrealBuildTool
 					}
 
 					string StageId = CultureProperties["StageId"];
-					if (string.IsNullOrEmpty(StageId))
+					if (String.IsNullOrEmpty(StageId))
 					{
 						Logger.LogWarning("Missing StageId value: {Culture}", CultureResources);
 						continue;
 					}
 
 					string CultureId = CultureProperties["CultureId"];
-					if (string.IsNullOrEmpty(CultureId))
+					if (String.IsNullOrEmpty(CultureId))
 					{
 						Logger.LogWarning("Missing CultureId value: {Culture}", CultureResources);
 						continue;
@@ -323,7 +333,7 @@ namespace UnrealBuildTool
 			// Read the PackageName from config
 			var DefaultName = (ProjectFile != null) ? ProjectFile.GetFileNameWithoutAnyExtensions() : (TargetName ?? "DefaultUEProject");
 			var PackageName = Regex.Replace(GetConfigString("PackageName", "ProjectName", DefaultName), "[^-.A-Za-z0-9]", "");
-			if (string.IsNullOrWhiteSpace(PackageName))
+			if (String.IsNullOrWhiteSpace(PackageName))
 			{
 				Logger.LogError("Invalid package name {Name}. Package names must only contain letters, numbers, dash, and period and must be at least one character long.", PackageName);
 				Logger.LogError("Consider using the setting [{IniSection}]:PackageName to provide a specific value.", IniSection_PlatformTargetSettings);
@@ -394,7 +404,7 @@ namespace UnrealBuildTool
 			{
 				// Break apart the version number into individual elements
 				string[] SplitVersionString = VersionNumber.Split('.');
-				VersionNumber = string.Format("{0}.{1}.{2}.{3}",
+				VersionNumber = String.Format("{0}.{1}.{2}.{3}",
 					SplitVersionString[0],
 					SplitVersionString[1],
 					BuildVersionForPackage.Changelist / 10000,
@@ -494,7 +504,7 @@ namespace UnrealBuildTool
 			// generate the list of AppX cultures to stage
 			foreach (string UEStageId in SelectedUECultureIds)
 			{
-				if (!UEStageIdToAppXCultureId.TryGetValue(UEStageId, out string? AppXCultureId) || string.IsNullOrEmpty(AppXCultureId))
+				if (!UEStageIdToAppXCultureId.TryGetValue(UEStageId, out string? AppXCultureId) || String.IsNullOrEmpty(AppXCultureId))
 				{
 					// use the culture directly - no remapping required
 					AppXCultureId = UEStageId;
@@ -505,7 +515,7 @@ namespace UnrealBuildTool
 			}
 
 			// look up the default AppX culture
-			if (!UEStageIdToAppXCultureId.TryGetValue(DefaultUECultureId, out DefaultAppXCultureId) || string.IsNullOrEmpty(DefaultAppXCultureId))
+			if (!UEStageIdToAppXCultureId.TryGetValue(DefaultUECultureId, out DefaultAppXCultureId) || String.IsNullOrEmpty(DefaultAppXCultureId))
 			{
 				// use the default culture directly - no remapping required
 				DefaultAppXCultureId = DefaultUECultureId;
diff --git a/Engine/Source/Programs/UnrealBuildTool/Platform/Windows/MicrosoftPlatformSDK.cs b/Engine/Source/Programs/UnrealBuildTool/Platform/Windows/MicrosoftPlatformSDK.cs
--- a/Engine/Source/Programs/UnrealBuildTool/Platform/Windows/MicrosoftPlatformSDK.cs
+++ b/Engine/Source/Programs/UnrealBuildTool/Platform/Windows/MicrosoftPlatformSDK.cs
@@ -192,7 +192,7 @@ namespace UnrealBuildTool
 
 			// Figure out which version number to look for
 			VersionNumber? WindowsSdkVersion = null;
-			if (!string.IsNullOrEmpty(DesiredVersion))
+			if (!String.IsNullOrEmpty(DesiredVersion))
 			{
 				if (String.Compare(DesiredVersion, "Latest", StringComparison.InvariantCultureIgnoreCase) == 0 && CachedWindowsSdkDirs!.Count > 0)
 				{
@@ -536,7 +536,7 @@ namespace UnrealBuildTool
 
 			// Figure out the actual version number that we want
 			ToolChainInstallation? ToolChain = null;
-			if (string.IsNullOrEmpty(CompilerVersion))
+			if (String.IsNullOrEmpty(CompilerVersion))
 			{
 				ToolChain = SelectToolChain(ToolChains, x => x, Architecture);
 				if (ToolChain == null)
@@ -785,14 +785,14 @@ namespace UnrealBuildTool
 						DirectoryReference? PlatformDir;
 						if (UEBuildPlatformSDK.TryGetHostPlatformAutoSDKDir(out PlatformDir))
 						{
-							string VSDir = string.Empty;
+							string VSDir = String.Empty;
 							switch (Compiler)
 							{
 								case WindowsCompiler.VisualStudio2019: VSDir = "VS2019"; break;
 								case WindowsCompiler.VisualStudio2022: VSDir = "VS2022"; break;
 							}
 
-							if (!string.IsNullOrEmpty(VSDir))
+							if (!String.IsNullOrEmpty(VSDir))
 							{
 								DirectoryReference ReleaseBaseDir = DirectoryReference.Combine(PlatformDir, "Win64", VSDir);
 								FindVisualStudioToolChains(ReleaseBaseDir, null, WindowsCompilerChannel.Latest, ToolChains, IsAutoSdk: true, Logger);
@@ -1336,14 +1336,14 @@ namespace UnrealBuildTool
 				DirectoryReference? PlatformDir;
 				if (UEBuildPlatformSDK.TryGetHostPlatformAutoSDKDir(out PlatformDir))
 				{
-					string VSDir = string.Empty;
+					string VSDir = String.Empty;
 					switch (Compiler)
 					{
 						case WindowsCompiler.VisualStudio2019: VSDir = "VS2019"; break;
 						case WindowsCompiler.VisualStudio2022: VSDir = "VS2022"; break;
 					}
 
-					if (!string.IsNullOrEmpty(VSDir))
+					if (!String.IsNullOrEmpty(VSDir))
 					{
 						DirectoryReference DiaSdkDir = DirectoryReference.Combine(PlatformDir, "Win64", "DIA SDK", VSDir);
 						if (IsValidDiaSdkDir(DiaSdkDir))
diff --git a/Engine/Source/Programs/UnrealBuildTool/Platform/Windows/PVSToolChain.cs b/Engine/Source/Programs/UnrealBuildTool/Platform/Windows/PVSToolChain.cs
--- a/Engine/Source/Programs/UnrealBuildTool/Platform/Windows/PVSToolChain.cs
+++ b/Engine/Source/Programs/UnrealBuildTool/Platform/Windows/PVSToolChain.cs
@@ -429,7 +429,7 @@ namespace UnrealBuildTool
 								int LineNumber;
 								bool bFalseAlarm;
 								int Level;
-								if (int.TryParse(LineNumberStr, out LineNumber) && bool.TryParse(FalseAlarmStr, out bFalseAlarm) && int.TryParse(LevelStr, out Level))
+								if (Int32.TryParse(LineNumberStr, out LineNumber) && Boolean.TryParse(FalseAlarmStr, out bFalseAlarm) && Int32.TryParse(LevelStr, out Level))
 								{
 									bCanParse = true;
 
@@ -563,7 +563,9 @@ namespace UnrealBuildTool
 			catch (Exception Ex)
 			{
 				if (Ex is BuildException)
+				{
 					throw;
+				}
 
 				throw new BuildException(Ex, "Failed to obtain PVS-Studio version.");
 			}
@@ -627,10 +629,14 @@ namespace UnrealBuildTool
 
 			bool compileAsCCode;
 			if (CFlagLastPosition == CppFlagLastPosition)
+			{
 				//ни один флаг, определяющий язык, не задан. Определяем по расширению файла
 				compileAsCCode = Path.GetExtension(sourceFileName).Equals(".c", StringComparison.InvariantCultureIgnoreCase);
+			}
 			else
+			{
 				compileAsCCode = CFlagLastPosition > CppFlagLastPosition;
+			}
 
 			return compileAsCCode;
 		}
diff --git a/Engine/Source/Programs/UnrealBuildTool/Platform/Windows/ResXResource/UEAppXResources.cs b/Engine/Source/Programs/UnrealBuildTool/Platform/Windows/ResXResource/UEAppXResources.cs
--- a/Engine/Source/Programs/UnrealBuildTool/Platform/Windows/ResXResource/UEAppXResources.cs
+++ b/Engine/Source/Programs/UnrealBuildTool/Platform/Windows/ResXResource/UEAppXResources.cs
@@ -447,7 +447,7 @@ namespace UnrealBuildTool
 			{
 				DefaultCultureString = Value;
 			}
-			if (string.IsNullOrEmpty(DefaultCultureString))
+			if (String.IsNullOrEmpty(DefaultCultureString))
 			{
 				DefaultCultureString = DefaultValue;
 			}
@@ -458,7 +458,7 @@ namespace UnrealBuildTool
 			foreach (string CultureId in GetAllCultureIds())
 			{
 				Resources ThisCultureResources = GetCultureResources(CultureId);
-				if (ThisCultureResources.SourceStrings.TryGetValue(ConfigKey, out string? CultureString) && !string.IsNullOrEmpty(CultureString))
+				if (ThisCultureResources.SourceStrings.TryGetValue(ConfigKey, out string? CultureString) && !String.IsNullOrEmpty(CultureString))
 				{
 					ThisCultureResources.StringResources.Add(ResourceEntryName, CultureString + ValueSuffix);
 				}
diff --git a/Engine/Source/Programs/UnrealBuildTool/Platform/Windows/ResXResource/UEResXWriter.cs b/Engine/Source/Programs/UnrealBuildTool/Platform/Windows/ResXResource/UEResXWriter.cs
--- a/Engine/Source/Programs/UnrealBuildTool/Platform/Windows/ResXResource/UEResXWriter.cs
+++ b/Engine/Source/Programs/UnrealBuildTool/Platform/Windows/ResXResource/UEResXWriter.cs
@@ -56,7 +56,7 @@ namespace UnrealBuildTool
 				Data.Attributes.Append(Attr);
 			}
 
-			if (!string.IsNullOrEmpty(InSpace))
+			if (!System.String.IsNullOrEmpty(InSpace))
 			{
 				XmlAttribute Attr = Document.CreateAttribute("xml:space");
 				Attr.Value = "preserve";
diff --git a/Engine/Source/Programs/UnrealBuildTool/Platform/Windows/UEBuildWindows.cs b/Engine/Source/Programs/UnrealBuildTool/Platform/Windows/UEBuildWindows.cs
--- a/Engine/Source/Programs/UnrealBuildTool/Platform/Windows/UEBuildWindows.cs
+++ b/Engine/Source/Programs/UnrealBuildTool/Platform/Windows/UEBuildWindows.cs
@@ -48,7 +48,11 @@ namespace UnrealBuildTool
 		{
 			get
 			{
-				if (WindowsToolchainArchitectures.ContainsKey(this)) return WindowsToolchainArchitectures[this].ToolChain;
+				if (WindowsToolchainArchitectures.ContainsKey(this))
+				{
+					return WindowsToolchainArchitectures[this].ToolChain;
+				}
+
 				throw new BuildException($"Unknown architecture {ToString()} passed to UnrealArch.WindowsToolChain");
 			}
 		}
@@ -60,7 +64,11 @@ namespace UnrealBuildTool
 		{
 			get
 			{
-				if (WindowsToolchainArchitectures.ContainsKey(this)) return WindowsToolchainArchitectures[this].LibDir;
+				if (WindowsToolchainArchitectures.ContainsKey(this))
+				{
+					return WindowsToolchainArchitectures[this].LibDir;
+				}
+
 				throw new BuildException($"Unknown architecture {ToString()} passed to UnrealArch.WindowsLibDir");
 			}
 		}
@@ -72,7 +80,11 @@ namespace UnrealBuildTool
 		{
 			get
 			{
-				if (WindowsToolchainArchitectures.ContainsKey(this)) return WindowsToolchainArchitectures[this].SystemLibDir;
+				if (WindowsToolchainArchitectures.ContainsKey(this))
+				{
+					return WindowsToolchainArchitectures[this].SystemLibDir;
+				}
+
 				throw new BuildException($"Unknown architecture {ToString()} passed to UnrealArch.WindowsSystemLibDir");
 			}
 		}
@@ -84,7 +96,11 @@ namespace UnrealBuildTool
 		{
 			get
 			{
-				if (WindowsToolchainArchitectures.ContainsKey(this)) return WindowsToolchainArchitectures[this].ToolChain;
+				if (WindowsToolchainArchitectures.ContainsKey(this))
+				{
+					return WindowsToolchainArchitectures[this].ToolChain;
+				}
+
 				throw new BuildException($"Unknown architecture {ToString()} passed to UnrealArch.WindowsName");
 			}
 		}
@@ -596,7 +612,7 @@ namespace UnrealBuildTool
 		public ReadOnlyWindowsTargetRules(WindowsTargetRules Inner)
 		{
 			this.Inner = Inner;
-			this.PVS = new ReadOnlyPVSTargetSettings(Inner.PVS);
+			PVS = new ReadOnlyPVSTargetSettings(Inner.PVS);
 		}
 
 		/// <summary>
@@ -1447,7 +1463,7 @@ namespace UnrealBuildTool
 				string PGOPlatform = Target.Platform == UnrealTargetPlatform.Win64 ? "Windows" : Target.Platform.ToString();
 
 				CompileEnvironment.PGODirectory = DirectoryReference.Combine(Target.ProjectFile.Directory, "Platforms", PGOPlatform, "Build", "PGO").FullName;
-				CompileEnvironment.PGOFilenamePrefix = string.Format("{0}-{1}-{2}", Target.Name, Target.Platform, Target.Configuration);
+				CompileEnvironment.PGOFilenamePrefix = String.Format("{0}-{1}-{2}", Target.Name, Target.Platform, Target.Configuration);
 
 				LinkEnvironment.PGODirectory = CompileEnvironment.PGODirectory;
 				LinkEnvironment.PGOFilenamePrefix = CompileEnvironment.PGOFilenamePrefix;
diff --git a/Engine/Source/Programs/UnrealBuildTool/Platform/Windows/VCCompileAction.cs b/Engine/Source/Programs/UnrealBuildTool/Platform/Windows/VCCompileAction.cs
--- a/Engine/Source/Programs/UnrealBuildTool/Platform/Windows/VCCompileAction.cs
+++ b/Engine/Source/Programs/UnrealBuildTool/Platform/Windows/VCCompileAction.cs
@@ -316,9 +316,9 @@ namespace UnrealBuildTool
 		/// <param name="Environment">Compiler executable</param>
 		public VCCompileAction(VCEnvironment Environment)
 		{
-			this.CompilerExe = FileItem.GetItemByFileReference(Environment.CompilerPath);
-			this.CompilerType = Environment.Compiler;
-			this.ToolChainVersion = Environment.ToolChainVersion.ToString();
+			CompilerExe = FileItem.GetItemByFileReference(Environment.CompilerPath);
+			CompilerType = Environment.Compiler;
+			ToolChainVersion = Environment.ToolChainVersion.ToString();
 		}
 
 		/// <summary>
diff --git a/Engine/Source/Programs/UnrealBuildTool/Platform/Windows/VCEnvironment.cs b/Engine/Source/Programs/UnrealBuildTool/Platform/Windows/VCEnvironment.cs
--- a/Engine/Source/Programs/UnrealBuildTool/Platform/Windows/VCEnvironment.cs
+++ b/Engine/Source/Programs/UnrealBuildTool/Platform/Windows/VCEnvironment.cs
@@ -118,18 +118,18 @@ namespace UnrealBuildTool
 		[SupportedOSPlatform("windows")]
 		public VCEnvironment(VCEnvironmentParameters Params, ILogger Logger)
 		{
-			this.Compiler = Params.Compiler;
-			this.CompilerDir = Params.CompilerDir;
-			this.CompilerVersion = Params.CompilerVersion;
-			this.Architecture = Params.Architecture;
-			this.ToolChain = Params.ToolChain;
-			this.ToolChainDir = Params.ToolChainDir;
-			this.ToolChainVersion = Params.ToolChainVersion;
-			this.WindowsSdkDir = Params.WindowsSdkDir;
-			this.WindowsSdkVersion = Params.WindowsSdkVersion;
-			this.RedistDir = Params.RedistDir;
-			this.bUseCPPWinRT = Params.bUseCPPWinRT;
-			this.bAllowClangLinker = Params.bAllowClangLinker;
+			Compiler = Params.Compiler;
+			CompilerDir = Params.CompilerDir;
+			CompilerVersion = Params.CompilerVersion;
+			Architecture = Params.Architecture;
+			ToolChain = Params.ToolChain;
+			ToolChainDir = Params.ToolChainDir;
+			ToolChainVersion = Params.ToolChainVersion;
+			WindowsSdkDir = Params.WindowsSdkDir;
+			WindowsSdkVersion = Params.WindowsSdkVersion;
+			RedistDir = Params.RedistDir;
+			bUseCPPWinRT = Params.bUseCPPWinRT;
+			bAllowClangLinker = Params.bAllowClangLinker;
 
 			// Get the compiler and linker paths from the Toolchain directory
 			CompilerPath = GetCompilerToolPath(Compiler, Architecture, CompilerDir);
@@ -263,7 +263,7 @@ namespace UnrealBuildTool
 		/// <summary>
 		/// Gets the path to the resource compiler.
 		/// </summary>
-		virtual protected FileReference GetResourceCompilerToolPath(DirectoryReference WindowsSdkDir, VersionNumber WindowsSdkVersion, ILogger Logger)
+		protected virtual FileReference GetResourceCompilerToolPath(DirectoryReference WindowsSdkDir, VersionNumber WindowsSdkVersion, ILogger Logger)
 		{
 			FileReference ResourceCompilerPath = FileReference.Combine(WindowsSdkDir, "bin", WindowsSdkVersion.ToString(), "x64", "rc.exe");
 			if (FileReference.Exists(ResourceCompilerPath))
@@ -521,15 +521,15 @@ namespace UnrealBuildTool
 			// Store the final parameters
 			this.Platform = Platform;
 			this.Compiler = Compiler;
-			this.CompilerDir = SelectedCompilerDir;
+			CompilerDir = SelectedCompilerDir;
 			this.CompilerVersion = SelectedCompilerVersion;
 			this.Architecture = Architecture;
 			this.ToolChain = ToolChain;
-			this.ToolChainDir = SelectedToolChainDir;
-			this.ToolChainVersion = SelectedToolChainVersion;
-			this.WindowsSdkDir = SelectedWindowsSdkDir;
+			ToolChainDir = SelectedToolChainDir;
+			ToolChainVersion = SelectedToolChainVersion;
+			WindowsSdkDir = SelectedWindowsSdkDir;
 			this.WindowsSdkVersion = SelectedWindowsSdkVersion;
-			this.RedistDir = SelectedRedistDir;
+			RedistDir = SelectedRedistDir;
 			this.bUseCPPWinRT = bUseCPPWinRT;
 			this.bAllowClangLinker = bAllowClangLinker;
 		}
diff --git a/Engine/Source/Programs/UnrealBuildTool/Platform/Windows/VCToolChain.cs b/Engine/Source/Programs/UnrealBuildTool/Platform/Windows/VCToolChain.cs
--- a/Engine/Source/Programs/UnrealBuildTool/Platform/Windows/VCToolChain.cs
+++ b/Engine/Source/Programs/UnrealBuildTool/Platform/Windows/VCToolChain.cs
@@ -34,7 +34,7 @@ namespace UnrealBuildTool
 			: base(Logger)
 		{
 			this.Target = Target;
-			this.EnvVars = Target.WindowsPlatform.Environment!;
+			EnvVars = Target.WindowsPlatform.Environment!;
 
 			Logger.LogDebug("Compiler: {Path}", EnvVars.CompilerPath);
 			Logger.LogDebug("Linker: {Path}", EnvVars.LinkerPath);
@@ -105,7 +105,7 @@ namespace UnrealBuildTool
 			ExternalDependencies.Add(FileItem.GetItemByFileReference(EnvVars.LinkerPath));
 		}
 
-		static public void AddDefinition(List<string> Arguments, string Definition)
+		public static void AddDefinition(List<string> Arguments, string Definition)
 		{
 			// Split the definition into name and value
 			int ValueIdx = Definition.IndexOf('=');
@@ -119,7 +119,7 @@ namespace UnrealBuildTool
 			}
 		}
 
-		static public void AddDefinition(List<string> Arguments, string Variable, string? Value)
+		public static void AddDefinition(List<string> Arguments, string Variable, string? Value)
 		{
 			// If the value has a space in it and isn't wrapped in quotes, do that now
 			if (Value != null && !Value.StartsWith("\"") && (Value.Contains(" ") || Value.Contains("$")))
@@ -327,7 +327,10 @@ namespace UnrealBuildTool
 					Arguments.Add("-Xclang -analyzer-config -Xclang path-diagnostics-alternate=true");
 
 					// Run shallow analyze if requested.
-					if (Target.StaticAnalyzerMode == StaticAnalyzerMode.Shallow) Arguments.Add("-Xclang -analyzer-config -Xclang mode=shallow");
+					if (Target.StaticAnalyzerMode == StaticAnalyzerMode.Shallow)
+					{
+						Arguments.Add("-Xclang -analyzer-config -Xclang mode=shallow");
+					}
 
 					if (CompileEnvironment.StaticAnalyzerCheckers.Count > 0)
 					{
@@ -1253,7 +1256,7 @@ namespace UnrealBuildTool
 			}
 
 			// Add any extra options from the target
-			if (!string.IsNullOrEmpty(Target.WindowsPlatform.AdditionalLinkerOptions))
+			if (!String.IsNullOrEmpty(Target.WindowsPlatform.AdditionalLinkerOptions))
 			{
 				Arguments.Add(Target.WindowsPlatform.AdditionalLinkerOptions);
 			}
@@ -1724,7 +1727,7 @@ namespace UnrealBuildTool
 							$"-CompileTimingFile={ExpectedCompileTimeFile}",
 						};
 
-						Action AggregateTimingInfoAction = MakefileBuilder.CreateRecursiveAction<AggregateParsedTimingInfo>(ActionType.ParseTimingInfo, string.Join(" ", ActionArgs));
+						Action AggregateTimingInfoAction = MakefileBuilder.CreateRecursiveAction<AggregateParsedTimingInfo>(ActionType.ParseTimingInfo, String.Join(" ", ActionArgs));
 						AggregateTimingInfoAction.WorkingDirectory = Unreal.EngineSourceDirectory;
 						AggregateTimingInfoAction.StatusDescription = $"Aggregating {TimingJsonFiles.Count} Timing File(s)";
 						AggregateTimingInfoAction.bCanExecuteRemotely = false;
@@ -1749,7 +1752,7 @@ namespace UnrealBuildTool
 							$"-HeadersFile={HeadersOutputFile.FullName}",
 						};
 
-						Action AggregateTimingInfoAction = MakefileBuilder.CreateRecursiveAction<AggregateClangTimingInfo>(ActionType.ParseTimingInfo, string.Join(" ", AggregateActionArgs));
+						Action AggregateTimingInfoAction = MakefileBuilder.CreateRecursiveAction<AggregateClangTimingInfo>(ActionType.ParseTimingInfo, String.Join(" ", AggregateActionArgs));
 						AggregateTimingInfoAction.WorkingDirectory = Unreal.EngineSourceDirectory;
 						AggregateTimingInfoAction.StatusDescription = $"Aggregating {TimingJsonFiles.Count} Timing File(s)";
 						AggregateTimingInfoAction.bCanExecuteRemotely = false;
@@ -1767,7 +1770,7 @@ namespace UnrealBuildTool
 							$"-ArchiveFile={ArchiveOutputFile.FullName}",
 						};
 
-						Action ArchiveTimingInfoAction = MakefileBuilder.CreateRecursiveAction<AggregateClangTimingInfo>(ActionType.ParseTimingInfo, string.Join(" ", ArchiveActionArgs));
+						Action ArchiveTimingInfoAction = MakefileBuilder.CreateRecursiveAction<AggregateClangTimingInfo>(ActionType.ParseTimingInfo, String.Join(" ", ArchiveActionArgs));
 						ArchiveTimingInfoAction.WorkingDirectory = Unreal.EngineSourceDirectory;
 						ArchiveTimingInfoAction.StatusDescription = $"Archiving {TimingJsonFiles.Count} Timing File(s)";
 						ArchiveTimingInfoAction.bCanExecuteRemotely = false;
@@ -2124,7 +2127,9 @@ namespace UnrealBuildTool
 		{
 			// This is not supported for now.. If someone wants it we can implement it
 			if (CompileEnvironment.Architectures.bIsMultiArch)
+			{
 				return;
+			}
 
 			VCCompileAction BaseCompileAction = CreateBaseCompileAction(CompileEnvironment);
 			AppendCLArguments_CPP(CompileEnvironment, BaseCompileAction.Arguments);
@@ -2326,7 +2331,7 @@ namespace UnrealBuildTool
 			List<string> InputFileNames = new List<string>();
 			foreach (FileItem InputFile in LinkEnvironment.InputFiles)
 			{
-				InputFileNames.Add(string.Format("\"{0}\"", NormalizeCommandLinePath(InputFile)));
+				InputFileNames.Add(String.Format("\"{0}\"", NormalizeCommandLinePath(InputFile)));
 				PrerequisiteItems.Add(InputFile);
 			}
 
@@ -2334,18 +2339,18 @@ namespace UnrealBuildTool
 			{
 				foreach (FileReference Library in LinkEnvironment.Libraries)
 				{
-					InputFileNames.Add(string.Format("\"{0}\"", NormalizeCommandLinePath(Library)));
+					InputFileNames.Add(String.Format("\"{0}\"", NormalizeCommandLinePath(Library)));
 					PrerequisiteItems.Add(FileItem.GetItemByFileReference(Library));
 				}
 				foreach (string SystemLibrary in LinkEnvironment.SystemLibraries)
 				{
-					InputFileNames.Add(string.Format("\"{0}\"", SystemLibrary));
+					InputFileNames.Add(String.Format("\"{0}\"", SystemLibrary));
 				}
 
 				foreach (FileItem NatvisFile in LinkEnvironment.DebuggerVisualizerFiles)
 				{
 					PrerequisiteItems.Add(NatvisFile);
-					Arguments.Add(string.Format("/NATVIS:\"{0}\"",
+					Arguments.Add(String.Format("/NATVIS:\"{0}\"",
 						NormalizeCommandLinePath(NatvisFile)));
 				}
 
@@ -2532,7 +2537,7 @@ namespace UnrealBuildTool
 				int PGCFileIndex = 0;
 				foreach (string SrcFilePath in PGCFiles)
 				{
-					string DestFileName = string.Format("{0}!{1}.pgc", LinkEnvironment.PGOFilenamePrefix, ++PGCFileIndex);
+					string DestFileName = String.Format("{0}!{1}.pgc", LinkEnvironment.PGOFilenamePrefix, ++PGCFileIndex);
 					string DestFilePath = Path.Combine(LinkEnvironment.OutputDirectory.FullName, DestFileName);
 
 					Logger.LogInformation("{Source} -> {Target}", SrcFilePath, DestFilePath);
diff --git a/Engine/Source/Programs/UnrealBuildTool/ProjectFiles/CMake/CMakefileGenerator.cs b/Engine/Source/Programs/UnrealBuildTool/ProjectFiles/CMake/CMakefileGenerator.cs
--- a/Engine/Source/Programs/UnrealBuildTool/ProjectFiles/CMake/CMakefileGenerator.cs
+++ b/Engine/Source/Programs/UnrealBuildTool/ProjectFiles/CMake/CMakefileGenerator.cs
@@ -47,7 +47,7 @@ namespace UnrealBuildTool
 		/// This determines if engine files are included in the source lists.
 		/// </remarks>
 		/// <returns><value>true</value> if we should treat this as a project build; <value>false</value> otherwise.</returns>
-		public bool IsProjectBuild => !string.IsNullOrEmpty(GameProjectName);
+		public bool IsProjectBuild => !String.IsNullOrEmpty(GameProjectName);
 
 		/// <summary>
 		/// The file extension for this project file.
diff --git a/Engine/Source/Programs/UnrealBuildTool/ProjectFiles/CodeLite/CodeLiteGenerator.cs b/Engine/Source/Programs/UnrealBuildTool/ProjectFiles/CodeLite/CodeLiteGenerator.cs
--- a/Engine/Source/Programs/UnrealBuildTool/ProjectFiles/CodeLite/CodeLiteGenerator.cs
+++ b/Engine/Source/Programs/UnrealBuildTool/ProjectFiles/CodeLite/CodeLiteGenerator.cs
@@ -40,7 +40,7 @@ namespace UnrealBuildTool
 		//
 		// Returns CodeLite's project filename extension.
 		//
-		override public string ProjectFileExtension => ".project";
+		public override string ProjectFileExtension => ".project";
 		protected override bool WritePrimaryProjectFile(ProjectFile? UBTProject, PlatformProjectGeneratorCollection PlatformProjectGenerators, ILogger Logger)
 		{
 			string SolutionFileName = PrimaryProjectName + SolutionExtension;
diff --git a/Engine/Source/Programs/UnrealBuildTool/ProjectFiles/Eddie/EddieGenerator.cs b/Engine/Source/Programs/UnrealBuildTool/ProjectFiles/Eddie/EddieGenerator.cs
--- a/Engine/Source/Programs/UnrealBuildTool/ProjectFiles/Eddie/EddieGenerator.cs
+++ b/Engine/Source/Programs/UnrealBuildTool/ProjectFiles/Eddie/EddieGenerator.cs
@@ -16,7 +16,7 @@ namespace UnrealBuildTool
 		{
 		}
 
-		override public string ProjectFileExtension => ".wkst";
+		public override string ProjectFileExtension => ".wkst";
 
 		public override void CleanProjectFiles(DirectoryReference InPrimaryProjectDirectory, string InPrimaryProjectName, DirectoryReference InIntermediateProjectFilesPath, ILogger Logger)
 		{
diff --git a/Engine/Source/Programs/UnrealBuildTool/ProjectFiles/Eddie/EddieProject.cs b/Engine/Source/Programs/UnrealBuildTool/ProjectFiles/Eddie/EddieProject.cs
--- a/Engine/Source/Programs/UnrealBuildTool/ProjectFiles/Eddie/EddieProject.cs
+++ b/Engine/Source/Programs/UnrealBuildTool/ProjectFiles/Eddie/EddieProject.cs
@@ -207,7 +207,7 @@ namespace UnrealBuildTool
 			StringBuilder ProjectFileContent = new StringBuilder();
 
 			ProjectFileContent.Append("# @Eddie Workset@" + ProjectFileGenerator.NewLine);
-			ProjectFileContent.Append("AddWorkset \"" + this.ToString() + ".wkst\" \"" + ProjectFilePath.FullName + "\"" + ProjectFileGenerator.NewLine);
+			ProjectFileContent.Append("AddWorkset \"" + ToString() + ".wkst\" \"" + ProjectFilePath.FullName + "\"" + ProjectFileGenerator.NewLine);
 
 			ParseSourceFilesIntoGroups();
 			EmitProject(ProjectFileContent, Folders, Logger);
