UnrealBuildTool: More automated code cleanup

diff --git a/Engine/Source/Programs/UnrealBuildTool/ProjectFiles/KDevelop/KDevelopGenerator.cs b/Engine/Source/Programs/UnrealBuildTool/ProjectFiles/KDevelop/KDevelopGenerator.cs
--- a/Engine/Source/Programs/UnrealBuildTool/ProjectFiles/KDevelop/KDevelopGenerator.cs
+++ b/Engine/Source/Programs/UnrealBuildTool/ProjectFiles/KDevelop/KDevelopGenerator.cs
@@ -30,7 +30,7 @@ namespace UnrealBuildTool
 		}
 
 		/// File extension for project files we'll be generating (e.g. ".vcxproj")
-		override public string ProjectFileExtension => ".kdev4";
+		public override string ProjectFileExtension => ".kdev4";
 
 		protected override bool WritePrimaryProjectFile(ProjectFile? UBTProject, PlatformProjectGeneratorCollection PlatformProjectGenerators, ILogger Logger)
 		{
@@ -383,7 +383,7 @@ namespace UnrealBuildTool
 				foreach (string CurDefine in KDevelopProject.IntelliSensePreprocessorDefinitions)
 				{
 					SplitDefinitionAndValue(CurDefine, out Key, out Value);
-					if (string.IsNullOrEmpty(Value))
+					if (String.IsNullOrEmpty(Value))
 					{
 						DefineHolder.Add(String.Format("{0} \\\n", Key));
 					}
diff --git a/Engine/Source/Programs/UnrealBuildTool/ProjectFiles/Make/MakefileGenerator.cs b/Engine/Source/Programs/UnrealBuildTool/ProjectFiles/Make/MakefileGenerator.cs
--- a/Engine/Source/Programs/UnrealBuildTool/ProjectFiles/Make/MakefileGenerator.cs
+++ b/Engine/Source/Programs/UnrealBuildTool/ProjectFiles/Make/MakefileGenerator.cs
@@ -33,13 +33,13 @@ namespace UnrealBuildTool
 		}
 
 		/// True if we should include IntelliSense data in the generated project files when possible
-		override public bool ShouldGenerateIntelliSenseData()
+		public override bool ShouldGenerateIntelliSenseData()
 		{
 			return bGenerateIntelliSenseData;
 		}
 
 		/// File extension for project files we'll be generating (e.g. ".vcxproj")
-		override public string ProjectFileExtension => ".mk";
+		public override string ProjectFileExtension => ".mk";
 
 		protected override bool WritePrimaryProjectFile(ProjectFile? UBTProject, PlatformProjectGeneratorCollection PlatformProjectGenerators, ILogger Logger)
 		{
diff --git a/Engine/Source/Programs/UnrealBuildTool/ProjectFiles/PlatformProjectGenerator.cs b/Engine/Source/Programs/UnrealBuildTool/ProjectFiles/PlatformProjectGenerator.cs
--- a/Engine/Source/Programs/UnrealBuildTool/ProjectFiles/PlatformProjectGenerator.cs
+++ b/Engine/Source/Programs/UnrealBuildTool/ProjectFiles/PlatformProjectGenerator.cs
@@ -13,7 +13,7 @@ namespace UnrealBuildTool
 	/// </summary>
 	abstract class PlatformProjectGenerator
 	{
-		static public readonly string DefaultPlatformConfigurationType = "Makefile";
+		public static readonly string DefaultPlatformConfigurationType = "Makefile";
 		protected readonly ILogger Logger;
 
 		/// <summary>
diff --git a/Engine/Source/Programs/UnrealBuildTool/ProjectFiles/ProjectFileGenerator.cs b/Engine/Source/Programs/UnrealBuildTool/ProjectFiles/ProjectFileGenerator.cs
--- a/Engine/Source/Programs/UnrealBuildTool/ProjectFiles/ProjectFileGenerator.cs
+++ b/Engine/Source/Programs/UnrealBuildTool/ProjectFiles/ProjectFileGenerator.cs
@@ -1453,7 +1453,9 @@ namespace UnrealBuildTool
 							bAppendPlatformSuffix = true;
 						}
 					}
-					else switch (CurArgument.ToUpperInvariant())
+					else
+					{
+						switch (CurArgument.ToUpperInvariant())
 						{
 							case "-ALLPLATFORMS":
 								IncludeAllPlatforms = true;
@@ -1538,6 +1540,7 @@ namespace UnrealBuildTool
 								bIncludeTempTargets = true;
 								break;
 						}
+					}
 				}
 			}
 
@@ -2291,7 +2294,7 @@ namespace UnrealBuildTool
 		/// </summary>
 		/// <param name="InPlatform"></param>
 		/// <returns>true if valid, false if not</returns>
-		static public bool IsValidDesktopPlatform(UnrealTargetPlatform InPlatform)
+		public static bool IsValidDesktopPlatform(UnrealTargetPlatform InPlatform)
 		{
 			if (BuildHostPlatform.Current.Platform == UnrealTargetPlatform.Linux)
 			{
@@ -3274,7 +3277,7 @@ namespace UnrealBuildTool
 				catch (Exception ex)
 				{
 					// Unable to write to the project file.
-					string Message = string.Format("Error while trying to write file {0}.  The file is probably read-only.", FileName);
+					string Message = String.Format("Error while trying to write file {0}.  The file is probably read-only.", FileName);
 					Logger.LogInformation("");
 					Logger.LogError("{Message}", Message);
 					throw new BuildException(ex, Message);
@@ -3544,18 +3547,31 @@ namespace UnrealBuildTool
 				if (Char.IsWhiteSpace(A) && Char.IsWhiteSpace(B) && CanIgnoreWhitespace(A) && CanIgnoreWhitespace(B))
 				{
 					// Skip whitespaces in both strings
-					for (IndexA++; IndexA < StrA.Length && Char.IsWhiteSpace(StrA[IndexA]) == true; IndexA++) ;
-					for (IndexB++; IndexB < StrB.Length && Char.IsWhiteSpace(StrB[IndexB]) == true; IndexB++) ;
+					for (IndexA++; IndexA < StrA.Length && Char.IsWhiteSpace(StrA[IndexA]) == true; IndexA++)
+					{
+						;
+					}
+
+					for (IndexB++; IndexB < StrB.Length && Char.IsWhiteSpace(StrB[IndexB]) == true; IndexB++)
+					{
+						;
+					}
 				}
 				else if (Char.IsWhiteSpace(A) && IndexA > 0 && StrA[IndexA - 1] == '\n')
 				{
 					// Skip whitespaces in StrA at the beginning of each line
-					for (IndexA++; IndexA < StrA.Length && Char.IsWhiteSpace(StrA[IndexA]) == true; IndexA++) ;
+					for (IndexA++; IndexA < StrA.Length && Char.IsWhiteSpace(StrA[IndexA]) == true; IndexA++)
+					{
+						;
+					}
 				}
 				else if (Char.IsWhiteSpace(B) && IndexB > 0 && StrB[IndexB - 1] == '\n')
 				{
 					// Skip whitespaces in StrA at the beginning of each line
-					for (IndexB++; IndexB < StrB.Length && Char.IsWhiteSpace(StrB[IndexB]) == true; IndexB++) ;
+					for (IndexB++; IndexB < StrB.Length && Char.IsWhiteSpace(StrB[IndexB]) == true; IndexB++)
+					{
+						;
+					}
 				}
 				else if (A != B)
 				{
diff --git a/Engine/Source/Programs/UnrealBuildTool/ProjectFiles/QMake/QMakefileGenerator.cs b/Engine/Source/Programs/UnrealBuildTool/ProjectFiles/QMake/QMakefileGenerator.cs
--- a/Engine/Source/Programs/UnrealBuildTool/ProjectFiles/QMake/QMakefileGenerator.cs
+++ b/Engine/Source/Programs/UnrealBuildTool/ProjectFiles/QMake/QMakefileGenerator.cs
@@ -31,7 +31,7 @@ namespace UnrealBuildTool
 		}
 
 		/// File extension for project files we'll be generating (e.g. ".vcxproj")
-		override public string ProjectFileExtension => ".pro";
+		public override string ProjectFileExtension => ".pro";
 
 		protected override bool WritePrimaryProjectFile(ProjectFile? UBTProject, PlatformProjectGeneratorCollection PlatformProjectGenerators, ILogger Logger)
 		{
@@ -177,7 +177,7 @@ namespace UnrealBuildTool
 					if (!DefinesOnly.Contains(define))
 					{
 						// Logger.LogInformation (CurDefine);
-						if (string.IsNullOrEmpty(value))
+						if (String.IsNullOrEmpty(value))
 						{
 							DefinesAndValues.Add("\t");
 							DefinesAndValues.Add(String.Format("{0}=", define));
@@ -201,10 +201,16 @@ namespace UnrealBuildTool
 							DefinesAndValues.Add(define);
 							DefinesAndValues.Add("=");
 							if (hasSpaces)
+							{
 								DefinesAndValues.Add("\"");
+							}
+
 							DefinesAndValues.Add(value);
 							if (hasSpaces)
+							{
 								DefinesAndValues.Add("\"");
+							}
+
 							DefinesAndValues.Add(" \\\n");
 						}
 						DefinesOnly.Add(define);
@@ -245,7 +251,7 @@ namespace UnrealBuildTool
 			}
 			else
 			{
-				BuildCommand = string.Format("build=bash $$unrealRootPath{0}Engine{0}Build{0}BatchFiles{0}{1}{0}Build.sh\n", Seperator, CurrentPlatform);
+				BuildCommand = String.Format("build=bash $$unrealRootPath{0}Engine{0}Build{0}BatchFiles{0}{1}{0}Build.sh\n", Seperator, CurrentPlatform);
 			}
 
 			string UnrealRootPath = Unreal.RootDirectory.FullName;
diff --git a/Engine/Source/Programs/UnrealBuildTool/ProjectFiles/Rider/RiderProjectFile.cs b/Engine/Source/Programs/UnrealBuildTool/ProjectFiles/Rider/RiderProjectFile.cs
--- a/Engine/Source/Programs/UnrealBuildTool/ProjectFiles/Rider/RiderProjectFile.cs
+++ b/Engine/Source/Programs/UnrealBuildTool/ProjectFiles/Rider/RiderProjectFile.cs
@@ -78,7 +78,10 @@ namespace UnrealBuildTool
 				{
 					foreach (ProjectTarget ProjectTarget in ProjectTargets.OfType<ProjectTarget>())
 					{
-						if (TargetTypes.Any() && !TargetTypes.Contains(ProjectTarget.TargetRules!.Type)) continue;
+						if (TargetTypes.Any() && !TargetTypes.Contains(ProjectTarget.TargetRules!.Type))
+						{
+							continue;
+						}
 
 						// Skip Programs for all configs except for current platform + Development & Debug configurations
 						if (ProjectTarget.TargetRules!.Type == TargetType.Program &&
@@ -231,7 +234,7 @@ namespace UnrealBuildTool
 
 		private void ExportModuleCpp(UEBuildModuleCPP ModuleCPP, CppCompileEnvironment ModuleCompileEnvironment, JsonWriter Writer, ILogger Logger)
 		{
-			Writer.WriteValue("GeneratedCodeDirectory", ModuleCPP.GeneratedCodeDirectory != null ? ModuleCPP.GeneratedCodeDirectory.FullName : string.Empty);
+			Writer.WriteValue("GeneratedCodeDirectory", ModuleCPP.GeneratedCodeDirectory != null ? ModuleCPP.GeneratedCodeDirectory.FullName : String.Empty);
 
 			ToolchainInfo ModuleToolchainInfo = GenerateToolchainInfo(ModuleCompileEnvironment);
 			if (!ModuleToolchainInfo.Equals(RootToolchainInfo))
@@ -248,7 +251,9 @@ namespace UnrealBuildTool
 			{
 				string CorrectFilePathPch;
 				if (ExtractWrappedIncludeFile(ModuleCompileEnvironment.PrecompiledHeaderIncludeFilename, Logger, out CorrectFilePathPch))
+				{
 					Writer.WriteValue("SharedPCHFilePath", CorrectFilePathPch);
+				}
 			}
 		}
 
@@ -384,7 +389,10 @@ namespace UnrealBuildTool
 		/// <param name="Modules">Sequence of Modules to write. May be null.</param>
 		private static void ExportJsonModuleArray(JsonWriter Writer, string ArrayName, IEnumerable<UEBuildModule>? Modules)
 		{
-			if (Modules == null || !Modules.Any()) return;
+			if (Modules == null || !Modules.Any())
+			{
+				return;
+			}
 
 			Writer.WriteArrayStart(ArrayName);
 			foreach (UEBuildModule Module in Modules)
@@ -402,7 +410,10 @@ namespace UnrealBuildTool
 		/// <param name="Strings">Sequence of strings to write. May be null.</param>
 		private static void ExportJsonStringArray(JsonWriter Writer, string ArrayName, IEnumerable<string> Strings)
 		{
-			if (Strings == null || !Strings.Any()) return;
+			if (Strings == null || !Strings.Any())
+			{
+				return;
+			}
 
 			Writer.WriteArrayStart(ArrayName);
 			foreach (string String in Strings)
@@ -444,7 +455,10 @@ namespace UnrealBuildTool
 		private static void ExportPluginsFromTarget(UEBuildTarget Target, JsonWriter Writer, ILogger Logger)
 		{
 			Target.SetupPlugins(Logger);
-			if (Target.BuildPlugins == null || !Target.BuildPlugins.Any()) return;
+			if (Target.BuildPlugins == null || !Target.BuildPlugins.Any())
+			{
+				return;
+			}
 
 			Writer.WriteObjectStart("Plugins");
 			foreach (UEBuildPlugin plugin in Target.BuildPlugins!)
@@ -530,7 +544,7 @@ namespace UnrealBuildTool
 				string? UseLibcxxEnvVarOverride = Environment.GetEnvironmentVariable("UE_LINUX_USE_LIBCXX");
 				// assumes a single architecture
 				UnrealArch TargetArchitecture = Target.Architectures.SingleArchitecture;
-				if (string.IsNullOrEmpty(UseLibcxxEnvVarOverride) || UseLibcxxEnvVarOverride == "1")
+				if (String.IsNullOrEmpty(UseLibcxxEnvVarOverride) || UseLibcxxEnvVarOverride == "1")
 				{
 					if (TargetArchitecture == UnrealArch.X64 ||
 						TargetArchitecture == UnrealArch.Arm64)
@@ -580,7 +594,10 @@ namespace UnrealBuildTool
 
 		private void ExportBuildInfo(JsonWriter Writer, UEBuildTarget Target, PlatformProjectGeneratorCollection PlatformProjectGenerators, bool bBuildByDefault)
 		{
-			if (IsStubProject) return;
+			if (IsStubProject)
+			{
+				return;
+			}
 
 			Writer.WriteObjectStart("BuildInfo");
 			UnrealTargetPlatform HostPlatform = BuildHostPlatform.Current.Platform;
@@ -683,7 +700,11 @@ namespace UnrealBuildTool
 
 		private static void WriteField(string ModuleOrTargetName, JsonWriter Writer, Tuple<string, object?> Field, ILogger Logger)
 		{
-			if (Field.Item2 == null) return;
+			if (Field.Item2 == null)
+			{
+				return;
+			}
+
 			string Name = Field.Item1;
 			if (Field.Item2 is bool)
 			{
@@ -693,7 +714,9 @@ namespace UnrealBuildTool
 			{
 				string FieldValue = (string)Field.Item2;
 				if (FieldValue != "")
+				{
 					Writer.WriteValue(Name, (string)Field.Item2);
+				}
 			}
 			else if (Field.Item2 is int)
 			{
@@ -733,7 +756,9 @@ namespace UnrealBuildTool
 			{
 				IEnumerable<string> FieldValue = (IEnumerable<string>)Field.Item2;
 				if (FieldValue.Any())
+				{
 					Writer.WriteStringArrayField(Name, FieldValue);
+				}
 			}
 			else
 			{
@@ -779,7 +804,9 @@ namespace UnrealBuildTool
 				object? Value = typeof(ReadOnlyTargetRules).GetProperty(PlatformName)?.GetValue(CurrentTarget.Rules);
 				object? CompilerField = Value?.GetType().GetProperty("Compiler")?.GetValue(Value);
 				if (CompilerField != null)
+				{
 					ToolchainInfo.Compiler = CompilerField.ToString()!;
+				}
 			}
 
 			return ToolchainInfo;
@@ -816,7 +843,7 @@ namespace UnrealBuildTool
 			private readonly Dictionary<string, IList<string>> CachedIncludePaths =
 				new Dictionary<string, IList<string>>();
 
-			private string CurrentlyProcessedSDK = string.Empty;
+			private string CurrentlyProcessedSDK = String.Empty;
 			private Process? XcrunProcess;
 			private bool IsReadingIncludesSection;
 
@@ -838,7 +865,7 @@ namespace UnrealBuildTool
 
 			private void CalculateSystemIncludePaths(string SDKPath)
 			{
-				if (CurrentlyProcessedSDK != string.Empty)
+				if (CurrentlyProcessedSDK != String.Empty)
 				{
 					throw new InvalidOperationException("Cannot calculate include paths for several platforms at once");
 				}
@@ -849,7 +876,7 @@ namespace UnrealBuildTool
 				{
 					string AppName = "xcrun";
 					string Arguments = "clang++ -Wp,-v -x c++ - -fsyntax-only" +
-									   (string.IsNullOrEmpty(SDKPath) ? string.Empty : (" -isysroot " + SDKPath));
+									   (String.IsNullOrEmpty(SDKPath) ? String.Empty : (" -isysroot " + SDKPath));
 					XcrunProcess.StartInfo.FileName = AppName;
 					XcrunProcess.StartInfo.Arguments = Arguments;
 					XcrunProcess.StartInfo.UseShellExecute = false;
@@ -868,7 +895,7 @@ namespace UnrealBuildTool
 
 				XcrunProcess = null;
 				IsReadingIncludesSection = false;
-				CurrentlyProcessedSDK = string.Empty;
+				CurrentlyProcessedSDK = String.Empty;
 			}
 
 			private void OnOutputDataReceived(object Sender, DataReceivedEventArgs Args)
diff --git a/Engine/Source/Programs/UnrealBuildTool/ProjectFiles/Rider/RiderProjectFileGenerator.cs b/Engine/Source/Programs/UnrealBuildTool/ProjectFiles/Rider/RiderProjectFileGenerator.cs
--- a/Engine/Source/Programs/UnrealBuildTool/ProjectFiles/Rider/RiderProjectFileGenerator.cs
+++ b/Engine/Source/Programs/UnrealBuildTool/ProjectFiles/Rider/RiderProjectFileGenerator.cs
@@ -124,10 +124,16 @@ namespace UnrealBuildTool
 				List<UnrealTargetPlatform> FilteredPlatforms = PlatformsToGenerate.Where(it =>
 				{
 					// Skip deprecated platforms if they are not specified in commandline arguments directly 
-					if (DeprecatedPlatforms.Contains(it) && !Platforms.Contains(it)) return false;
+					if (DeprecatedPlatforms.Contains(it) && !Platforms.Contains(it))
+					{
+						return false;
+					}
 
 					if (UEBuildPlatform.IsPlatformAvailable(it))
+					{
 						return true;
+					}
+
 					Logger.LogWarning(
 						"Platform {Platform} is not a valid platform to build. Check that the SDK is installed properly",
 						it);
diff --git a/Engine/Source/Programs/UnrealBuildTool/ProjectFiles/Rider/ToolchainInfo.cs b/Engine/Source/Programs/UnrealBuildTool/ProjectFiles/Rider/ToolchainInfo.cs
--- a/Engine/Source/Programs/UnrealBuildTool/ProjectFiles/Rider/ToolchainInfo.cs
+++ b/Engine/Source/Programs/UnrealBuildTool/ProjectFiles/Rider/ToolchainInfo.cs
@@ -34,16 +34,25 @@ namespace UnrealBuildTool
 		{
 			foreach (FieldInfo FieldInfo in typeof(ToolchainInfo).GetFields())
 			{
-				if (FieldInfo.GetValue(this) == null) continue;
+				if (FieldInfo.GetValue(this) == null)
+				{
+					continue;
+				}
+
 				if (typeof(List<string>).IsAssignableFrom(FieldInfo.FieldType))
 				{
 					List<string> LocalField = (List<string>)FieldInfo.GetValue(this)!;
 					HashSet<string> OtherField = new HashSet<string>((List<string>)FieldInfo.GetValue(Other)!);
 					IEnumerable<string> Result = LocalField.Where(Item => OtherField.Contains(Item));
-					if (Result.Any()) yield return new Tuple<string, object?>(FieldInfo.Name, Result);
+					if (Result.Any())
+					{
+						yield return new Tuple<string, object?>(FieldInfo.Name, Result);
+					}
 				}
 				else if (!FieldInfo.GetValue(this)!.Equals(FieldInfo.GetValue(Other)))
+				{
 					yield return new Tuple<string, object?>(FieldInfo.Name, FieldInfo.GetValue(this));
+				}
 			}
 		}
 
@@ -57,8 +66,16 @@ namespace UnrealBuildTool
 
 		public bool Equals(ToolchainInfo? Other)
 		{
-			if (ReferenceEquals(null, Other)) return false;
-			if (ReferenceEquals(this, Other)) return true;
+			if (ReferenceEquals(null, Other))
+			{
+				return false;
+			}
+
+			if (ReferenceEquals(this, Other))
+			{
+				return true;
+			}
+
 			return CppStandard == Other.CppStandard && bUseRTTI == Other.bUseRTTI &&
 				bEnableExceptions == Other.bEnableExceptions && bIsBuildingLibrary == Other.bIsBuildingLibrary &&
 				bIsBuildingDLL == Other.bIsBuildingDLL && Architecture == Other.Architecture &&
@@ -73,8 +90,16 @@ namespace UnrealBuildTool
 
 		public override bool Equals(object? Obj)
 		{
-			if (ReferenceEquals(null, Obj)) return false;
-			if (ReferenceEquals(this, Obj)) return true;
+			if (ReferenceEquals(null, Obj))
+			{
+				return false;
+			}
+
+			if (ReferenceEquals(this, Obj))
+			{
+				return true;
+			}
+
 			return Obj is ToolchainInfo && Equals((ToolchainInfo)Obj);
 		}
 
diff --git a/Engine/Source/Programs/UnrealBuildTool/ProjectFiles/VisualStudio/VCProject.cs b/Engine/Source/Programs/UnrealBuildTool/ProjectFiles/VisualStudio/VCProject.cs
--- a/Engine/Source/Programs/UnrealBuildTool/ProjectFiles/VisualStudio/VCProject.cs
+++ b/Engine/Source/Programs/UnrealBuildTool/ProjectFiles/VisualStudio/VCProject.cs
@@ -74,16 +74,16 @@ namespace UnrealBuildTool
 	abstract class MSBuildProjectFile : ProjectFile
 	{
 		/// The project file version string
-		static public readonly string VCProjectFileVersionString = "10.0.30319.1";
+		public static readonly string VCProjectFileVersionString = "10.0.30319.1";
 
 		/// The build configuration name to use for stub project configurations.  These are projects whose purpose
 		/// is to make it easier for developers to find source files and to provide IntelliSense data for the module
 		/// to Visual Studio
-		static public readonly string StubProjectConfigurationName = "BuiltWithUnrealBuildTool";
+		public static readonly string StubProjectConfigurationName = "BuiltWithUnrealBuildTool";
 
 		/// The name of the Visual C++ platform to use for stub project configurations
 		/// NOTE: We always use Win32 for the stub project's platform, since that is guaranteed to be supported by Visual Studio
-		static public readonly string StubProjectPlatformName = "Win64";
+		public static readonly string StubProjectPlatformName = "Win64";
 
 		/// <summary>
 		/// The Guid representing the project type e.g. C# or C++
@@ -660,7 +660,7 @@ namespace UnrealBuildTool
 					Arguments.Add("/Zc:enumTypes");
 				}
 			}
-			return string.Empty;
+			return String.Empty;
 		}
 
 		/// <summary>
@@ -678,7 +678,7 @@ namespace UnrealBuildTool
 				return "-fcoroutines-ts";
 			}
 
-			return string.Empty;
+			return String.Empty;
 		}
 
 		HashSet<string>? InvalidConfigPlatformNames;
@@ -1353,10 +1353,10 @@ namespace UnrealBuildTool
 
 						// filter here. It's a little more graceful to do it where this info is built but easier to follow if we filter 
 						// things our right before they're written.
-						if (!string.IsNullOrEmpty(ForceIncludePaths))
+						if (!String.IsNullOrEmpty(ForceIncludePaths))
 						{
 							IEnumerable<string> PathList = ForceIncludePaths.Split(new[] { ';' }, StringSplitOptions.RemoveEmptyEntries);
-							ForceIncludePaths = string.Join(";", PathList.Where(P => !IncludePathIsFilteredOut(new DirectoryReference(P))));
+							ForceIncludePaths = String.Join(";", PathList.Where(P => !IncludePathIsFilteredOut(new DirectoryReference(P))));
 						}
 
 						ForceIncludePaths ??= String.Empty;
@@ -2063,7 +2063,7 @@ namespace UnrealBuildTool
 					// Get the output directory
 					DirectoryReference OutputDirectory = DirectoryReference.Combine(RootOutputDirectory, "Binaries", UBTPlatformName);
 
-					if (!string.IsNullOrEmpty(TargetRulesObject.ExeBinariesSubFolder))
+					if (!String.IsNullOrEmpty(TargetRulesObject.ExeBinariesSubFolder))
 					{
 						OutputDirectory = DirectoryReference.Combine(OutputDirectory, TargetRulesObject.ExeBinariesSubFolder);
 					}
@@ -2131,7 +2131,7 @@ namespace UnrealBuildTool
 
 					if (ProjGenerator != null)
 					{
-						string? ForeignUProjectPath = (IsForeignProject && !string.IsNullOrEmpty(UProjectPath)) ? UProjectPath : null;
+						string? ForeignUProjectPath = (IsForeignProject && !String.IsNullOrEmpty(UProjectPath)) ? UProjectPath : null;
 						VCUserFileContent.Append(ProjGenerator.GetVisualStudioUserFileStrings(Combination.Platform!.Value, Configuration, ConditionString, TargetRulesObject, Combination.ProjectTarget.TargetFilePath, ProjectFilePath, ProjectName, ForeignUProjectPath));
 					}
 				}
@@ -2214,7 +2214,9 @@ namespace UnrealBuildTool
 						XmlNodeList ConfigNodeList = PropertyGroup.GetElementsByTagName("Configurations");
 						// if this property group does not set configurations we do not care about it
 						if (ConfigNodeList.Count == 0)
+						{
 							continue;
+						}
 
 						if (PropertyGroup.HasAttribute("Condition"))
 						{
diff --git a/Engine/Source/Programs/UnrealBuildTool/ProjectFiles/VisualStudio/VCProjectFileGenerator.cs b/Engine/Source/Programs/UnrealBuildTool/ProjectFiles/VisualStudio/VCProjectFileGenerator.cs
--- a/Engine/Source/Programs/UnrealBuildTool/ProjectFiles/VisualStudio/VCProjectFileGenerator.cs
+++ b/Engine/Source/Programs/UnrealBuildTool/ProjectFiles/VisualStudio/VCProjectFileGenerator.cs
@@ -146,11 +146,11 @@ namespace UnrealBuildTool
 
 		public override string[] GetTargetArguments(string[] Arguments)
 		{
-			return Arguments.Where(s => string.Equals(s, BuildToolOverride, StringComparison.InvariantCultureIgnoreCase)).ToArray();
+			return Arguments.Where(s => String.Equals(s, BuildToolOverride, StringComparison.InvariantCultureIgnoreCase)).ToArray();
 		}
 
 		/// File extension for project files we'll be generating (e.g. ".vcxproj")
-		override public string ProjectFileExtension => ".vcxproj";
+		public override string ProjectFileExtension => ".vcxproj";
 
 		/// <summary>
 		/// </summary>
@@ -216,7 +216,7 @@ namespace UnrealBuildTool
 
 
 		/// "4.0", "12.0", or "14.0", etc...
-		static public string GetProjectFileToolVersionString(VCProjectFileFormat ProjectFileFormat)
+		public static string GetProjectFileToolVersionString(VCProjectFileFormat ProjectFileFormat)
 		{
 			switch (ProjectFileFormat)
 			{
@@ -225,11 +225,11 @@ namespace UnrealBuildTool
 				case VCProjectFileFormat.VisualStudio2022:
 					return "17.0";
 			}
-			return string.Empty;
+			return String.Empty;
 		}
 
 		/// for instance: <PlatformToolset>v110</PlatformToolset>
-		static public string GetProjectFilePlatformToolsetVersionString(VCProjectFileFormat ProjectFileFormat)
+		public static string GetProjectFilePlatformToolsetVersionString(VCProjectFileFormat ProjectFileFormat)
 		{
 			switch (ProjectFileFormat)
 			{
@@ -239,10 +239,10 @@ namespace UnrealBuildTool
 					return "v143";
 
 			}
-			return string.Empty;
+			return String.Empty;
 		}
 
-		static public void AppendPlatformToolsetProperty(StringBuilder VCProjectFileContent, VCProjectFileFormat ProjectFileFormat)
+		public static void AppendPlatformToolsetProperty(StringBuilder VCProjectFileContent, VCProjectFileFormat ProjectFileFormat)
 		{
 			string ToolVersionString = GetProjectFileToolVersionString(ProjectFileFormat);
 			string PlatformToolsetVersionString = GetProjectFilePlatformToolsetVersionString(ProjectFileFormat);
@@ -345,7 +345,7 @@ namespace UnrealBuildTool
 										Logger.LogInformation("To force Visual Studio 2022 solutions to always be generated add the following to BuildConfiguration.xml:");
 										Logger.LogInformation("  <VCProjectFileGenerator>\r\n    <Version>VisualStudio2022</Version>\r\n  </VCProjectFileGenerator>");
 									}
-									Logger.LogInformation(string.Empty);
+									Logger.LogInformation(String.Empty);
 								}
 							}
 						}
@@ -707,7 +707,7 @@ namespace UnrealBuildTool
 						string UnrealVSGuid = "ddbf523f-7eb6-4887-bd51-85a714ff87eb";
 						VCSolutionFileContent.AppendLine("\t# UnrealVS Section");
 						VCSolutionFileContent.AppendLine("\tGlobalSection({0}) = preSolution", UnrealVSGuid);
-						VCSolutionFileContent.AppendLine("\t\tAvailablePlatforms={0}", string.Join(";", PlatformsValidForProjects.Select(platform => platform.ToString())));
+						VCSolutionFileContent.AppendLine("\t\tAvailablePlatforms={0}", String.Join(";", PlatformsValidForProjects.Select(platform => platform.ToString())));
 						VCSolutionFileContent.AppendLine("\tEndGlobalSection");
 					}
 
@@ -879,7 +879,7 @@ namespace UnrealBuildTool
 				VsConfigFileContent.AppendLine("  \"version\": \"1.0\",");
 				VsConfigFileContent.AppendLine("  \"components\": [");
 				IEnumerable<string> Components = MicrosoftPlatformSDK.GetVisualStudioSuggestedComponents(Settings.ProjectFileFormat);
-				string ComponentsString = string.Join($",{Environment.NewLine}    ", Components.Select(x => $"\"{x}\""));
+				string ComponentsString = String.Join($",{Environment.NewLine}    ", Components.Select(x => $"\"{x}\""));
 				VsConfigFileContent.AppendLine($"    {ComponentsString}");
 				VsConfigFileContent.AppendLine("  ]");
 				VsConfigFileContent.AppendLine("}");
diff --git a/Engine/Source/Programs/UnrealBuildTool/ProjectFiles/VisualStudioCode/VSCodeProjectFileGenerator.cs b/Engine/Source/Programs/UnrealBuildTool/ProjectFiles/VisualStudioCode/VSCodeProjectFileGenerator.cs
--- a/Engine/Source/Programs/UnrealBuildTool/ProjectFiles/VisualStudioCode/VSCodeProjectFileGenerator.cs
+++ b/Engine/Source/Programs/UnrealBuildTool/ProjectFiles/VisualStudioCode/VSCodeProjectFileGenerator.cs
@@ -239,7 +239,7 @@ namespace UnrealBuildTool
 			private string TabString = "";
 		}
 
-		override public string ProjectFileExtension => ".vscode";
+		public override string ProjectFileExtension => ".vscode";
 
 		public override void CleanProjectFiles(DirectoryReference InPrimaryProjectDirectory, string InPrimaryProjectName, DirectoryReference InIntermediateProjectFilesPath, ILogger Logger)
 		{
@@ -343,7 +343,7 @@ namespace UnrealBuildTool
 			{
 				CompilerPath = FileReference.FromString(LinuxCommon.WhichClang(Logger));
 				string? InternalSDKPath = UEBuildPlatform.GetSDK(UnrealTargetPlatform.Linux)?.GetInternalSDKPath();
-				if (!string.IsNullOrEmpty(InternalSDKPath))
+				if (!String.IsNullOrEmpty(InternalSDKPath))
 				{
 					SysRootPath = DirectoryReference.FromString(InternalSDKPath);
 				}
@@ -620,9 +620,11 @@ namespace UnrealBuildTool
 
 							// we do not generate intellisense for every target, as that just causes a lot of redundancy, as such we will not find a mapping for a lot of the targets
 							if (BuildTarget == null)
+							{
 								continue;
+							}
 
-							string Name = string.Format("{0} {1} {2} {3} ({4})", ProjectTarget.Name, ProjectTarget.Type, BuildTarget.Platform, BuildTarget.Configuration, Project.Name);
+							string Name = String.Format("{0} {1} {2} {3} ({4})", ProjectTarget.Name, ProjectTarget.Type, BuildTarget.Platform, BuildTarget.Configuration, Project.Name);
 							WriteConfiguration(Name, Project.Name, Project.SourceProject.SourceFiles.Select(x => x.Reference), BuildTarget.CppStandard, BuildTarget.CompilerPath!, BuildTarget.SysRootPath, BuildTarget.ModuleCommandLines, OutFile, OutputDirectory);
 
 							CompilerPath = BuildTarget.CompilerPath;
@@ -740,7 +742,7 @@ namespace UnrealBuildTool
 
 			if (ShouldGenerateIntelliSenseData())
 			{
-				FileReference CompileCommands = FileReference.Combine(OutputDirectory, string.Format("compileCommands_{0}.json", ProjectName));
+				FileReference CompileCommands = FileReference.Combine(OutputDirectory, String.Format("compileCommands_{0}.json", ProjectName));
 				WriteCompileCommands(CompileCommands, SourceFiles, CompilerPath!, ModuleCommandLines);
 				OutFile.AddField("compileCommands", MakePathString(CompileCommands, bInAbsolute: true, bForceSkipQuotes: true));
 			}
@@ -954,7 +956,7 @@ namespace UnrealBuildTool
 								}
 								OutFile.AddUnnamedField("-waitmutex");
 
-								if (!string.IsNullOrEmpty(CleanParam))
+								if (!String.IsNullOrEmpty(CleanParam))
 								{
 									OutFile.AddUnnamedField(CleanParam);
 								}
@@ -1165,7 +1167,7 @@ namespace UnrealBuildTool
 			if (PackageName.Contains("[PROJECT]"))
 			{
 				// project name must start with a letter
-				if (!char.IsLetter(Target.Name[0]))
+				if (!Char.IsLetter(Target.Name[0]))
 				{
 					Trace.TraceWarning("Package name segments must all start with a letter. Please replace [PROJECT] with a valid name");
 				}
@@ -1182,7 +1184,7 @@ namespace UnrealBuildTool
 				for (int Index = 0; Index < ProjectName.Length; Index++)
 				{
 					char c = ProjectName[Index];
-					if (c != '.' && c != '_' && !char.IsLetterOrDigit(c))
+					if (c != '.' && c != '_' && !Char.IsLetterOrDigit(c))
 					{
 						Trace.TraceWarning("Project name contains illegal characters (only letters, numbers, and underscore allowed); please replace [PROJECT] with a valid name");
 						ProjectName.Replace(c, '_');
@@ -1393,7 +1395,7 @@ namespace UnrealBuildTool
 				OutFile.AddField("type", "coreclr");
 				OutFile.AddField("request", "launch");
 
-				if (!string.IsNullOrEmpty(InBuildTaskName))
+				if (!String.IsNullOrEmpty(InBuildTaskName))
 				{
 					OutFile.AddField("preLaunchTask", InBuildTaskName);
 				}
diff --git a/Engine/Source/Programs/UnrealBuildTool/ProjectFiles/VisualStudioMac/VisualStudioMacProjectFileGenerator.cs b/Engine/Source/Programs/UnrealBuildTool/ProjectFiles/VisualStudioMac/VisualStudioMacProjectFileGenerator.cs
--- a/Engine/Source/Programs/UnrealBuildTool/ProjectFiles/VisualStudioMac/VisualStudioMacProjectFileGenerator.cs
+++ b/Engine/Source/Programs/UnrealBuildTool/ProjectFiles/VisualStudioMac/VisualStudioMacProjectFileGenerator.cs
@@ -29,7 +29,7 @@ namespace UnrealBuildTool
 		}
 
 		/// True if we should include IntelliSense data in the generated project files when possible
-		override public bool ShouldGenerateIntelliSenseData()
+		public override bool ShouldGenerateIntelliSenseData()
 		{
 			return false;
 		}
diff --git a/Engine/Source/Programs/UnrealBuildTool/ProjectFiles/Xcode/XcodeFrameworkWrapperProject.cs b/Engine/Source/Programs/UnrealBuildTool/ProjectFiles/Xcode/XcodeFrameworkWrapperProject.cs
--- a/Engine/Source/Programs/UnrealBuildTool/ProjectFiles/Xcode/XcodeFrameworkWrapperProject.cs
+++ b/Engine/Source/Programs/UnrealBuildTool/ProjectFiles/Xcode/XcodeFrameworkWrapperProject.cs
@@ -15,11 +15,11 @@ namespace UnrealBuildTool.ProjectFiles.Xcode
 	/// </summary>
 	class XcodeFrameworkWrapperProject
 	{
-		private readonly static string PROJECT_FILE_SEARCH_EXPRESSION = "*.pbxproj";
-		private readonly static string TEMPLATE_NAME = "FrameworkWrapper";
-		private readonly static string FRAMEWORK_WRAPPER_TEMPLATE_DIRECTORY = Path.Combine(Unreal.EngineDirectory.ToNormalizedPath(), "Build", "IOS", "Resources", TEMPLATE_NAME);
-		private readonly static string TEMPLATE_PROJECT_NAME = "PROJECT_NAME";
-		private readonly static string COMMANDLINE_FILENAME = "uecommandline.txt";
+		private static readonly string PROJECT_FILE_SEARCH_EXPRESSION = "*.pbxproj";
+		private static readonly string TEMPLATE_NAME = "FrameworkWrapper";
+		private static readonly string FRAMEWORK_WRAPPER_TEMPLATE_DIRECTORY = Path.Combine(Unreal.EngineDirectory.ToNormalizedPath(), "Build", "IOS", "Resources", TEMPLATE_NAME);
+		private static readonly string TEMPLATE_PROJECT_NAME = "PROJECT_NAME";
+		private static readonly string COMMANDLINE_FILENAME = "uecommandline.txt";
 
 		/// <summary>
 		/// Recursively copies all of the files and directories that are inside <paramref name="SourceDirectory"/> into <paramref name="DestinationDirectory"/>.
diff --git a/Engine/Source/Programs/UnrealBuildTool/ProjectFiles/Xcode/XcodeProject.Scheme.cs b/Engine/Source/Programs/UnrealBuildTool/ProjectFiles/Xcode/XcodeProject.Scheme.cs
--- a/Engine/Source/Programs/UnrealBuildTool/ProjectFiles/Xcode/XcodeProject.Scheme.cs
+++ b/Engine/Source/Programs/UnrealBuildTool/ProjectFiles/Xcode/XcodeProject.Scheme.cs
@@ -97,7 +97,7 @@ namespace UnrealBuildTool.XcodeProjectXcconfig
 				Content.WriteLine("            ActionType = \"Xcode.IDEStandardExecutionActionsCore.ExecutionActionType.ShellScriptAction\">");
 				Content.WriteLine("            <ActionContent");
 				Content.WriteLine("               title = \"Touch UBT generated tiles\"");
-				Content.WriteLine("               scriptText = \"" + string.Join("&#10;", Lines) + "\">");
+				Content.WriteLine("               scriptText = \"" + String.Join("&#10;", Lines) + "\">");
 				Content.WriteLine("               <EnvironmentBuildable>");
 				Content.WriteLine("                  <BuildableReference");
 				Content.WriteLine("                     BuildableIdentifier = \"primary\"");
@@ -164,7 +164,7 @@ namespace UnrealBuildTool.XcodeProjectXcconfig
 				Content.WriteLine("            ReferencedContainer = \"container:" + ProjectName + ".xcodeproj\">");
 				Content.WriteLine("         </BuildableReference>");
 				Content.WriteLine("      </BuildableProductRunnable>");
-				if (string.IsNullOrEmpty(OldCommandLineArguments))
+				if (String.IsNullOrEmpty(OldCommandLineArguments))
 				{
 					if (bHasEditorConfiguration && GameProjectPath != "")
 					{
diff --git a/Engine/Source/Programs/UnrealBuildTool/ProjectFiles/Xcode/XcodeProject.cs b/Engine/Source/Programs/UnrealBuildTool/ProjectFiles/Xcode/XcodeProject.cs
--- a/Engine/Source/Programs/UnrealBuildTool/ProjectFiles/Xcode/XcodeProject.cs
+++ b/Engine/Source/Programs/UnrealBuildTool/ProjectFiles/Xcode/XcodeProject.cs
@@ -420,9 +420,9 @@ namespace UnrealBuildTool.XcodeProjectXcconfig
 			}
 
 			this.bMakeProjectPerTarget = bMakeProjectPerTarget;
-			this.bForDistribution = bIsForDistribution;
-			this.BundleIdentifier = BundleID;
-			this.DisplayName = string.IsNullOrEmpty(AppName) ? "$(UE_PRODUCT_NAME)" : AppName;
+			bForDistribution = bIsForDistribution;
+			BundleIdentifier = BundleID;
+			DisplayName = String.IsNullOrEmpty(AppName) ? "$(UE_PRODUCT_NAME)" : AppName;
 		}
 
 		public void InitializeUProjectFileLocation(XcodeProjectFile ProjectFile)
@@ -457,20 +457,20 @@ namespace UnrealBuildTool.XcodeProjectXcconfig
 		public bool Initialize(XcodeProjectFile ProjectFile, List<UnrealTargetConfiguration> Configurations, ILogger Logger)
 		{
 			this.ProjectFile = ProjectFile;
-			this.bIsForeignProject = ProjectFile.IsForeignProject;
-			this.bIsStubProject = ProjectFile.IsStubProject;
-			this.bIsContentOnlyProject = ProjectFile.IsContentOnlyProject;
+			bIsForeignProject = ProjectFile.IsForeignProject;
+			bIsStubProject = ProjectFile.IsStubProject;
+			bIsContentOnlyProject = ProjectFile.IsContentOnlyProject;
 			this.Logger = Logger;
 
 			InitializeUProjectFileLocation(ProjectFile);
 
 			// setup BundleIdentifier from ini file (if there's a specified plist file with one, that will override this)
-			if (string.IsNullOrEmpty(BundleIdentifier))
+			if (String.IsNullOrEmpty(BundleIdentifier))
 			{
 				ConfigHierarchy Ini = ConfigCache.ReadHierarchy(ConfigHierarchyType.Engine, UProjectFileLocation?.Directory, UnrealTargetPlatform.Mac);
 				Ini.GetString($"/Script/MacTargetPlatform.XcodeProjectSettings", "ModernBundleIdentifier", out BundleIdentifier);
 			}
-			if (string.IsNullOrEmpty(BundleIdentifier))
+			if (String.IsNullOrEmpty(BundleIdentifier))
 			{
 				BundleIdentifier = "$(UE_SIGNING_PREFIX).$(UE_PRODUCT_NAME_STRIPPED)";
 			}
@@ -483,7 +483,7 @@ namespace UnrealBuildTool.XcodeProjectXcconfig
 			else if (ProjectFile.ProjectTargets[0].TargetRules!.Type == TargetType.Program)
 			{
 				DirectoryReference? ProgramFinder = DirectoryReference.Combine(ProjectFile.BaseDir);
-				while (ProgramFinder != null && string.Compare(ProgramFinder.GetDirectoryName(), "Source", true) != 0)
+				while (ProgramFinder != null && String.Compare(ProgramFinder.GetDirectoryName(), "Source", true) != 0)
 				{
 					ProgramFinder = ProgramFinder.ParentDirectory;
 				}
@@ -712,7 +712,7 @@ namespace UnrealBuildTool.XcodeProjectXcconfig
 											DirectoryReference MacBinaryDir = DirectoryReference.Combine(OutputDirectory, "Mac");
 											DirectoryReference IOSBinaryDir = DirectoryReference.Combine(OutputDirectory, "IOS");
 											DirectoryReference TVOSBinaryDir = DirectoryReference.Combine(OutputDirectory, "TVOS");
-											if (!string.IsNullOrEmpty(ProjectTarget.TargetRules.ExeBinariesSubFolder))
+											if (!String.IsNullOrEmpty(ProjectTarget.TargetRules.ExeBinariesSubFolder))
 											{
 												MacBinaryDir = DirectoryReference.Combine(MacBinaryDir, ProjectTarget.TargetRules.ExeBinariesSubFolder);
 												IOSBinaryDir = DirectoryReference.Combine(IOSBinaryDir, ProjectTarget.TargetRules.ExeBinariesSubFolder);
@@ -999,7 +999,7 @@ namespace UnrealBuildTool.XcodeProjectXcconfig
 			MiscItems.Add($");");
 
 			//			string Script = string.Join("&#10", ScriptLines);
-			string Script = string.Join("\\n", ScriptLines);
+			string Script = String.Join("\\n", ScriptLines);
 			MiscItems.Add($"shellPath = /bin/sh;");
 			MiscItems.Add($"shellScript = \"{Script}\";");
 			if (bInstallOnly)
@@ -1252,7 +1252,7 @@ namespace UnrealBuildTool.XcodeProjectXcconfig
 		{
 			this.TargetType = TargetType;
 			this.Platform = Platform;
-			this.UnrealData = Project.UnrealData;
+			UnrealData = Project.UnrealData;
 
 			BuildConfigList = new XcodeBuildConfigList(Platform, Name, Project.UnrealData, bIncludeAllPlatformsInConfig: false);
 			References.Add(BuildConfigList);
@@ -1481,7 +1481,7 @@ namespace UnrealBuildTool.XcodeProjectXcconfig
 				XcodeFrameworkBuildPhase FrameworkPhase = new XcodeFrameworkBuildPhase(Project.FileCollection);
 
 				// filter frameworks that need to installed into the .app (either the framework or a bundle inside a .zip)
-				IEnumerable<UEBuildFramework> InstalledFrameworks = Frameworks.Where(x => x.bCopyFramework || !string.IsNullOrEmpty(x.CopyBundledAssets));
+				IEnumerable<UEBuildFramework> InstalledFrameworks = Frameworks.Where(x => x.bCopyFramework || !String.IsNullOrEmpty(x.CopyBundledAssets));
 				// filter frameworks that need to be unzipped before we compile
 				IEnumerable<UEBuildFramework> ZippedFrameworks = Frameworks.Where(x => x.ZipFile != null);
 
@@ -1517,7 +1517,7 @@ namespace UnrealBuildTool.XcodeProjectXcconfig
 					}
 
 					// set up the CopyBundle to be copied
-					if (!string.IsNullOrEmpty(Framework.CopyBundledAssets))
+					if (!String.IsNullOrEmpty(Framework.CopyBundledAssets))
 					{
 						ResourcesBuildPhase.AddFolderResource(DirectoryReference.Combine(BundleRootDir, Framework.CopyBundledAssets), "Resources");
 					}
@@ -1639,7 +1639,7 @@ namespace UnrealBuildTool.XcodeProjectXcconfig
 				// @todo: get a version for  games, like IOS has
 				MarketingVersion = MacToolChain.LoadEngineDisplayVersion();
 
-				string SupportedMacArchitectures = string.Join(" ", XcodeUtils.GetSupportedMacArchitectures(BuildConfig.BuildTarget, UnrealData.UProjectFileLocation).Architectures.Select(x => x.AppleName));
+				string SupportedMacArchitectures = String.Join(" ", XcodeUtils.GetSupportedMacArchitectures(BuildConfig.BuildTarget, UnrealData.UProjectFileLocation).Architectures.Select(x => x.AppleName));
 				ExtraConfigLines.Add($"VALID_ARCHS = {SupportedMacArchitectures}");
 			}
 			else
@@ -1668,7 +1668,7 @@ namespace UnrealBuildTool.XcodeProjectXcconfig
 
 					// only iphone deals with orientation
 					List<string> SupportedOrientations = XcodeUtils.GetSupportedOrientations(PlatformIni);
-					ExtraConfigLines.Add($"INFOPLIST_KEY_UISupportedInterfaceOrientations = \"{string.Join(" ", SupportedOrientations)}\"");
+					ExtraConfigLines.Add($"INFOPLIST_KEY_UISupportedInterfaceOrientations = \"{String.Join(" ", SupportedOrientations)}\"");
 				}
 				else // tvos
 				{
@@ -1775,11 +1775,11 @@ namespace UnrealBuildTool.XcodeProjectXcconfig
 				Xcconfig.AppendLine("");
 				Xcconfig.AppendLine("// Code-signing settings");
 				Xcconfig.AppendLine("CODE_SIGN_STYLE = " + (bAutomaticSigning ? "Automatic" : "Manual"));
-				if (!string.IsNullOrEmpty(SigningTeam))
+				if (!String.IsNullOrEmpty(SigningTeam))
 				{
 					Xcconfig.AppendLine($"DEVELOPMENT_TEAM = {SigningTeam}");
 				}
-				if (!string.IsNullOrEmpty(ProvisioningProfile))
+				if (!String.IsNullOrEmpty(ProvisioningProfile))
 				{
 					Xcconfig.AppendLine($"PROVISIONING_PROFILE_SPECIFIER = {ProvisioningProfile}");
 				}
diff --git a/Engine/Source/Programs/UnrealBuildTool/ProjectFiles/Xcode/XcodeProjectFileGenerator.cs b/Engine/Source/Programs/UnrealBuildTool/ProjectFiles/Xcode/XcodeProjectFileGenerator.cs
--- a/Engine/Source/Programs/UnrealBuildTool/ProjectFiles/Xcode/XcodeProjectFileGenerator.cs
+++ b/Engine/Source/Programs/UnrealBuildTool/ProjectFiles/Xcode/XcodeProjectFileGenerator.cs
@@ -71,11 +71,11 @@ namespace UnrealBuildTool
 		/// <summary>
 		/// A static copy of ProjectPlatforms from the base class
 		/// </summary>
-		static public List<UnrealTargetPlatform> XcodePlatforms = new();
+		public static List<UnrealTargetPlatform> XcodePlatforms = new();
 
-		static public List<UnrealTargetPlatform?> WorkspacePlatforms = new();
-		static public List<UnrealTargetPlatform?> RunTargetPlatforms = new();
-		static public List<UnrealTargetPlatform?> NullPlatformList = new() { null };
+		public static List<UnrealTargetPlatform?> WorkspacePlatforms = new();
+		public static List<UnrealTargetPlatform?> RunTargetPlatforms = new();
+		public static List<UnrealTargetPlatform?> NullPlatformList = new() { null };
 		internal static Dictionary<Tuple<ProjectFile, UnrealTargetPlatform>, IEnumerable<UEBuildFramework>> TargetFrameworks = new();
 		internal static Dictionary<Tuple<ProjectFile, UnrealTargetPlatform>, IEnumerable<UEBuildBundleResource>> TargetBundles = new();
 
@@ -84,7 +84,7 @@ namespace UnrealBuildTool
 		/// <summary>
 		/// Should we generate only a run project (no build/index targets)
 		/// </summary>
-		static public bool bGenerateRunOnlyProject = false;
+		public static bool bGenerateRunOnlyProject = false;
 
 
 		public XcodeProjectFileGenerator(FileReference? InOnlyGameProject, CommandLineArguments CommandLine)
@@ -161,7 +161,7 @@ namespace UnrealBuildTool
 		}
 
 		/// File extension for project files we'll be generating (e.g. ".vcxproj")
-		override public string ProjectFileExtension => ".xcodeproj";
+		public override string ProjectFileExtension => ".xcodeproj";
 
 		/// <summary>
 		/// </summary>
@@ -324,7 +324,7 @@ namespace UnrealBuildTool
 							foreach (XcodeProjectXcconfig.XcodeProjectFile XcodeProject in SupportedProjects)
 							{
 								// if we are only generating a single target project, skip any others now
-								if (!string.IsNullOrEmpty(SingleTargetName) && !XcodeProject.ProjectTargets.Any(x => x.TargetRules?.Name == SingleTargetName))
+								if (!String.IsNullOrEmpty(SingleTargetName) && !XcodeProject.ProjectTargets.Any(x => x.TargetRules?.Name == SingleTargetName))
 								{
 									continue;
 								}
diff --git a/Engine/Source/Programs/UnrealBuildTool/ProjectFiles/Xcode/XcodeProjectLegacy.cs b/Engine/Source/Programs/UnrealBuildTool/ProjectFiles/Xcode/XcodeProjectLegacy.cs
--- a/Engine/Source/Programs/UnrealBuildTool/ProjectFiles/Xcode/XcodeProjectLegacy.cs
+++ b/Engine/Source/Programs/UnrealBuildTool/ProjectFiles/Xcode/XcodeProjectLegacy.cs
@@ -296,14 +296,14 @@ namespace UnrealBuildTool.XcodeProjectLegacy
 
 				if (IsGeneratedProject)
 				{
-					PBXBuildFileSection.Append(string.Format("\t\t{0} /* {1} in {2} */ = {{isa = PBXBuildFile; fileRef = {3} /* {1} */; }};" + ProjectFileGenerator.NewLine,
+					PBXBuildFileSection.Append(String.Format("\t\t{0} /* {1} in {2} */ = {{isa = PBXBuildFile; fileRef = {3} /* {1} */; }};" + ProjectFileGenerator.NewLine,
 						SourceFile.FileGuid,
 						FileName,
 						GetFileCategory(FileExtension),
 						SourceFile.FileRefGuid));
 				}
 
-				PBXFileReferenceSection.Append(string.Format("\t\t{0} /* {1} */ = {{isa = PBXFileReference; explicitFileType = {2}; name = \"{1}\"; path = \"{3}\"; sourceTree = SOURCE_ROOT; }};" + ProjectFileGenerator.NewLine,
+				PBXFileReferenceSection.Append(String.Format("\t\t{0} /* {1} */ = {{isa = PBXFileReference; explicitFileType = {2}; name = \"{1}\"; path = \"{3}\"; sourceTree = SOURCE_ROOT; }};" + ProjectFileGenerator.NewLine,
 					SourceFile.FileRefGuid,
 					FileName,
 					GetFileType(FileExtension),
@@ -321,7 +321,7 @@ namespace UnrealBuildTool.XcodeProjectLegacy
 				}
 			}
 
-			PBXFileReferenceSection.Append(string.Format("\t\t{0} /* {1} */ = {{isa = PBXFileReference; explicitFileType = {2}; path = {1}; sourceTree = BUILT_PRODUCTS_DIR; }};" + ProjectFileGenerator.NewLine, TargetAppGuid, TargetName, bIsAppBundle ? "wrapper.application" : "\"compiled.mach-o.executable\""));
+			PBXFileReferenceSection.Append(String.Format("\t\t{0} /* {1} */ = {{isa = PBXFileReference; explicitFileType = {2}; path = {1}; sourceTree = BUILT_PRODUCTS_DIR; }};" + ProjectFileGenerator.NewLine, TargetAppGuid, TargetName, bIsAppBundle ? "wrapper.application" : "\"compiled.mach-o.executable\""));
 		}
 
 		private void GenerateSectionsWithExtensions(StringBuilder PBXBuildFileSection, StringBuilder PBXFileReferenceSection, StringBuilder PBXCopyFilesBuildPhaseSection, StringBuilder PBXResourcesBuildPhaseSection,
@@ -347,16 +347,16 @@ namespace UnrealBuildTool.XcodeProjectLegacy
 						XcodeExtensionInfo ExtensionInfo = new XcodeExtensionInfo(DI.Name);
 						AllExtensions.Add(ExtensionInfo);
 
-						PBXBuildFileSection.Append(string.Format("\t\t{0} /* {1} in Embed App Extensions */ = {{isa = PBXBuildFile; fileRef = {2} /* {1} */; settings = {{ATTRIBUTES = (RemoveHeadersOnCopy, ); }}; }};" + ProjectFileGenerator.NewLine,
+						PBXBuildFileSection.Append(String.Format("\t\t{0} /* {1} in Embed App Extensions */ = {{isa = PBXBuildFile; fileRef = {2} /* {1} */; settings = {{ATTRIBUTES = (RemoveHeadersOnCopy, ); }}; }};" + ProjectFileGenerator.NewLine,
 							ExtensionGuid,
 							Extension,
 							ExtensionInfo.ProductGuid));
 
-						PBXFileReferenceSection.Append(string.Format("\t\t{0} /* {1} */ = {{isa = PBXFileReference; explicitFileType = wrapper.app-extension; path = \"{1}\"; sourceTree = BUILT_PRODUCTS_DIR; }};" + ProjectFileGenerator.NewLine,
+						PBXFileReferenceSection.Append(String.Format("\t\t{0} /* {1} */ = {{isa = PBXFileReference; explicitFileType = wrapper.app-extension; path = \"{1}\"; sourceTree = BUILT_PRODUCTS_DIR; }};" + ProjectFileGenerator.NewLine,
 							ExtensionInfo.ProductGuid,
 							Extension));
 
-						PBXCopyFilesBuildPhaseSection.Append(string.Format("\t\t\t\t{0} /* {1} in Embed App Extensions */," + ProjectFileGenerator.NewLine,
+						PBXCopyFilesBuildPhaseSection.Append(String.Format("\t\t\t\t{0} /* {1} in Embed App Extensions */," + ProjectFileGenerator.NewLine,
 							ExtensionGuid,
 							Extension));
 
@@ -377,13 +377,13 @@ namespace UnrealBuildTool.XcodeProjectLegacy
 								// for each resource, put it into the File/FileRef section, and into the CopyResuorceBuildPhase
 								string ResourceGuid = XcodeProjectFileGenerator.MakeXcodeGuid();
 								string ResourceRefGuid = XcodeProjectFileGenerator.MakeXcodeGuid();
-								PBXBuildFileSection.Append(string.Format("\t\t{0} /* {1} in Embed App Extensions */ = {{isa = PBXBuildFile; fileRef = {2} /* {1} */; }};" + ProjectFileGenerator.NewLine,
+								PBXBuildFileSection.Append(String.Format("\t\t{0} /* {1} in Embed App Extensions */ = {{isa = PBXBuildFile; fileRef = {2} /* {1} */; }};" + ProjectFileGenerator.NewLine,
 									ResourceGuid,
 									FSI.Name,
 									ResourceRefGuid));
 
 								// lastKnownFileType = wrapper.app-extension; 
-								PBXFileReferenceSection.Append(string.Format("\t\t{0} /* {1} */ = {{isa = PBXFileReference; lastKnownFileType = folder.assetcatalog; path = \"{2}\"; sourceTree = \"<absolute>\"; }};" + ProjectFileGenerator.NewLine,
+								PBXFileReferenceSection.Append(String.Format("\t\t{0} /* {1} */ = {{isa = PBXFileReference; lastKnownFileType = folder.assetcatalog; path = \"{2}\"; sourceTree = \"<absolute>\"; }};" + ProjectFileGenerator.NewLine,
 									ResourceRefGuid,
 									FSI.Name,
 									// @todo: make this relative path!! 
@@ -449,18 +449,18 @@ namespace UnrealBuildTool.XcodeProjectLegacy
 		{
 			if (!Group.bIsReference)
 			{
-				Content.Append(string.Format("\t\t{0} = {{{1}", Group.GroupGuid, ProjectFileGenerator.NewLine));
+				Content.Append(String.Format("\t\t{0} = {{{1}", Group.GroupGuid, ProjectFileGenerator.NewLine));
 				Content.Append("\t\t\tisa = PBXGroup;" + ProjectFileGenerator.NewLine);
 				Content.Append("\t\t\tchildren = (" + ProjectFileGenerator.NewLine);
 
 				foreach (XcodeFileGroup ChildGroup in Group.Children.Values)
 				{
-					Content.Append(string.Format("\t\t\t\t{0} /* {1} */,{2}", ChildGroup.GroupGuid, ChildGroup.GroupName, ProjectFileGenerator.NewLine));
+					Content.Append(String.Format("\t\t\t\t{0} /* {1} */,{2}", ChildGroup.GroupGuid, ChildGroup.GroupName, ProjectFileGenerator.NewLine));
 				}
 
 				foreach (XcodeSourceFile File in Group.Files)
 				{
-					Content.Append(string.Format("\t\t\t\t{0} /* {1} */,{2}", File.FileRefGuid, File.Reference.GetFileName(), ProjectFileGenerator.NewLine));
+					Content.Append(String.Format("\t\t\t\t{0} /* {1} */,{2}", File.FileRefGuid, File.Reference.GetFileName(), ProjectFileGenerator.NewLine));
 				}
 
 				Content.Append("\t\t\t);" + ProjectFileGenerator.NewLine);
@@ -493,7 +493,7 @@ namespace UnrealBuildTool.XcodeProjectLegacy
 		private void AppendSourcesBuildPhaseSection(StringBuilder Content, StringBuilder SectionContent, string SourcesBuildPhaseGuid)
 		{
 			Content.Append("/* Begin PBXSourcesBuildPhase section */" + ProjectFileGenerator.NewLine);
-			Content.Append(string.Format("\t\t{0} = {{{1}", SourcesBuildPhaseGuid, ProjectFileGenerator.NewLine));
+			Content.Append(String.Format("\t\t{0} = {{{1}", SourcesBuildPhaseGuid, ProjectFileGenerator.NewLine));
 			Content.Append("\t\t\tisa = PBXSourcesBuildPhase;" + ProjectFileGenerator.NewLine);
 			Content.Append("\t\t\tbuildActionMask = 2147483647;" + ProjectFileGenerator.NewLine);
 			Content.Append("\t\t\tfiles = (" + ProjectFileGenerator.NewLine);
@@ -527,7 +527,7 @@ namespace UnrealBuildTool.XcodeProjectLegacy
 		private void AppendCopyExtensionsBuildPhaseSection(StringBuilder Content, StringBuilder SectionContent, string CopyFilesBuildPhaseGuid)
 		{
 			Content.Append("/* Begin PBXCopyFilesBuildPhase section */" + ProjectFileGenerator.NewLine);
-			Content.Append(string.Format("\t{0} /* Embed App Extensions */ = {{{1}", CopyFilesBuildPhaseGuid, ProjectFileGenerator.NewLine));
+			Content.Append(String.Format("\t{0} /* Embed App Extensions */ = {{{1}", CopyFilesBuildPhaseGuid, ProjectFileGenerator.NewLine));
 			Content.Append("\t\tisa = PBXCopyFilesBuildPhase;" + ProjectFileGenerator.NewLine);
 			Content.Append("\t\tbuildActionMask = 2147483647;" + ProjectFileGenerator.NewLine);
 			Content.Append("\t\tdstPath = \"\";" + ProjectFileGenerator.NewLine);
@@ -552,21 +552,21 @@ namespace UnrealBuildTool.XcodeProjectLegacy
 			Content.Append("/* Begin PBXGroup section */" + ProjectFileGenerator.NewLine);
 
 			// Main group
-			Content.Append(string.Format("\t\t{0} = {{{1}", MainGroupGuid, ProjectFileGenerator.NewLine));
+			Content.Append(String.Format("\t\t{0} = {{{1}", MainGroupGuid, ProjectFileGenerator.NewLine));
 			Content.Append("\t\t\tisa = PBXGroup;" + ProjectFileGenerator.NewLine);
 			Content.Append("\t\t\tchildren = (" + ProjectFileGenerator.NewLine);
 
 			foreach (XcodeFileGroup Group in RootGroup.Children.Values)
 			{
-				Content.Append(string.Format("\t\t\t\t{0} /* {1} */,{2}", Group.GroupGuid, Group.GroupName, ProjectFileGenerator.NewLine));
+				Content.Append(String.Format("\t\t\t\t{0} /* {1} */,{2}", Group.GroupGuid, Group.GroupName, ProjectFileGenerator.NewLine));
 			}
 
 			foreach (XcodeSourceFile File in RootGroup.Files)
 			{
-				Content.Append(string.Format("\t\t\t\t{0} /* {1} */,{2}", File.FileRefGuid, File.Reference.GetFileName(), ProjectFileGenerator.NewLine));
+				Content.Append(String.Format("\t\t\t\t{0} /* {1} */,{2}", File.FileRefGuid, File.Reference.GetFileName(), ProjectFileGenerator.NewLine));
 			}
 
-			Content.Append(string.Format("\t\t\t\t{0} /* Products */,{1}", ProductRefGroupGuid, ProjectFileGenerator.NewLine));
+			Content.Append(String.Format("\t\t\t\t{0} /* Products */,{1}", ProductRefGroupGuid, ProjectFileGenerator.NewLine));
 			Content.Append("\t\t\t);" + ProjectFileGenerator.NewLine);
 			Content.Append("\t\t\tsourceTree = \"<group>\";" + ProjectFileGenerator.NewLine);
 			Content.Append("\t\t};" + ProjectFileGenerator.NewLine);
@@ -578,13 +578,13 @@ namespace UnrealBuildTool.XcodeProjectLegacy
 			}
 
 			// Products group
-			Content.Append(string.Format("\t\t{0} /* Products */ = {{{1}", ProductRefGroupGuid, ProjectFileGenerator.NewLine));
+			Content.Append(String.Format("\t\t{0} /* Products */ = {{{1}", ProductRefGroupGuid, ProjectFileGenerator.NewLine));
 			Content.Append("\t\t\tisa = PBXGroup;" + ProjectFileGenerator.NewLine);
 			Content.Append("\t\t\tchildren = (" + ProjectFileGenerator.NewLine);
-			Content.Append(string.Format("\t\t\t\t{0} /* {1} */,{2}", TargetAppGuid, TargetName, ProjectFileGenerator.NewLine));
+			Content.Append(String.Format("\t\t\t\t{0} /* {1} */,{2}", TargetAppGuid, TargetName, ProjectFileGenerator.NewLine));
 			foreach (XcodeExtensionInfo EI in AllExtensions)
 			{
-				Content.Append(string.Format("\t\t\t\t{0} /* {1} */,{2}", EI.ProductGuid, EI.Name, ProjectFileGenerator.NewLine));
+				Content.Append(String.Format("\t\t\t\t{0} /* {1} */,{2}", EI.ProductGuid, EI.Name, ProjectFileGenerator.NewLine));
 			}
 			Content.Append("\t\t\t);" + ProjectFileGenerator.NewLine);
 			Content.Append("\t\t\tname = Products;" + ProjectFileGenerator.NewLine);
@@ -638,11 +638,11 @@ namespace UnrealBuildTool.XcodeProjectLegacy
 
 			Dictionary<string, string> BuildPhases = new Dictionary<string, string>();
 			// add optional build phases
-			if (!string.IsNullOrEmpty(CopyExtensionsBuildPhaseGuid))
+			if (!String.IsNullOrEmpty(CopyExtensionsBuildPhaseGuid))
 			{
 				BuildPhases.Add(CopyExtensionsBuildPhaseGuid, "Embed App Extensions");
 			}
-			if (!string.IsNullOrEmpty(ShellScriptSectionGuid))
+			if (!String.IsNullOrEmpty(ShellScriptSectionGuid))
 			{
 				BuildPhases.Add(ShellScriptSectionGuid, "Shell Script");
 			}
@@ -765,7 +765,7 @@ namespace UnrealBuildTool.XcodeProjectLegacy
 				"fi\\n";
 
 			Content.Append("/* Begin PBXShellScriptBuildPhase section */" + ProjectFileGenerator.NewLine);
-			Content.Append(string.Format("\t\t{0} /* Sign Frameworks */ = {{" + ProjectFileGenerator.NewLine, ShellScriptGuid));
+			Content.Append(String.Format("\t\t{0} /* Sign Frameworks */ = {{" + ProjectFileGenerator.NewLine, ShellScriptGuid));
 			Content.Append("\t\t\tisa = PBXShellScriptBuildPhase;" + ProjectFileGenerator.NewLine);
 			Content.Append("\t\t\tbuildActionMask = 2147483647;" + ProjectFileGenerator.NewLine);
 			Content.Append("\t\t\tfiles = (" + ProjectFileGenerator.NewLine);
@@ -778,7 +778,7 @@ namespace UnrealBuildTool.XcodeProjectLegacy
 			Content.Append("\t\t\tname = \"Sign Manual Frameworks\";" + ProjectFileGenerator.NewLine);
 			Content.Append("\t\t\trunOnlyForDeploymentPostprocessing = 0;" + ProjectFileGenerator.NewLine);
 			Content.Append("\t\t\tshellPath = /bin/sh;" + ProjectFileGenerator.NewLine);
-			Content.Append(string.Format("\t\t\tshellScript = \"{0}\";" + ProjectFileGenerator.NewLine, ShellScript));
+			Content.Append(String.Format("\t\t\tshellScript = \"{0}\";" + ProjectFileGenerator.NewLine, ShellScript));
 			Content.Append("\t\t};" + ProjectFileGenerator.NewLine);
 			Content.Append("/* End PBXShellScriptBuildPhase section */" + ProjectFileGenerator.NewLine);
 		}
@@ -1039,17 +1039,17 @@ namespace UnrealBuildTool.XcodeProjectLegacy
 			}
 			if (IOSRunTimeVersion != null)
 			{
-				Content.Append("\t\t\t\t\"VALID_ARCHS[sdk=iphoneos*]\" = \"" + string.Join(" ", this.SupportedIOSArchitectures) + "\";" + ProjectFileGenerator.NewLine);
+				Content.Append("\t\t\t\t\"VALID_ARCHS[sdk=iphoneos*]\" = \"" + String.Join(" ", SupportedIOSArchitectures) + "\";" + ProjectFileGenerator.NewLine);
 				Content.Append("\t\t\t\tIPHONEOS_DEPLOYMENT_TARGET = " + IOSRunTimeVersion + ";" + ProjectFileGenerator.NewLine);
 				Content.Append("\t\t\t\t\"PRODUCT_NAME[sdk=iphoneos*]\" = \"" + TargetName + "\";" + ProjectFileGenerator.NewLine); // @todo: change to Path.GetFileName(Config.IOSExecutablePath) when we stop using payload
 				Content.Append("\t\t\t\t\"TARGETED_DEVICE_FAMILY[sdk=iphoneos*]\" = \"" + IOSRunTimeDevices + "\";" + ProjectFileGenerator.NewLine);
 				Content.Append("\t\t\t\t\"SDKROOT[sdk=iphoneos]\" = iphoneos;" + ProjectFileGenerator.NewLine);
-				if (!string.IsNullOrEmpty(TEAM_IOS))
+				if (!String.IsNullOrEmpty(TEAM_IOS))
 				{
 					Content.Append("\t\t\t\t\"DEVELOPMENT_TEAM[sdk=iphoneos*]\" = " + TEAM_IOS + ";" + ProjectFileGenerator.NewLine);
 				}
 				Content.Append("\t\t\t\t\"CODE_SIGN_IDENTITY[sdk=iphoneos*]\" = \"" + IOS_CERT + "\";" + ProjectFileGenerator.NewLine);
-				if (!bAutomaticSigning && !string.IsNullOrEmpty(UUID_IOS))
+				if (!bAutomaticSigning && !String.IsNullOrEmpty(UUID_IOS))
 				{
 					Content.Append("\t\t\t\t\"PROVISIONING_PROFILE_SPECIFIER[sdk=iphoneos*]\" = \"" + UUID_IOS + "\";" + ProjectFileGenerator.NewLine);
 				}
@@ -1060,17 +1060,17 @@ namespace UnrealBuildTool.XcodeProjectLegacy
 			}
 			if (TVOSRunTimeVersion != null)
 			{
-				Content.Append("\t\t\t\t\"VALID_ARCHS[sdk=appletvos*]\" = \"" + string.Join(" ", this.SupportedIOSArchitectures) + "\";" + ProjectFileGenerator.NewLine);
+				Content.Append("\t\t\t\t\"VALID_ARCHS[sdk=appletvos*]\" = \"" + String.Join(" ", SupportedIOSArchitectures) + "\";" + ProjectFileGenerator.NewLine);
 				Content.Append("\t\t\t\tTVOS_DEPLOYMENT_TARGET = " + TVOSRunTimeVersion + ";" + ProjectFileGenerator.NewLine);
 				Content.Append("\t\t\t\t\"PRODUCT_NAME[sdk=appletvos*]\" = \"" + TargetName + "\";" + ProjectFileGenerator.NewLine); // @todo: change to Path.GetFileName(Config.TVOSExecutablePath) when we stop using payload
 				Content.Append("\t\t\t\t\"TARGETED_DEVICE_FAMILY[sdk=appletvos*]\" = \"" + TVOSRunTimeDevices + "\";" + ProjectFileGenerator.NewLine);
 				Content.Append("\t\t\t\t\"SDKROOT[sdk=appletvos]\" = appletvos;" + ProjectFileGenerator.NewLine);
-				if (!string.IsNullOrEmpty(TEAM_TVOS))
+				if (!String.IsNullOrEmpty(TEAM_TVOS))
 				{
 					Content.Append("\t\t\t\t\"DEVELOPMENT_TEAM[sdk=appletvos*]\" = " + TEAM_TVOS + ";" + ProjectFileGenerator.NewLine);
 				}
 				Content.Append("\t\t\t\t\"CODE_SIGN_IDENTITY[sdk=appletvos*]\" = \"" + TVOS_CERT + "\";" + ProjectFileGenerator.NewLine);
-				if (!bAutomaticSigning && !string.IsNullOrEmpty(UUID_TVOS))
+				if (!bAutomaticSigning && !String.IsNullOrEmpty(UUID_TVOS))
 				{
 					Content.Append("\t\t\t\t\"PROVISIONING_PROFILE_SPECIFIER[sdk=appletvos*]\" = \"" + UUID_TVOS + "\";" + ProjectFileGenerator.NewLine);
 				}
@@ -1081,7 +1081,7 @@ namespace UnrealBuildTool.XcodeProjectLegacy
 			}
 			if (bSupportMac)
 			{
-				Content.Append("\t\t\t\t\"VALID_ARCHS[sdk=macosx*]\" = \"" + string.Join(" ", SupportedMacArchitectures.Architectures.Select(x => x.AppleName)) + "\";" + ProjectFileGenerator.NewLine);
+				Content.Append("\t\t\t\t\"VALID_ARCHS[sdk=macosx*]\" = \"" + String.Join(" ", SupportedMacArchitectures.Architectures.Select(x => x.AppleName)) + "\";" + ProjectFileGenerator.NewLine);
 				Content.Append("\t\t\t\t\"PRODUCT_NAME[sdk=macosx*]\" = \"" + MacExecutableFileName + "\";" + ProjectFileGenerator.NewLine);
 				Content.Append("\t\t\t\t\"CONFIGURATION_BUILD_DIR[sdk=macosx*]\" = \"" + MacExecutableDir + "\";" + ProjectFileGenerator.NewLine);
 				Content.Append("\t\t\t\t\"SDKROOT[sdk=macosx*]\" = macosx;" + ProjectFileGenerator.NewLine);
@@ -1323,7 +1323,7 @@ namespace UnrealBuildTool.XcodeProjectLegacy
 			Content.Append("\t\t\tbuildSettings = {" + ProjectFileGenerator.NewLine);
 			if (bMacOnly)
 			{
-				Content.Append("\t\t\t\tVALID_ARCHS = \"" + string.Join(" ", SupportedMacArchitectures.Architectures.Select(x => x.AppleName)) + "\";" + ProjectFileGenerator.NewLine);
+				Content.Append("\t\t\t\tVALID_ARCHS = \"" + String.Join(" ", SupportedMacArchitectures.Architectures.Select(x => x.AppleName)) + "\";" + ProjectFileGenerator.NewLine);
 				Content.Append("\t\t\t\tSUPPORTED_PLATFORMS = \"macosx\";" + ProjectFileGenerator.NewLine);
 			}
 			else
@@ -1332,15 +1332,15 @@ namespace UnrealBuildTool.XcodeProjectLegacy
 				string SupportedPlatforms = "macosx";
 				if (InstalledPlatformInfo.IsValidPlatform(UnrealTargetPlatform.IOS, EProjectType.Code))
 				{
-					ValidArchs = ValidArchs.Union(this.SupportedIOSArchitectures);
+					ValidArchs = ValidArchs.Union(SupportedIOSArchitectures);
 					SupportedPlatforms += " iphoneos";
 				}
 				if (InstalledPlatformInfo.IsValidPlatform(UnrealTargetPlatform.TVOS, EProjectType.Code))
 				{
-					ValidArchs = ValidArchs.Union(this.SupportedIOSArchitectures);
+					ValidArchs = ValidArchs.Union(SupportedIOSArchitectures);
 					SupportedPlatforms += " appletvos";
 				}
-				Content.Append("\t\t\t\tVALID_ARCHS = \"" + string.Join(" ", ValidArchs.Select(x => x.AppleName)) + "\";" + ProjectFileGenerator.NewLine);
+				Content.Append("\t\t\t\tVALID_ARCHS = \"" + String.Join(" ", ValidArchs.Select(x => x.AppleName)) + "\";" + ProjectFileGenerator.NewLine);
 				Content.Append("\t\t\t\tSUPPORTED_PLATFORMS = \"" + SupportedPlatforms + "\";" + ProjectFileGenerator.NewLine);
 			}
 			Content.Append("\t\t\t\tGCC_PREPROCESSOR_DEFINITIONS = ();" + ProjectFileGenerator.NewLine);
@@ -1506,7 +1506,7 @@ namespace UnrealBuildTool.XcodeProjectLegacy
 											DirectoryReference MacBinaryDir = DirectoryReference.Combine(OutputDirectory, "Mac");
 											DirectoryReference IOSBinaryDir = DirectoryReference.Combine(OutputDirectory, "IOS");
 											DirectoryReference TVOSBinaryDir = DirectoryReference.Combine(OutputDirectory, "TVOS");
-											if (!string.IsNullOrEmpty(ProjectTarget.TargetRules.ExeBinariesSubFolder))
+											if (!String.IsNullOrEmpty(ProjectTarget.TargetRules.ExeBinariesSubFolder))
 											{
 												MacBinaryDir = DirectoryReference.Combine(MacBinaryDir, ProjectTarget.TargetRules.ExeBinariesSubFolder);
 												IOSBinaryDir = DirectoryReference.Combine(IOSBinaryDir, ProjectTarget.TargetRules.ExeBinariesSubFolder);
@@ -1660,7 +1660,7 @@ namespace UnrealBuildTool.XcodeProjectLegacy
 			Content.Append("               ReferencedContainer = \"container:" + TargetName + ".xcodeproj\">" + ProjectFileGenerator.NewLine);
 			Content.Append("            </BuildableReference>" + ProjectFileGenerator.NewLine);
 			Content.Append("      </BuildableProductRunnable>" + ProjectFileGenerator.NewLine);
-			if (string.IsNullOrEmpty(OldCommandLineArguments))
+			if (String.IsNullOrEmpty(OldCommandLineArguments))
 			{
 				if (bHasEditorConfiguration && GameProjectPath != "")
 				{
diff --git a/Engine/Source/Programs/UnrealBuildTool/Storage/Impl/FileSystemStorageProvider.cs b/Engine/Source/Programs/UnrealBuildTool/Storage/Impl/FileSystemStorageProvider.cs
--- a/Engine/Source/Programs/UnrealBuildTool/Storage/Impl/FileSystemStorageProvider.cs
+++ b/Engine/Source/Programs/UnrealBuildTool/Storage/Impl/FileSystemStorageProvider.cs
@@ -65,10 +65,10 @@ namespace UnrealBuildTool.Storage.Impl
 
 			public StorageWriter(FileReference Location)
 			{
-				this.FinalLocation = Location;
+				FinalLocation = Location;
 				DirectoryReference.CreateDirectory(FinalLocation.Directory);
 
-				this.TempLocation = new FileReference(String.Format("{0}.{1}", Location.FullName, ProcessId));
+				TempLocation = new FileReference(String.Format("{0}.{1}", Location.FullName, ProcessId));
 				Stream = FileReference.Open(TempLocation, FileMode.Create, FileAccess.Write, FileShare.Read | FileShare.Delete);
 			}
 
diff --git a/Engine/Source/Programs/UnrealBuildTool/System/Action.cs b/Engine/Source/Programs/UnrealBuildTool/System/Action.cs
--- a/Engine/Source/Programs/UnrealBuildTool/System/Action.cs
+++ b/Engine/Source/Programs/UnrealBuildTool/System/Action.cs
@@ -696,7 +696,7 @@ namespace UnrealBuildTool
 		{
 			this.Inner = Inner;
 			this.Target = Target;
-			this.IsHighPriority = Inner.bIsHighPriority ? 1 : 0;
+			IsHighPriority = Inner.bIsHighPriority ? 1 : 0;
 		}
 
 		/// <summary>
@@ -709,7 +709,9 @@ namespace UnrealBuildTool
 			if (VisitedActions.Add(this))
 			{
 				if (bIsHighPriority)
+				{
 					Interlocked.Exchange(ref IsHighPriority, 1);
+				}
 
 				Interlocked.Increment(ref NumTotalDependentActions);
 				foreach (LinkedAction PrerequisiteAction in PrerequisiteActions)
diff --git a/Engine/Source/Programs/UnrealBuildTool/System/ActionExecutor.cs b/Engine/Source/Programs/UnrealBuildTool/System/ActionExecutor.cs
--- a/Engine/Source/Programs/UnrealBuildTool/System/ActionExecutor.cs
+++ b/Engine/Source/Programs/UnrealBuildTool/System/ActionExecutor.cs
@@ -17,7 +17,7 @@ namespace UnrealBuildTool
 			get;
 		}
 
-		static protected double MemoryPerActionBytesOverride
+		protected static double MemoryPerActionBytesOverride
 		{
 			get;
 			private set;
diff --git a/Engine/Source/Programs/UnrealBuildTool/System/BuildHostPlatform.cs b/Engine/Source/Programs/UnrealBuildTool/System/BuildHostPlatform.cs
--- a/Engine/Source/Programs/UnrealBuildTool/System/BuildHostPlatform.cs
+++ b/Engine/Source/Programs/UnrealBuildTool/System/BuildHostPlatform.cs
@@ -36,7 +36,7 @@ namespace UnrealBuildTool
 		/// <summary>
 		/// Host platform singleton.
 		/// </summary>
-		static public BuildHostPlatform Current
+		public static BuildHostPlatform Current
 		{
 			get
 			{
@@ -66,22 +66,22 @@ namespace UnrealBuildTool
 		/// <summary>
 		/// Gets the current host platform type.
 		/// </summary>
-		abstract public UnrealTargetPlatform Platform { get; }
+		public abstract UnrealTargetPlatform Platform { get; }
 
 		/// <summary>
 		/// Gets the path to the shell for this platform
 		/// </summary>
-		abstract public FileReference Shell { get; }
+		public abstract FileReference Shell { get; }
 
 		/// <summary>
 		/// The type of shell returned by the Shell parameter
 		/// </summary>
-		abstract public ShellType ShellType { get; }
+		public abstract ShellType ShellType { get; }
 
 		/// <summary>
 		/// The executable binary suffix for this platform
 		/// </summary>
-		abstract public string BinarySuffix { get; }
+		public abstract string BinarySuffix { get; }
 
 		/// <summary>
 		/// Class that holds information about a running process
@@ -278,7 +278,7 @@ namespace UnrealBuildTool
 
 		public override ShellType ShellType => ShellType.Sh;
 
-		public override string BinarySuffix => string.Empty;
+		public override string BinarySuffix => String.Empty;
 
 		/// <summary>
 		/// Currently Mono returns incomplete process names in Process.GetProcesses() so we need to parse 'ps' output.
@@ -398,7 +398,7 @@ namespace UnrealBuildTool
 
 		public override ShellType ShellType => ShellType.Sh;
 
-		public override string BinarySuffix => string.Empty;
+		public override string BinarySuffix => String.Empty;
 
 		/// <summary>
 		/// Currently Mono returns incomplete process names in Process.GetProcesses() so we need to use /proc
diff --git a/Engine/Source/Programs/UnrealBuildTool/System/ConfigCache.cs b/Engine/Source/Programs/UnrealBuildTool/System/ConfigCache.cs
--- a/Engine/Source/Programs/UnrealBuildTool/System/ConfigCache.cs
+++ b/Engine/Source/Programs/UnrealBuildTool/System/ConfigCache.cs
@@ -720,7 +720,7 @@ namespace UnrealBuildTool
 				}
 
 				// move past blank lines between sections
-				while (string.IsNullOrWhiteSpace(Lines[NextSectionIndex - 1]))
+				while (String.IsNullOrWhiteSpace(Lines[NextSectionIndex - 1]))
 				{
 					NextSectionIndex--;
 				}
diff --git a/Engine/Source/Programs/UnrealBuildTool/System/ConfigFile.cs b/Engine/Source/Programs/UnrealBuildTool/System/ConfigFile.cs
--- a/Engine/Source/Programs/UnrealBuildTool/System/ConfigFile.cs
+++ b/Engine/Source/Programs/UnrealBuildTool/System/ConfigFile.cs
@@ -137,11 +137,11 @@ namespace UnrealBuildTool
 
 
 		// Remap of config names/sections
-		static private Dictionary<string, string> SectionNameRemap = new();
-		static private Dictionary<string, Dictionary<string, string>> SectionKeyRemap = new();
-		static private HashSet<string> WarnedKeys = new(StringComparer.InvariantCultureIgnoreCase);
+		private static Dictionary<string, string> SectionNameRemap = new();
+		private static Dictionary<string, Dictionary<string, string>> SectionKeyRemap = new();
+		private static HashSet<string> WarnedKeys = new(StringComparer.InvariantCultureIgnoreCase);
 
-		static private string RemapSectionOrKey(Dictionary<string, string>? Remap, string Key, string Context)
+		private static string RemapSectionOrKey(Dictionary<string, string>? Remap, string Key, string Context)
 		{
 			if (Remap != null)
 			{
diff --git a/Engine/Source/Programs/UnrealBuildTool/System/ConfigHierarchy.cs b/Engine/Source/Programs/UnrealBuildTool/System/ConfigHierarchy.cs
--- a/Engine/Source/Programs/UnrealBuildTool/System/ConfigHierarchy.cs
+++ b/Engine/Source/Programs/UnrealBuildTool/System/ConfigHierarchy.cs
@@ -478,7 +478,7 @@ namespace UnrealBuildTool
 		/// <param name="Text">The text to parse</param>
 		/// <param name="Value">The parsed value, if successful</param>
 		/// <returns>True if the text was parsed, false otherwise</returns>
-		static public bool TryParse(string Text, out bool Value)
+		public static bool TryParse(string Text, out bool Value)
 		{
 			// C# Boolean type expects "False" or "True" but since we're not case sensitive, we need to suppor that manually
 			if (Text == "1" || Text.Equals("true", StringComparison.InvariantCultureIgnoreCase))
@@ -504,7 +504,7 @@ namespace UnrealBuildTool
 		/// <param name="Text">The text to parse</param>
 		/// <param name="Value">The parsed value, if successful</param>
 		/// <returns>True if the text was parsed, false otherwise</returns>
-		static public bool TryParse(string Text, out int Value)
+		public static bool TryParse(string Text, out int Value)
 		{
 			return Int32.TryParse(Text, out Value);
 		}
@@ -1238,14 +1238,14 @@ namespace UnrealBuildTool
 			string? AltRegex = null;
 			if (bIsArrayProperty)
 			{
-				PrimaryRegex = string.Format("{0}\\s*=\\s*\\((.*?)\\)", Property);
+				PrimaryRegex = String.Format("{0}\\s*=\\s*\\((.*?)\\)", Property);
 			}
 			else
 			{
 				// handle quoted strings, allowing for escaped quotation marks (basically doing " followed by whatever, until we see a quote that was not proceeded by a \, and gather the whole mess in an outer group)
-				PrimaryRegex = string.Format("{0}\\s*=\\s*\"((.*?)[^\\\\])\"", Property);
+				PrimaryRegex = String.Format("{0}\\s*=\\s*\"((.*?)[^\\\\])\"", Property);
 				// when no quotes, we skip over whitespace, and we end when we see whitespace, a comma or a ). This will handle (Ip = 192.168.0.1 , Name=....) , and return only '192.168.0.1'
-				AltRegex = string.Format("{0}\\s*=\\s*(.*?)[\\s,\\)]", Property);
+				AltRegex = String.Format("{0}\\s*=\\s*(.*?)[\\s,\\)]", Property);
 			}
 
 			// attempt to match it!
@@ -1293,8 +1293,8 @@ namespace UnrealBuildTool
 		public static string? GetMapValue(string Input, string Key)
 		{
 			// handle quoted strings, allowing for escaped quotation marks (and possibly the key in quotes as well)
-			string PrimaryRegex = string.Format("{0}\"?\\s*,\\s*\"((.*?)[^\\\\])\"", Key);
-			string AltRegex = string.Format("{0}\"?\\s*,\\s*(.*?)[\\s,\\)]", Key);
+			string PrimaryRegex = String.Format("{0}\"?\\s*,\\s*\"((.*?)[^\\\\])\"", Key);
+			string AltRegex = String.Format("{0}\"?\\s*,\\s*(.*?)[\\s,\\)]", Key);
 
 			// attempt to match it!
 			Match Result = Regex.Match(Input, PrimaryRegex);
diff --git a/Engine/Source/Programs/UnrealBuildTool/System/CppCompileEnvironment.cs b/Engine/Source/Programs/UnrealBuildTool/System/CppCompileEnvironment.cs
--- a/Engine/Source/Programs/UnrealBuildTool/System/CppCompileEnvironment.cs
+++ b/Engine/Source/Programs/UnrealBuildTool/System/CppCompileEnvironment.cs
@@ -180,7 +180,7 @@ namespace UnrealBuildTool
 					return PerArchPrecompiledHeaderFile;
 				}
 			}
-			return this.PrecompiledHeaderFile;
+			return PrecompiledHeaderFile;
 		}
 	}
 
@@ -661,12 +661,12 @@ namespace UnrealBuildTool
 			this.Configuration = Configuration;
 			this.Architectures = Architectures;
 			this.MetadataCache = MetadataCache;
-			this.SharedPCHs = new List<PrecompiledHeaderTemplate>();
-			this.UserIncludePaths = new HashSet<DirectoryReference>();
-			this.SystemIncludePaths = new HashSet<DirectoryReference>();
-			this.SharedUserIncludePaths = new HashSet<DirectoryReference>();
-			this.SharedSystemIncludePaths = new HashSet<DirectoryReference>();
-			this.ModuleInterfacePaths = new HashSet<DirectoryReference>();
+			SharedPCHs = new List<PrecompiledHeaderTemplate>();
+			UserIncludePaths = new HashSet<DirectoryReference>();
+			SystemIncludePaths = new HashSet<DirectoryReference>();
+			SharedUserIncludePaths = new HashSet<DirectoryReference>();
+			SharedSystemIncludePaths = new HashSet<DirectoryReference>();
+			ModuleInterfacePaths = new HashSet<DirectoryReference>();
 		}
 
 		/// <summary>
diff --git a/Engine/Source/Programs/UnrealBuildTool/System/DataDrivenPlatformInfo.cs b/Engine/Source/Programs/UnrealBuildTool/System/DataDrivenPlatformInfo.cs
--- a/Engine/Source/Programs/UnrealBuildTool/System/DataDrivenPlatformInfo.cs
+++ b/Engine/Source/Programs/UnrealBuildTool/System/DataDrivenPlatformInfo.cs
@@ -171,7 +171,7 @@ namespace UnrealBuildTool
 						List<string> Chain = new List<string>();
 						if (IniParents.TryGetValue(Pair.Key, out CurrentPlatform))
 						{
-							while (!string.IsNullOrEmpty(CurrentPlatform))
+							while (!String.IsNullOrEmpty(CurrentPlatform))
 							{
 								// insert at 0 to reverse the order
 								Chain.Insert(0, CurrentPlatform);
diff --git a/Engine/Source/Programs/UnrealBuildTool/System/DeployExports.cs b/Engine/Source/Programs/UnrealBuildTool/System/DeployExports.cs
--- a/Engine/Source/Programs/UnrealBuildTool/System/DeployExports.cs
+++ b/Engine/Source/Programs/UnrealBuildTool/System/DeployExports.cs
@@ -23,7 +23,7 @@ namespace UnrealBuildTool
 		/// <returns>The default deploy folder name.</returns>
 		public static string GetDefaultDeployFolder(string ShortProjectName)
 		{
-			string DeployFolder = string.Format("{0}-{1}-{2}",
+			string DeployFolder = String.Format("{0}-{1}-{2}",
 				ShortProjectName,
 				Unreal.RootDirectory.GetDirectoryName(),
 				Environment.UserName);
diff --git a/Engine/Source/Programs/UnrealBuildTool/System/DynamicCompilation.cs b/Engine/Source/Programs/UnrealBuildTool/System/DynamicCompilation.cs
--- a/Engine/Source/Programs/UnrealBuildTool/System/DynamicCompilation.cs
+++ b/Engine/Source/Programs/UnrealBuildTool/System/DynamicCompilation.cs
@@ -361,7 +361,7 @@ namespace UnrealBuildTool
 
 		private static string FormatVersionNumber(ReadOnlyBuildVersion Version)
 		{
-			return string.Format("{0}.{1}.{2}", Version.MajorVersion, Version.MinorVersion, Version.PatchVersion);
+			return String.Format("{0}.{1}.{2}", Version.MajorVersion, Version.MinorVersion, Version.PatchVersion);
 		}
 	}
 }
diff --git a/Engine/Source/Programs/UnrealBuildTool/System/EncryptionAndSigning.cs b/Engine/Source/Programs/UnrealBuildTool/System/EncryptionAndSigning.cs
--- a/Engine/Source/Programs/UnrealBuildTool/System/EncryptionAndSigning.cs
+++ b/Engine/Source/Programs/UnrealBuildTool/System/EncryptionAndSigning.cs
@@ -192,7 +192,7 @@ namespace UnrealBuildTool
 			/// </summary>
 			public bool ContainsEncryptionKey(string KeyGuid)
 			{
-				if (string.IsNullOrEmpty(KeyGuid))
+				if (String.IsNullOrEmpty(KeyGuid))
 				{
 					return false;
 				}
@@ -285,7 +285,7 @@ namespace UnrealBuildTool
 						(Settings.SigningKey.PublicKey.Exponent.Length > 64) ||
 						(Settings.SigningKey.PublicKey.Modulus.Length > 64))
 					{
-						throw new Exception(string.Format("[{0}] Signing keys parsed from encryption.ini are too long. They must be a maximum of 64 bytes long!", InProjectDirectory));
+						throw new Exception(String.Format("[{0}] Signing keys parsed from encryption.ini are too long. They must be a maximum of 64 bytes long!", InProjectDirectory));
 					}
 				}
 
@@ -340,7 +340,7 @@ namespace UnrealBuildTool
 				// Parse encryption key
 				string EncryptionKeyString;
 				Ini.GetString(SectionName, "EncryptionKey", out EncryptionKeyString);
-				if (!string.IsNullOrEmpty(EncryptionKeyString))
+				if (!String.IsNullOrEmpty(EncryptionKeyString))
 				{
 					Settings.EncryptionKey = new EncryptionKey();
 					Settings.EncryptionKey.Key = System.Convert.FromBase64String(EncryptionKeyString);
@@ -515,7 +515,7 @@ namespace UnrealBuildTool
 				int CharsToParse = Math.Min(2, InString.Length);
 				string Value = InString.Substring(InString.Length - CharsToParse);
 				InString = InString.Substring(0, InString.Length - CharsToParse);
-				Bytes.Add(byte.Parse(Value, System.Globalization.NumberStyles.AllowHexSpecifier));
+				Bytes.Add(Byte.Parse(Value, System.Globalization.NumberStyles.AllowHexSpecifier));
 			}
 
 			while (Bytes.Count < InMinimumLength)
diff --git a/Engine/Source/Programs/UnrealBuildTool/System/ExternalExecution.cs b/Engine/Source/Programs/UnrealBuildTool/System/ExternalExecution.cs
--- a/Engine/Source/Programs/UnrealBuildTool/System/ExternalExecution.cs
+++ b/Engine/Source/Programs/UnrealBuildTool/System/ExternalExecution.cs
@@ -176,12 +176,12 @@ namespace UnrealBuildTool
 			this.ModuleDirectories = ModuleDirectories;
 			this.ModuleIncudeBase = ModuleIncudeBase;
 			this.ModuleType = ModuleType.ToString();
-			this.OverrideModuleType = OverrideType.ToString();
-			this.PublicUObjectClassesHeaders = new List<FileItem>();
-			this.PublicUObjectHeaders = new List<FileItem>();
-			this.InternalUObjectHeaders = new List<FileItem>();
-			this.PrivateUObjectHeaders = new List<FileItem>();
-			this.PublicDefines = new List<string>();
+			OverrideModuleType = OverrideType.ToString();
+			PublicUObjectClassesHeaders = new List<FileItem>();
+			PublicUObjectHeaders = new List<FileItem>();
+			InternalUObjectHeaders = new List<FileItem>();
+			PrivateUObjectHeaders = new List<FileItem>();
+			PublicDefines = new List<string>();
 			this.GeneratedCodeDirectory = GeneratedCodeDirectory;
 			this.GeneratedCodeVersion = GeneratedCodeVersion;
 			this.bIsReadOnly = bIsReadOnly;
diff --git a/Engine/Source/Programs/UnrealBuildTool/System/InstalledPlatformInfo.cs b/Engine/Source/Programs/UnrealBuildTool/System/InstalledPlatformInfo.cs
--- a/Engine/Source/Programs/UnrealBuildTool/System/InstalledPlatformInfo.cs
+++ b/Engine/Source/Programs/UnrealBuildTool/System/InstalledPlatformInfo.cs
@@ -232,7 +232,7 @@ namespace UnrealBuildTool
 
 		private static bool ParseSubValue(string TrimmedLine, string Match, out string Result)
 		{
-			Result = string.Empty;
+			Result = String.Empty;
 			int MatchIndex = TrimmedLine.IndexOf(Match);
 			if (MatchIndex < 0)
 			{
@@ -376,7 +376,7 @@ namespace UnrealBuildTool
 				{
 					// Check whether filter accepts this configuration and it has required file
 					if (ConfigFilter(PlatformConfiguration)
-					&& (string.IsNullOrEmpty(PlatformConfiguration.RequiredFile)
+					&& (String.IsNullOrEmpty(PlatformConfiguration.RequiredFile)
 					|| File.Exists(PlatformConfiguration.RequiredFile)))
 					{
 						return true;
@@ -408,25 +408,25 @@ namespace UnrealBuildTool
 		private static void WriteConfigFileEntry(InstalledPlatformConfiguration Config, ref List<String> OutEntries)
 		{
 			string ConfigDescription = "+InstalledPlatformConfigurations=(";
-			ConfigDescription += string.Format("PlatformName=\"{0}\", ", Config.Platform.ToString());
+			ConfigDescription += String.Format("PlatformName=\"{0}\", ", Config.Platform.ToString());
 			if (Config.Configuration != UnrealTargetConfiguration.Unknown)
 			{
-				ConfigDescription += string.Format("Configuration=\"{0}\", ", Config.Configuration.ToString());
+				ConfigDescription += String.Format("Configuration=\"{0}\", ", Config.Configuration.ToString());
 			}
 			if (Config.PlatformType != TargetType.Program)
 			{
-				ConfigDescription += string.Format("PlatformType=\"{0}\", ", Config.PlatformType.ToString());
+				ConfigDescription += String.Format("PlatformType=\"{0}\", ", Config.PlatformType.ToString());
 			}
-			ConfigDescription += string.Format("Architecture=\"{0}\", ", Config.Architecture);
-			if (!string.IsNullOrEmpty(Config.RequiredFile))
+			ConfigDescription += String.Format("Architecture=\"{0}\", ", Config.Architecture);
+			if (!String.IsNullOrEmpty(Config.RequiredFile))
 			{
-				ConfigDescription += string.Format("RequiredFile=\"{0}\", ", Config.RequiredFile);
+				ConfigDescription += String.Format("RequiredFile=\"{0}\", ", Config.RequiredFile);
 			}
 			if (Config.ProjectType != EProjectType.Unknown)
 			{
-				ConfigDescription += string.Format("ProjectType=\"{0}\", ", Config.ProjectType.ToString());
+				ConfigDescription += String.Format("ProjectType=\"{0}\", ", Config.ProjectType.ToString());
 			}
-			ConfigDescription += string.Format("bCanBeDisplayed={0})", Config.bCanBeDisplayed.ToString());
+			ConfigDescription += String.Format("bCanBeDisplayed={0})", Config.bCanBeDisplayed.ToString());
 
 			OutEntries.Add(ConfigDescription);
 		}
diff --git a/Engine/Source/Programs/UnrealBuildTool/System/JunkDeleter.cs b/Engine/Source/Programs/UnrealBuildTool/System/JunkDeleter.cs
--- a/Engine/Source/Programs/UnrealBuildTool/System/JunkDeleter.cs
+++ b/Engine/Source/Programs/UnrealBuildTool/System/JunkDeleter.cs
@@ -12,7 +12,7 @@ namespace UnrealBuildTool
 		/// <summary>
 		/// Loads JunkManifest.txt file and removes all junk files/folders defined in it.
 		/// </summary>
-		static public void DeleteJunk(ILogger Logger)
+		public static void DeleteJunk(ILogger Logger)
 		{
 			List<string> JunkManifest = LoadJunkManifest();
 			DeleteAllJunk(JunkManifest, Logger);
@@ -22,7 +22,7 @@ namespace UnrealBuildTool
 		/// Loads JunkManifest.txt file.
 		/// </summary>
 		/// <returns>Junk manifest file contents.</returns>
-		static private List<string> LoadJunkManifest()
+		private static List<string> LoadJunkManifest()
 		{
 			string ManifestPath = ".." + Path.DirectorySeparatorChar + "Build" + Path.DirectorySeparatorChar + "JunkManifest.txt";
 			List<string> JunkManifest = new List<string>();
@@ -93,7 +93,7 @@ namespace UnrealBuildTool
 		/// </summary>
 		/// <param name="JunkManifest">JunkManifest.txt entries.</param>
 		/// <param name="Logger">Logger for output</param>
-		static private void DeleteAllJunk(List<string> JunkManifest, ILogger Logger)
+		private static void DeleteAllJunk(List<string> JunkManifest, ILogger Logger)
 		{
 			foreach (string Junk in JunkManifest)
 			{
@@ -135,7 +135,7 @@ namespace UnrealBuildTool
 			}
 		}
 
-		static private bool IsFile(string PathToCheck)
+		private static bool IsFile(string PathToCheck)
 		{
 			string FileName = Path.GetFileName(PathToCheck);
 			if (String.IsNullOrEmpty(FileName) == false)
@@ -161,7 +161,7 @@ namespace UnrealBuildTool
 		/// </summary>
 		/// <param name="DirectoryPath">Path.</param>
 		/// <param name="Logger">Logger for output</param>
-		static private void DeleteDirectory(string DirectoryPath, ILogger Logger)
+		private static void DeleteDirectory(string DirectoryPath, ILogger Logger)
 		{
 			try
 			{
@@ -180,7 +180,7 @@ namespace UnrealBuildTool
 		/// </summary>
 		/// <param name="Filename">Filename.</param>
 		/// <param name="Logger">Logger for output</param>
-		static private void DeleteFile(string Filename, ILogger Logger)
+		private static void DeleteFile(string Filename, ILogger Logger)
 		{
 			try
 			{
diff --git a/Engine/Source/Programs/UnrealBuildTool/System/Plugins.cs b/Engine/Source/Programs/UnrealBuildTool/System/Plugins.cs
--- a/Engine/Source/Programs/UnrealBuildTool/System/Plugins.cs
+++ b/Engine/Source/Programs/UnrealBuildTool/System/Plugins.cs
@@ -259,7 +259,7 @@ namespace UnrealBuildTool
 			{
 				foreach (var Plugin in Pair.Value)
 				{
-					if (string.Equals(Plugin.Name, PluginName, StringComparison.InvariantCultureIgnoreCase))
+					if (String.Equals(Plugin.Name, PluginName, StringComparison.InvariantCultureIgnoreCase))
 					{
 						return Plugin;
 					}
diff --git a/Engine/Source/Programs/UnrealBuildTool/System/RestrictedFolders.cs b/Engine/Source/Programs/UnrealBuildTool/System/RestrictedFolders.cs
--- a/Engine/Source/Programs/UnrealBuildTool/System/RestrictedFolders.cs
+++ b/Engine/Source/Programs/UnrealBuildTool/System/RestrictedFolders.cs
@@ -50,7 +50,7 @@ namespace UnrealBuildTool
 		/// </summary>
 		/// <param name="Name">Name of the folder</param>
 		/// <returns>New restricted folder instance</returns>
-		static private RestrictedFolder FindOrAddByName(string Name)
+		private static RestrictedFolder FindOrAddByName(string Name)
 		{
 			return new RestrictedFolder(StringRegistry.FindOrAddByName(Name));
 		}
diff --git a/Engine/Source/Programs/UnrealBuildTool/System/RulesAssembly.cs b/Engine/Source/Programs/UnrealBuildTool/System/RulesAssembly.cs
--- a/Engine/Source/Programs/UnrealBuildTool/System/RulesAssembly.cs
+++ b/Engine/Source/Programs/UnrealBuildTool/System/RulesAssembly.cs
@@ -37,7 +37,10 @@ namespace UnrealBuildTool
 			{
 				return CompiledAssembly.GetName().Name;
 			}
-			else return null;
+			else
+			{
+				return null;
+			}
 		}
 
 		/// <summary>
@@ -466,7 +469,7 @@ namespace UnrealBuildTool
 					(Target.ProjectFile != null && ModuleFileName.ContainsAnyNames(DisallowedPlatformsAndGroups, Target.ProjectFile.Directory)))
 				{
 					throw new BuildException("Platform module file {0} is not allowed (only platforms '{1}', and their groups, are allowed. This indicates a module reference not being checked with something like IsPlatformAvailableForTarget()).",
-						ModuleFileName, string.Join(",", Target.OptedInModulePlatforms));
+						ModuleFileName, String.Join(",", Target.OptedInModulePlatforms));
 				}
 			}
 
diff --git a/Engine/Source/Programs/UnrealBuildTool/System/SingleInstanceMutex.cs b/Engine/Source/Programs/UnrealBuildTool/System/SingleInstanceMutex.cs
--- a/Engine/Source/Programs/UnrealBuildTool/System/SingleInstanceMutex.cs
+++ b/Engine/Source/Programs/UnrealBuildTool/System/SingleInstanceMutex.cs
@@ -55,7 +55,7 @@ namespace UnrealBuildTool
 		public static string GetUniqueMutexForPath(string Name, string UniquePath)
 		{
 			// generate a md5 hash of the path, as GetHashCode is not guaranteed to generate a stable hash
-			return string.Format("Global\\{0}_{1}", Name, ContentHash.MD5(UniquePath.ToUpperInvariant()));
+			return String.Format("Global\\{0}_{1}", Name, ContentHash.MD5(UniquePath.ToUpperInvariant()));
 		}
 
 		/// <summary>
diff --git a/Engine/Source/Programs/UnrealBuildTool/System/SourceFileMetadataCache.cs b/Engine/Source/Programs/UnrealBuildTool/System/SourceFileMetadataCache.cs
--- a/Engine/Source/Programs/UnrealBuildTool/System/SourceFileMetadataCache.cs
+++ b/Engine/Source/Programs/UnrealBuildTool/System/SourceFileMetadataCache.cs
@@ -160,7 +160,7 @@ namespace UnrealBuildTool
 		private SourceFileMetadataCache(FileReference Location, DirectoryReference BaseDir, SourceFileMetadataCache? Parent, ILogger Logger)
 		{
 			this.Location = Location;
-			this.BaseDirectory = BaseDir;
+			BaseDirectory = BaseDir;
 			this.Parent = Parent;
 			this.Logger = Logger;
 
@@ -315,7 +315,10 @@ namespace UnrealBuildTool
 				{
 					IncludeSpan = IncludeSpan.Slice("#include".Length).TrimStart();
 					if (IncludeSpan.IsEmpty)
+					{
 						continue;
+					}
+
 					char EndChar;
 					bool TrimQuotation = true;
 					if (IncludeSpan[0] == '"')
@@ -357,9 +360,13 @@ namespace UnrealBuildTool
 				{
 					ReadOnlySpan<char> Span = Line.AsSpan(HeaderUnitIndex + "HEADER_UNIT_".Length);
 					if (Span.StartsWith("UNSUPPORTED"))
+					{
 						HeaderFileInfo.UnitType = HeaderUnitType.Unsupported;
+					}
 					else if (Span.StartsWith("SKIP"))
+					{
 						HeaderFileInfo.UnitType = HeaderUnitType.Skip;
+					}
 				}
 			}
 
diff --git a/Engine/Source/Programs/UnrealBuildTool/System/SourceFileWorkingSet.cs b/Engine/Source/Programs/UnrealBuildTool/System/SourceFileWorkingSet.cs
--- a/Engine/Source/Programs/UnrealBuildTool/System/SourceFileWorkingSet.cs
+++ b/Engine/Source/Programs/UnrealBuildTool/System/SourceFileWorkingSet.cs
@@ -108,9 +108,9 @@ namespace UnrealBuildTool
 		public GitSourceFileWorkingSet(string GitPath, DirectoryReference RootDir, GitSourceFileWorkingSet? Inner, ILogger Logger)
 		{
 			this.RootDir = RootDir;
-			this.Files = new HashSet<FileReference>();
-			this.Directories = new List<DirectoryReference>();
-			this.ErrorOutput = new List<string>();
+			Files = new HashSet<FileReference>();
+			Directories = new List<DirectoryReference>();
+			ErrorOutput = new List<string>();
 			this.Inner = Inner;
 			this.Logger = Logger;
 
diff --git a/Engine/Source/Programs/UnrealBuildTool/System/Subnet.cs b/Engine/Source/Programs/UnrealBuildTool/System/Subnet.cs
--- a/Engine/Source/Programs/UnrealBuildTool/System/Subnet.cs
+++ b/Engine/Source/Programs/UnrealBuildTool/System/Subnet.cs
@@ -41,7 +41,7 @@ namespace UnrealBuildTool
 		public Subnet(IPAddress Prefix, int MaskBits)
 		{
 			this.Prefix = Prefix;
-			this.PrefixBytes = Prefix.GetAddressBytes();
+			PrefixBytes = Prefix.GetAddressBytes();
 			this.MaskBits = MaskBits;
 		}
 
@@ -54,7 +54,7 @@ namespace UnrealBuildTool
 		{
 			int SlashIdx = Text.IndexOf('/');
 			IPAddress Address = IPAddress.Parse(Text.Substring(0, SlashIdx));
-			return new Subnet(Address, int.Parse(Text.Substring(SlashIdx + 1)));
+			return new Subnet(Address, Int32.Parse(Text.Substring(SlashIdx + 1)));
 		}
 
 		/// <summary>
diff --git a/Engine/Source/Programs/UnrealBuildTool/System/SupportedPlatformGroupsAttribute.cs b/Engine/Source/Programs/UnrealBuildTool/System/SupportedPlatformGroupsAttribute.cs
--- a/Engine/Source/Programs/UnrealBuildTool/System/SupportedPlatformGroupsAttribute.cs
+++ b/Engine/Source/Programs/UnrealBuildTool/System/SupportedPlatformGroupsAttribute.cs
@@ -32,8 +32,8 @@ namespace UnrealBuildTool
 						SupportedPlatforms.UnionWith(UnrealTargetPlatform.GetValidPlatforms().Where(x => x.IsInGroup(Group)));
 						continue;
 					}
-					throw new BuildException(string.Format("The platform group name {0} is not a valid platform group name. Valid names are ({1})", Name,
-						string.Join(",", UnrealPlatformGroup.GetValidGroupNames())));
+					throw new BuildException(String.Format("The platform group name {0} is not a valid platform group name. Valid names are ({1})", Name,
+						String.Join(",", UnrealPlatformGroup.GetValidGroupNames())));
 				}
 			}
 			catch (BuildException Ex)
diff --git a/Engine/Source/Programs/UnrealBuildTool/System/SupportedPlatformsAttribute.cs b/Engine/Source/Programs/UnrealBuildTool/System/SupportedPlatformsAttribute.cs
--- a/Engine/Source/Programs/UnrealBuildTool/System/SupportedPlatformsAttribute.cs
+++ b/Engine/Source/Programs/UnrealBuildTool/System/SupportedPlatformsAttribute.cs
@@ -39,7 +39,7 @@ namespace UnrealBuildTool
 		/// <param name="Category">Category of platforms to add</param>
 		public SupportedPlatformsAttribute(UnrealPlatformClass Category)
 		{
-			this.Platforms = Utils.GetPlatformsInClass(Category);
+			Platforms = Utils.GetPlatformsInClass(Category);
 		}
 	}
 }
diff --git a/Engine/Source/Programs/UnrealBuildTool/System/TargetMakefile.cs b/Engine/Source/Programs/UnrealBuildTool/System/TargetMakefile.cs
--- a/Engine/Source/Programs/UnrealBuildTool/System/TargetMakefile.cs
+++ b/Engine/Source/Programs/UnrealBuildTool/System/TargetMakefile.cs
@@ -45,7 +45,7 @@ namespace UnrealBuildTool
 		/// </summary>
 		/// <param name="InlinedGenCppNames">Gen.cpp file names</param>
 		/// <returns>Hash of file names</returns>
-		static public int CalculateInlineGenCppHash(IEnumerable<string> InlinedGenCppNames)
+		public static int CalculateInlineGenCppHash(IEnumerable<string> InlinedGenCppNames)
 		{
 			int InlineHash = 0;
 			foreach (string FileName in InlinedGenCppNames)
@@ -409,9 +409,9 @@ namespace UnrealBuildTool
 			TargetType TargetType, ConfigValueTracker ConfigValueTracker, bool bDeployAfterCompile,
 			FileReference[]? UbtPlugins, FileReference[]? EnabledUbtPlugins, FileReference[]? EnabledUhtPlugins)
 		{
-			this.CreateTimeUtc = UnrealBuildTool.StartTimeUtc;
-			this.ModifiedTimeUtc = CreateTimeUtc;
-			this.Diagnostics = new List<string>();
+			CreateTimeUtc = UnrealBuildTool.StartTimeUtc;
+			ModifiedTimeUtc = CreateTimeUtc;
+			Diagnostics = new List<string>();
 			this.ExternalMetadata = ExternalMetadata;
 			this.ExecutableFile = ExecutableFile;
 			this.ReceiptFile = ReceiptFile;
@@ -423,23 +423,23 @@ namespace UnrealBuildTool
 			this.UbtPlugins = UbtPlugins;
 			this.EnabledUbtPlugins = EnabledUbtPlugins;
 			this.EnabledUhtPlugins = EnabledUhtPlugins;
-			this.Actions = new List<IExternalAction>();
-			this.OutputItems = new List<FileItem>();
-			this.ModuleNameToOutputItems = new Dictionary<string, FileItem[]>(StringComparer.OrdinalIgnoreCase);
-			this.HotReloadModuleNames = new HashSet<string>(StringComparer.OrdinalIgnoreCase);
-			this.SourceDirectories = new List<DirectoryItem>();
-			this.DirectoryToSourceFiles = new();
-			this.HeaderFiles = new();
-			this.WorkingSet = new HashSet<FileItem>();
-			this.CandidatesForWorkingSet = new HashSet<FileItem>();
-			this.UObjectModules = new List<UHTModuleInfo>();
-			this.UObjectModuleHeaders = new List<UHTModuleHeaderInfo>();
+			Actions = new List<IExternalAction>();
+			OutputItems = new List<FileItem>();
+			ModuleNameToOutputItems = new Dictionary<string, FileItem[]>(StringComparer.OrdinalIgnoreCase);
+			HotReloadModuleNames = new HashSet<string>(StringComparer.OrdinalIgnoreCase);
+			SourceDirectories = new List<DirectoryItem>();
+			DirectoryToSourceFiles = new();
+			HeaderFiles = new();
+			WorkingSet = new HashSet<FileItem>();
+			CandidatesForWorkingSet = new HashSet<FileItem>();
+			UObjectModules = new List<UHTModuleInfo>();
+			UObjectModuleHeaders = new List<UHTModuleHeaderInfo>();
 #if __VPROJECT_AVAILABLE__
-			this.VNIModules = new List<VNIModuleInfo>();
+			VNIModules = new List<VNIModuleInfo>();
 #endif
-			this.PluginFiles = new HashSet<FileItem>();
-			this.ExternalDependencies = new HashSet<FileItem>();
-			this.InternalDependencies = new HashSet<FileItem>();
+			PluginFiles = new HashSet<FileItem>();
+			ExternalDependencies = new HashSet<FileItem>();
+			InternalDependencies = new HashSet<FileItem>();
 		}
 
 		/// <summary>
@@ -679,8 +679,8 @@ namespace UnrealBuildTool
 				// Check if the arguments are different
 				if (!Enumerable.SequenceEqual(Makefile.AdditionalArguments!, Arguments))
 				{
-					Logger.LogDebug("Old command line arguments:\n", string.Join(' ', Makefile.AdditionalArguments!));
-					Logger.LogDebug("New command line arguments:\n", string.Join(' ', Arguments));
+					Logger.LogDebug("Old command line arguments:\n", String.Join(' ', Makefile.AdditionalArguments!));
+					Logger.LogDebug("New command line arguments:\n", String.Join(' ', Arguments));
 					ReasonNotLoaded = "command line arguments changed";
 					return null;
 				}
@@ -827,13 +827,13 @@ namespace UnrealBuildTool
 					if (!ExternalDependency.Exists)
 					{
 						Logger.LogDebug("{File} has been deleted since makefile was built.", ExternalDependency.Location);
-						ReasonNotLoaded = string.Format("{0} deleted", ExternalDependency.Location.GetFileName());
+						ReasonNotLoaded = String.Format("{0} deleted", ExternalDependency.Location.GetFileName());
 						return false;
 					}
 					if (ExternalDependency.LastWriteTimeUtc > Makefile.CreateTimeUtc)
 					{
 						Logger.LogDebug("{File} has been modified since makefile was built.", ExternalDependency.Location);
-						ReasonNotLoaded = string.Format("{0} modified", ExternalDependency.Location.GetFileName());
+						ReasonNotLoaded = String.Format("{0} modified", ExternalDependency.Location.GetFileName());
 						return false;
 					}
 				}
@@ -844,13 +844,13 @@ namespace UnrealBuildTool
 					if (!InternalDependency.Exists)
 					{
 						Logger.LogDebug("{File} has been deleted since makefile was written.", InternalDependency.Location);
-						ReasonNotLoaded = string.Format("{0} deleted", InternalDependency.Location.GetFileName());
+						ReasonNotLoaded = String.Format("{0} deleted", InternalDependency.Location.GetFileName());
 						return false;
 					}
 					if (InternalDependency.LastWriteTimeUtc > Makefile.ModifiedTimeUtc)
 					{
 						Logger.LogDebug("{File} has been modified since makefile was written.", InternalDependency.Location);
-						ReasonNotLoaded = string.Format("{0} modified", InternalDependency.Location.GetFileName());
+						ReasonNotLoaded = String.Format("{0} modified", InternalDependency.Location.GetFileName());
 						return false;
 					}
 				}
@@ -862,7 +862,7 @@ namespace UnrealBuildTool
 					if (!Makefile.PluginFiles.Contains(PluginFileItem))
 					{
 						Logger.LogDebug("{File} has been added", PluginFile.GetFileName());
-						ReasonNotLoaded = string.Format("{0} has been added", PluginFile.GetFileName());
+						ReasonNotLoaded = String.Format("{0} has been added", PluginFile.GetFileName());
 						return false;
 					}
 				}
@@ -898,7 +898,7 @@ namespace UnrealBuildTool
 					if (!WorkingSet.Contains(SourceFile) && SourceFile.LastWriteTimeUtc > Makefile.CreateTimeUtc)
 					{
 						Logger.LogDebug("{File} was part of source working set and now is not; invalidating makefile", SourceFile.Location);
-						ReasonNotLoaded = string.Format("working set of source files changed");
+						ReasonNotLoaded = String.Format("working set of source files changed");
 						return false;
 					}
 				}
@@ -909,7 +909,7 @@ namespace UnrealBuildTool
 					if (WorkingSet.Contains(SourceFile) && SourceFile.LastWriteTimeUtc > Makefile.CreateTimeUtc)
 					{
 						Logger.LogDebug("{File} was not part of source working set and should be", SourceFile.Location);
-						ReasonNotLoaded = string.Format("working set of source files changed");
+						ReasonNotLoaded = String.Format("working set of source files changed");
 						return false;
 					}
 				}
