UnrealBuildTool: More automated code cleanup

diff --git a/Engine/Source/Programs/UnrealBuildTool/System/Timeline.cs b/Engine/Source/Programs/UnrealBuildTool/System/Timeline.cs
--- a/Engine/Source/Programs/UnrealBuildTool/System/Timeline.cs
+++ b/Engine/Source/Programs/UnrealBuildTool/System/Timeline.cs
@@ -60,7 +60,7 @@ namespace UnrealBuildTool
 				this.Name = Name;
 				this.StartTime = StartTime;
 				this.FinishTime = FinishTime;
-				this.Span = TraceSpan.Create(Name);
+				Span = TraceSpan.Create(Name);
 			}
 
 			/// <summary>
diff --git a/Engine/Source/Programs/UnrealBuildTool/System/TopologicalSorter.cs b/Engine/Source/Programs/UnrealBuildTool/System/TopologicalSorter.cs
--- a/Engine/Source/Programs/UnrealBuildTool/System/TopologicalSorter.cs
+++ b/Engine/Source/Programs/UnrealBuildTool/System/TopologicalSorter.cs
@@ -324,7 +324,7 @@ namespace UnrealBuildTool
 			public GraphNode(T Data)
 			{
 				this.Data = Data;
-				this.Mark = Mark.None;
+				Mark = Mark.None;
 			}
 		}
 	}
diff --git a/Engine/Source/Programs/UnrealBuildTool/System/Unity.cs b/Engine/Source/Programs/UnrealBuildTool/System/Unity.cs
--- a/Engine/Source/Programs/UnrealBuildTool/System/Unity.cs
+++ b/Engine/Source/Programs/UnrealBuildTool/System/Unity.cs
@@ -116,7 +116,9 @@ namespace UnrealBuildTool
 			public void EndCurrentUnityFile()
 			{
 				if (CurrentCollection.Files.Count == 0)
+				{
 					return;
+				}
 
 				UnityFiles.Add(CurrentCollection);
 				CurrentCollection = new FileCollection();
@@ -208,9 +210,14 @@ namespace UnrealBuildTool
 					bool bAIsGenerated = A.AbsolutePath.EndsWith(".gen.cpp");
 					bool bBIsGenerated = B.AbsolutePath.EndsWith(".gen.cpp");
 					if (bAIsGenerated && !bBIsGenerated)
+					{
 						return -1;
+					}
+
 					if (!bAIsGenerated && bBIsGenerated)
+					{
 						return 1;
+					}
 
 					return String.Compare(A.AbsolutePath, B.AbsolutePath, StringComparison.OrdinalIgnoreCase);
 				});
@@ -260,16 +267,16 @@ namespace UnrealBuildTool
 				{
 					if (Target.bDetailedUnityFiles)
 					{
-						UnityCPPFileName = string.Format("{0}{1}.{2}_of_{3}.cpp", ModulePrefix, BaseName, CurrentUnityFileCount, AllUnityFiles.Count);
+						UnityCPPFileName = String.Format("{0}{1}.{2}_of_{3}.cpp", ModulePrefix, BaseName, CurrentUnityFileCount, AllUnityFiles.Count);
 					}
 					else
 					{
-						UnityCPPFileName = string.Format("{0}{1}.{2}.cpp", ModulePrefix, BaseName, CurrentUnityFileCount);
+						UnityCPPFileName = String.Format("{0}{1}.{2}.cpp", ModulePrefix, BaseName, CurrentUnityFileCount);
 					}
 				}
 				else
 				{
-					UnityCPPFileName = string.Format("{0}{1}.cpp", ModulePrefix, BaseName);
+					UnityCPPFileName = String.Format("{0}{1}.cpp", ModulePrefix, BaseName);
 				}
 				FileReference UnityCPPFilePath = FileReference.Combine(IntermediateDirectory, UnityCPPFileName);
 
diff --git a/Engine/Source/Programs/UnrealBuildTool/System/UnrealPluginLanguage.cs b/Engine/Source/Programs/UnrealBuildTool/System/UnrealPluginLanguage.cs
--- a/Engine/Source/Programs/UnrealBuildTool/System/UnrealPluginLanguage.cs
+++ b/Engine/Source/Programs/UnrealBuildTool/System/UnrealPluginLanguage.cs
@@ -499,13 +499,13 @@ namespace UnrealBuildTool
 		private UnrealTargetPlatform TargetPlatform;
 
 		/** Trace flag to enable debugging */
-		static private bool bGlobalTrace = false;
-		static private bool bGlobalTraceFilename = false;
+		private static bool bGlobalTrace = false;
+		private static bool bGlobalTraceFilename = false;
 
 		/** Project file reference */
 		private FileReference? ProjectFile;
 
-		static private XDocument XMLDummy = XDocument.Parse("<manifest></manifest>");
+		private static XDocument XMLDummy = XDocument.Parse("<manifest></manifest>");
 
 		private class UPLContext
 		{
@@ -669,19 +669,19 @@ namespace UnrealBuildTool
 			StringBuilder Text = new StringBuilder();
 			foreach (KeyValuePair<string, bool> Variable in Context.BoolVariables)
 			{
-				Text.AppendLine(string.Format("\tbool {0} = {1}", Variable.Key, Variable.Value.ToString().ToLower()));
+				Text.AppendLine(String.Format("\tbool {0} = {1}", Variable.Key, Variable.Value.ToString().ToLower()));
 			}
 			foreach (KeyValuePair<string, int> Variable in Context.IntVariables)
 			{
-				Text.AppendLine(string.Format("\tint {0} = {1}", Variable.Key, Variable.Value));
+				Text.AppendLine(String.Format("\tint {0} = {1}", Variable.Key, Variable.Value));
 			}
 			foreach (KeyValuePair<string, string> Variable in Context.StringVariables)
 			{
-				Text.AppendLine(string.Format("\tstring {0} = {1}", Variable.Key, Variable.Value));
+				Text.AppendLine(String.Format("\tstring {0} = {1}", Variable.Key, Variable.Value));
 			}
 			foreach (KeyValuePair<string, XElement> Variable in Context.ElementVariables)
 			{
-				Text.AppendLine(string.Format("\telement {0} = {1}", Variable.Key, Variable.Value));
+				Text.AppendLine(String.Format("\telement {0} = {1}", Variable.Key, Variable.Value));
 			}
 			return Text.ToString();
 		}
@@ -868,7 +868,7 @@ namespace UnrealBuildTool
 		private int StringToInt(UPLContext Context, XElement Node, string? Input)
 		{
 			int Result = 0;
-			if (!int.TryParse(Input, out Result))
+			if (!Int32.TryParse(Input, out Result))
 			{
 				Logger.LogWarning("\nInvalid integer '{Input}' in '{Node}' (defaulting to 0)", Input, TraceNodeString(Context, Node));
 			}
@@ -907,7 +907,7 @@ namespace UnrealBuildTool
 			return bExpand ? ExpandVariables(Context, Result) : Result;
 		}
 
-		static private Dictionary<string, ConfigCacheIni_UPL>? ConfigCache = null;
+		private static Dictionary<string, ConfigCacheIni_UPL>? ConfigCache = null;
 
 		private ConfigCacheIni_UPL GetConfigCacheIni_UPL(string baseIniName)
 		{
diff --git a/Engine/Source/Programs/UnrealBuildTool/System/Utils.cs b/Engine/Source/Programs/UnrealBuildTool/System/Utils.cs
--- a/Engine/Source/Programs/UnrealBuildTool/System/Utils.cs
+++ b/Engine/Source/Programs/UnrealBuildTool/System/Utils.cs
@@ -306,7 +306,7 @@ namespace UnrealBuildTool
 		public static int RunLocalProcessAndPrintfOutput(ProcessStartInfo StartInfo, ILogger Logger)
 		{
 			string AppName = Path.GetFileNameWithoutExtension(StartInfo.FileName);
-			string LogFilenameBase = string.Format("{0}_{1}", AppName, DateTime.Now.ToString("yyyy.MM.dd-HH.mm.ss"));
+			string LogFilenameBase = String.Format("{0}_{1}", AppName, DateTime.Now.ToString("yyyy.MM.dd-HH.mm.ss"));
 			string LogDir = Path.Combine(Unreal.EngineDirectory.FullName, "Programs", "AutomationTool", "Saved", "Logs");
 			string LogFilename = "";
 			for (int Attempt = 1; Attempt < 100; ++Attempt)
@@ -316,7 +316,7 @@ namespace UnrealBuildTool
 					if (!Directory.Exists(LogDir))
 					{
 						string? IniPath = UnrealBuildTool.GetRemoteIniPath();
-						if (string.IsNullOrEmpty(IniPath))
+						if (String.IsNullOrEmpty(IniPath))
 						{
 							break;
 						}
@@ -356,12 +356,12 @@ namespace UnrealBuildTool
 				if (Args != null && Args.Data != null)
 				{
 					string data = Args.Data.TrimEnd();
-					if (string.IsNullOrEmpty(data))
+					if (String.IsNullOrEmpty(data))
 					{
 						return;
 					}
 
-					if (!string.IsNullOrEmpty(LogFilename))
+					if (!String.IsNullOrEmpty(LogFilename))
 					{
 						File.AppendAllLines(LogFilename, data.Split('\n'));
 					}
@@ -376,7 +376,7 @@ namespace UnrealBuildTool
 			LocalProcess.OutputDataReceived += Output;
 			LocalProcess.ErrorDataReceived += Output;
 			var ExitCode = RunLocalProcess(LocalProcess);
-			if (ExitCode != 0 && !string.IsNullOrEmpty(LogFilename))
+			if (ExitCode != 0 && !String.IsNullOrEmpty(LogFilename))
 			{
 				Logger.LogError("Process \'{AppName}\' failed. Details are in \'{LogFilename}\'", AppName, LogFilename);
 			}
@@ -673,7 +673,7 @@ namespace UnrealBuildTool
 			for (; ; )
 			{
 				// An empty path is finished
-				if (string.IsNullOrEmpty(LocalString))
+				if (String.IsNullOrEmpty(LocalString))
 				{
 					break;
 				}
@@ -816,11 +816,11 @@ namespace UnrealBuildTool
 		/*
 		 * Read and write classes with xml specifiers
 		 */
-		static private void UnknownAttributeDelegate(object? sender, XmlAttributeEventArgs e)
+		private static void UnknownAttributeDelegate(object? sender, XmlAttributeEventArgs e)
 		{
 		}
 
-		static private void UnknownNodeDelegate(object? sender, XmlNodeEventArgs e)
+		private static void UnknownNodeDelegate(object? sender, XmlNodeEventArgs e)
 		{
 		}
 
@@ -831,7 +831,7 @@ namespace UnrealBuildTool
 		/// <param name="FileName">The XML file to read from</param>
 		/// <param name="Logger">Logger for output</param>
 		/// <returns>New deserialized instance of type T</returns>
-		static public T ReadClass<T>(string FileName, ILogger Logger) where T : new()
+		public static T ReadClass<T>(string FileName, ILogger Logger) where T : new()
 		{
 			T Instance = new T();
 			StreamReader? XmlStream = null;
@@ -874,7 +874,7 @@ namespace UnrealBuildTool
 		/// <param name="DefaultNameSpace">Default namespace for the output elements</param>
 		/// <param name="Logger">Logger for output</param>
 		/// <returns>True if the file was written successfully</returns>
-		static public bool WriteClass<T>(T Data, string FileName, string DefaultNameSpace, ILogger Logger)
+		public static bool WriteClass<T>(T Data, string FileName, string DefaultNameSpace, ILogger Logger)
 		{
 			bool bSuccess = true;
 			StreamWriter? XmlStream = null;
@@ -927,7 +927,9 @@ namespace UnrealBuildTool
 		{
 			// null objects are always invalid
 			if (p == null)
+			{
 				return false;
+			}
 			// due to multithreading on Windows, lock the object
 			lock (p)
 			{
@@ -1025,7 +1027,7 @@ namespace UnrealBuildTool
 		}
 
 		[DllImport("kernel32.dll", SetLastError = true)]
-		extern static bool GetLogicalProcessorInformationEx(LOGICAL_PROCESSOR_RELATIONSHIP RelationshipType, IntPtr Buffer, ref uint ReturnedLength);
+		static extern bool GetLogicalProcessorInformationEx(LOGICAL_PROCESSOR_RELATIONSHIP RelationshipType, IntPtr Buffer, ref uint ReturnedLength);
 
 		/// <summary>
 		/// Gets the number of logical cores. We use this rather than Environment.ProcessorCount when possible to handle machines with > 64 cores (the single group limit available to the .NET framework).
@@ -1083,7 +1085,7 @@ namespace UnrealBuildTool
 
 		// int sysctlbyname(const char *name, void *oldp, size_t *oldlenp, void *newp, size_t newlen); // from man page
 		[DllImport("libc")]
-		extern static int sysctlbyname(string name, out int oldp, ref UInt64 oldlenp, IntPtr newp, UInt64 newlen);
+		static extern int sysctlbyname(string name, out int oldp, ref UInt64 oldlenp, IntPtr newp, UInt64 newlen);
 
 		/// <summary>
 		/// Gets the number of physical cores, excluding hyper threading.
@@ -1237,11 +1239,11 @@ namespace UnrealBuildTool
 
 		// kern_return_t host_statistics64(host_t host_priv, host_flavor_t flavor, host_info64_t host_info64_out, mach_msg_type_number_t *host_info64_outCnt); // from <mach/mach_host.h>
 		[DllImport("libc")]
-		extern static int host_statistics64(IntPtr host_priv, int flavor, out vm_statistics64 host_info64_out, ref uint host_info_count);
+		static extern int host_statistics64(IntPtr host_priv, int flavor, out vm_statistics64 host_info64_out, ref uint host_info_count);
 
 		// mach_port_t mach_host_self() // from <mach/mach_init.h>
 		[DllImport("libc")]
-		extern static IntPtr mach_host_self();
+		static extern IntPtr mach_host_self();
 
 		/// <summary>
 		/// Gets the total memory bytes free, based on what is known to the garbage collector.
@@ -1584,8 +1586,8 @@ namespace UnrealBuildTool
 		/// </summary>
 		static readonly Dictionary<FileReference, (int WriteRequestCount, int ActualWriteCount)> WriteFileIfChangedRecord = new Dictionary<FileReference, (int, int)>();
 
-		static internal FileReference? WriteFileIfChangedTrace = null;
-		static internal string WriteFileIfChangedContext = "";
+		internal static FileReference? WriteFileIfChangedTrace = null;
+		internal static string WriteFileIfChangedContext = "";
 
 		static void RecordWriteFileIfChanged(FileReference File, bool bNew, bool bChanged, ILogger Logger)
 		{
@@ -1784,7 +1786,7 @@ namespace UnrealBuildTool
 		/// <param name="Logger">Logger for output</param>
 		internal static void WriteFileIfChanged(FileItem FileItem, IEnumerable<string> ContentLines, StringComparison Comparison, ILogger Logger)
 		{
-			var Contents = string.Join(Environment.NewLine, ContentLines);
+			var Contents = String.Join(Environment.NewLine, ContentLines);
 			WriteFileIfChanged(FileItem, Contents, Comparison, Logger);
 		}
 
diff --git a/Engine/Source/Programs/UnrealBuildTool/System/VersionNumber.cs b/Engine/Source/Programs/UnrealBuildTool/System/VersionNumber.cs
--- a/Engine/Source/Programs/UnrealBuildTool/System/VersionNumber.cs
+++ b/Engine/Source/Programs/UnrealBuildTool/System/VersionNumber.cs
@@ -199,7 +199,7 @@ namespace UnrealBuildTool
 			List<int> Components = new List<int>();
 			foreach (string TextElement in Text.Split(Delimiters))
 			{
-				Components.Add(int.Parse(TextElement));
+				Components.Add(Int32.Parse(TextElement));
 			}
 			return new VersionNumber(Components.ToArray());
 		}
@@ -216,7 +216,7 @@ namespace UnrealBuildTool
 			foreach (string TextElement in Text.Split(Delimiters))
 			{
 				int Component;
-				if (!int.TryParse(TextElement, out Component))
+				if (!Int32.TryParse(TextElement, out Component))
 				{
 					OutNumber = null;
 					return false;
diff --git a/Engine/Source/Programs/UnrealBuildTool/System/XmlConfig.cs b/Engine/Source/Programs/UnrealBuildTool/System/XmlConfig.cs
--- a/Engine/Source/Programs/UnrealBuildTool/System/XmlConfig.cs
+++ b/Engine/Source/Programs/UnrealBuildTool/System/XmlConfig.cs
@@ -205,7 +205,7 @@ namespace UnrealBuildTool
 			}
 			catch (ReflectionTypeLoadException Ex)
 			{
-				Console.WriteLine("TypeLoadException: {0}", string.Join("\n", Ex.LoaderExceptions.Select(x => x?.Message)));
+				Console.WriteLine("TypeLoadException: {0}", String.Join("\n", Ex.LoaderExceptions.Select(x => x?.Message)));
 				throw;
 			}
 			return ConfigTypes;
diff --git a/Engine/Source/Programs/UnrealBuildTool/System/XmlConfigData.cs b/Engine/Source/Programs/UnrealBuildTool/System/XmlConfigData.cs
--- a/Engine/Source/Programs/UnrealBuildTool/System/XmlConfigData.cs
+++ b/Engine/Source/Programs/UnrealBuildTool/System/XmlConfigData.cs
@@ -103,7 +103,7 @@ namespace UnrealBuildTool
 
 			public ValueInfo(FieldInfo FieldInfo, object Value, FileReference SourceFile, XmlConfigFileAttribute XmlConfigAttribute)
 			{
-				this.Target = new TargetField(FieldInfo);
+				Target = new TargetField(FieldInfo);
 				this.Value = Value;
 				this.SourceFile = SourceFile;
 				this.XmlConfigAttribute = XmlConfigAttribute;
@@ -111,7 +111,7 @@ namespace UnrealBuildTool
 
 			public ValueInfo(PropertyInfo PropertyInfo, object Value, FileReference SourceFile, XmlConfigFileAttribute XmlConfigAttribute)
 			{
-				this.Target = new TargetProperty(PropertyInfo);
+				Target = new TargetProperty(PropertyInfo);
 				this.Value = Value;
 				this.SourceFile = SourceFile;
 				this.XmlConfigAttribute = XmlConfigAttribute;
diff --git a/Engine/Source/Programs/UnrealBuildTool/ToolChain/AppleToolChain.cs b/Engine/Source/Programs/UnrealBuildTool/ToolChain/AppleToolChain.cs
--- a/Engine/Source/Programs/UnrealBuildTool/ToolChain/AppleToolChain.cs
+++ b/Engine/Source/Programs/UnrealBuildTool/ToolChain/AppleToolChain.cs
@@ -76,7 +76,7 @@ namespace UnrealBuildTool
 				{
 					throw new BuildException("Unable to get xcode version");
 				}
-				if (int.Parse(InstalledSdkVersion.Substring(0, 2)) < 13)
+				if (Int32.Parse(InstalledSdkVersion.Substring(0, 2)) < 13)
 				{
 					throw new BuildException("Building for macOS, iOS and tvOS requires Xcode 13.4.1 or newer, Xcode " + InstalledSdkVersion + " detected");
 				}
@@ -113,8 +113,8 @@ namespace UnrealBuildTool
 								string[] Tokens = SDKString.Split(".".ToCharArray());
 								if (Tokens.Length == 2)
 								{
-									Major = int.Parse(Tokens[0]);
-									Minor = int.Parse(Tokens[1]);
+									Major = Int32.Parse(Tokens[0]);
+									Minor = Int32.Parse(Tokens[1]);
 								}
 							}
 							catch (Exception)
@@ -191,7 +191,7 @@ namespace UnrealBuildTool
 		/// <returns></returns>
 		protected string FormatArchitectureArg(UnrealArchitectures InArchitectures)
 		{
-			return "-arch " + string.Join('+', InArchitectures.Architectures.Select(x => x.AppleName));
+			return "-arch " + String.Join('+', InArchitectures.Architectures.Select(x => x.AppleName));
 		}
 
 		protected DirectoryReference GetMacDevSrcRoot()
@@ -377,7 +377,7 @@ namespace UnrealBuildTool
 
 			bool bStaticAnalysis = false;
 			string? StaticAnalysisMode = Environment.GetEnvironmentVariable("CLANG_STATIC_ANALYZER_MODE");
-			if (!string.IsNullOrEmpty(StaticAnalysisMode))
+			if (!String.IsNullOrEmpty(StaticAnalysisMode))
 			{
 				bStaticAnalysis = true;
 			}
@@ -489,7 +489,7 @@ namespace UnrealBuildTool
 				}
 				else
 				{
-					if (!int.TryParse(Versions[0], out Major) || !int.TryParse(Versions[1], out Minor) || !int.TryParse(Versions[2], out Patch))
+					if (!Int32.TryParse(Versions[0], out Major) || !Int32.TryParse(Versions[1], out Minor) || !Int32.TryParse(Versions[2], out Patch))
 					{
 						Log.TraceInformationOnce("Unable to parse version tokens: {0}", Tokens[3]);
 					}
@@ -596,7 +596,7 @@ namespace UnrealBuildTool
 				"-automated",
 			};
 
-			if (!string.IsNullOrEmpty(TargetName))
+			if (!String.IsNullOrEmpty(TargetName))
 			{
 				Options.Add($"-singletarget={TargetName}");
 			}
@@ -641,7 +641,7 @@ namespace UnrealBuildTool
 				//$"-sdk {SDKName}",
 			};
 
-			return Utils.RunLocalProcessAndLogOutput("/usr/bin/env", string.Join(" ", Arguments), Logger);
+			return Utils.RunLocalProcessAndLogOutput("/usr/bin/env", String.Join(" ", Arguments), Logger);
 		}
 
 		#endregion
diff --git a/Engine/Source/Programs/UnrealBuildTool/ToolChain/ClangToolChain.cs b/Engine/Source/Programs/UnrealBuildTool/ToolChain/ClangToolChain.cs
--- a/Engine/Source/Programs/UnrealBuildTool/ToolChain/ClangToolChain.cs
+++ b/Engine/Source/Programs/UnrealBuildTool/ToolChain/ClangToolChain.cs
@@ -249,7 +249,7 @@ namespace UnrealBuildTool
 						$"-HeadersFile={HeadersOutputFile.FullName}",
 					};
 
-					Action AggregateTimingInfoAction = MakefileBuilder.CreateRecursiveAction<AggregateClangTimingInfo>(ActionType.ParseTimingInfo, string.Join(" ", AggregateActionArgs));
+					Action AggregateTimingInfoAction = MakefileBuilder.CreateRecursiveAction<AggregateClangTimingInfo>(ActionType.ParseTimingInfo, String.Join(" ", AggregateActionArgs));
 					AggregateTimingInfoAction.WorkingDirectory = Unreal.EngineSourceDirectory;
 					AggregateTimingInfoAction.StatusDescription = $"Aggregating {TimingJsonFiles.Count} Timing File(s)";
 					AggregateTimingInfoAction.bCanExecuteRemotely = false;
@@ -267,7 +267,7 @@ namespace UnrealBuildTool
 						$"-ArchiveFile={ArchiveOutputFile.FullName}",
 					};
 
-					Action ArchiveTimingInfoAction = MakefileBuilder.CreateRecursiveAction<AggregateClangTimingInfo>(ActionType.ParseTimingInfo, string.Join(" ", ArchiveActionArgs));
+					Action ArchiveTimingInfoAction = MakefileBuilder.CreateRecursiveAction<AggregateClangTimingInfo>(ActionType.ParseTimingInfo, String.Join(" ", ArchiveActionArgs));
 					ArchiveTimingInfoAction.WorkingDirectory = Unreal.EngineSourceDirectory;
 					ArchiveTimingInfoAction.StatusDescription = $"Archiving {TimingJsonFiles.Count} Timing File(s)";
 					ArchiveTimingInfoAction.bCanExecuteRemotely = false;
@@ -675,7 +675,7 @@ namespace UnrealBuildTool
 		/// <param name="Arguments"></param>
 		protected virtual void GetCompileArguments_AdditionalArgs(CppCompileEnvironment CompileEnvironment, List<string> Arguments)
 		{
-			if (!string.IsNullOrWhiteSpace(CompileEnvironment.AdditionalArguments))
+			if (!String.IsNullOrWhiteSpace(CompileEnvironment.AdditionalArguments))
 			{
 				Arguments.Add(CompileEnvironment.AdditionalArguments);
 			}
@@ -712,7 +712,10 @@ namespace UnrealBuildTool
 			Arguments.Add("-Xclang -analyzer-config -Xclang path-diagnostics-alternate=true");
 
 			// Run shallow analyze if requested.
-			if (StaticAnalyzerMode == StaticAnalyzerMode.Shallow) Arguments.Add("-Xclang -analyzer-config -Xclang mode=shallow");
+			if (StaticAnalyzerMode == StaticAnalyzerMode.Shallow)
+			{
+				Arguments.Add("-Xclang -analyzer-config -Xclang mode=shallow");
+			}
 
 			if (CompileEnvironment.StaticAnalyzerCheckers.Count > 0)
 			{
@@ -995,7 +998,9 @@ namespace UnrealBuildTool
 		{
 			// This is not supported for now.. If someone wants it we can implement it
 			if (CompileEnvironment.Architectures.bIsMultiArch)
+			{
 				return;
+			}
 
 			List<string> GlobalArguments = new();
 			if (!CompileEnvironment.bHasSharedResponseFile)
@@ -1039,7 +1044,10 @@ namespace UnrealBuildTool
 				// To be conservative we add a dummy define to inflate the response file size and force clang to use the response file mode when we are close to the limit.
 				int CmdLineLength = Info.Clang.ToString().Length;
 				foreach (string Line in ResponseFileContents)
+				{
 					CmdLineLength += 1 + Line.Length;
+				}
+
 				bool bIsInDangerZone = CmdLineLength >= ClangCmdlineDangerZone && CmdLineLength <= ClangCmdLineMaxSize;
 				if (bIsInDangerZone)
 				{
@@ -1110,7 +1118,7 @@ namespace UnrealBuildTool
 
 				if (RuntimePlatform.IsWindows)
 				{
-					int CmdLineLength = Info.Clang.ToString().Length + string.Join(' ', ResponseFileContents).Length;
+					int CmdLineLength = Info.Clang.ToString().Length + String.Join(' ', ResponseFileContents).Length;
 					bool bIsInDangerZone = CmdLineLength >= ClangCmdlineDangerZone && CmdLineLength <= ClangCmdLineMaxSize;
 					if (bIsInDangerZone)
 					{
diff --git a/Engine/Source/Programs/UnrealBuildTool/ToolChain/IWYUToolChain.cs b/Engine/Source/Programs/UnrealBuildTool/ToolChain/IWYUToolChain.cs
--- a/Engine/Source/Programs/UnrealBuildTool/ToolChain/IWYUToolChain.cs
+++ b/Engine/Source/Programs/UnrealBuildTool/ToolChain/IWYUToolChain.cs
@@ -146,7 +146,7 @@ namespace UnrealBuildTool
 		{
 			// set UE_LINUX_USE_LIBCXX to either 0 or 1. If unset, defaults to 1.
 			string? UseLibcxxEnvVarOverride = Environment.GetEnvironmentVariable("UE_LINUX_USE_LIBCXX");
-			if (string.IsNullOrEmpty(UseLibcxxEnvVarOverride) || UseLibcxxEnvVarOverride == "1")
+			if (String.IsNullOrEmpty(UseLibcxxEnvVarOverride) || UseLibcxxEnvVarOverride == "1")
 			{
 				return true;
 			}
diff --git a/Engine/Source/Programs/UnrealBuildTool/ToolChain/RemoteMac.cs b/Engine/Source/Programs/UnrealBuildTool/ToolChain/RemoteMac.cs
--- a/Engine/Source/Programs/UnrealBuildTool/ToolChain/RemoteMac.cs
+++ b/Engine/Source/Programs/UnrealBuildTool/ToolChain/RemoteMac.cs
@@ -130,17 +130,17 @@ namespace UnrealBuildTool
 		/// <param name="bPrepareForSecondaryMac">Added arguments when preparing a secondary Mac for debug. False by default.</param>
 		public RemoteMac(FileReference? ProjectFile, ILogger Logger, bool bIsPrimary = true, bool bPrepareForSecondaryMac = false)
 		{
-			this.RsyncExe = FileReference.Combine(Unreal.EngineDirectory, "Extras", "ThirdPartyNotUE", "cwrsync", "bin", "rsync.exe");
-			this.SshExe = FileReference.Combine(Unreal.EngineDirectory, "Extras", "ThirdPartyNotUE", "cwrsync", "bin", "ssh.exe");
+			RsyncExe = FileReference.Combine(Unreal.EngineDirectory, "Extras", "ThirdPartyNotUE", "cwrsync", "bin", "rsync.exe");
+			SshExe = FileReference.Combine(Unreal.EngineDirectory, "Extras", "ThirdPartyNotUE", "cwrsync", "bin", "ssh.exe");
 			this.ProjectFile = ProjectFile;
 			if (ProjectFile != null)
 			{
-				this.ProjectDescriptor = ProjectDescriptor.FromFile(ProjectFile);
-				this.AdditionalPaths = new List<DirectoryReference>();
-				this.ProjectDescriptor.AddAdditionalPaths(this.AdditionalPaths, ProjectFile.Directory);
-				if (this.AdditionalPaths.Count == 0)
+				ProjectDescriptor = ProjectDescriptor.FromFile(ProjectFile);
+				AdditionalPaths = new List<DirectoryReference>();
+				ProjectDescriptor.AddAdditionalPaths(AdditionalPaths, ProjectFile.Directory);
+				if (AdditionalPaths.Count == 0)
 				{
-					this.AdditionalPaths = null;
+					AdditionalPaths = null;
 				}
 			}
 
@@ -163,7 +163,7 @@ namespace UnrealBuildTool
 
 				if ((bIsPrimary ? Ini.GetString("/Script/IOSRuntimeSettings.IOSRuntimeSettings", "RemoteServerName", out IniServerName) : Ini.GetString("/Script/IOSRuntimeSettings.IOSRuntimeSettings", "SecondaryRemoteServerName", out IniServerName)) && !String.IsNullOrEmpty(IniServerName))
 				{
-					this.ServerName = IniServerName;
+					ServerName = IniServerName;
 				}
 				else
 				{
@@ -174,7 +174,7 @@ namespace UnrealBuildTool
 				string IniUserName;
 				if ((bIsPrimary ? Ini.GetString("/Script/IOSRuntimeSettings.IOSRuntimeSettings", "RSyncUsername", out IniUserName) : Ini.GetString("/Script/IOSRuntimeSettings.IOSRuntimeSettings", "SecondaryRSyncUsername", out IniUserName)) && !String.IsNullOrEmpty(IniUserName))
 				{
-					this.UserName = IniUserName;
+					UserName = IniUserName;
 				}
 			}
 
@@ -183,7 +183,7 @@ namespace UnrealBuildTool
 			if (PortIdx != -1)
 			{
 				string Port = ServerName.Substring(PortIdx + 1);
-				if (!int.TryParse(Port, out ServerPort))
+				if (!Int32.TryParse(Port, out ServerPort))
 				{
 					throw new BuildException("Unable to parse port number from '{0}'", ServerName);
 				}
@@ -1019,7 +1019,7 @@ namespace UnrealBuildTool
 			{
 				ProjectDir = ProjectFile.Directory;
 			}
-			else if (!string.IsNullOrEmpty(UnrealBuildTool.GetRemoteIniPath()))
+			else if (!String.IsNullOrEmpty(UnrealBuildTool.GetRemoteIniPath()))
 			{
 				ProjectDir = new DirectoryReference(UnrealBuildTool.GetRemoteIniPath()!);
 				if (ProjectDir.IsUnderDirectory(Unreal.EngineDirectory))
diff --git a/Engine/Source/Programs/UnrealBuildTool/ToolChain/UEToolChain.cs b/Engine/Source/Programs/UnrealBuildTool/ToolChain/UEToolChain.cs
--- a/Engine/Source/Programs/UnrealBuildTool/ToolChain/UEToolChain.cs
+++ b/Engine/Source/Programs/UnrealBuildTool/ToolChain/UEToolChain.cs
@@ -262,7 +262,7 @@ namespace UnrealBuildTool
 		{
 			string ProcessOutput = Utils.RunLocalProcessAndReturnStdOut(Command.FullName, ToolArg, Logger);
 
-			if (string.IsNullOrEmpty(Expression))
+			if (String.IsNullOrEmpty(Expression))
 			{
 				return ProcessOutput;
 			}
diff --git a/Engine/Source/Programs/UnrealBuildTool/UnrealBuildTool.cs b/Engine/Source/Programs/UnrealBuildTool/UnrealBuildTool.cs
--- a/Engine/Source/Programs/UnrealBuildTool/UnrealBuildTool.cs
+++ b/Engine/Source/Programs/UnrealBuildTool/UnrealBuildTool.cs
@@ -20,17 +20,17 @@ namespace UnrealBuildTool
 		/// Save the application startup time. This can be used as the timestamp for build makefiles, to determine a base time after which any
 		/// modifications should invalidate it.
 		/// </summary>
-		static public DateTime StartTimeUtc { get; } = DateTime.UtcNow;
+		public static DateTime StartTimeUtc { get; } = DateTime.UtcNow;
 
 		/// <summary>
 		/// The environment at boot time.
 		/// </summary>
-		static public System.Collections.IDictionary? InitialEnvironment;
+		public static System.Collections.IDictionary? InitialEnvironment;
 
 		/// <summary>
 		/// Whether we're running with an installed project
 		/// </summary>
-		static private bool? bIsProjectInstalled;
+		private static bool? bIsProjectInstalled;
 
 		/// <summary>
 		/// If we are running with an installed project, specifies the path to it
@@ -119,7 +119,7 @@ namespace UnrealBuildTool
 		/// Returns true if UnrealBuildTool is running using an installed project (ie. a mod kit)
 		/// </summary>
 		/// <returns>True if running using an installed project</returns>
-		static public bool IsProjectInstalled()
+		public static bool IsProjectInstalled()
 		{
 			if (!bIsProjectInstalled.HasValue)
 			{
@@ -142,7 +142,7 @@ namespace UnrealBuildTool
 		/// Gets the installed project file
 		/// </summary>
 		/// <returns>Location of the installed project file</returns>
-		static public FileReference? GetInstalledProjectFile()
+		public static FileReference? GetInstalledProjectFile()
 		{
 			if (IsProjectInstalled())
 			{
@@ -159,7 +159,7 @@ namespace UnrealBuildTool
 		/// </summary>
 		/// <param name="File">File to test</param>
 		/// <returns>True if the file is part of the installed distribution, false otherwise</returns>
-		static public bool IsFileInstalled(FileReference File)
+		public static bool IsFileInstalled(FileReference File)
 		{
 			if (Unreal.IsEngineInstalled() && File.IsUnderDirectory(Unreal.EngineDirectory))
 			{
@@ -177,18 +177,18 @@ namespace UnrealBuildTool
 		/// </summary>
 		/// <returns>A string containing the path to the UBT assembly.</returns>
 		[Obsolete("Deprecated in UE5.1 - use UnrealBuildTool.DotnetPath Unreal.UnrealBuildToolDllPath")]
-		static public FileReference GetUBTPath() => Unreal.UnrealBuildToolPath;
+		public static FileReference GetUBTPath() => Unreal.UnrealBuildToolPath;
 
 		/// <summary>
 		/// The Unreal remote tool ini directory.  This should be valid if compiling using a remote server
 		/// </summary>
 		/// <returns>The directory path</returns>
-		static public string? GetRemoteIniPath()
+		public static string? GetRemoteIniPath()
 		{
 			return RemoteIniPath;
 		}
 
-		static public void SetRemoteIniPath(string Path)
+		public static void SetRemoteIniPath(string Path)
 		{
 			RemoteIniPath = Path;
 		}
@@ -325,7 +325,7 @@ namespace UnrealBuildTool
 			public GlobalOptions(CommandLineArguments Arguments)
 			{
 				Arguments.ApplyTo(this);
-				if (!string.IsNullOrEmpty(RemoteIni))
+				if (!String.IsNullOrEmpty(RemoteIni))
 				{
 					UnrealBuildTool.SetRemoteIniPath(RemoteIni);
 				}
@@ -488,8 +488,8 @@ namespace UnrealBuildTool
 				AssemblyUtils.InstallAssemblyResolver(Path.GetDirectoryName(Assembly.GetEntryAssembly()!.GetOriginalLocation())!);
 
 				// Add the application directory to PATH
-				string path = Environment.GetEnvironmentVariable("PATH") ?? string.Empty;
-				if (!path.Split(';').Any(x => string.Compare(x, Unreal.UnrealBuildToolDllPath.Directory.FullName, true) == 0))
+				string path = Environment.GetEnvironmentVariable("PATH") ?? String.Empty;
+				if (!path.Split(';').Any(x => String.Compare(x, Unreal.UnrealBuildToolDllPath.Directory.FullName, true) == 0))
 				{
 					Environment.SetEnvironmentVariable("PATH", $"{path};{Unreal.UnrealBuildToolDllPath.Directory.FullName}");
 				}
@@ -565,7 +565,7 @@ namespace UnrealBuildTool
 					DirectoryReference RunsDir = DirectoryReference.Combine(Unreal.EngineDirectory, "Intermediate", "UbtRuns");
 					Directory.CreateDirectory(RunsDir.FullName);
 					string ModuleFileName = Process.GetCurrentProcess().MainModule?.FileName ?? "";
-					if (!string.IsNullOrEmpty(ModuleFileName))
+					if (!String.IsNullOrEmpty(ModuleFileName))
 					{
 						ModuleFileName = Path.GetFullPath(ModuleFileName);
 					}
@@ -581,7 +581,7 @@ namespace UnrealBuildTool
 				try
 				{
 					// If the temp directory is already overridden from a parent process, do not override again
-					if (string.IsNullOrEmpty(Environment.GetEnvironmentVariable("UnrealBuildTool_TMP")))
+					if (String.IsNullOrEmpty(Environment.GetEnvironmentVariable("UnrealBuildTool_TMP")))
 					{
 						DirectoryReference OverrideTempDirectory = new DirectoryReference(Path.Combine(Path.GetTempPath(), "UnrealBuildTool"));
 						if (Options.TempDirectory != null)
