UnrealBuildTool: More automated code cleanup

#rnx

diff --git a/Engine/Source/Programs/UnrealBuildTool/Configuration/TargetRules.cs b/Engine/Source/Programs/UnrealBuildTool/Configuration/TargetRules.cs
--- a/Engine/Source/Programs/UnrealBuildTool/Configuration/TargetRules.cs
+++ b/Engine/Source/Programs/UnrealBuildTool/Configuration/TargetRules.cs
@@ -361,7 +361,7 @@ namespace UnrealBuildTool
 	/// Attribute used to mark fields which must match between targets in the shared build environment
 	/// </summary>
 	[AttributeUsage(AttributeTargets.Field | AttributeTargets.Property, AllowMultiple = false)]
-	class RequiresUniqueBuildEnvironmentAttribute : Attribute
+	sealed class RequiresUniqueBuildEnvironmentAttribute : Attribute
 	{
 	}
 
@@ -369,7 +369,7 @@ namespace UnrealBuildTool
 	/// Attribute used to mark configurable sub-objects
 	/// </summary>
 	[AttributeUsage(AttributeTargets.Field, AllowMultiple = false)]
-	class ConfigSubObjectAttribute : Attribute
+	sealed class ConfigSubObjectAttribute : Attribute
 	{
 	}
 
@@ -1030,7 +1030,6 @@ namespace UnrealBuildTool
 		}
 		private bool bCompileAgainstCoreUObjectPrivate = true;
 
-
 		/// <summary>
 		/// Enabled for builds that need to initialize the ApplicationCore module. Command line utilities do not normally need this.
 		/// </summary>
@@ -3255,7 +3254,6 @@ namespace UnrealBuildTool
 
 		public bool bIsEngineInstalled => Inner.bIsEngineInstalled;
 
-
 		public IReadOnlyList<string>? DisableUnityBuildForModules => Inner.DisableUnityBuildForModules;
 
 		public IReadOnlyList<string>? EnableOptimizeCodeForModules => Inner.EnableOptimizeCodeForModules;
diff --git a/Engine/Source/Programs/UnrealBuildTool/Configuration/TestModuleRules.cs b/Engine/Source/Programs/UnrealBuildTool/Configuration/TestModuleRules.cs
--- a/Engine/Source/Programs/UnrealBuildTool/Configuration/TestModuleRules.cs
+++ b/Engine/Source/Programs/UnrealBuildTool/Configuration/TestModuleRules.cs
@@ -147,7 +147,6 @@ namespace UnrealBuildTool
 			}
 		}
 
-
 #pragma warning disable 8602
 #pragma warning disable 8604
 		/// <summary>
@@ -249,7 +248,6 @@ namespace UnrealBuildTool
 
 			GenPropsDoc.Save(GeneratedPropertiesScriptFile);
 
-
 			// Platform-specific configurations
 			string GeneratedPropertiesPlatformFile;
 
@@ -361,9 +359,8 @@ namespace UnrealBuildTool
 		{
 			IEnumerable<XElement> NextChunk = ElementUpsertAfter.ElementsAfterSelf(BuildGraphNamespace + "Property")
 				.Where(prop => prop.Attribute("Name").Value.EndsWith(FlagSuffix));
-			if (NextChunk
-				.Where(prop => prop.Attribute("Name").Value == TestName + FlagSuffix)
-				.Count() == 0)
+			if (!NextChunk
+				.Where(prop => prop.Attribute("Name").Value == TestName + FlagSuffix).Any())
 			{
 				XElement ElementInsert = new XElement(BuildGraphNamespace + "Property");
 				ElementInsert.SetAttributeValue("Name", TestName + FlagSuffix);
@@ -382,9 +379,8 @@ namespace UnrealBuildTool
 		{
 			IEnumerable<XElement> NextChunk = ElementUpsertAfter.ElementsAfterSelf(BuildGraphNamespace + "Option")
 				.Where(prop => prop.Attribute("Name").Value.StartsWith(OptionPrefix) && prop.Attribute("Name").Value.EndsWith(OptionSuffix));
-			if (NextChunk
-				.Where(prop => prop.Attribute("Name").Value == OptionPrefix + OptionRadix + OptionSuffix)
-				.Count() == 0)
+			if (!NextChunk
+				.Where(prop => prop.Attribute("Name").Value == OptionPrefix + OptionRadix + OptionSuffix).Any())
 			{
 				XElement ElementInsert = new XElement(BuildGraphNamespace + "Option");
 				ElementInsert.SetAttributeValue("Name", OptionPrefix + OptionRadix + OptionSuffix);
diff --git a/Engine/Source/Programs/UnrealBuildTool/Configuration/UEBuildModuleCPP.cs b/Engine/Source/Programs/UnrealBuildTool/Configuration/UEBuildModuleCPP.cs
--- a/Engine/Source/Programs/UnrealBuildTool/Configuration/UEBuildModuleCPP.cs
+++ b/Engine/Source/Programs/UnrealBuildTool/Configuration/UEBuildModuleCPP.cs
@@ -665,7 +665,7 @@ namespace UnrealBuildTool
 			{
 				Unity.GetAdaptiveFiles(Target, CPPFiles, InputFiles.HeaderFiles, CompileEnvironment, WorkingSet, Rules.ShortName ?? Name, IntermediateDirectory, Graph,
 					out List<FileItem> NormalFiles, out List<FileItem> AdaptiveFiles);
-				if (NormalFiles.Where(file => !file.HasExtension(".gen.cpp")).Count() == 0)
+				if (!NormalFiles.Where(file => !file.HasExtension(".gen.cpp")).Any())
 				{
 					NormalFiles = CPPFiles;
 					AdaptiveFiles.RemoveAll(new HashSet<FileItem>(NormalFiles).Contains);
diff --git a/Engine/Source/Programs/UnrealBuildTool/Configuration/UEBuildTarget.cs b/Engine/Source/Programs/UnrealBuildTool/Configuration/UEBuildTarget.cs
--- a/Engine/Source/Programs/UnrealBuildTool/Configuration/UEBuildTarget.cs
+++ b/Engine/Source/Programs/UnrealBuildTool/Configuration/UEBuildTarget.cs
@@ -18,7 +18,6 @@ using Microsoft.Extensions.Logging;
 using OpenTracing.Util;
 using UnrealBuildBase;
 
-
 namespace UnrealBuildTool
 {
 	/// <summary>
@@ -98,7 +97,6 @@ namespace UnrealBuildTool
 			return Original;
 		}
 
-
 		/// <summary>
 		/// 
 		/// </summary>
@@ -199,7 +197,6 @@ namespace UnrealBuildTool
 				String.Join(",", GetUniqueStringRegistry().GetStringNames())));
 		}
 
-
 		/// <summary>
 		/// 
 		/// </summary>
@@ -459,7 +456,6 @@ namespace UnrealBuildTool
 
 		#endregion
 
-
 		/// <summary>
 		/// Return the string representation
 		/// </summary>
@@ -515,8 +511,6 @@ namespace UnrealBuildTool
 			return GetUniqueStringRegistry().HasString(Name);
 		}
 
-
-
 		/// <summary>
 		/// this group is just to lump Win32 and Win64 into Windows directories, removing the special Windows logic in MakeListOfUnsupportedPlatforms
 		/// </summary>
@@ -735,7 +729,6 @@ namespace UnrealBuildTool
 			return Name;
 		}
 
-
 		/// <summary>
 		/// 
 		/// </summary>
@@ -799,7 +792,6 @@ namespace UnrealBuildTool
 				String.Join(",", GetUniqueStringRegistry().GetStringNames())));
 		}
 
-
 		/// <summary>
 		/// 
 		/// </summary>
@@ -979,7 +971,6 @@ namespace UnrealBuildTool
 			return new UnrealArchitectures(ArchString, ValidationPlatform);
 		}
 
-
 		/// <summary>
 		/// The list of architecture names in this set, sorted by architecture name
 		/// </summary>
@@ -1244,7 +1235,6 @@ namespace UnrealBuildTool
 				RulesObject.DependencyListFileNames.Add(DependencyListFile);
 			}
 
-
 			// If we're compiling a plugin, and this target is monolithic, just create the object files
 			if (Descriptor.ForeignPlugin != null && RulesObject.LinkType == TargetLinkType.Monolithic)
 			{
@@ -3077,7 +3067,7 @@ namespace UnrealBuildTool
 							!IsRedistributable(DependencyModule) && DependencyModule.Name != AppName
 						);
 
-					if (NonRedistModules.Count() != 0)
+					if (NonRedistModules.Any())
 					{
 						IEnumerable<UEBuildModule> NonRedistDeps = AllDependencies.Where((DependantModule) =>
 							DependantModule.GetDirectDependencyModules().Intersect(NonRedistModules).Any()
@@ -5332,7 +5322,6 @@ namespace UnrealBuildTool
 			}
 		}
 
-
 		/// <summary>
 		/// Combines a list of paths with a base path.
 		/// </summary>
diff --git a/Engine/Source/Programs/UnrealBuildTool/Executors/ParallelExecutor.cs b/Engine/Source/Programs/UnrealBuildTool/Executors/ParallelExecutor.cs
--- a/Engine/Source/Programs/UnrealBuildTool/Executors/ParallelExecutor.cs
+++ b/Engine/Source/Programs/UnrealBuildTool/Executors/ParallelExecutor.cs
@@ -19,15 +19,6 @@ namespace UnrealBuildTool
 	/// </summary>
 	class ParallelExecutor : ActionExecutor
 	{
-		/// <summary>
-		/// Maximum processor count for local execution. 
-		/// </summary>
-		[XmlConfigFile]
-		[Obsolete("ParallelExecutor.MaxProcessorCount is deprecated. Please update xml to use BuildConfiguration.MaxParallelActions")]
-#pragma warning disable 0169
-		private static int MaxProcessorCount;
-#pragma warning restore 0169
-
 		/// <summary>
 		/// Processor count multiplier for local execution. Can be below 1 to reserve CPU for other tasks.
 		/// When using the local executor (not XGE), run a single action on each CPU core. Note that you can set this to a larger value
@@ -35,14 +26,14 @@ namespace UnrealBuildTool
 		/// This value is ignored if the CPU does not support hyper-threading.
 		/// </summary>
 		[XmlConfigFile]
-		private static double ProcessorCountMultiplier = 1.0;
+		private static double ProcessorCountMultiplier { get; set; } = 1.0;
 
 		/// <summary>
 		/// Free memory per action in bytes, used to limit the number of parallel actions if the machine is memory starved.
 		/// Set to 0 to disable free memory checking.
 		/// </summary>
 		[XmlConfigFile]
-		private static double MemoryPerActionBytes = 1.5 * 1024 * 1024 * 1024;
+		private static double MemoryPerActionBytes { get; set; } = 1.5 * 1024 * 1024 * 1024;
 
 		/// <summary>
 		/// The priority to set for spawned processes.
@@ -50,54 +41,54 @@ namespace UnrealBuildTool
 		/// Default: BelowNormal or Normal for an Asymmetrical processor as BelowNormal can cause scheduling issues.
 		/// </summary>
 		[XmlConfigFile]
-		protected static ProcessPriorityClass ProcessPriority = Utils.IsAsymmetricalProcessor() ? ProcessPriorityClass.Normal : ProcessPriorityClass.BelowNormal;
+		protected static ProcessPriorityClass ProcessPriority { get; set; } = Utils.IsAsymmetricalProcessor() ? ProcessPriorityClass.Normal : ProcessPriorityClass.BelowNormal;
 
 		/// <summary>
 		/// When enabled, will stop compiling targets after a compile error occurs.
 		/// </summary>
 		[XmlConfigFile]
-		private static bool bStopCompilationAfterErrors = false;
+		private static bool bStopCompilationAfterErrors { get; set; } = false;
 
 		/// <summary>
 		/// Whether to show compilation times along with worst offenders or not.
 		/// </summary>
 		[XmlConfigFile]
-		private static bool bShowCompilationTimes = false;
+		private static bool bShowCompilationTimes { get; set; } = false;
 
 		/// <summary>
 		/// Whether to show compilation times for each executed action
 		/// </summary>
 		[XmlConfigFile]
-		private static bool bShowPerActionCompilationTimes = false;
+		private static bool bShowPerActionCompilationTimes { get; set; } = false;
 
 		/// <summary>
 		/// Whether to log command lines for actions being executed
 		/// </summary>
 		[XmlConfigFile]
-		private static bool bLogActionCommandLines = false;
+		private static bool bLogActionCommandLines { get; set; } = false;
 
 		/// <summary>
 		/// Add target names for each action executed
 		/// </summary>
 		[XmlConfigFile]
-		private static bool bPrintActionTargetNames = false;
+		private static bool bPrintActionTargetNames { get; set; } = false;
 
 		/// <summary>
 		/// Whether to take into account the Action's weight when determining to do more work or not.
 		/// </summary>
 		[XmlConfigFile]
-		protected static bool bUseActionWeights = false;
+		protected static bool bUseActionWeights { get; set; } = false;
 
 		/// <summary>
 		/// Whether to show CPU utilization after the work is complete.
 		/// </summary>
 		[XmlConfigFile]
-		protected static bool bShowCPUUtilization = false;
+		protected static bool bShowCPUUtilization { get; set; } = false;
 
 		/// <summary>
 		/// Collapse non-error output lines
 		/// </summary>
-		private bool bCompactOutput = false;
+		private bool bCompactOutput { get; set; } = false;
 
 		/// <summary>
 		/// How many processes that will be executed in parallel
diff --git a/Engine/Source/Programs/UnrealBuildTool/Executors/SNDBS.cs b/Engine/Source/Programs/UnrealBuildTool/Executors/SNDBS.cs
--- a/Engine/Source/Programs/UnrealBuildTool/Executors/SNDBS.cs
+++ b/Engine/Source/Programs/UnrealBuildTool/Executors/SNDBS.cs
@@ -259,7 +259,7 @@ namespace UnrealBuildTool
 						["run_locally"] = !(a.bCanExecuteRemotely && a.bCanExecuteRemotelyWithSNDBS)
 					};
 
-					if (a.PrerequisiteItems.Count() > 0)
+					if (a.PrerequisiteItems.Any())
 					{
 						Job["explicit_input_files"] = a.PrerequisiteItems.Where(i => !(i.AbsolutePath.EndsWith(".rsp") || i.AbsolutePath.EndsWith(".response"))).Select(i => new Dictionary<string, object>()
 						{
diff --git a/Engine/Source/Programs/UnrealBuildTool/Platform/Windows/MicrosoftPlatformSDK.cs b/Engine/Source/Programs/UnrealBuildTool/Platform/Windows/MicrosoftPlatformSDK.cs
--- a/Engine/Source/Programs/UnrealBuildTool/Platform/Windows/MicrosoftPlatformSDK.cs
+++ b/Engine/Source/Programs/UnrealBuildTool/Platform/Windows/MicrosoftPlatformSDK.cs
@@ -82,8 +82,6 @@ namespace UnrealBuildTool
 			return false;
 		}
 
-
-
 		#region Windows Specific SDK discovery
 
 		#region Cached Locations
@@ -508,14 +506,14 @@ namespace UnrealBuildTool
 		/// <returns>True if the given compiler is installed</returns>
 		public static bool HasCompiler(WindowsCompiler Compiler, UnrealArch Architecture, ILogger Logger)
 		{
-			return FindToolChainInstallations(Compiler, Logger).Where(x => (x.Architecture == Architecture)).Count() > 0;
+			return FindToolChainInstallations(Compiler, Logger).Where(x => (x.Architecture == Architecture)).Any();
 		}
 
 		public static bool HasValidCompiler(WindowsCompiler Compiler, UnrealArch Architecture, ILogger Logger)
 		{
 			// since this is static, we get the instance for it
 			MicrosoftPlatformSDK SDK = (MicrosoftPlatformSDK)GetSDKForPlatform("Win64")!;
-			return FindToolChainInstallations(Compiler, Logger).Where(x => (x.Version >= MinimumVisualCppVersion && x.Architecture == Architecture)).Count() > 0;
+			return FindToolChainInstallations(Compiler, Logger).Where(x => (x.Version >= MinimumVisualCppVersion && x.Architecture == Architecture)).Any();
 		}
 
 		/// <summary>
@@ -654,7 +652,7 @@ namespace UnrealBuildTool
 				.ThenByDescending(x => x.IsAutoSdk)
 				.ThenByDescending(x => x.Version);
 
-			if (SortedToolChains.Count() > 0)
+			if (SortedToolChains.Any())
 			{
 				IEnumerable<ToolChainInstallation> AvailableToolChains = SortedToolChains.Where(x => x.Error == null);
 				IEnumerable<ToolChainInstallation> UnavailableToolChains = SortedToolChains.Where(x => x.Error != null);
@@ -937,7 +935,6 @@ namespace UnrealBuildTool
 			}
 		}
 
-
 		/// <summary>
 		/// Finds the most appropriate redist directory for the given toolchain version
 		/// </summary>
@@ -1056,7 +1053,6 @@ namespace UnrealBuildTool
 			}
 		}
 
-
 		/// <summary>
 		/// Add an Intel OneAPI toolchain
 		/// </summary>
@@ -1153,7 +1149,6 @@ namespace UnrealBuildTool
 			}
 		}
 
-
 		/// <summary>
 		/// Determines if the given path is a valid Visual C++ version number
 		/// </summary>
diff --git a/Engine/Source/Programs/UnrealBuildTool/ProjectFiles/Xcode/XcodeProject.File.cs b/Engine/Source/Programs/UnrealBuildTool/ProjectFiles/Xcode/XcodeProject.File.cs
--- a/Engine/Source/Programs/UnrealBuildTool/ProjectFiles/Xcode/XcodeProject.File.cs
+++ b/Engine/Source/Programs/UnrealBuildTool/ProjectFiles/Xcode/XcodeProject.File.cs
@@ -113,7 +113,6 @@ namespace UnrealBuildTool.XcodeProjectXcconfig
 		}
 	}
 
-
 	class XcodeFileCollection
 	{
 		static readonly string[] ResourceExtensions =
@@ -161,7 +160,6 @@ namespace UnrealBuildTool.XcodeProjectXcconfig
 
 		const string FileTypeDefault = "file.txt";
 
-
 		class ManualFileReference
 		{
 			public string Guid = XcodeProjectFileGenerator.MakeXcodeGuid();
@@ -189,7 +187,6 @@ namespace UnrealBuildTool.XcodeProjectXcconfig
 		private List<ManualFileReference> ManualFiles = new();
 		private Dictionary<string, string> GuidsForGroups = new();
 
-
 		public XcodeFileCollection(XcodeProjectFile ProjectFile)
 		{
 			ProjectDirectory = ProjectFile.ProjectFilePath.Directory;
@@ -236,7 +233,7 @@ namespace UnrealBuildTool.XcodeProjectXcconfig
 					{
 						if (!SourceToBuildFileMap.TryGetValue(SourceDir, out BuildFileOffset))
 						{
-							while (SourceDir != null && DirectoryReference.EnumerateFiles(SourceDir, "*.Build.cs").Count() == 0)
+							while (SourceDir != null && !DirectoryReference.EnumerateFiles(SourceDir, "*.Build.cs").Any())
 							{
 								SourceDir = SourceDir.ParentDirectory!;
 							}
@@ -343,7 +340,6 @@ namespace UnrealBuildTool.XcodeProjectXcconfig
 			}
 		}
 
-
 		public void Write(StringBuilder Content)
 		{
 			Content.WriteLine("/* Begin PBXBuildFile section */");
@@ -365,7 +361,6 @@ namespace UnrealBuildTool.XcodeProjectXcconfig
 			Content.WriteLine("/* End PBXBuildFile section */");
 			Content.WriteLine();
 
-
 			Content.WriteLine("/* Begin PBXFileReference section */");
 			foreach (XcodeSourceFile File in AllFiles)
 			{
@@ -401,7 +396,6 @@ namespace UnrealBuildTool.XcodeProjectXcconfig
 
 			Content.WriteLine("/* Begin PBXGroup section */");
 
-
 			// write main/root group and it's children (adding the manual groups to the main group, which is the true param here0
 			WriteGroup(Content, RootGroup, bIsRootGroup: true);
 
@@ -425,7 +419,6 @@ namespace UnrealBuildTool.XcodeProjectXcconfig
 			Content.WriteLine("/* End PBXGroup section */");
 		}
 
-
 		private void WriteGroup(StringBuilder Content, XcodeFileGroup Group, bool bIsRootGroup)
 		{
 			if (!Group.bIsReference)
@@ -494,7 +487,6 @@ namespace UnrealBuildTool.XcodeProjectXcconfig
 			return null;
 		}
 
-
 		/// <summary>
 		/// Gets Xcode file type based on its extension
 		/// </summary>
@@ -572,7 +564,6 @@ namespace UnrealBuildTool.XcodeProjectXcconfig
 						return null;
 					}
 
-
 					CurrentGroup = new XcodeFileGroup(Path.GetFileName(CurrentPath), CurrentPath, GetFileType(Path.GetExtension(CurrentPath)).StartsWith("folder"));
 					if (Groups.Count == 0)
 					{
diff --git a/Engine/Source/Programs/UnrealBuildTool/System/ConfigFileAttribute.cs b/Engine/Source/Programs/UnrealBuildTool/System/ConfigFileAttribute.cs
--- a/Engine/Source/Programs/UnrealBuildTool/System/ConfigFileAttribute.cs
+++ b/Engine/Source/Programs/UnrealBuildTool/System/ConfigFileAttribute.cs
@@ -8,7 +8,7 @@ namespace UnrealBuildTool
 	/// Attribute indicating a value which should be populated from a UE .ini config file
 	/// </summary>
 	[AttributeUsage(AttributeTargets.Field | AttributeTargets.Property)]
-	public class ConfigFileAttribute : Attribute
+	public sealed class ConfigFileAttribute : Attribute
 	{
 		/// <summary>
 		/// Name of the config hierarchy to read from
diff --git a/Engine/Source/Programs/UnrealBuildTool/System/EncryptionAndSigning.cs b/Engine/Source/Programs/UnrealBuildTool/System/EncryptionAndSigning.cs
--- a/Engine/Source/Programs/UnrealBuildTool/System/EncryptionAndSigning.cs
+++ b/Engine/Source/Programs/UnrealBuildTool/System/EncryptionAndSigning.cs
@@ -320,7 +320,7 @@ namespace UnrealBuildTool
 			ConfigHierarchySection CryptoSection = Ini.FindSection(SectionName);
 
 			// If we have new format crypto keys, read them in over the top of the legacy settings
-			if (CryptoSection != null && CryptoSection.KeyNames.Count() > 0)
+			if (CryptoSection != null && CryptoSection.KeyNames.Any())
 			{
 				Ini.GetBool(SectionName, "bEnablePakSigning", out bool bEnablePakSigning);
 				Settings.bEnablePakSigning = bEnablePakSigning;
diff --git a/Engine/Source/Programs/UnrealBuildTool/System/NativeProjects.cs b/Engine/Source/Programs/UnrealBuildTool/System/NativeProjects.cs
--- a/Engine/Source/Programs/UnrealBuildTool/System/NativeProjects.cs
+++ b/Engine/Source/Programs/UnrealBuildTool/System/NativeProjects.cs
@@ -77,8 +77,6 @@ namespace UnrealBuildTool
 			return CachedTargetNameToProjectFile.TryGetValue(InTargetName, out OutProjectFileName);
 		}
 
-
-
 		#region Hybrid content-only/code-based project
 
 		/// <summary>
@@ -226,10 +224,6 @@ namespace UnrealBuildTool
 			return true;
 		}
 
-
-
-
-
 		/// <summary>
 		/// Determines if a project (given a .uproject file) has source code. This is determined by finding at least one .Target.cs file
 		/// </summary>
@@ -242,18 +236,17 @@ namespace UnrealBuildTool
 			DirectoryReference TempSourceDir = DirectoryReference.Combine(UProjectFile.Directory, "Intermediate", "Source");
 
 			// check to see if we have a Target.cs file in Source or Intermediate/Source
-			if (DirectoryReference.Exists(SourceDir) && DirectoryReference.EnumerateFiles(SourceDir, "*.Target.cs", SearchOption.TopDirectoryOnly).Count() > 0)
+			if (DirectoryReference.Exists(SourceDir) && DirectoryReference.EnumerateFiles(SourceDir, "*.Target.cs", SearchOption.TopDirectoryOnly).Any())
 			{
 				return true;
 			}
-			if (bCheckForTempTargets && DirectoryReference.Exists(TempSourceDir) && DirectoryReference.EnumerateFiles(TempSourceDir, "*.Target.cs", SearchOption.TopDirectoryOnly).Count() > 0)
+			if (bCheckForTempTargets && DirectoryReference.Exists(TempSourceDir) && DirectoryReference.EnumerateFiles(TempSourceDir, "*.Target.cs", SearchOption.TopDirectoryOnly).Any())
 			{
 				return true;
 			}
 			return false;
 		}
 
-
 		private static bool RequiresTempTarget(FileReference UProjectFile, List<UnrealTargetPlatform> Platforms, List<UnrealTargetConfiguration> Configurations, [NotNullWhen(true)] out string? Reason, ILogger Logger)
 		{
 			// no reason by default
diff --git a/Engine/Source/Programs/UnrealBuildTool/System/ObsoleteOverrideAttribute.cs b/Engine/Source/Programs/UnrealBuildTool/System/ObsoleteOverrideAttribute.cs
--- a/Engine/Source/Programs/UnrealBuildTool/System/ObsoleteOverrideAttribute.cs
+++ b/Engine/Source/Programs/UnrealBuildTool/System/ObsoleteOverrideAttribute.cs
@@ -8,7 +8,7 @@ namespace UnrealBuildTool
 	/// Attribute used to denote that a method should no longer be overriden. Used by RulesCompiler.
 	/// </summary>
 	[AttributeUsage(AttributeTargets.Method)]
-	public class ObsoleteOverrideAttribute : Attribute
+	public sealed class ObsoleteOverrideAttribute : Attribute
 	{
 		/// <summary>
 		/// Message to display to the user if the method is overridden.
diff --git a/Engine/Source/Programs/UnrealBuildTool/System/ToolMode.cs b/Engine/Source/Programs/UnrealBuildTool/System/ToolMode.cs
--- a/Engine/Source/Programs/UnrealBuildTool/System/ToolMode.cs
+++ b/Engine/Source/Programs/UnrealBuildTool/System/ToolMode.cs
@@ -62,7 +62,7 @@ namespace UnrealBuildTool
 	/// <summary>
 	/// Attribute used to specify options for a UBT mode.
 	/// </summary>
-	class ToolModeAttribute : Attribute
+	sealed class ToolModeAttribute : Attribute
 	{
 		/// <summary>
 		/// Name of this mode
diff --git a/Engine/Source/Programs/UnrealBuildTool/System/Unity.cs b/Engine/Source/Programs/UnrealBuildTool/System/Unity.cs
--- a/Engine/Source/Programs/UnrealBuildTool/System/Unity.cs
+++ b/Engine/Source/Programs/UnrealBuildTool/System/Unity.cs
@@ -109,7 +109,6 @@ namespace UnrealBuildTool
 				}
 			}
 
-
 			/// <summary>
 			/// Starts a new unity file.  If the current unity file contains no files, this function has no effect, i.e. you will not get an empty unity file.
 			/// </summary>
@@ -187,7 +186,7 @@ namespace UnrealBuildTool
 			// Every single file in the module appears in the working set. Don't bother using adaptive unity for this module.
 			// Otherwise it would make full builds really slow.
 			GetAdaptiveFiles(Target, CPPFiles, HeaderFiles, CompileEnvironment, WorkingSet, BaseName, IntermediateDirectory, Graph, out NormalFiles, out AdaptiveFiles);
-			if (NormalFiles.Where(file => !file.HasExtension(".gen.cpp")).Count() == 0)
+			if (!NormalFiles.Where(file => !file.HasExtension(".gen.cpp")).Any())
 			{
 				NormalFiles = CPPFiles;
 				AdaptiveFiles.RemoveAll(new HashSet<FileItem>(NormalFiles).Contains);
diff --git a/Engine/Source/Programs/UnrealBuildTool/System/XmlConfigFileAttribute.cs b/Engine/Source/Programs/UnrealBuildTool/System/XmlConfigFileAttribute.cs
--- a/Engine/Source/Programs/UnrealBuildTool/System/XmlConfigFileAttribute.cs
+++ b/Engine/Source/Programs/UnrealBuildTool/System/XmlConfigFileAttribute.cs
@@ -8,7 +8,7 @@ namespace UnrealBuildTool
 	/// Marks a field as being serializable from a config file
 	/// </summary>
 	[AttributeUsage(AttributeTargets.Field | AttributeTargets.Property, AllowMultiple = true)]
-	class XmlConfigFileAttribute : Attribute
+	sealed class XmlConfigFileAttribute : Attribute
 	{
 		/// <summary>
 		/// The category for this config value. Optional; defaults to the declaring type name.
