UnrealBuildTool: More automated code cleanup

* Use object type rather than var
* Remove double newlines
* Use pattern matching

#rnx

diff --git a/Engine/Source/Programs/UnrealBuildTool/Configuration/TestModuleRules.cs b/Engine/Source/Programs/UnrealBuildTool/Configuration/TestModuleRules.cs
--- a/Engine/Source/Programs/UnrealBuildTool/Configuration/TestModuleRules.cs
+++ b/Engine/Source/Programs/UnrealBuildTool/Configuration/TestModuleRules.cs
@@ -147,7 +147,6 @@ namespace UnrealBuildTool
 			}
 		}
 
-
 #pragma warning disable 8602
 #pragma warning disable 8604
 		/// <summary>
@@ -228,7 +227,7 @@ namespace UnrealBuildTool
 					InsertOrUpdateTestOption(ref lastUpdatedNode, TestMetadata.TestName, "Run" + TestMetadata.TestShortName + "Tests", "Run", "Tests", false.ToString());
 
 					List<UnrealTargetPlatform> AllSupportedPlatforms = new List<UnrealTargetPlatform>();
-					var SupportedPlatforms = GetType().GetCustomAttributes(typeof(SupportedPlatformsAttribute), false);
+					object[] SupportedPlatforms = GetType().GetCustomAttributes(typeof(SupportedPlatformsAttribute), false);
 					// If none specified we assume Win64
 					if (SupportedPlatforms.Length == 0)
 					{
@@ -237,26 +236,25 @@ namespace UnrealBuildTool
 					}
 					else
 					{
-						foreach (var Platform in SupportedPlatforms)
+						foreach (object Platform in SupportedPlatforms)
 						{
 							AllSupportedPlatforms.AddRange(((SupportedPlatformsAttribute)Platform).Platforms);
 						}
 					}
 
-					InsertOrUpdateTestFlagProperty(ref lastUpdatedNode, TestMetadata.TestName, "SupportedPlatforms", AllSupportedPlatforms.Aggregate("", (current, next) => (current == "" ? next.ToString() : current + ";" + next.ToString())));
+					InsertOrUpdateTestFlagProperty(ref lastUpdatedNode, TestMetadata.TestName, "SupportedPlatforms", AllSupportedPlatforms.Aggregate("", (current, next) => (String.IsNullOrEmpty(current) ? next.ToString() : current + ";" + next.ToString())));
 				}
 			}
 
 			GenPropsDoc.Save(GeneratedPropertiesScriptFile);
 
-
 			// Platform-specific configurations
 			string GeneratedPropertiesPlatformFile;
 
 			string NonPublicPathPlatform;
 
 			// Generate peroperty file for each platform
-			foreach (var ValidPlatform in UnrealTargetPlatform.GetValidPlatforms())
+			foreach (UnrealTargetPlatform ValidPlatform in UnrealTargetPlatform.GetValidPlatforms())
 			{
 				bool IsRestrictedPlatformName = IsPlatformRestricted(ValidPlatform);
 				if (IsRestrictedPlatformName)
@@ -361,9 +359,8 @@ namespace UnrealBuildTool
 		{
 			IEnumerable<XElement> NextChunk = ElementUpsertAfter.ElementsAfterSelf(BuildGraphNamespace + "Property")
 				.Where(prop => prop.Attribute("Name").Value.EndsWith(FlagSuffix));
-			if (NextChunk
-				.Where(prop => prop.Attribute("Name").Value == TestName + FlagSuffix)
-				.Count() == 0)
+			if (!NextChunk
+				.Where(prop => prop.Attribute("Name").Value == TestName + FlagSuffix).Any())
 			{
 				XElement ElementInsert = new XElement(BuildGraphNamespace + "Property");
 				ElementInsert.SetAttributeValue("Name", TestName + FlagSuffix);
@@ -382,9 +379,8 @@ namespace UnrealBuildTool
 		{
 			IEnumerable<XElement> NextChunk = ElementUpsertAfter.ElementsAfterSelf(BuildGraphNamespace + "Option")
 				.Where(prop => prop.Attribute("Name").Value.StartsWith(OptionPrefix) && prop.Attribute("Name").Value.EndsWith(OptionSuffix));
-			if (NextChunk
-				.Where(prop => prop.Attribute("Name").Value == OptionPrefix + OptionRadix + OptionSuffix)
-				.Count() == 0)
+			if (!NextChunk
+				.Where(prop => prop.Attribute("Name").Value == OptionPrefix + OptionRadix + OptionSuffix).Any())
 			{
 				XElement ElementInsert = new XElement(BuildGraphNamespace + "Option");
 				ElementInsert.SetAttributeValue("Name", OptionPrefix + OptionRadix + OptionSuffix);
diff --git a/Engine/Source/Programs/UnrealBuildTool/Configuration/UEBuildDeploy.cs b/Engine/Source/Programs/UnrealBuildTool/Configuration/UEBuildDeploy.cs
--- a/Engine/Source/Programs/UnrealBuildTool/Configuration/UEBuildDeploy.cs
+++ b/Engine/Source/Programs/UnrealBuildTool/Configuration/UEBuildDeploy.cs
@@ -20,7 +20,7 @@ namespace UnrealBuildTool
 
 		protected bool CopyFilesRecursively(string FromDir, string ToDir)
 		{
-			var FromDirInfo = new DirectoryInfo(FromDir);
+			DirectoryInfo FromDirInfo = new DirectoryInfo(FromDir);
 
 			if (!FromDirInfo.Exists)
 			{
@@ -37,7 +37,7 @@ namespace UnrealBuildTool
 				try
 				{
 					// Ensure we can overrwrite dest file as non-read only.
-					var ToDestFile = Path.Combine(ToDir, FromFile.Name);
+					string ToDestFile = Path.Combine(ToDir, FromFile.Name);
 					if (File.Exists(ToDestFile))
 					{
 						File.SetAttributes(ToDestFile, File.GetAttributes(ToDestFile) & ~FileAttributes.ReadOnly);
diff --git a/Engine/Source/Programs/UnrealBuildTool/Configuration/UEBuildModuleCPP.cs b/Engine/Source/Programs/UnrealBuildTool/Configuration/UEBuildModuleCPP.cs
--- a/Engine/Source/Programs/UnrealBuildTool/Configuration/UEBuildModuleCPP.cs
+++ b/Engine/Source/Programs/UnrealBuildTool/Configuration/UEBuildModuleCPP.cs
@@ -339,7 +339,7 @@ namespace UnrealBuildTool
 					IncludePaths.Add(GeneratedCodeDirectoryVNI!);
 				}
 
-				foreach (var (Subdirectory, _) in Rules.GenerateHeaderFuncs)
+				foreach ((string Subdirectory, Action<ILogger, DirectoryReference> _) in Rules.GenerateHeaderFuncs)
 				{
 					IncludePaths.Add(DirectoryReference.Combine(GeneratedCodeDirectory, Subdirectory));
 				}
@@ -480,7 +480,7 @@ namespace UnrealBuildTool
 			// Compile all the generated CPP files
 			if (GeneratedCppDirectories != null && !CompileEnvironment.bHackHeaderGenerator)
 			{
-				var GeneratedFiles = new Dictionary<string, FileItem>();
+				Dictionary<string, FileItem> GeneratedFiles = new Dictionary<string, FileItem>();
 				foreach (string GeneratedDir in GeneratedCppDirectories)
 				{
 					if (!Directory.Exists(GeneratedDir))
@@ -507,7 +507,7 @@ namespace UnrealBuildTool
 					// Remove any generated files from the compile list if they are inlined
 					foreach (FileItem CPPFileItem in CPPFiles)
 					{
-						var ListOfInlinedGenCpps = ModuleCompileEnvironment.MetadataCache.GetListOfInlinedGeneratedCppFiles(CPPFileItem);
+						IList<string> ListOfInlinedGenCpps = ModuleCompileEnvironment.MetadataCache.GetListOfInlinedGeneratedCppFiles(CPPFileItem);
 						foreach (string ListOfInlinedGenCppsItem in ListOfInlinedGenCpps)
 						{
 							string Prefix = "UHT/";
@@ -529,12 +529,12 @@ namespace UnrealBuildTool
 
 					if (Rules.bEnableNonInlinedGenCppWarnings)
 					{
-						var CPPFilesLookup = new Dictionary<string, FileItem>();
-						foreach (var CPPFile in CPPFiles)
+						Dictionary<string, FileItem> CPPFilesLookup = new Dictionary<string, FileItem>();
+						foreach (FileItem CPPFile in CPPFiles)
 						{
 							CPPFilesLookup.Add(Utils.GetFilenameWithoutAnyExtensions(CPPFile.Name), CPPFile);
 						}
-						foreach (var Name in GeneratedFiles.Keys)
+						foreach (string Name in GeneratedFiles.Keys)
 						{
 							if (!Name.StartsWith("UHT/"))
 							{
@@ -665,7 +665,7 @@ namespace UnrealBuildTool
 			{
 				Unity.GetAdaptiveFiles(Target, CPPFiles, InputFiles.HeaderFiles, CompileEnvironment, WorkingSet, Rules.ShortName ?? Name, IntermediateDirectory, Graph,
 					out List<FileItem> NormalFiles, out List<FileItem> AdaptiveFiles);
-				if (NormalFiles.Where(file => !file.HasExtension(".gen.cpp")).Count() == 0)
+				if (!NormalFiles.Where(file => !file.HasExtension(".gen.cpp")).Any())
 				{
 					NormalFiles = CPPFiles;
 					AdaptiveFiles.RemoveAll(new HashSet<FileItem>(NormalFiles).Contains);
@@ -967,7 +967,7 @@ namespace UnrealBuildTool
 									return false;
 								}
 
-								foreach (var Definition in DefinitionsDictionary)
+								foreach (KeyValuePair<string, string> Definition in DefinitionsDictionary)
 								{
 									if (ParentInstance.DefinitionsDictionary.TryGetValue(Definition.Key, out string? ParentValue))
 									{
diff --git a/Engine/Source/Programs/UnrealBuildTool/Configuration/UEBuildPlatform.cs b/Engine/Source/Programs/UnrealBuildTool/Configuration/UEBuildPlatform.cs
--- a/Engine/Source/Programs/UnrealBuildTool/Configuration/UEBuildPlatform.cs
+++ b/Engine/Source/Programs/UnrealBuildTool/Configuration/UEBuildPlatform.cs
@@ -64,7 +64,6 @@ namespace UnrealBuildTool
 			return UnrealTargetPlatform.GetValidPlatforms().Where(x => UEBuildPlatform.TryGetBuildPlatform(x, out _)).Select(x => UEBuildPlatform.GetBuildPlatform(x).ArchitectureConfig);
 		}
 
-
 		/// <summary>
 		/// The multi-architecture mode for this platform (potentially single-architecture)
 		/// </summary>
@@ -101,7 +100,6 @@ namespace UnrealBuildTool
 			return ActiveArchitectures(ProjectFile, TargetName);
 		}
 
-
 		/// <summary>
 		/// Returns the set all architectures potentially supported by this project. Can be used by project file gnenerators to restrict IDE architecture options
 		/// Defaults to AllSupportedArchitectures
@@ -167,8 +165,6 @@ namespace UnrealBuildTool
 			return AllSupportedArchitectures.SingleArchitecture;
 		}
 
-
-
 		/// <summary>
 		/// Simple constructor for platforms with a single architecture
 		/// </summary>
@@ -189,13 +185,8 @@ namespace UnrealBuildTool
 			this.Mode = Mode;
 			AllSupportedArchitectures = new UnrealArchitectures(SupportedArchitectures);
 		}
-
-
 	}
 
-
-
-
 	abstract class UEBuildPlatform
 	{
 		private static Dictionary<UnrealTargetPlatform, UEBuildPlatform> BuildPlatformDictionary = new Dictionary<UnrealTargetPlatform, UEBuildPlatform>();
@@ -296,7 +287,6 @@ namespace UnrealBuildTool
 								continue;
 							}
 
-
 							// We need all platforms to be registered when we run -validateplatform command to check SDK status of each
 							if (bIncludeNonInstalledPlatforms || InstalledPlatformInfo.IsValidPlatform(TempInst.TargetPlatform))
 							{
@@ -813,7 +803,7 @@ namespace UnrealBuildTool
 
 				foreach (DirectoryReference Path in PublicSystemIncludePaths.Where(x => !x.IsUnderDirectory(Module.ModuleDirectory)))
 				{
-					var OtherModule = OtherModules.FirstOrDefault(x => Path.IsUnderDirectory(x.ModuleDirectory));
+					UEBuildModule? OtherModule = OtherModules.FirstOrDefault(x => Path.IsUnderDirectory(x.ModuleDirectory));
 					if (OtherModule is UEBuildModuleExternal)
 					{
 						LoggerFunc("External module '{Name}' is adding '{Path}' from external module '{OtherModule}' to PublicSystemIncludePaths. Did you intend to add a public reference?", Module.Name, Path, OtherModule.Name);
@@ -856,7 +846,7 @@ namespace UnrealBuildTool
 						}
 					}
 
-					var OtherModule = OtherModules.FirstOrDefault(x => Path.IsUnderDirectory(x.ModuleDirectory));
+					UEBuildModule? OtherModule = OtherModules.FirstOrDefault(x => Path.IsUnderDirectory(x.ModuleDirectory));
 					if (OtherModule is UEBuildModuleExternal)
 					{
 						LoggerFunc("Module '{Name}' is adding '{Path}' from external module '{OtherModule}' to PublicIncludePaths. Did you intend to add a public reference?", Module.Name, Path, OtherModule.Name);
@@ -897,7 +887,7 @@ namespace UnrealBuildTool
 						}
 					}
 
-					var OtherModule = OtherModules.FirstOrDefault(x => Path.IsUnderDirectory(x.ModuleDirectory));
+					UEBuildModule? OtherModule = OtherModules.FirstOrDefault(x => Path.IsUnderDirectory(x.ModuleDirectory));
 					if (OtherModule is UEBuildModuleExternal)
 					{
 						LoggerFunc("Module '{Name}' is adding '{Path}' from external module '{OtherModule}' to PrivateIncludePaths. Did you intend to add a private reference?", Module.Name, Path, OtherModule.Name);
@@ -941,7 +931,7 @@ namespace UnrealBuildTool
 						}
 					}
 
-					var OtherModule = OtherModules.FirstOrDefault(x => Path.IsUnderDirectory(x.ModuleDirectory));
+					UEBuildModule? OtherModule = OtherModules.FirstOrDefault(x => Path.IsUnderDirectory(x.ModuleDirectory));
 					if (OtherModule is UEBuildModuleExternal)
 					{
 						LoggerFunc("Module '{Name}' is adding '{Path}' from external module '{OtherModule}' to InternalncludePaths. Did you intend to add a public reference?", Module.Name, Path, OtherModule.Name);
@@ -1232,7 +1222,6 @@ namespace UnrealBuildTool
 		{
 		}
 
-
 		/// <summary>
 		/// Allows the platform to modify the binary link environment before the binary is built
 		/// </summary>
@@ -1259,7 +1248,6 @@ namespace UnrealBuildTool
 			return false;
 		}
 
-
 		/// <summary>
 		/// Checks if platform is part of a given platform group
 		/// </summary>
diff --git a/Engine/Source/Programs/UnrealBuildTool/Configuration/UEBuildTarget.cs b/Engine/Source/Programs/UnrealBuildTool/Configuration/UEBuildTarget.cs
--- a/Engine/Source/Programs/UnrealBuildTool/Configuration/UEBuildTarget.cs
+++ b/Engine/Source/Programs/UnrealBuildTool/Configuration/UEBuildTarget.cs
@@ -18,7 +18,6 @@ using Microsoft.Extensions.Logging;
 using OpenTracing.Util;
 using UnrealBuildBase;
 
-
 namespace UnrealBuildTool
 {
 	/// <summary>
@@ -98,7 +97,6 @@ namespace UnrealBuildTool
 			return Original;
 		}
 
-
 		/// <summary>
 		/// 
 		/// </summary>
@@ -199,7 +197,6 @@ namespace UnrealBuildTool
 				String.Join(",", GetUniqueStringRegistry().GetStringNames())));
 		}
 
-
 		/// <summary>
 		/// 
 		/// </summary>
@@ -459,7 +456,6 @@ namespace UnrealBuildTool
 
 		#endregion
 
-
 		/// <summary>
 		/// Return the string representation
 		/// </summary>
@@ -515,8 +511,6 @@ namespace UnrealBuildTool
 			return GetUniqueStringRegistry().HasString(Name);
 		}
 
-
-
 		/// <summary>
 		/// this group is just to lump Win32 and Win64 into Windows directories, removing the special Windows logic in MakeListOfUnsupportedPlatforms
 		/// </summary>
@@ -735,7 +729,6 @@ namespace UnrealBuildTool
 			return Name;
 		}
 
-
 		/// <summary>
 		/// 
 		/// </summary>
@@ -799,7 +792,6 @@ namespace UnrealBuildTool
 				String.Join(",", GetUniqueStringRegistry().GetStringNames())));
 		}
 
-
 		/// <summary>
 		/// 
 		/// </summary>
@@ -979,7 +971,6 @@ namespace UnrealBuildTool
 			return new UnrealArchitectures(ArchString, ValidationPlatform);
 		}
 
-
 		/// <summary>
 		/// The list of architecture names in this set, sorted by architecture name
 		/// </summary>
@@ -1244,7 +1235,6 @@ namespace UnrealBuildTool
 				RulesObject.DependencyListFileNames.Add(DependencyListFile);
 			}
 
-
 			// If we're compiling a plugin, and this target is monolithic, just create the object files
 			if (Descriptor.ForeignPlugin != null && RulesObject.LinkType == TargetLinkType.Monolithic)
 			{
@@ -2399,7 +2389,7 @@ namespace UnrealBuildTool
 			}
 
 			// Prefetch directory items for UHT folders since they are going to be used later
-			foreach (var Module in ModulesToGenerateHeadersFor)
+			foreach (UEBuildModuleCPP Module in ModulesToGenerateHeadersFor)
 			{
 				if (Module.GeneratedCodeDirectoryUHT != null)
 				{
@@ -2421,9 +2411,9 @@ namespace UnrealBuildTool
 			VNITask.Wait();
 #endif
 
-			foreach (var Module in Modules.Values.OfType<UEBuildModuleCPP>())
+			foreach (UEBuildModuleCPP Module in Modules.Values.OfType<UEBuildModuleCPP>())
 			{
-				foreach (var (Subdirectory, Func) in Module.Rules.GenerateHeaderFuncs)
+				foreach ((string Subdirectory, Action<ILogger, DirectoryReference> Func) in Module.Rules.GenerateHeaderFuncs)
 				{
 					DirectoryReference GenDirectory = DirectoryReference.Combine(Module.GeneratedCodeDirectory!, Subdirectory);
 					DirectoryReference.CreateDirectory(GenDirectory);
@@ -3077,7 +3067,7 @@ namespace UnrealBuildTool
 							!IsRedistributable(DependencyModule) && DependencyModule.Name != AppName
 						);
 
-					if (NonRedistModules.Count() != 0)
+					if (NonRedistModules.Any())
 					{
 						IEnumerable<UEBuildModule> NonRedistDeps = AllDependencies.Where((DependantModule) =>
 							DependantModule.GetDirectDependencyModules().Intersect(NonRedistModules).Any()
@@ -3523,7 +3513,7 @@ namespace UnrealBuildTool
 			// Create the PCH instances
 			for (int SharedPCHIndex = NumSharedPCHs - 1; SharedPCHIndex >= 0; SharedPCHIndex--)
 			{
-				foreach (var ModuleInfo in ModuleInfos[SharedPCHIndex])
+				foreach (Tuple<UEBuildModuleCPP, CppCompileEnvironment> ModuleInfo in ModuleInfos[SharedPCHIndex])
 				{
 					ModuleInfo.Item1.FindOrCreateSharedPCH(ToolChain, GlobalCompileEnvironment.SharedPCHs[SharedPCHIndex], ModuleInfo.Item2, Graph);
 				}
@@ -3565,7 +3555,7 @@ namespace UnrealBuildTool
 						// OptimizeCode
 						foreach (PrecompiledHeaderInstance Instance in SharedPCH.Instances.Where(instance => instance.Modules.First().Rules.OptimizeCode == ModuleRules.CodeOptimization.Never))
 						{
-							foreach (var Module in Instance.Modules)
+							foreach (UEBuildModuleCPP Module in Instance.Modules)
 							{
 								if (Module.RulesFile.ToFileInfo().IsReadOnly && !(Target.DisableOptimizeCodeForModules?.Contains(Module.Name) ?? false)) // support for local module changes
 								{
@@ -5332,7 +5322,6 @@ namespace UnrealBuildTool
 			}
 		}
 
-
 		/// <summary>
 		/// Combines a list of paths with a base path.
 		/// </summary>
diff --git a/Engine/Source/Programs/UnrealBuildTool/Executors/Experimental/FASTBuild.cs b/Engine/Source/Programs/UnrealBuildTool/Executors/Experimental/FASTBuild.cs
--- a/Engine/Source/Programs/UnrealBuildTool/Executors/Experimental/FASTBuild.cs
+++ b/Engine/Source/Programs/UnrealBuildTool/Executors/Experimental/FASTBuild.cs
@@ -943,8 +943,8 @@ namespace UnrealBuildTool
 				}
 				else
 				{
-					var numericDirectories = Directory.GetDirectories(VCEnv.GetToolPath().ToString()).Where(d => Path.GetFileName(d).All(Char.IsDigit));
-					var cluiDirectories = numericDirectories.Where(d => Directory.GetFiles(d, "clui.dll").Any());
+					IEnumerable<string> numericDirectories = Directory.GetDirectories(VCEnv.GetToolPath().ToString()).Where(d => Path.GetFileName(d).All(Char.IsDigit));
+					IEnumerable<string> cluiDirectories = numericDirectories.Where(d => Directory.GetFiles(d, "clui.dll").Any());
 					if (cluiDirectories.Any())
 					{
 						cluiSubDirName = Path.GetFileName(cluiDirectories.First());
@@ -1002,7 +1002,6 @@ namespace UnrealBuildTool
 						Logger.LogInformation("Using path : {PrefferedMSVCRedistPath} for vccorlib_.dll (MSVC redist)..." +
 							"\n\t...Add an entry for MsvcCRTRedistVersion in BuildConfiguration.xml to specify a version number", PrefferedMSVCRedistPath.ToString());
 					}
-
 				}
 
 				PotentialMSVCRedistPaths = new List<String>(Directory.EnumerateDirectories(String.Format("{0}/{1}", PrefferedMSVCRedistPath, VCEnv.Architecture)));
@@ -1023,7 +1022,6 @@ namespace UnrealBuildTool
 					//AddText(string.Format("\t\t'{0}/Redist/MSVC/{1}/x64/Microsoft.VC141.CRT/vccorlib{2}.dll'\n", VCEnv.GetVCInstallDirectory(), VCEnv.ToolChainVersion, platformVersionNumber));
 				}
 
-
 				AddText("\t}\n"); //End extra files
 
 				AddText($"\t.CompilerFamily = 'msvc'\n");
@@ -1103,7 +1101,7 @@ namespace UnrealBuildTool
 			}
 
 			string[] SpecialCompilerOptions = { "/Fo", "/fo", "/Yc", "/Yu", "/Fp", "-o", "-dependencies=", "-compiler=" };
-			var ParsedCompilerOptions = ParseCommandLineOptions(Action.CommandPath.GetFileName(), Action.CommandArguments, SpecialCompilerOptions, Logger);
+			Dictionary<string, string> ParsedCompilerOptions = ParseCommandLineOptions(Action.CommandPath.GetFileName(), Action.CommandArguments, SpecialCompilerOptions, Logger);
 
 			string OutputObjectFileName = GetOptionValue(ParsedCompilerOptions, IsMSVC() ? "/Fo" : "-o", Action, Logger, ProblemIfNotFound: !IsMSVC());
 
diff --git a/Engine/Source/Programs/UnrealBuildTool/Executors/ImmediateActionQueue.cs b/Engine/Source/Programs/UnrealBuildTool/Executors/ImmediateActionQueue.cs
--- a/Engine/Source/Programs/UnrealBuildTool/Executors/ImmediateActionQueue.cs
+++ b/Engine/Source/Programs/UnrealBuildTool/Executors/ImmediateActionQueue.cs
@@ -336,7 +336,7 @@ namespace UnrealBuildTool
 
 			if (readArtifacts)
 			{
-				var runAction = (LinkedAction action) =>
+				Func<LinkedAction, Func<Task>> runAction = (LinkedAction action) =>
 				{
 					return new Func<Task>(async () =>
 					{
@@ -1028,7 +1028,6 @@ namespace UnrealBuildTool
 					{
 						Logger.LogInformation("{Description} [ Time {ExecutionTime:0.00} s ]", Description, Result.ExecutionTime.TotalSeconds);
 					}
-
 				}
 				Logger.LogInformation("");
 			}
diff --git a/Engine/Source/Programs/UnrealBuildTool/Executors/SNDBS.cs b/Engine/Source/Programs/UnrealBuildTool/Executors/SNDBS.cs
--- a/Engine/Source/Programs/UnrealBuildTool/Executors/SNDBS.cs
+++ b/Engine/Source/Programs/UnrealBuildTool/Executors/SNDBS.cs
@@ -91,7 +91,7 @@ namespace UnrealBuildTool
 			{
 				string BrokerHostName = "";
 				Regex FindHost = new Regex(@"Active broker is ""(\S +)"" \((\S+)\)");
-				var LocalProcess = new Process();
+				Process LocalProcess = new Process();
 				LocalProcess.StartInfo = new ProcessStartInfo(SNDBSUtilExe, $"-connected");
 				LocalProcess.OutputDataReceived += (Sender, Args) =>
 				{
@@ -242,14 +242,14 @@ namespace UnrealBuildTool
 
 			int IdCounter = 0;
 			// Build the json script file to describe all the actions and their dependencies
-			var ActionIds = Actions.ToDictionary(a => a, a => new Guid(++IdCounter, 0, 0, 0, 0, 0, 0, 0, 0, 0, 0).ToString());
+			Dictionary<LinkedAction, string> ActionIds = Actions.ToDictionary(a => a, a => new Guid(++IdCounter, 0, 0, 0, 0, 0, 0, 0, 0, 0, 0).ToString());
 			JsonSerializerOptions JsonOption = new JsonSerializerOptions();
 			JsonOption.Encoder = System.Text.Encodings.Web.JavaScriptEncoder.UnsafeRelaxedJsonEscaping;
 			File.WriteAllText(ScriptFile.FullName, JsonSerializer.Serialize(new Dictionary<string, object>()
 			{
 				["jobs"] = Actions.Select(a =>
 				{
-					var Job = new Dictionary<string, object>()
+					Dictionary<string, object> Job = new Dictionary<string, object>()
 					{
 						["id"] = ActionIds[a],
 						["title"] = a.StatusDescription,
@@ -259,7 +259,7 @@ namespace UnrealBuildTool
 						["run_locally"] = !(a.bCanExecuteRemotely && a.bCanExecuteRemotelyWithSNDBS)
 					};
 
-					if (a.PrerequisiteItems.Count() > 0)
+					if (a.PrerequisiteItems.Any())
 					{
 						Job["explicit_input_files"] = a.PrerequisiteItems.Where(i => !(i.AbsolutePath.EndsWith(".rsp") || i.AbsolutePath.EndsWith(".response"))).Select(i => new Dictionary<string, object>()
 						{
@@ -277,10 +277,10 @@ namespace UnrealBuildTool
 			PrepareToolTemplates();
 			bool bHasRewrites = GenerateSNDBSIncludeRewriteRules();
 
-			var ConfigList = TargetDescriptors.Select(Descriptor => $"{Descriptor.Name}|{Descriptor.Platform}|{Descriptor.Configuration}");
-			var ConfigDescription = String.Join(",", ConfigList);
+			IEnumerable<string> ConfigList = TargetDescriptors.Select(Descriptor => $"{Descriptor.Name}|{Descriptor.Platform}|{Descriptor.Configuration}");
+			string ConfigDescription = String.Join(",", ConfigList);
 
-			var StartInfo = new ProcessStartInfo(
+			ProcessStartInfo StartInfo = new ProcessStartInfo(
 				SNDBSBuildExe,
 				$"-q -p \"{ConfigDescription}\" -s \"{ScriptFile}\" -templates \"{IntermediateDir}\""
 				);
@@ -317,7 +317,7 @@ namespace UnrealBuildTool
 							Writer.Write(++NumCompletedActions, NumActions);
 
 							Text = Args.Data.Substring(ProgressMarkupPrefix.Length).Trim();
-							var ActionInfo = Text.Split(':');
+							string[] ActionInfo = Text.Split(':');
 							Logger.LogInformation("[{NumCompletedActions}/{NumActions}] {ActionInfo0} {ActionInfo1}", NumCompletedActions, NumActions, ActionInfo[0], ActionInfo[1]);
 							CurrentStatus = ActionInfo[1];
 							return;
@@ -369,10 +369,10 @@ namespace UnrealBuildTool
 
 		private void PrepareToolTemplates()
 		{
-			foreach (var Template in ActiveTemplates)
+			foreach (KeyValuePair<string, string> Template in ActiveTemplates)
 			{
-				var TemplateFile = FileReference.Combine(IntermediateDir, $"{Template.Key}.sn-dbs-tool.ini");
-				var TemplateText = Template.Value;
+				FileReference TemplateFile = FileReference.Combine(IntermediateDir, $"{Template.Key}.sn-dbs-tool.ini");
+				string TemplateText = Template.Value;
 
 				foreach (Nullable<DictionaryEntry> Variable in Environment.GetEnvironmentVariables(EnvironmentVariableTarget.Process))
 				{
@@ -389,7 +389,7 @@ namespace UnrealBuildTool
 		private bool GenerateSNDBSIncludeRewriteRules()
 		{
 			// Get all distinct platform names being used in this build.
-			var Platforms = TargetDescriptors
+			List<string> Platforms = TargetDescriptors
 				.Select(TargetDescriptor => UEBuildPlatform.GetBuildPlatform(TargetDescriptor.Platform).GetPlatformName())
 				.Distinct()
 				.ToList();
@@ -397,7 +397,7 @@ namespace UnrealBuildTool
 			if (Platforms.Count > 0)
 			{
 				// language=regex
-				var Lines = new[]
+				string[] Lines = new[]
 				{
 					@"pattern1=^COMPILED_PLATFORM_HEADER\(\s*([^ ,]+)\s*\)",
 					$"expansions1={String.Join("|", Platforms.Select(Name => $"{Name}/{Name}$1|{Name}$1"))}",
diff --git a/Engine/Source/Programs/UnrealBuildTool/Modes/AggregateClangTimingInfo.cs b/Engine/Source/Programs/UnrealBuildTool/Modes/AggregateClangTimingInfo.cs
--- a/Engine/Source/Programs/UnrealBuildTool/Modes/AggregateClangTimingInfo.cs
+++ b/Engine/Source/Programs/UnrealBuildTool/Modes/AggregateClangTimingInfo.cs
@@ -180,7 +180,7 @@ namespace UnrealBuildTool
 
 			if (AggregateFile != null)
 			{
-				var Tasks = Task.WhenAll(SourceFiles.Select(x => GetTraceData(x, Logger)));
+				Task<TraceData[]> Tasks = Task.WhenAll(SourceFiles.Select(x => GetTraceData(x, Logger)));
 				Tasks.Wait();
 				List<TraceData> TraceDatas = Tasks.Result.OrderBy(x => x.Module).ThenBy(x => x.Name).ToList();
 
@@ -198,7 +198,7 @@ namespace UnrealBuildTool
 
 			if (HeadersFile != null)
 			{
-				var Tasks = Task.WhenAll(SourceFiles.Select(x => ParseTimingDataFile(x, Logger)));
+				Task<ClangTrace[]> Tasks = Task.WhenAll(SourceFiles.Select(x => ParseTimingDataFile(x, Logger)));
 				Tasks.Wait();
 				List<ClangTrace> ClangTraces = Tasks.Result.ToList();
 
@@ -221,7 +221,7 @@ namespace UnrealBuildTool
 				using (StreamWriter Writer = new StreamWriter(TempFilePath))
 				{
 					Writer.WriteLine("Source,Count,Total,Min,Max,Average");
-					foreach (var Data in Sources.OrderBy(x => x.Key.FullName).Where(x => x.Key.HasExtension(".h") || x.Key.HasExtension(".inl")))
+					foreach (KeyValuePair<FileReference, List<long>> Data in Sources.OrderBy(x => x.Key.FullName).Where(x => x.Key.HasExtension(".h") || x.Key.HasExtension(".inl")))
 					{
 						Writer.WriteLine($"{Data.Key.FullName},{Data.Value.Count},{Data.Value.Sum()},{Data.Value.Min()},{Data.Value.Max()},{Data.Value.Average()}");
 					}
@@ -258,7 +258,6 @@ namespace UnrealBuildTool
 			return Task.FromResult(0);
 		}
 
-
 		private async Task<ClangTrace> ParseTimingDataFile(FileReference SourceFile, ILogger Logger)
 		{
 			if (!ClangTraceCache.ContainsKey(SourceFile))
diff --git a/Engine/Source/Programs/UnrealBuildTool/Modes/CleanMode.cs b/Engine/Source/Programs/UnrealBuildTool/Modes/CleanMode.cs
--- a/Engine/Source/Programs/UnrealBuildTool/Modes/CleanMode.cs
+++ b/Engine/Source/Programs/UnrealBuildTool/Modes/CleanMode.cs
@@ -226,7 +226,7 @@ namespace UnrealBuildTool
 				});
 			}
 
-			var DeleteDirectories = Task.Run(() =>
+			Task DeleteDirectories = Task.Run(() =>
 			{
 				using ScopedTimer Timer = new ScopedTimer($"Delete {DirectoriesToDelete.Count} directories", Logger);
 				Parallel.ForEach(DirectoriesToDelete, DirectoryToDelete =>
@@ -249,7 +249,7 @@ namespace UnrealBuildTool
 				});
 			});
 
-			var DeleteFiles = Task.Run(() =>
+			Task DeleteFiles = Task.Run(() =>
 			{
 				using ScopedTimer Timer = new ScopedTimer($"Delete {FilesToDelete.Count} files", Logger);
 				Parallel.ForEach(FilesToDelete, FileToDelete =>
diff --git a/Engine/Source/Programs/UnrealBuildTool/Modes/FixIncludePathsMode.cs b/Engine/Source/Programs/UnrealBuildTool/Modes/FixIncludePathsMode.cs
--- a/Engine/Source/Programs/UnrealBuildTool/Modes/FixIncludePathsMode.cs
+++ b/Engine/Source/Programs/UnrealBuildTool/Modes/FixIncludePathsMode.cs
@@ -111,12 +111,12 @@ namespace UnrealBuildTool
 							InputFiles.AddRange(InputFileCollection.CCFiles);
 							InputFiles.AddRange(InputFileCollection.CFiles);
 
-							var FileList = new List<FileReference>();
+							List<FileReference> FileList = new List<FileReference>();
 							foreach (FileItem InputFile in InputFiles)
 							{
 								if (FileFilter == null || FileFilter.Matches(InputFile.Location.MakeRelativeTo(Unreal.RootDirectory)))
 								{
-									var fileRef = new FileReference(InputFile.AbsolutePath);
+									FileReference fileRef = new FileReference(InputFile.AbsolutePath);
 									FileList.Add(fileRef);
 								}
 							}
@@ -126,7 +126,7 @@ namespace UnrealBuildTool
 							{
 								if (IncludeFilter == null || IncludeFilter.Matches(InputFile.Location.MakeRelativeTo(Unreal.RootDirectory)))
 								{
-									var fileRef = new FileReference(InputFile.AbsolutePath);
+									FileReference fileRef = new FileReference(InputFile.AbsolutePath);
 									IncludeFileList.Add(fileRef);
 								}
 							}
@@ -155,7 +155,7 @@ namespace UnrealBuildTool
 							Dictionary<string, string?> PreferredPathCache = new();
 							CppCompileEnvironment env = Module.CreateCompileEnvironmentForIntellisense(Target.Rules, BinaryCompileEnvironment, Logger);
 
-							foreach (var InputFile in FileList)
+							foreach (FileReference InputFile in FileList)
 							{
 								List<int> LinesUpdated = new();
 								string[] Text = FileReference.ReadAllLines(InputFile);
@@ -163,7 +163,7 @@ namespace UnrealBuildTool
 
 								for (int i = 0; i < Text.Length; i++)
 								{
-									var Line = Text[i];
+									string Line = Text[i];
 									int LineNumber = i + 1;
 									Match IncludeMatch = IncludeRegex.Match(Line);
 									if (IncludeMatch.Success)
@@ -195,7 +195,7 @@ namespace UnrealBuildTool
 											// search include paths
 											FileReference? FoundIncludeFile = null;
 											DirectoryReference? FoundIncludePath = null;
-											foreach (var IncludePath in IncludePaths)
+											foreach (DirectoryReference IncludePath in IncludePaths)
 											{
 												string Path = System.IO.Path.GetFullPath(System.IO.Path.Combine(IncludePath.FullName, Include));
 												if (System.IO.File.Exists(Path))
@@ -405,10 +405,10 @@ namespace UnrealBuildTool
 
 		private void SortIncludes(FileReference File, List<int> LinesUpdated, string[] Text)
 		{
-			var HeaderSort = new HeaderSortComparison(File.GetFileNameWithoutExtension() + ".h");
-			foreach (var LineIndex in LinesUpdated)
+			HeaderSortComparison HeaderSort = new HeaderSortComparison(File.GetFileNameWithoutExtension() + ".h");
+			foreach (int LineIndex in LinesUpdated)
 			{
-				var FirstIncludeIndex = LineIndex;
+				int FirstIncludeIndex = LineIndex;
 				for (int i = LineIndex - 1; i >= 0; i--)
 				{
 					Match IncludeMatch = IncludeRegex.Match(Text[i]);
@@ -422,7 +422,7 @@ namespace UnrealBuildTool
 					}
 				}
 
-				var LastIncludeIndex = LineIndex;
+				int LastIncludeIndex = LineIndex;
 				for (int i = LineIndex + 1; i < Text.Length; i++)
 				{
 					Match IncludeMatch = IncludeRegex.Match(Text[i]);
diff --git a/Engine/Source/Programs/UnrealBuildTool/Modes/GenerateClangDatabase.cs b/Engine/Source/Programs/UnrealBuildTool/Modes/GenerateClangDatabase.cs
--- a/Engine/Source/Programs/UnrealBuildTool/Modes/GenerateClangDatabase.cs
+++ b/Engine/Source/Programs/UnrealBuildTool/Modes/GenerateClangDatabase.cs
@@ -240,7 +240,7 @@ namespace UnrealBuildTool
 										}
 
 										// find source file
-										var SourceFile = Action.PrerequisiteItems.FirstOrDefault(fi => fi.HasExtension(".cpp") || fi.HasExtension(".c") || fi.HasExtension(".c"));
+										FileItem? SourceFile = Action.PrerequisiteItems.FirstOrDefault(fi => fi.HasExtension(".cpp") || fi.HasExtension(".c") || fi.HasExtension(".c"));
 										if (SourceFile != null)
 										{
 											FileToCommand[SourceFile.Location] = CommandBuilder.ToString();
@@ -341,7 +341,7 @@ namespace UnrealBuildTool
 				OldPath = OldPath.Remove(OldPath.Length - 1, 1);
 			}
 
-			var FileReference = new FileReference(OldPath);
+			FileReference FileReference = new FileReference(OldPath);
 			return Line.Replace(OldPath, FileReference.FullName.Replace("\\", "/"));
 		}
 	}
diff --git a/Engine/Source/Programs/UnrealBuildTool/Modes/IWYUMode.cs b/Engine/Source/Programs/UnrealBuildTool/Modes/IWYUMode.cs
--- a/Engine/Source/Programs/UnrealBuildTool/Modes/IWYUMode.cs
+++ b/Engine/Source/Programs/UnrealBuildTool/Modes/IWYUMode.cs
@@ -14,7 +14,6 @@ using EpicGames.Core;
 using Microsoft.Extensions.Logging;
 using UnrealBuildBase;
 
-
 // IWYUMode is a mode that can be used to clean up includes in source code. It uses the clang based tool include-what-you-use (IWYU) to figure out what is needed in each .h/.cpp file
 // and then cleans up accordingly. This mode can be used to clean up unreal as well as plugins and projects on top of unreal.
 // Note, IWYU is not perfect. There are still c++ features not supported so even though it will do a good job cleaning up it might require a little bit of hands on but not much.
@@ -41,7 +40,6 @@ using UnrealBuildBase;
 // 
 //   .build target UnrealEditor linux development -- -Mode=IWYU -PathToUpdate=Engine/Plugins/FX/Niagara -UpdateOnlyPrivate
 
-
 namespace UnrealBuildTool
 {
 	/// <summary>
@@ -437,9 +435,9 @@ namespace UnrealBuildTool
 				// Traverse all modules to figure out which ones should be updated. This is based on -ModuleToUpdate and -PathToUpdate
 				List<UEBuildModule> ReferencedModules = new();
 				HashSet<UEBuildModule> IgnoreReferencedModules = new();
-				foreach (var Binary in Target.Binaries)
+				foreach (UEBuildBinary Binary in Target.Binaries)
 				{
-					foreach (var Module in Binary.Modules)
+					foreach (UEBuildModule Module in Binary.Modules)
 					{
 						if (IgnoreReferencedModules.Add(Module))
 						{
@@ -451,13 +449,13 @@ namespace UnrealBuildTool
 				foreach (UEBuildModule Module in ReferencedModules)
 				{
 					NameToModule.TryAdd(Module.Name, Module);
-					foreach (var Module2 in Module.GetDependencies(true, true))
+					foreach (UEBuildModule Module2 in Module.GetDependencies(true, true))
 					{
 						NameToModule.TryAdd(Module2.Name, Module2);
 					}
 				}
 
-				foreach (var Module in NameToModule.Values)
+				foreach (UEBuildModule Module in NameToModule.Values)
 				{
 					bool ShouldUpdate = false;
 					string ModuleDir = Module.ModuleDirectory.FullName.Replace('\\', '/');
@@ -486,14 +484,14 @@ namespace UnrealBuildTool
 						{
 							if (Module.PublicDependencyModules != null)
 							{
-								foreach (var Dependency in Module.PublicDependencyModules)
+								foreach (UEBuildModule Dependency in Module.PublicDependencyModules)
 								{
 									DependsOnModule = DependsOnModule || Dependency == UEModuleToUpdate;
 								}
 							}
 							if (Module.PrivateDependencyModules != null)
 							{
-								foreach (var Dependency in Module.PrivateDependencyModules!)
+								foreach (UEBuildModule Dependency in Module.PrivateDependencyModules!)
 								{
 									DependsOnModule = DependsOnModule || Dependency == UEModuleToUpdate;
 								}
@@ -636,7 +634,7 @@ namespace UnrealBuildTool
 
 					TocContent = new();
 
-					foreach (var OutputItem in OutputItems)
+					foreach (FileItem OutputItem in OutputItems)
 					{
 						if (OutputItem.Name.EndsWith(".iwyu"))
 						{
@@ -664,12 +662,11 @@ namespace UnrealBuildTool
 						IWYUFile!.Name = IWYUFilePath;
 
 						// Traverse the cpp/inl/h file entries inside the .iwyu file
-						foreach (var Info in IWYUFile!.Files)
+						foreach (IWYUInfo Info in IWYUFile!.Files)
 						{
 							Info.Source = IWYUFile;
 							Info.IsCpp = Info.File.EndsWith(".cpp");
 
-
 							// We track .gen.cpp in a special list, they need special treatment later
 							if (Info.File.Contains(".gen.cpp", StringComparison.Ordinal))
 							{
@@ -687,7 +684,7 @@ namespace UnrealBuildTool
 								if (!Infos.TryAdd(Info.File, Info))
 								{
 									// This is a valid scenario when Foo.cpp also registers Foo.h... and then Foo.h registers itself.
-									var ExistingEntry = Infos[Info.File];
+									IWYUInfo ExistingEntry = Infos[Info.File];
 									if (IWYUFile.Files.Count == 1)
 									{
 										if (ExistingEntry.Source!.Files.Count == 1)
@@ -727,7 +724,7 @@ namespace UnrealBuildTool
 							if (Info.File.EndsWith(".h"))
 							{
 								// We need to add entries for .generated.h. They are not in the iwyu files
-								foreach (var Include in Info.Includes)
+								foreach (IWYUIncludeEntry Include in Info.Includes)
 								{
 									if (Include.Full.Contains("/UHT/"))
 									{
@@ -791,7 +788,7 @@ namespace UnrealBuildTool
 				return new(Include, Entry);
 			}
 
-			var SpecialIncludes = new Dictionary<string, IWYUIncludeEntry?>(
+			Dictionary<string, IWYUIncludeEntry?> SpecialIncludes = new Dictionary<string, IWYUIncludeEntry?>(
 				new[]
 				{
 					GetInfo("UObject/ObjectMacros.h", "Source/Runtime/CoreUObject/Public/UObject/ObjectMacros.h"),
@@ -802,7 +799,7 @@ namespace UnrealBuildTool
 				}
 			);
 
-			var ForwardingHeaders = new Dictionary<string, IWYUIncludeEntry?>()
+			Dictionary<string, IWYUIncludeEntry?> ForwardingHeaders = new Dictionary<string, IWYUIncludeEntry?>()
 			{
 				{ "TMap", SpecialIncludes["Containers/ContainersFwd.h"] },
 				{ "TSet", SpecialIncludes["Containers/ContainersFwd.h"] },
@@ -811,7 +808,6 @@ namespace UnrealBuildTool
 				{ "TOptional", SpecialIncludes["Misc/OptionalFwd.h"] },
 			};
 
-
 			// Add all .generated.h files as entries in the lookup and explicitly add the includes they have which will never be removed
 			Logger.LogInformation($"Generating infos for .generated.h files...");
 			if (GeneratedHeaderInfos.Count > 0)
@@ -820,7 +816,7 @@ namespace UnrealBuildTool
 				IWYUIncludeEntry ScriptMacrosInclude = SpecialIncludes["UObject/ScriptMacros.h"]!;
 				IWYUIncludeEntry VerseInteropUtilsInclude = SpecialIncludes["VerseInteropUtils.h"]!;
 
-				foreach (var Gen in GeneratedHeaderInfos)
+				foreach (string Gen in GeneratedHeaderInfos)
 				{
 					IWYUInfo GenInfo = new();
 					GenInfo.File = Gen;
@@ -844,7 +840,7 @@ namespace UnrealBuildTool
 
 			Logger.LogInformation($"Resolving Includes...");
 
-			var PlatformSDK = new LinuxPlatformSDK(Logger);
+			LinuxPlatformSDK PlatformSDK = new LinuxPlatformSDK(Logger);
 			DirectoryReference? BaseLinuxPath = PlatformSDK.GetBaseLinuxPathForArchitecture(LinuxPlatform.DefaultHostArchitecture);
 			string SystemPath = BaseLinuxPath!.FullName.Replace('\\', '/');
 
@@ -866,7 +862,6 @@ namespace UnrealBuildTool
 					}
 				}
 
-
 				if (!Info.IsCpp)
 				{
 					// See if we need to replace forward declarations with includes
@@ -930,7 +925,7 @@ namespace UnrealBuildTool
 					}
 				}
 				*/
-				foreach (var Include in Info.Includes)
+				foreach (IWYUIncludeEntry Include in Info.Includes)
 				{
 					if (Include.Resolved == null)
 					{
@@ -985,10 +980,10 @@ namespace UnrealBuildTool
 			{
 				// First, check which files .gen.cpp will see
 				HashSet<string> SeenTransitiveIncludes = new();
-				foreach (var SeenInclude in GeneratedCpp.IncludesSeenInFile)
+				foreach (IWYUIncludeEntry SeenInclude in GeneratedCpp.IncludesSeenInFile)
 				{
 					SeenTransitiveIncludes.Add(SeenInclude.Full);
-					foreach (var Include in GeneratedCpp.Includes)
+					foreach (IWYUIncludeEntry Include in GeneratedCpp.Includes)
 					{
 						IWYUInfo? IncludeInfo;
 						if (SeenInclude.Printable == Include.Printable && Infos.TryGetValue(Include.Full, out IncludeInfo))
@@ -1010,7 +1005,7 @@ namespace UnrealBuildTool
 				{
 					int NameStart = GeneratedCpp.File.LastIndexOf('/');
 					string HeaderName = GeneratedCpp.File.Substring(NameStart + 1, GeneratedCpp.File.Length - NameStart - ".gen.cpp".Length) + "h";
-					foreach (var Include in GeneratedCpp.Includes)
+					foreach (IWYUIncludeEntry Include in GeneratedCpp.Includes)
 					{
 						NameStart = Include.Full.LastIndexOf('/');
 						string IncludeName = Include.Full.Substring(NameStart + 1);
@@ -1031,7 +1026,7 @@ namespace UnrealBuildTool
 					return;
 				}
 
-				foreach (var Include in GeneratedCpp.Includes)
+				foreach (IWYUIncludeEntry Include in GeneratedCpp.Includes)
 				{
 					if (SeenTransitiveIncludes.Contains(Include.Full))
 					{
@@ -1092,7 +1087,7 @@ namespace UnrealBuildTool
 			{
 				return null;
 			}
-			foreach (var Dir in Module.PublicIncludePaths.Union(Module.PrivateIncludePaths).Union(Module.PublicSystemIncludePaths))
+			foreach (DirectoryReference? Dir in Module.PublicIncludePaths.Union(Module.PrivateIncludePaths).Union(Module.PublicSystemIncludePaths))
 			{
 				FileReference FileReference = FileReference.Combine(Dir, Include);
 				FileItem FileItem = FileItem.GetItemByFileReference(FileReference);
@@ -1114,7 +1109,7 @@ namespace UnrealBuildTool
 				}
 			}
 
-			foreach (var PublicModule in Module.Rules.PublicDependencyModuleNames.Union(Module.Rules.PublicIncludePathModuleNames).Union(Module.Rules.PrivateIncludePathModuleNames).Union(Module.Rules.PrivateDependencyModuleNames))
+			foreach (string? PublicModule in Module.Rules.PublicDependencyModuleNames.Union(Module.Rules.PublicIncludePathModuleNames).Union(Module.Rules.PrivateIncludePathModuleNames).Union(Module.Rules.PrivateDependencyModuleNames))
 			{
 				UEBuildModule? DependencyModule;
 				if (NameToModule.TryGetValue(PublicModule, out DependencyModule))
@@ -1149,7 +1144,7 @@ namespace UnrealBuildTool
 
 		static void ParseSkippedHeaders(Dictionary<string, IWYUInfo> SkippedHeaders, Dictionary<string, IWYUInfo> Infos, Dictionary<string, UEBuildModule> PathToModule, Dictionary<string, UEBuildModule> NameToModule)
 		{
-			foreach (var Kvp in SkippedHeaders)
+			foreach (KeyValuePair<string, IWYUInfo> Kvp in SkippedHeaders)
 			{
 				Infos.Add(Kvp.Key, Kvp.Value);
 			}
@@ -1190,14 +1185,14 @@ namespace UnrealBuildTool
 						}
 						else if (LineSpan.StartsWith("include"))
 						{
-							var IncludeSpan = LineSpan.Slice("include".Length).TrimStart();
+							ReadOnlySpan<char> IncludeSpan = LineSpan.Slice("include".Length).TrimStart();
 							char LeadingIncludeChar = IncludeSpan[0];
 							if (LeadingIncludeChar == '"' || LeadingIncludeChar == '<')
 							{
-								var Start = IncludeSpan.Slice(1);
-								var EndIndex = Start.IndexOf(LeadingIncludeChar == '"' ? '"' : '>');
-								var FileSpan = Start.Slice(0, EndIndex);
-								var Include = FileSpan.ToString();
+								ReadOnlySpan<char> Start = IncludeSpan.Slice(1);
+								int EndIndex = Start.IndexOf(LeadingIncludeChar == '"' ? '"' : '>');
+								ReadOnlySpan<char> FileSpan = Start.Slice(0, EndIndex);
+								string Include = FileSpan.ToString();
 								string? File = GetFullName(Include, Info.File, Info.Module, NameToModule);
 								if (File != null)
 								{
@@ -1310,7 +1305,7 @@ namespace UnrealBuildTool
 				SortedSet<string> CleanedupIncludes = new();
 				SortedSet<string> ForwardDeclarationsToAdd = new();
 
-				foreach (var Include in Info.Includes)
+				foreach (IWYUIncludeEntry Include in Info.Includes)
 				{
 					// We never remove header with name matching cpp
 					string NameWithoutPath = Include.Printable;
@@ -1333,7 +1328,7 @@ namespace UnrealBuildTool
 					}
 					else if (!bNoTransitiveIncludes)
 					{
-						foreach (var Include2 in Info.Includes)
+						foreach (IWYUIncludeEntry Include2 in Info.Includes)
 						{
 							if (Include2.Resolved != null && Include != Include2)
 							{
@@ -1362,7 +1357,7 @@ namespace UnrealBuildTool
 				if (!bRemoveRedundantIncludes)
 				{
 					List<string> ToReadd = new();
-					foreach (var Seen in Info.IncludesSeenInFile)
+					foreach (IWYUIncludeEntry Seen in Info.IncludesSeenInFile)
 					{
 						string QuotedPrintable = Seen.System ? $"<{Seen.Printable}>" : $"\"{Seen.Printable}\"";
 						if (!CleanedupIncludes.Contains(QuotedPrintable) && Info.TransitiveIncludes.ContainsKey(Seen.Full))
@@ -1375,12 +1370,12 @@ namespace UnrealBuildTool
 				// Ignore forward declarations for cpp files. They are very rarely needed and we let the user add them manually instead
 				if (!IsCpp)
 				{
-					foreach (var ForwardDeclaration in Info.ForwardDeclarations)
+					foreach (IWYUForwardEntry ForwardDeclaration in Info.ForwardDeclarations)
 					{
 						bool Add = ForwardDeclaration.Present == false;
 						if (!bNoTransitiveIncludes)
 						{
-							foreach (var Include2 in Info.Includes)
+							foreach (IWYUIncludeEntry Include2 in Info.Includes)
 							{
 								if (Include2.Resolved != null && Include2.Resolved.TransitiveForwardDeclarations!.ContainsKey(ForwardDeclaration.Printable))
 								{
@@ -1396,7 +1391,6 @@ namespace UnrealBuildTool
 					}
 				}
 
-
 				// Read all lines of the header/source file
 				string[] ExistingLines = File.ReadAllLines(Info.File);
 
@@ -1419,25 +1413,23 @@ namespace UnrealBuildTool
 					}
 				}
 
-
 				bool ForceKeepScope = false;
 				bool ErrorOnMoreIncludes = false;
 				int LineIndex = -1;
 
-
 				// This makes sure that we have at least HAL/Platform.h included if the file contains XXX_API
 				bool Contains_API = false;
 				bool LookFor_API = false;
 				if (Info.Includes.Count == 0)
 				{
-					foreach (var Missing in Info.MissingIncludes)
+					foreach (IWYUMissingInclude Missing in Info.MissingIncludes)
 					{
 						LookFor_API = LookFor_API || Missing.Full.Contains("Definitions.h");
 					}
 				}
 
 				// Traverse all lines in file and figure out which includes that should be added or removed
-				foreach (var Line in ExistingLines)
+				foreach (string Line in ExistingLines)
 				{
 					++LineIndex;
 
@@ -1586,7 +1578,6 @@ namespace UnrealBuildTool
 					}
 				}
 
-
 				SortedSet<string> LinesToAdd = new();
 
 				foreach (string IncludeToAdd in IncludesToAdd)
@@ -1599,11 +1590,11 @@ namespace UnrealBuildTool
 					lock (ShouldLog)
 					{
 						System.Console.WriteLine(Info.File);
-						foreach (var I in LinesToAdd)
+						foreach (string I in LinesToAdd)
 						{
 							System.Console.WriteLine("  +" + I);
 						}
-						foreach (var Pair in LinesToRemove)
+						foreach (KeyValuePair<string, string> Pair in LinesToRemove)
 						{
 							System.Console.Write("  -" + Pair.Key);
 							string? Reason;
@@ -1613,7 +1604,7 @@ namespace UnrealBuildTool
 							}
 							System.Console.WriteLine();
 						}
-						foreach (var I in ForwardDeclarationsToAdd)
+						foreach (string I in ForwardDeclarationsToAdd)
 						{
 							System.Console.WriteLine("  +" + I);
 						}
@@ -1722,7 +1713,6 @@ namespace UnrealBuildTool
 								}
 							}
 
-
 							NewLines.Add(OldLine);
 						}
 						else
@@ -1780,7 +1770,7 @@ namespace UnrealBuildTool
 				if (!IsPrivate && Info.File.StartsWith(EngineDir) && !String.IsNullOrEmpty(HeaderDeprecationTag))
 				{
 					Dictionary<string, string> PrintableToFull = new();
-					foreach (var Seen in Info.IncludesSeenInFile)
+					foreach (IWYUIncludeEntry Seen in Info.IncludesSeenInFile)
 					{
 						PrintableToFull.TryAdd(Seen.Printable, Seen.Full);
 					}
@@ -1881,7 +1871,7 @@ namespace UnrealBuildTool
 			Func<bool> WaitForP4 = () =>
 			{
 				bool P4Success = true;
-				foreach (var P4Process in P4Processes)
+				foreach (System.Diagnostics.Process P4Process in P4Processes)
 				{
 					P4Process.WaitForExit();
 					if (P4Process.ExitCode != 0)
@@ -1898,7 +1888,7 @@ namespace UnrealBuildTool
 			if (!bNoP4)
 			{
 				List<IWYUInfo> ReadOnlyFileInfos = new();
-				foreach (var (Info, NewLines) in UpdatedFiles)
+				foreach ((IWYUInfo Info, List<string> NewLines) in UpdatedFiles)
 				{
 					if (new FileInfo(Info.File).IsReadOnly)
 					{
@@ -1912,7 +1902,7 @@ namespace UnrealBuildTool
 					// Should probably revisit this code to prevent 100s of p4 processes to start at once
 					Logger.LogInformation($"Opening {ReadOnlyFileInfos.Count} files for edit in P4... ({SkippedCount} files skipped)");
 					int ShowCount = 8;
-					foreach (var Info in ReadOnlyFileInfos)
+					foreach (IWYUInfo Info in ReadOnlyFileInfos)
 					{
 						Logger.LogInformation($"   edit {Info.File}");
 						if (--ShowCount == 0)
@@ -1929,7 +1919,7 @@ namespace UnrealBuildTool
 					int BatchSize = 10;
 					int BatchCount = 0;
 					int Index = 0;
-					foreach (var Info in ReadOnlyFileInfos)
+					foreach (IWYUInfo Info in ReadOnlyFileInfos)
 					{
 						if (!SkippedFiles.Contains(Info.Source!))
 						{
@@ -1956,7 +1946,7 @@ namespace UnrealBuildTool
 			bool WriteSuccess = true;
 			Logger.LogInformation($"Writing {UpdatedFiles.Count - SkippedCount} files to disk...");
 
-			foreach (var (Info, NewLines) in UpdatedFiles)
+			foreach ((IWYUInfo Info, List<string> NewLines) in UpdatedFiles)
 			{
 				if (SkippedFiles.Contains(Info.Source!))
 				{
@@ -1998,7 +1988,7 @@ namespace UnrealBuildTool
 		static void CalculateTransitive(IWYUInfo Root, IWYUInfo Info, Stack<string> Stack, Dictionary<string, string> TransitiveIncludes, Dictionary<string, string> TransitiveForwardDeclarations, bool UseSeenIncludes)
 		{
 			List<IWYUIncludeEntry> Includes = UseSeenIncludes ? Info.IncludesSeenInFile : Info.Includes;
-			foreach (var Include in Includes)
+			foreach (IWYUIncludeEntry Include in Includes)
 			{
 				string Key = Include.Full;
 
@@ -2017,7 +2007,7 @@ namespace UnrealBuildTool
 				Stack.Pop();
 			}
 
-			foreach (var ForwardDeclaration in Info.ForwardDeclarations)
+			foreach (IWYUForwardEntry ForwardDeclaration in Info.ForwardDeclarations)
 			{
 				string Key = ForwardDeclaration.Printable;
 				TransitiveForwardDeclarations.TryAdd(Key, Info.File);
@@ -2033,7 +2023,7 @@ namespace UnrealBuildTool
 			Logger.LogInformation($"Parsing seen headers not supporting being compiled...");
 			Parallel.ForEach(Infos.Values, Info =>
 			{
-				foreach (var Include in Info.IncludesSeenInFile)
+				foreach (IWYUIncludeEntry Include in Info.IncludesSeenInFile)
 				{
 					if (Include.Resolved == null)
 					{
@@ -2078,7 +2068,7 @@ namespace UnrealBuildTool
 					InfosToCompare.Add(new(Info, SeenTransitiveIncludes));
 				}
 
-				foreach (var Include in SeenTransitiveIncludes.Union(Info.TransitiveIncludes))
+				foreach (KeyValuePair<string, string> Include in SeenTransitiveIncludes.Union(Info.TransitiveIncludes))
 				{
 					AllSeenIncludes.TryAdd(Include.Key, 0);
 				}
@@ -2095,7 +2085,7 @@ namespace UnrealBuildTool
 				*/
 			});
 
-			var FileToSize = new Dictionary<string, long>(AllSeenIncludes);
+			Dictionary<string, long> FileToSize = new Dictionary<string, long>(AllSeenIncludes);
 
 			Logger.LogInformation($"Reading file sizes...");
 			Parallel.ForEach(AllSeenIncludes, Info =>
@@ -2109,12 +2099,12 @@ namespace UnrealBuildTool
 			{
 				long OptimizedSize = 0;
 				long SeenSize = 0;
-				foreach (var Include in Kvp.Item1.TransitiveIncludes.Keys)
+				foreach (string Include in Kvp.Item1.TransitiveIncludes.Keys)
 				{
 					OptimizedSize += FileToSize[Include];
 				}
 
-				foreach (var Include in Kvp.Item2.Keys)
+				foreach (string? Include in Kvp.Item2.Keys)
 				{
 					SeenSize += FileToSize[Include];
 				}
@@ -2164,7 +2154,7 @@ namespace UnrealBuildTool
 				long Saved = ByteSizeDiffList[I].Item1;
 				IWYUInfo File = ByteSizeDiffList[I].Item2;
 				long OptimizedSize = 0;
-				foreach (var Include in File.TransitiveIncludes.Keys)
+				foreach (string Include in File.TransitiveIncludes.Keys)
 				{
 					OptimizedSize += FileToSize[Include];
 				}
diff --git a/Engine/Source/Programs/UnrealBuildTool/Modes/InlineGeneratedCpps.cs b/Engine/Source/Programs/UnrealBuildTool/Modes/InlineGeneratedCpps.cs
--- a/Engine/Source/Programs/UnrealBuildTool/Modes/InlineGeneratedCpps.cs
+++ b/Engine/Source/Programs/UnrealBuildTool/Modes/InlineGeneratedCpps.cs
@@ -100,7 +100,7 @@ namespace UnrealBuildTool
 						{
 							UEBuildModuleCPP.InputFileCollection InputFileCollection = Module.FindInputFiles(Target.Platform, new Dictionary<DirectoryItem, FileItem[]>(), Logger);
 
-							var FileList = new List<FileItem>();
+							List<FileItem> FileList = new List<FileItem>();
 							foreach (FileItem InputFile in InputFileCollection.CPPFiles)
 							{
 								if (FileFilter == null || FileFilter.Matches(InputFile.Location.MakeRelativeTo(Unreal.RootDirectory)))
@@ -113,7 +113,7 @@ namespace UnrealBuildTool
 							{
 								CppCompileEnvironment env = Module.CreateCompileEnvironmentForIntellisense(Target.Rules, BinaryCompileEnvironment, Logger);
 
-								foreach (var InputFile in FileList)
+								foreach (FileItem InputFile in FileList)
 								{
 									if (BinaryCompileEnvironment.MetadataCache.GetListOfInlinedGeneratedCppFiles(InputFile).Any())
 									{
@@ -128,7 +128,7 @@ namespace UnrealBuildTool
 									}
 									else
 									{
-										var HeaderText = FileReference.ReadAllText(FoundHeader.Location);
+										string HeaderText = FileReference.ReadAllText(FoundHeader.Location);
 										if (!IncludeGenHeaderRegex.IsMatch(HeaderText))
 										{
 											continue;
@@ -143,18 +143,18 @@ namespace UnrealBuildTool
 										}
 									}
 
-									var TextLines = System.IO.File.ReadAllLines(InputFile.Location.FullName).ToList();
+									List<string> TextLines = System.IO.File.ReadAllLines(InputFile.Location.FullName).ToList();
 									if (TextLines.Any(Text => InlineReflectionMarkupRegex.IsMatch(Text)))
 									{
 										continue;
 									}
 
-									var CleanTextLines = new List<string>(TextLines.Count);
+									List<string> CleanTextLines = new List<string>(TextLines.Count);
 
 									bool bInComment = false;
 									for (int i = 0; i < TextLines.Count; i++)
 									{
-										var Line = TextLines[i];
+										string Line = TextLines[i];
 
 										if (Line.Contains("//", StringComparison.InvariantCultureIgnoreCase))
 										{
@@ -190,7 +190,7 @@ namespace UnrealBuildTool
 									int LastLineOfInclude = -1;
 									for (int i = 0; i < CleanTextLines.Count; i++)
 									{
-										var Line = CleanTextLines[i];
+										string Line = CleanTextLines[i];
 
 										if (IfEndBlocks.Count == 0 && IncludeRegex.IsMatch(Line))
 										{
@@ -217,7 +217,7 @@ namespace UnrealBuildTool
 												LastLineOfInclude = TextLines.Count - 1;
 												TextLines.Add(String.Empty);
 											}
-											else if (LastLineOfInclude + 1 == TextLines.Count || TextLines[LastLineOfInclude + 1] != String.Empty)
+											else if (LastLineOfInclude + 1 == TextLines.Count || !String.IsNullOrEmpty(TextLines[LastLineOfInclude + 1]))
 											{
 												TextLines.Insert(LastLineOfInclude + 1, String.Empty);
 											}
diff --git a/Engine/Source/Programs/UnrealBuildTool/Modes/PrintBuildGraphInfo.cs b/Engine/Source/Programs/UnrealBuildTool/Modes/PrintBuildGraphInfo.cs
--- a/Engine/Source/Programs/UnrealBuildTool/Modes/PrintBuildGraphInfo.cs
+++ b/Engine/Source/Programs/UnrealBuildTool/Modes/PrintBuildGraphInfo.cs
@@ -199,7 +199,7 @@ namespace UnrealBuildTool
 									Dependencies.Append("!!! ");
 								}
 
-								var DepModuleArray = DepModuleStack.Reverse().ToArray();
+								UEBuildModule[] DepModuleArray = DepModuleStack.Reverse().ToArray();
 								for (int DepModuleIndex = 0; DepModuleIndex < DepModuleArray.Length; DepModuleIndex++)
 								{
 									Dependencies.Append(DepModuleArray[DepModuleIndex].Name);
diff --git a/Engine/Source/Programs/UnrealBuildTool/Modes/ProfileUnitySizesMode.cs b/Engine/Source/Programs/UnrealBuildTool/Modes/ProfileUnitySizesMode.cs
--- a/Engine/Source/Programs/UnrealBuildTool/Modes/ProfileUnitySizesMode.cs
+++ b/Engine/Source/Programs/UnrealBuildTool/Modes/ProfileUnitySizesMode.cs
@@ -293,7 +293,6 @@ namespace UnrealBuildTool
 			return BestTimingData;
 		}
 
-
 		/// <summary>
 		/// Compiles the module and returns the timing information
 		/// </summary>
@@ -311,7 +310,7 @@ namespace UnrealBuildTool
 					Logger.LogInformation($"{LogPrefix}Deleting intermediate directory...");
 					try
 					{
-						var IntermDir = DirectoryItem.GetItemByDirectoryReference(Module.IntermediateDirectory);
+						DirectoryItem IntermDir = DirectoryItem.GetItemByDirectoryReference(Module.IntermediateDirectory);
 						IntermDir.CacheFiles();
 						IntermDir.ResetCachedInfo();
 						DirectoryReference.Delete(Module.IntermediateDirectory, true);
diff --git a/Engine/Source/Programs/UnrealBuildTool/Modes/QueryMode.cs b/Engine/Source/Programs/UnrealBuildTool/Modes/QueryMode.cs
--- a/Engine/Source/Programs/UnrealBuildTool/Modes/QueryMode.cs
+++ b/Engine/Source/Programs/UnrealBuildTool/Modes/QueryMode.cs
@@ -121,7 +121,7 @@ namespace UnrealBuildTool
 				return Task.FromResult(0);
 			}
 
-			var ResponseOptions = new JsonSerializerOptions
+			JsonSerializerOptions ResponseOptions = new JsonSerializerOptions
 			{
 				PropertyNamingPolicy = JsonNamingPolicy.CamelCase,
 				WriteIndented = bIndented,
@@ -229,7 +229,7 @@ namespace UnrealBuildTool
 			}
 
 			GenerateProjectFilesMode.TryParseProjectFileArgument(Arguments, Logger, out FileReference? ProjectFileArg);
-			var RawArgs = new List<string> { TargetName, TargetConfiguration, TargetPlatform };
+			List<string> RawArgs = new List<string> { TargetName, TargetConfiguration, TargetPlatform };
 			if (ProjectFileArg != null)
 			{
 				RawArgs.Add(ProjectFileArg.ToString());
@@ -254,8 +254,8 @@ namespace UnrealBuildTool
 					CurrentTarget = UEBuildTarget.Create(TargetDescriptors[0], false, false, bUsePrecompiled, Logger);
 				}
 
-				var CurrentTargetIntellisenseInfo = new TargetIntellisenseInfo();
-				var CurrentBrowseConfiguration = new GetBrowseConfigurationResponse { Success = true };
+				TargetIntellisenseInfo CurrentTargetIntellisenseInfo = new TargetIntellisenseInfo();
+				GetBrowseConfigurationResponse CurrentBrowseConfiguration = new GetBrowseConfigurationResponse { Success = true };
 
 				// Partially duplicated from UEBuildTarget.Build because we just want to get C++ compile actions without running UHT
 				// or generating link actions / full dependency graph 
@@ -303,7 +303,7 @@ namespace UnrealBuildTool
 							BrowseConfigurationFolders.Add(Dir.ToString());
 						}
 
-						var Settings = new TargetIntellisenseInfo.CompileSettings();
+						TargetIntellisenseInfo.CompileSettings Settings = new TargetIntellisenseInfo.CompileSettings();
 						if (OperatingSystem.IsWindows())
 						{
 							if (CurrentTarget.Platform == UnrealTargetPlatform.Win64)
diff --git a/Engine/Source/Programs/UnrealBuildTool/Modes/ServerMode.cs b/Engine/Source/Programs/UnrealBuildTool/Modes/ServerMode.cs
--- a/Engine/Source/Programs/UnrealBuildTool/Modes/ServerMode.cs
+++ b/Engine/Source/Programs/UnrealBuildTool/Modes/ServerMode.cs
@@ -124,7 +124,6 @@ namespace UnrealBuildTool
 		public List<TargetIntellisenseInfo.CompileSettings?> Settings { get; set; } = new();
 	}
 
-
 	// TODO: Deprecate 
 	[Obsolete]
 	[ToolMode("Server", ToolModeOptions.BuildPlatforms | ToolModeOptions.XmlConfig | ToolModeOptions.UseStartupTraceListener)]
@@ -167,7 +166,7 @@ namespace UnrealBuildTool
 
 			while (KeepRunning)
 			{
-				var Line = Console.ReadLine();
+				string? Line = Console.ReadLine();
 				if (Line == null) { continue; }
 				HandleCommand(Line, Logger);
 			}
@@ -176,11 +175,11 @@ namespace UnrealBuildTool
 
 		private void HandleCommand(string CommandString, ILogger Logger)
 		{
-			var jsonOptions = new JsonSerializerOptions
+			JsonSerializerOptions jsonOptions = new JsonSerializerOptions
 			{
 				PropertyNameCaseInsensitive = true
 			};
-			var responseOptions = new JsonSerializerOptions
+			JsonSerializerOptions responseOptions = new JsonSerializerOptions
 			{
 				PropertyNamingPolicy = JsonNamingPolicy.CamelCase
 			};
@@ -193,7 +192,7 @@ namespace UnrealBuildTool
 					return;
 				}
 				Logger.LogInformation($"Got command {CommandString}");
-				var command = MaybeCommand.Value;
+				Command command = MaybeCommand.Value;
 				switch (command.Type)
 				{
 					case CommandType.Quit:
@@ -201,35 +200,35 @@ namespace UnrealBuildTool
 						break;
 					case CommandType.DiscoverTargets:
 						{
-							var response = Handle_DiscoverTargets(Logger);
+							DiscoverTargetsResponse response = Handle_DiscoverTargets(Logger);
 							response.Sequence = command.Sequence;
 							Console.WriteLine(JsonSerializer.Serialize(response, responseOptions));
 						}
 						break;
 					case CommandType.SetTarget:
 						{
-							var response = Handle_SetTarget(JsonSerializer.Deserialize<SetTargetCommand>(CommandString, jsonOptions), Logger);
+							SetTargetResponse response = Handle_SetTarget(JsonSerializer.Deserialize<SetTargetCommand>(CommandString, jsonOptions), Logger);
 							response.Sequence = command.Sequence;
 							Console.WriteLine(JsonSerializer.Serialize(response, responseOptions));
 						}
 						break;
 					case CommandType.QueryFile:
 						{
-							var response = Handle_QueryFile(JsonSerializer.Deserialize<QueryFileCommand>(CommandString, jsonOptions), Logger);
+							QueryFileResponse response = Handle_QueryFile(JsonSerializer.Deserialize<QueryFileCommand>(CommandString, jsonOptions), Logger);
 							response.Sequence = command.Sequence;
 							Console.WriteLine(JsonSerializer.Serialize(response, responseOptions));
 						}
 						break;
 					case CommandType.GetBrowseConfiguration:
 						{
-							var response = Handle_GetBrowseConfiguration(JsonSerializer.Deserialize<GetBrowseConfigurationCommand>(CommandString, jsonOptions), Logger);
+							GetBrowseConfigurationResponse response = Handle_GetBrowseConfiguration(JsonSerializer.Deserialize<GetBrowseConfigurationCommand>(CommandString, jsonOptions), Logger);
 							response.Sequence = command.Sequence;
 							Console.WriteLine(JsonSerializer.Serialize(response, responseOptions));
 						}
 						break;
 					case CommandType.GetCompileSettings:
 						{
-							var response = Handle_GetCompileSettings(JsonSerializer.Deserialize<GetCompileSettingsCommand>(CommandString, jsonOptions), Logger);
+							GetCompileSettingsResponse response = Handle_GetCompileSettings(JsonSerializer.Deserialize<GetCompileSettingsCommand>(CommandString, jsonOptions), Logger);
 							response.Sequence = command.Sequence;
 							Console.WriteLine(JsonSerializer.Serialize(response, responseOptions));
 						}
@@ -291,7 +290,7 @@ namespace UnrealBuildTool
 
 			// TODO: Handle 0 target names
 
-			var Reply = new DiscoverTargetsResponse
+			DiscoverTargetsResponse Reply = new DiscoverTargetsResponse
 			{
 				Configurations = Configurations,
 				Platforms = Platforms.Select(x => x.ToString()).ToList(),
@@ -375,7 +374,7 @@ namespace UnrealBuildTool
 							BrowseConfigurationFolders.Add(Dir.ToString());
 						}
 
-						var Settings = new TargetIntellisenseInfo.CompileSettings();
+						TargetIntellisenseInfo.CompileSettings Settings = new TargetIntellisenseInfo.CompileSettings();
 						Settings.IncludePaths.AddRange(ModuleCompileEnvironment.SystemIncludePaths.Select(x => x.ToString()));
 						Settings.IncludePaths.AddRange(ModuleCompileEnvironment.UserIncludePaths.Select(x => x.ToString()));
 						Settings.Defines = ModuleCompileEnvironment.Definitions;
@@ -394,7 +393,7 @@ namespace UnrealBuildTool
 
 			CurrentBrowseConfiguration.Paths = BrowseConfigurationFolders.ToList();
 
-			var Reply = new SetTargetResponse
+			SetTargetResponse Reply = new SetTargetResponse
 			{
 			};
 			return Reply;
@@ -402,14 +401,14 @@ namespace UnrealBuildTool
 
 		private QueryFileResponse Handle_QueryFile(QueryFileCommand Command, ILogger Logger)
 		{
-			var FailResponse = new QueryFileResponse { Found = false };
+			QueryFileResponse FailResponse = new QueryFileResponse { Found = false };
 			if (CurrentTargetIntellisenseInfo == null)
 			{
 				return FailResponse;
 			}
 			try
 			{
-				var Path = System.IO.Path.GetFullPath(Command.AbsolutePath);
+				string Path = System.IO.Path.GetFullPath(Command.AbsolutePath);
 				if (CurrentTargetIntellisenseInfo.FindModuleForFile(new FileReference(Path)) != null)
 				{
 					return new QueryFileResponse { Found = true };
@@ -430,13 +429,13 @@ namespace UnrealBuildTool
 			}
 			if (Command.AbsolutePath != null)
 			{
-				var Path = System.IO.Path.GetFullPath(Command.AbsolutePath);
+				string Path = System.IO.Path.GetFullPath(Command.AbsolutePath);
 				UEBuildModule? Module = CurrentTargetIntellisenseInfo.FindModuleForDirectory(new DirectoryReference(Path));
 				if (Module != null)
 				{
 					Response.Success = true;
 					Response.Paths.AddRange(Module.ModuleDirectories.Select(x => x.ToString()));
-					if (CurrentTargetIntellisenseInfo.ModuleToCompileSettings.TryGetValue(Module, out var Settings))
+					if (CurrentTargetIntellisenseInfo.ModuleToCompileSettings.TryGetValue(Module, out TargetIntellisenseInfo.CompileSettings? Settings))
 					{
 						Response.Paths.AddRange(Settings.IncludePaths);
 						Response.WindowsSdkVersion = Settings.WindowsSdkVersion;
@@ -459,7 +458,7 @@ namespace UnrealBuildTool
 				foreach (string AbsolutePath in Command.AbsolutePaths)
 				{
 					TargetIntellisenseInfo.CompileSettings? Settings = null;
-					var Path = System.IO.Path.GetFullPath(AbsolutePath); // convert absolute path to dos? 
+					string Path = System.IO.Path.GetFullPath(AbsolutePath); // convert absolute path to dos? 
 					UEBuildModule? Module = CurrentTargetIntellisenseInfo!.FindModuleForFile(new FileReference(Path));
 					if (Module is not null)
 					{
diff --git a/Engine/Source/Programs/UnrealBuildTool/Modes/UnrealHeaderToolMode.cs b/Engine/Source/Programs/UnrealBuildTool/Modes/UnrealHeaderToolMode.cs
--- a/Engine/Source/Programs/UnrealBuildTool/Modes/UnrealHeaderToolMode.cs
+++ b/Engine/Source/Programs/UnrealBuildTool/Modes/UnrealHeaderToolMode.cs
@@ -753,7 +753,7 @@ namespace UnrealBuildTool.Modes
 					int BodyEndPos = BodyStartPos;
 					while (true)
 					{
-						using var SaveState = new UhtTokenSaveState(Reader);
+						using UhtTokenSaveState SaveState = new UhtTokenSaveState(Reader);
 						UhtToken Token = Reader.GetLine();
 						if (Token.TokenType == UhtTokenType.EndOfFile)
 						{
diff --git a/Engine/Source/Programs/UnrealBuildTool/Platform/Android/AndroidAARHandler.cs b/Engine/Source/Programs/UnrealBuildTool/Platform/Android/AndroidAARHandler.cs
--- a/Engine/Source/Programs/UnrealBuildTool/Platform/Android/AndroidAARHandler.cs
+++ b/Engine/Source/Programs/UnrealBuildTool/Platform/Android/AndroidAARHandler.cs
@@ -195,7 +195,7 @@ namespace UnrealBuildTool
 			string WorkVersion = InVersion;
 
 			// if has commas, keep the rightmost part with actual numbers
-			if (WorkVersion.Contains(","))
+			if (WorkVersion.Contains(','))
 			{
 				string[] CommaParts = WorkVersion.Split(',');
 				WorkVersion = "";
@@ -632,6 +632,5 @@ namespace UnrealBuildTool
 				return OutputFileNames;
 			}
 		}
-
 	}
 }
diff --git a/Engine/Source/Programs/UnrealBuildTool/Platform/Android/AndroidPlatformSDK.cs b/Engine/Source/Programs/UnrealBuildTool/Platform/Android/AndroidPlatformSDK.cs
--- a/Engine/Source/Programs/UnrealBuildTool/Platform/Android/AndroidPlatformSDK.cs
+++ b/Engine/Source/Programs/UnrealBuildTool/Platform/Android/AndroidPlatformSDK.cs
@@ -53,7 +53,6 @@ namespace UnrealBuildTool
 							string PathVar = BashProfileContents[LineIndex].Split('=')[1].Replace("\"", "");
 							return PathVar;
 						}
-
 					}
 				}
 			}
@@ -112,8 +111,7 @@ namespace UnrealBuildTool
 			return NDKToolchainVersion;
 		}
 
-
-		public override bool TryConvertVersionToInt(string? StringValue, out UInt64 OutValue, string? Hint)
+		public override bool TryConvertVersionToInt(string? StringValue, out ulong OutValue, string? Hint)
 		{
 			// convert r<num>[letter] to hex
 			if (!String.IsNullOrEmpty(StringValue))
@@ -137,7 +135,6 @@ namespace UnrealBuildTool
 			return false;
 		}
 
-
 		//public override SDKStatus PrintSDKInfoAndReturnValidity(LogEventType Verbosity, LogFormatOptions Options, LogEventType ErrorVerbosity, LogFormatOptions ErrorOptions)
 		//{
 		//	SDKStatus Validity = base.PrintSDKInfoAndReturnValidity(Verbosity, Options, ErrorVerbosity, ErrorOptions);
@@ -150,7 +147,6 @@ namespace UnrealBuildTool
 		//	return Validity;
 		//}
 
-
 		protected override bool PlatformSupportsAutoSDKs()
 		{
 			return true;
diff --git a/Engine/Source/Programs/UnrealBuildTool/Platform/Android/AndroidToolChain.cs b/Engine/Source/Programs/UnrealBuildTool/Platform/Android/AndroidToolChain.cs
--- a/Engine/Source/Programs/UnrealBuildTool/Platform/Android/AndroidToolChain.cs
+++ b/Engine/Source/Programs/UnrealBuildTool/Platform/Android/AndroidToolChain.cs
@@ -239,7 +239,6 @@ namespace UnrealBuildTool
 			ReadElfPath = Path.Combine(NDKPath, @"toolchains/llvm", ArchitecturePath, @"bin/llvm-readelf" + ExeExtension);
 		}
 
-
 		protected override ClangToolChainInfo GetToolChainInfo()
 		{
 			return new ClangToolChainInfo(FileReference.FromString(ClangPath)!, FileReference.FromString(ArPathArm64)!, Logger);
@@ -301,7 +300,7 @@ namespace UnrealBuildTool
 		{
 			string NDKVersion = GetNdkApiLevel();
 			int NDKVersionInt = MinNdk;
-			if (NDKVersion.Contains("-"))
+			if (NDKVersion.Contains('-'))
 			{
 				int Version;
 				if (Int32.TryParse(NDKVersion.Substring(NDKVersion.LastIndexOf('-') + 1), out Version))
@@ -371,7 +370,7 @@ namespace UnrealBuildTool
 				return false;
 			}
 
-			if (ApiString.Contains("-"))
+			if (ApiString.Contains('-'))
 			{
 				int Version;
 				if (Int32.TryParse(ApiString.Substring(ApiString.LastIndexOf('-') + 1), out Version))
@@ -393,7 +392,7 @@ namespace UnrealBuildTool
 			string ProjectNDKLevel;
 			Ini.GetString("/Script/AndroidRuntimeSettings.AndroidRuntimeSettings", "NDKAPILevelOverride", out ProjectNDKLevel!);
 			ProjectNDKLevel = ProjectNDKLevel.Trim();
-			if (ProjectNDKLevel != "")
+			if (!String.IsNullOrEmpty(ProjectNDKLevel))
 			{
 				NDKLevel = ProjectNDKLevel;
 			}
@@ -443,7 +442,6 @@ namespace UnrealBuildTool
 			return "android-" + MaxPlatform.ToString();
 		}
 
-
 		protected override void GetCompileArguments_IncludePaths(CppCompileEnvironment CompileEnvironment, List<string> Arguments)
 		{
 			// remove paths that are not meant for this architecture
@@ -458,7 +456,7 @@ namespace UnrealBuildTool
 		/// <inheritdoc/>
 		protected override string EscapePreprocessorDefinition(string Definition)
 		{
-			return Definition.Contains("\"") ? Definition.Replace("\"", "\\\"") : Definition;
+			return Definition.Contains('"') ? Definition.Replace("\"", "\\\"") : Definition;
 		}
 
 		/// <inheritdoc/>
@@ -466,7 +464,6 @@ namespace UnrealBuildTool
 		{
 			base.GetCompileArguments_WarningsAndErrors(CompileEnvironment, Arguments);
 
-
 			// @todo unlikely all needed
 			Arguments.Add("-Wno-local-type-template-args"); // engine triggers this
 			Arguments.Add("-Wno-return-type-c-linkage");    // needed for PhysX
@@ -586,7 +583,6 @@ namespace UnrealBuildTool
 				Arguments.Add("-ffunction-sections");   // Places each function in its own section of the output file, linker may be able to perform opts to improve locality of reference
 				Arguments.Add("-fdata-sections");       // Places each data item in its own section of the output file, linker may be able to perform opts to improve locality of reference
 			}
-
 		}
 
 		/// <inheritdoc/>
@@ -606,8 +602,6 @@ namespace UnrealBuildTool
 				}
 			}
 
-
-
 			//string? SanitizerMode = Environment.GetEnvironmentVariable("ENABLE_ADDRESS_SANITIZER");
 			//if ((SanitizerMode != null && SanitizerMode == "YES") || (Options.HasFlag(ClangToolChainOptions.EnableAddressSanitizer)))
 			//{
@@ -649,7 +643,6 @@ namespace UnrealBuildTool
 
 			string NativeGluePath = Path.GetFullPath(GetNativeGluePath());
 
-
 			if (BuildWithHiddenSymbolVisibility(CompileEnvironment))
 			{
 				Arguments.Add("-fvisibility=hidden");
@@ -670,7 +663,6 @@ namespace UnrealBuildTool
 			Arguments.Add($"-DPLATFORM_USED_NDK_VERSION_INTEGER={NDKApiLevel64Int}");       // NDK version
 			Arguments.Add("-DPLATFORM_64BITS=1");       // NDK version
 
-
 			if (CompileEnvironment.bCompileISPC)
 			{
 				Arguments.Add("-DINTEL_ISPC=1");
@@ -713,7 +705,6 @@ namespace UnrealBuildTool
 			return TargetFile;
 		}
 
-
 		protected virtual string GetLinkArguments(LinkEnvironment LinkEnvironment, UnrealArch Architecture)
 		{
 			string Result = "";
@@ -756,7 +747,6 @@ namespace UnrealBuildTool
 			// use lld as linker (requires llvm-strip)
 			Result += " -fuse-ld=lld";
 
-
 			// make sure the DT_SONAME field is set properly (or we can a warning toast at startup on new Android)
 			Result += " -Wl,-soname,libUnreal.so";
 
@@ -837,7 +827,6 @@ namespace UnrealBuildTool
 			return Result;
 		}
 
-
 		protected virtual void ModifyLibraries(LinkEnvironment LinkEnvironment)
 		{
 			// @todo Lumin: verify this works with base android
@@ -1239,7 +1228,7 @@ namespace UnrealBuildTool
 				{
 					string ObjTarget = Target;
 
-					if (Target.Contains("-"))
+					if (Target.Contains('-'))
 					{
 						// Remove lane width and gang size from obj file name
 						ObjTarget = Target.Split('-')[0];
diff --git a/Engine/Source/Programs/UnrealBuildTool/Platform/Android/UEDeployAndroid.cs b/Engine/Source/Programs/UnrealBuildTool/Platform/Android/UEDeployAndroid.cs
--- a/Engine/Source/Programs/UnrealBuildTool/Platform/Android/UEDeployAndroid.cs
+++ b/Engine/Source/Programs/UnrealBuildTool/Platform/Android/UEDeployAndroid.cs
@@ -211,7 +211,7 @@ namespace UnrealBuildTool
 		private int GetApiLevelInt(string ApiString)
 		{
 			int VersionInt = 0;
-			if (ApiString.Contains("-"))
+			if (ApiString.Contains('-'))
 			{
 				int Version;
 				if (Int32.TryParse(ApiString.Substring(ApiString.LastIndexOf('-') + 1), out Version))
@@ -236,7 +236,7 @@ namespace UnrealBuildTool
 				string ProjectSDKLevel;
 				Ini.GetString("/Script/AndroidRuntimeSettings.AndroidRuntimeSettings", "SDKAPILevelOverride", out ProjectSDKLevel);
 				ProjectSDKLevel = ProjectSDKLevel.Trim();
-				if (ProjectSDKLevel != "")
+				if (!String.IsNullOrEmpty(ProjectSDKLevel))
 				{
 					SDKLevel = ProjectSDKLevel;
 				}
@@ -326,7 +326,7 @@ namespace UnrealBuildTool
 			ConfigHierarchy Ini = GetConfigCacheIni(ConfigHierarchyType.Engine);
 			Ini.GetString("/Script/AndroidRuntimeSettings.AndroidRuntimeSettings", "BuildToolsOverride", out BestVersionString);
 
-			if (BestVersionString == null || BestVersionString == "" || BestVersionString == "latest")
+			if (BestVersionString == null || String.IsNullOrEmpty(BestVersionString) || BestVersionString == "latest")
 			{
 				// get a list of the directories in build-tools.. may be more than one set installed (or none which is bad)
 				string[] Subdirs = Directory.GetDirectories(Path.Combine(HomePath, "build-tools"));
@@ -998,7 +998,6 @@ namespace UnrealBuildTool
 				Ini.GetInt32("/Script/AndroidRuntimeSettings.AndroidRuntimeSettings", "StoreVersionOffsetX8664", out CachedStoreVersionOffsetX8664);
 			}
 
-
 			if (Architecture == UnrealArch.Arm64)
 			{
 				return CachedStoreVersion + CachedStoreVersionOffsetArm64;
@@ -1503,7 +1502,7 @@ namespace UnrealBuildTool
 
 				Logger.LogInformation("DebugVulkanLayerDirectory {LayerDir}", LayerDir);
 
-				if (LayerDir != "")
+				if (!String.IsNullOrEmpty(LayerDir))
 				{
 					string VulkanLayersDir = Path.Combine(Environment.ExpandEnvironmentVariables(LayerDir), NDKArch);
 
@@ -1595,7 +1594,7 @@ namespace UnrealBuildTool
 			// see [REDACTED_PROFILE_URL]/runtime/issues/29857
 			// also see UE-102580
 			// for rules see https://docs.microsoft.com/en-us/cpp/cpp/main-function-command-line-args
-			if (Params.Contains("\'"))
+			if (Params.Contains('\''))
 			{
 				Params = Params.Replace("\"", "\\\"");
 				Params = Params.Replace('\'', '\"');
@@ -1605,7 +1604,7 @@ namespace UnrealBuildTool
 			{
 				Logger.LogInformation("\nRunning: {Command} {Params}", Command, Params);
 			}
-			else if (OverrideDesc != "")
+			else if (!String.IsNullOrEmpty(OverrideDesc))
 			{
 				Logger.LogInformation("{Message}", OverrideDesc);
 				Logger.LogDebug("\nRunning: {Command} {Params}", Command, Params);
@@ -1633,7 +1632,7 @@ namespace UnrealBuildTool
 			// see [REDACTED_PROFILE_URL]/runtime/issues/29857
 			// also see UE-102580
 			// for rules see https://docs.microsoft.com/en-us/cpp/cpp/main-function-command-line-args
-			if (Params.Contains("\'"))
+			if (Params.Contains('\''))
 			{
 				Params = Params.Replace("\"", "\\\"");
 				Params = Params.Replace('\'', '\"');
@@ -1643,7 +1642,7 @@ namespace UnrealBuildTool
 			{
 				Logger.LogInformation("\nRunning: {Command} {Params}", Command, Params);
 			}
-			else if (OverrideDesc != "")
+			else if (!String.IsNullOrEmpty(OverrideDesc))
 			{
 				Logger.LogInformation("{Message}", OverrideDesc);
 				Logger.LogDebug("\nRunning: {Command} {Params}", Command, Params);
@@ -1833,7 +1832,7 @@ namespace UnrealBuildTool
 			// see [REDACTED_PROFILE_URL]/runtime/issues/29857
 			// also see UE-102580
 			// for rules see https://docs.microsoft.com/en-us/cpp/cpp/main-function-command-line-args
-			if (Params.Contains("\'"))
+			if (Params.Contains('\''))
 			{
 				Params = Params.Replace("\"", "\\\"");
 				Params = Params.Replace('\'', '\"');
@@ -1843,7 +1842,7 @@ namespace UnrealBuildTool
 			{
 				Logger.LogInformation("\nRunning: {Command} {Params}", Command, Params);
 			}
-			else if (OverrideDesc != "")
+			else if (!String.IsNullOrEmpty(OverrideDesc))
 			{
 				Logger.LogInformation("{Message}", OverrideDesc);
 				Logger.LogDebug("\nRunning: {Command} {Params}", Command, Params);
@@ -2410,7 +2409,7 @@ namespace UnrealBuildTool
 					}
 
 					// hyphens not allowed so change them to underscores in project name
-					if (ProjectName.Contains("-"))
+					if (ProjectName.Contains('-'))
 					{
 						Trace.TraceWarning("Project name contained hyphens, converted to underscore");
 						ProjectName = ProjectName.Replace("-", "_");
@@ -2438,7 +2437,7 @@ namespace UnrealBuildTool
 				}
 
 				// hyphens not allowed
-				if (PackageName.Contains("-"))
+				if (PackageName.Contains('-'))
 				{
 					throw new BuildException("Package names may not contain hyphens; please change in Android Project Settings. Currently set to '" + PackageName + "'");
 				}
@@ -2516,7 +2515,6 @@ namespace UnrealBuildTool
 			return EngineMajorVersion + "." + EngineMinorVersion + "." + EnginePatchVersion;
 		}
 
-
 		private string GenerateManifest(AndroidToolChain ToolChain, string ProjectName, TargetType InTargetType, string EngineDirectory, bool bIsForDistribution, bool bPackageDataInsideApk, string GameBuildFilesPath, bool bHasOBBFiles, bool bDisableVerifyOBBOnStartUp, UnrealArch UnrealArch, string CookFlavor, bool bUseExternalFilesDir, string Configuration, int SDKLevelInt, bool bIsEmbedded, bool bEnableBundle)
 		{
 			// Read the engine version
@@ -2709,7 +2707,7 @@ namespace UnrealBuildTool
 			string CookedFlavors = (bETC2Enabled ? "ETC2," : "") +
 									(bDXTEnabled ? "DXT," : "") +
 									(bASTCEnabled ? "ASTC," : "");
-			CookedFlavors = (CookedFlavors == "") ? "" : CookedFlavors.Substring(0, CookedFlavors.Length - 1);
+			CookedFlavors = (String.IsNullOrEmpty(CookedFlavors)) ? "" : CookedFlavors.Substring(0, CookedFlavors.Length - 1);
 
 			StringBuilder Text = new StringBuilder();
 			Text.AppendLine(XML_HEADER);
@@ -3036,7 +3034,7 @@ namespace UnrealBuildTool
 					foreach (string Permission in ExtraPermissions)
 					{
 						string TrimmedPermission = Permission.Trim(' ');
-						if (TrimmedPermission != "")
+						if (!String.IsNullOrEmpty(TrimmedPermission))
 						{
 							string PermissionString = String.Format("\t<uses-permission android:name=\"{0}\"/>", TrimmedPermission);
 							if (!Text.ToString().Contains(PermissionString))
@@ -3309,7 +3307,6 @@ namespace UnrealBuildTool
 					return true;
 				}
 			}
-
 		}
 
 		private bool RequiresOBB(bool bDisallowPackageInAPK, string OBBLocation)
@@ -4813,7 +4810,7 @@ namespace UnrealBuildTool
 
 				// check if any plugins want to increase the required compile SDK version
 				string CompileSDKMin = UPL.ProcessPluginNode(NDKArch, "minimumSDKAPI", "");
-				if (CompileSDKMin != "")
+				if (!String.IsNullOrEmpty(CompileSDKMin))
 				{
 					int CompileSDKVersionInt;
 					if (!Int32.TryParse(CompileSDKVersion, out CompileSDKVersionInt))
@@ -4937,7 +4934,7 @@ namespace UnrealBuildTool
 							string ShellParametersEnd = RuntimePlatform.IsWindows ? "" : "'";
 							RunCommandLineProgramWithExceptionAndFiltering(UnrealBuildGradlePath, ShellExecutable, ShellParametersBegin + "\"" + GradleScriptPath + "\" " + GradleOptions + ShellParametersEnd, Logger, "Making .apk with Gradle...");
 
-							if (GradleSecondCallOptions != "")
+							if (!String.IsNullOrEmpty(GradleSecondCallOptions))
 							{
 								RunCommandLineProgramWithExceptionAndFiltering(UnrealBuildGradlePath, ShellExecutable, ShellParametersBegin + "\"" + GradleScriptPath + "\" " + GradleSecondCallOptions + ShellParametersEnd, Logger, "Additional Gradle steps...");
 							}
@@ -5311,7 +5308,7 @@ namespace UnrealBuildTool
 					string ShellParametersEnd = RuntimePlatform.IsWindows ? "" : "'";
 					RunCommandLineProgramWithExceptionAndFiltering(UnrealBuildGradlePath, ShellExecutable, ShellParametersBegin + "\"" + GradleScriptPath + "\" " + GradleOptions + ShellParametersEnd, Logger, "Making .aab with Gradle...");
 
-					if (GradleSecondCallOptions != "")
+					if (!String.IsNullOrEmpty(GradleSecondCallOptions))
 					{
 						RunCommandLineProgramWithExceptionAndFiltering(UnrealBuildGradlePath, ShellExecutable, ShellParametersBegin + "\"" + GradleScriptPath + "\" " + GradleSecondCallOptions + ShellParametersEnd, Logger, "Additional Gradle steps...");
 					}
@@ -5373,7 +5370,7 @@ namespace UnrealBuildTool
 						string ShellParametersEnd = RuntimePlatform.IsWindows ? "" : "'";
 						RunCommandLineProgramWithExceptionAndFiltering(UnrealBuildGradlePath, ShellExecutable, ShellParametersBegin + "\"" + GradleScriptPath + "\" " + GradleOptions + ShellParametersEnd, Logger, "Making .aab with Gradle...");
 
-						if (GradleSecondCallOptions != "")
+						if (!String.IsNullOrEmpty(GradleSecondCallOptions))
 						{
 							RunCommandLineProgramWithExceptionAndFiltering(UnrealBuildGradlePath, ShellExecutable, ShellParametersBegin + "\"" + GradleScriptPath + "\" " + GradleSecondCallOptions + ShellParametersEnd, Logger, "Additional Gradle steps...");
 						}
@@ -5975,19 +5972,19 @@ namespace UnrealBuildTool
 					}
 				}
 
-				if (VersionCode != "")
+				if (!String.IsNullOrEmpty(VersionCode))
 				{
 					BuildGradleContent.AppendLine(String.Format("\t\tversionCode {0}", VersionCode));
 				}
-				if (VersionName != "")
+				if (!String.IsNullOrEmpty(VersionName))
 				{
 					BuildGradleContent.AppendLine(String.Format("\t\tversionName \"{0}\"", VersionName));
 				}
-				if (MinSdkVersion != "")
+				if (!String.IsNullOrEmpty(MinSdkVersion))
 				{
 					BuildGradleContent.AppendLine(String.Format("\t\tminSdkVersion = {0}", MinSdkVersion));
 				}
-				if (TargetSdkVersion != "")
+				if (!String.IsNullOrEmpty(TargetSdkVersion))
 				{
 					BuildGradleContent.AppendLine(String.Format("\t\ttargetSdkVersion = {0}", TargetSdkVersion));
 				}
diff --git a/Engine/Source/Programs/UnrealBuildTool/Platform/IOS/IOSCodeSigning.cs b/Engine/Source/Programs/UnrealBuildTool/Platform/IOS/IOSCodeSigning.cs
--- a/Engine/Source/Programs/UnrealBuildTool/Platform/IOS/IOSCodeSigning.cs
+++ b/Engine/Source/Programs/UnrealBuildTool/Platform/IOS/IOSCodeSigning.cs
@@ -94,7 +94,7 @@ namespace UnrealBuildTool
 			{
 				// Make sure we have a useful friendly name
 				string FriendlyName = Cert.FriendlyName;
-				if ((FriendlyName == "") || (FriendlyName == null))
+				if (String.IsNullOrEmpty(FriendlyName))
 				{
 					FriendlyName = GetCommonNameFromCert(Cert);
 				}
@@ -103,13 +103,12 @@ namespace UnrealBuildTool
 			}
 		}
 
-
 		private static string CertToolData = "";
 
 		/// <summary>
 		/// When receiving the cert tool process call
 		/// </summary>
-		public static void OutputReceivedCertToolProcessCall(Object Sender, DataReceivedEventArgs Line)
+		public static void OutputReceivedCertToolProcessCall(object Sender, DataReceivedEventArgs Line)
 		{
 			if ((Line != null) && !String.IsNullOrEmpty(Line.Data))
 			{
@@ -246,7 +245,7 @@ namespace UnrealBuildTool
 				bool bPassesNameCheck = p.ApplicationIdentifier.Substring(p.ApplicationIdentifierPrefix.Length + 1) == CFBundleIdentifier;
 				bool bPassesCompanyCheck = false;
 				bool bPassesWildCardCheck = false;
-				if (p.ApplicationIdentifier.Contains("*"))
+				if (p.ApplicationIdentifier.Contains('*'))
 				{
 					string CompanyName = p.ApplicationIdentifier.Substring(p.ApplicationIdentifierPrefix.Length + 1);
 
@@ -313,10 +312,9 @@ namespace UnrealBuildTool
 				{
 					Provisions.Add(new FileReference(ProvisionFile));
 				}
-
 			}
 			_logger?.LogInformation(" MATCHED-Provision:{0},File:{1},Cert:{2} ", SelectedProvision, SelectedFile, SelectedCert);
-			if (SelectedProvision != "")
+			if (!String.IsNullOrEmpty(SelectedProvision))
 			{
 				MatchedProvision = new FileReference(SelectedProvision);
 			}
@@ -463,7 +461,7 @@ namespace UnrealBuildTool
 				try
 				{
 					string RawInfoPList = File.ReadAllText(plistFile, Encoding.UTF8);
-					Info = new PListHelper(RawInfoPList); ;
+					Info = new PListHelper(RawInfoPList);
 				}
 				catch (Exception ex)
 				{
@@ -645,7 +643,7 @@ namespace UnrealBuildTool
 			else if (Value is Dictionary<string, object> dict)
 			{
 				ValueElement = Doc.CreateElement("dict");
-				foreach (var KVP in dict)
+				foreach (KeyValuePair<string, object> KVP in dict)
 				{
 					AddKeyValuePair(ValueElement, KVP.Key, KVP.Value);
 				}
@@ -680,9 +678,9 @@ namespace UnrealBuildTool
 					ValueElement = Doc.CreateElement("array");
 					if (Value is Array array)
 					{
-						foreach (var A in array)
+						foreach (object? A in array)
 						{
-							var childNode = ConvertValueToPListFormat(A);
+							XmlElement? childNode = ConvertValueToPListFormat(A);
 							if (childNode != null)
 							{
 								ValueElement.AppendChild(childNode);
@@ -696,9 +694,9 @@ namespace UnrealBuildTool
 				ValueElement = Doc.CreateElement("array");
 				if (Value is IList list && ValueElement != null)
 				{
-					foreach (var A in list)
+					foreach (object? A in list)
 					{
-						var child = ConvertValueToPListFormat(A);
+						XmlElement? child = ConvertValueToPListFormat(A);
 						if (child != null)
 						{
 							ValueElement.AppendChild(child);
@@ -706,19 +704,19 @@ namespace UnrealBuildTool
 					}
 				}
 			}
-			else if (Value is bool)
+			else if (Value is bool vbool)
 			{
-				ValueElement = Doc.CreateElement(((bool)Value) ? "true" : "false");
+				ValueElement = Doc.CreateElement(vbool ? "true" : "false");
 			}
-			else if (Value is double)
+			else if (Value is double vdouble)
 			{
 				ValueElement = Doc.CreateElement("real");
-				ValueElement.InnerText = ((double)Value).ToString();
+				ValueElement.InnerText = vdouble.ToString();
 			}
-			else if (Value is int)
+			else if (Value is int vint)
 			{
 				ValueElement = Doc.CreateElement("integer");
-				ValueElement.InnerText = ((int)Value).ToString();
+				ValueElement.InnerText = vint.ToString();
 			}
 			else
 			{
@@ -741,7 +739,7 @@ namespace UnrealBuildTool
 			KeyElement.InnerText = KeyName;
 
 			DictRoot.AppendChild(KeyElement);
-			var newChild = ConvertValueToPListFormat(Value);
+			XmlElement? newChild = ConvertValueToPListFormat(Value);
 			if (newChild != null)
 			{
 				DictRoot.AppendChild(newChild);
@@ -936,7 +934,6 @@ namespace UnrealBuildTool
 			}
 		}
 
-
 		/// <summary>
 		/// Returns each of the entries in the value tag of type array for a given key
 		/// If the key is missing, an empty array is returned.
@@ -1049,7 +1046,6 @@ namespace UnrealBuildTool
 				KeyNode?.ParentNode?.InsertAfter(ValueNode!, KeyNode);
 
 			}
-
 		}
 
 		/// <summary>
@@ -1534,7 +1530,7 @@ namespace UnrealBuildTool
 					}
 					else if (Phase == 1)
 					{
-						if (TestProvision.ApplicationIdentifier.Contains("*"))
+						if (TestProvision.ApplicationIdentifier.Contains('*'))
 						{
 							string CompanyName = TestProvision.ApplicationIdentifier.Substring(TestProvision.ApplicationIdentifierPrefix.Length + 1);
 							if (CompanyName != "*")
@@ -1546,7 +1542,7 @@ namespace UnrealBuildTool
 					}
 					else
 					{
-						if (TestProvision.ApplicationIdentifier.Contains("*"))
+						if (TestProvision.ApplicationIdentifier.Contains('*'))
 						{
 							string CompanyName = TestProvision.ApplicationIdentifier.Substring(TestProvision.ApplicationIdentifierPrefix.Length + 1);
 							bPassesNameCheck = CompanyName == "*";
@@ -1586,7 +1582,6 @@ namespace UnrealBuildTool
 						}
 					}
 
-
 					DateTime CurrentUTCTime = DateTime.UtcNow;
 					bool bPassesDateCheck = (CurrentUTCTime >= TestProvision?.CreationDate) && (CurrentUTCTime < TestProvision?.ExpirationDate);
 					if (!bPassesDateCheck)
@@ -1673,7 +1668,6 @@ namespace UnrealBuildTool
 				string NewApplicationIdentifier = String.Format("{0}.{1}", ApplicationIdentifierPrefix, CFBundleIdentifier);
 				XCentPList.SetString("application-identifier", NewApplicationIdentifier);
 
-
 				// Replace the keychain access groups
 				// Note: This isn't robust, it ignores the existing value in the wildcard and uses the same value for
 				// each entry.  If there is a legitimate need for more than one entry in the access group list, then
@@ -1683,7 +1677,7 @@ namespace UnrealBuildTool
 				for (int i = 0; i < KeyGroups.Count; ++i)
 				{
 					string Entry = KeyGroups[i];
-					if (Entry.Contains("*"))
+					if (Entry.Contains('*'))
 					{
 						Entry = NewApplicationIdentifier;
 					}
@@ -1757,7 +1751,7 @@ namespace UnrealBuildTool
 						XCentPList.GetString("com.apple.developer.associated-domains", out AssociatedDomainsString);
 
 						//check if the value is string
-						if (AssociatedDomainsString != null && AssociatedDomainsString.Contains("*"))
+						if (AssociatedDomainsString != null && AssociatedDomainsString.Contains('*'))
 						{
 							XCentPList.RemoveKeyValue("com.apple.developer.associated-domains");
 						}
@@ -1766,7 +1760,7 @@ namespace UnrealBuildTool
 							//check if the value is an array
 							List<string> AssociatedDomainsGroup = XCentPList.GetArray("com.apple.developer.associated-domains", "string");
 
-							if (AssociatedDomainsGroup.Count == 1 && AssociatedDomainsGroup[0].Contains("*"))
+							if (AssociatedDomainsGroup.Count == 1 && AssociatedDomainsGroup[0].Contains('*'))
 							{
 								XCentPList.RemoveKeyValue("com.apple.developer.associated-domains");
 							}
@@ -1865,7 +1859,6 @@ namespace UnrealBuildTool
 				ApplicationIdentifier = "(unknown)";
 			}
 
-
 			// check for get-task-allow
 			if (XCentPList != null)
 			{
@@ -1911,7 +1904,6 @@ namespace UnrealBuildTool
 			}
 			return bFound;
 		}
-
 	}
 
 	/// <summary>
@@ -2037,7 +2029,6 @@ namespace UnrealBuildTool
 		/// </summary>
 		public static string AppDisplayName = "Unreal iOS Configuration";
 
-
 		/// <summary>
 		/// The game directory
 		/// </summary>
@@ -2421,6 +2412,5 @@ namespace UnrealBuildTool
 			return true;
 		}
 	}
-
 }
 
diff --git a/Engine/Source/Programs/UnrealBuildTool/Platform/IOS/IOSExports.cs b/Engine/Source/Programs/UnrealBuildTool/Platform/IOS/IOSExports.cs
--- a/Engine/Source/Programs/UnrealBuildTool/Platform/IOS/IOSExports.cs
+++ b/Engine/Source/Programs/UnrealBuildTool/Platform/IOS/IOSExports.cs
@@ -238,21 +238,21 @@ namespace UnrealBuildTool
 
 				string entitlementXML = MobileProvisionContent.GetNodeXMLValueByName("com.apple.developer.icloud-services");
 
-				if (!entitlementXML.Contains("*") || Platform == UnrealTargetPlatform.TVOS)
+				if (!entitlementXML.Contains('*') || Platform == UnrealTargetPlatform.TVOS)
 				{
 					// for iOS, replace the generic value (*) with the default
 					iCloudServicesXML = entitlementXML;
 				}
 
 				entitlementXML = MobileProvisionContent.GetNodeXMLValueByName("com.apple.developer.ubiquity-container-identifiers");
-				if (!entitlementXML.Contains("*") || !bForDistribution)
+				if (!entitlementXML.Contains('*') || !bForDistribution)
 				{
 					// for distribution, replace the generic value (*) with the default
 					UbiquityContainerIdentifiersXML = entitlementXML;
 				}
 
 				entitlementXML = MobileProvisionContent.GetNodeXMLValueByName("com.apple.developer.ubiquity-kvstore-identifier");
-				if (!entitlementXML.Contains("*") || !bForDistribution)
+				if (!entitlementXML.Contains('*') || !bForDistribution)
 				{
 					// for distribution, replace the generic value (*) with the default
 					UbiquityKVStoreIdentifiersXML = entitlementXML;
@@ -279,25 +279,25 @@ namespace UnrealBuildTool
 				Text.AppendLine(String.Format("\t<{0}/>", bForDistribution ? "false" : "true"));
 				if (bCloudKitSupported)
 				{
-					if (iCloudContainerIdentifiersXML != "")
+					if (!String.IsNullOrEmpty(iCloudContainerIdentifiersXML))
 					{
 						Text.AppendLine("\t<key>com.apple.developer.icloud-container-identifiers</key>");
 						Text.AppendLine(iCloudContainerIdentifiersXML);
 					}
 
-					if (iCloudServicesXML != "")
+					if (!String.IsNullOrEmpty(iCloudServicesXML))
 					{
 						Text.AppendLine("\t<key>com.apple.developer.icloud-services</key>");
 						Text.AppendLine(iCloudServicesXML);
 					}
 
-					if (UbiquityContainerIdentifiersXML != "")
+					if (!String.IsNullOrEmpty(UbiquityContainerIdentifiersXML))
 					{
 						Text.AppendLine("\t<key>com.apple.developer.ubiquity-container-identifiers</key>");
 						Text.AppendLine(UbiquityContainerIdentifiersXML);
 					}
 
-					if (UbiquityKVStoreIdentifiersXML != "")
+					if (!String.IsNullOrEmpty(UbiquityKVStoreIdentifiersXML))
 					{
 						Text.AppendLine("\t<key>com.apple.developer.ubiquity-kvstore-identifier</key>");
 						Text.AppendLine(UbiquityKVStoreIdentifiersXML);
@@ -362,7 +362,7 @@ namespace UnrealBuildTool
 
 			// create a pList key named ICloudContainerIdentifier
 			// to be used at run-time when intializing the CloudKit services
-			if (iCloudContainerIdentifier != "")
+			if (!String.IsNullOrEmpty(iCloudContainerIdentifier))
 			{
 				string PListFile = IntermediateDir + "/" + AppName + "-Info.plist";
 				if (File.Exists(PListFile))
@@ -411,10 +411,8 @@ namespace UnrealBuildTool
 					{
 						throw new BuildException("plist is invalid {0}\n{1}", e, OldPListData);
 					}
-
 				}
 			}
 		}
-
 	}
 }
diff --git a/Engine/Source/Programs/UnrealBuildTool/Platform/IOS/IOSToolChain.cs b/Engine/Source/Programs/UnrealBuildTool/Platform/IOS/IOSToolChain.cs
--- a/Engine/Source/Programs/UnrealBuildTool/Platform/IOS/IOSToolChain.cs
+++ b/Engine/Source/Programs/UnrealBuildTool/Platform/IOS/IOSToolChain.cs
@@ -1,6 +1,5 @@
 // Copyright Epic Games, Inc. All Rights Reserved.
 
-
 using System;
 using System.Collections.Generic;
 using System.Diagnostics;
@@ -812,7 +811,6 @@ namespace UnrealBuildTool
 					const int UNIX_WRITE = 2;
 					const int UNIX_READ = 4;
 
-
 					int MyPermissions = UNIX_READ | UNIX_WRITE;
 					int OtherPermissions = UNIX_READ;
 
@@ -1488,8 +1486,6 @@ namespace UnrealBuildTool
 			FileReference.Copy(Target.OutputPath, StagedExecutablePath, true);
 			string RemoteShadowDirectoryMac = Target.OutputPath.Directory.FullName;
 
-
-
 			if (Target.bCreateStubIPA || Target.bBuildAsFramework)
 			{
 				string FrameworkDerivedDataDir = Path.GetFullPath(Target.ProjectDirectory + "/Intermediate/" + (Target.Platform == UnrealTargetPlatform.IOS ? "IOS" : "TVOS") + "/DerivedData");
@@ -1579,7 +1575,7 @@ namespace UnrealBuildTool
 						MobileProvisionFile = ProvisioningData.MobileProvisionFile;
 						MobileProvisionUUID = ProvisioningData.MobileProvisionUUID;
 						TeamUUID = ProvisioningData.TeamUUID;
-						if (ProvisioningData.BundleIdentifier == null || !ProvisioningData.BundleIdentifier.Contains("*"))
+						if (ProvisioningData.BundleIdentifier == null || !ProvisioningData.BundleIdentifier.Contains('*'))
 						{
 							// If the BundleIndentifer contains a wild card it will not be valid to use in the plist.
 							BundleID = ProvisioningData.BundleIdentifier;
@@ -1593,7 +1589,7 @@ namespace UnrealBuildTool
 						MobileProvisionUUID = MobileProvision.GetUniqueId();
 						MobileProvision.TryGetTeamUniqueId(out TeamUUID);
 						string BundleIdentifier = MobileProvision.GetBundleIdentifier();
-						if (!BundleIdentifier.Contains("*"))
+						if (!BundleIdentifier.Contains('*'))
 						{
 							// If the BundleIndentifer contains a wild card it will not be valid to use in the plist.
 							BundleID = BundleIdentifier;
@@ -1861,6 +1857,5 @@ namespace UnrealBuildTool
 				BundleDependencies.Add(FileItem.GetItemByFileReference(Binary.OutputFilePath));
 			}
 		}
-
 	};
 }
diff --git a/Engine/Source/Programs/UnrealBuildTool/Platform/IOS/UEDeployIOS.cs b/Engine/Source/Programs/UnrealBuildTool/Platform/IOS/UEDeployIOS.cs
--- a/Engine/Source/Programs/UnrealBuildTool/Platform/IOS/UEDeployIOS.cs
+++ b/Engine/Source/Programs/UnrealBuildTool/Platform/IOS/UEDeployIOS.cs
@@ -213,7 +213,6 @@ namespace UnrealBuildTool
 			return MinVersionToReturn;
 		}
 
-
 		public static void WritePlistFile(FileReference PlistFile, DirectoryReference? ProjectLocation, UnrealPluginLanguage? UPL, string GameName, bool bIsUnrealGame, ILogger Logger)
 		{
 			ConfigHierarchy Ini = ConfigCache.ReadHierarchy(ConfigHierarchyType.Engine, ProjectLocation, UnrealTargetPlatform.IOS);
@@ -278,7 +277,6 @@ namespace UnrealBuildTool
 			string BundleShortVersion;
 			Ini.GetString("/Script/IOSRuntimeSettings.IOSRuntimeSettings", "VersionInfo", out BundleShortVersion);
 
-
 			// Get Google Support details
 			bool bEnableGoogleSupport = true;
 			Ini.GetBool("/Script/IOSRuntimeSettings.IOSRuntimeSettings", "bEnableGoogleSupport", out bEnableGoogleSupport);
@@ -347,7 +345,7 @@ namespace UnrealBuildTool
 			Text.AppendLine("\t<true/>");
 			Text.AppendLine("\t<key>UIViewControllerBasedStatusBarAppearance</key>");
 			Text.AppendLine("\t<false/>");
-			if (InterfaceOrientation != "")
+			if (!String.IsNullOrEmpty(InterfaceOrientation))
 			{
 				Text.AppendLine(InterfaceOrientation);
 			}
@@ -438,7 +436,6 @@ namespace UnrealBuildTool
 			Text.AppendLine("</dict>");
 			Text.AppendLine("</plist>");
 
-
 			DirectoryReference.CreateDirectory(PlistFile.Directory);
 
 			if (UPL != null)
@@ -697,7 +694,7 @@ namespace UnrealBuildTool
 			Text.AppendLine("\t<true/>");
 			Text.AppendLine("\t<key>UIViewControllerBasedStatusBarAppearance</key>");
 			Text.AppendLine("\t<false/>");
-			if (InterfaceOrientation != "")
+			if (!String.IsNullOrEmpty(InterfaceOrientation))
 			{
 				Text.AppendLine(InterfaceOrientation);
 			}
@@ -1001,7 +998,6 @@ namespace UnrealBuildTool
 			{
 				CopyFiles(InEngineDir + "/Build/IOS/Resources/Graphics", AppDirectory, "LaunchScreenIOS.png", true);
 			}
-
 		}
 
 		protected virtual void CopyLaunchScreenResources(string InEngineDir, string AppDirectory, string BuildDirectory, ILogger Logger)
diff --git a/Engine/Source/Programs/UnrealBuildTool/Platform/Linux/LinuxToolChain.cs b/Engine/Source/Programs/UnrealBuildTool/Platform/Linux/LinuxToolChain.cs
--- a/Engine/Source/Programs/UnrealBuildTool/Platform/Linux/LinuxToolChain.cs
+++ b/Engine/Source/Programs/UnrealBuildTool/Platform/Linux/LinuxToolChain.cs
@@ -169,7 +169,6 @@ namespace UnrealBuildTool
 			StringWriter Out = new StringWriter();
 			Out.NewLine = bUseCmdExe ? "\r\n" : "\n";
 
-
 			if (!Options.HasFlag(ClangToolChainOptions.DisableDumpSyms) || Options.HasFlag(ClangToolChainOptions.PreservePSYM))
 			{
 				// dump_syms
@@ -551,7 +550,7 @@ namespace UnrealBuildTool
 			if (String.IsNullOrEmpty(Key)) { return ""; }
 			if (!String.IsNullOrEmpty(Value))
 			{
-				if (!Value.StartsWith("\"") && (Value.Contains(" ") || Value.Contains("$")))
+				if (!Value.StartsWith("\"") && (Value.Contains(' ') || Value.Contains('$')))
 				{
 					Value = Value.Trim('\"');       // trim any leading or trailing quotes
 					Value = "\"" + Value + "\"";    // ensure wrap string with double quotes
@@ -990,7 +989,6 @@ namespace UnrealBuildTool
 			return FileItem.GetItemByFileReference(LinkEnvironment.OutputFilePath);
 		}
 
-
 		public override FileItem LinkFiles(LinkEnvironment LinkEnvironment, bool bBuildImportLibraryOnly, IActionGraphBuilder Graph)
 		{
 			Debug.Assert(!bBuildImportLibraryOnly);
@@ -1184,7 +1182,7 @@ namespace UnrealBuildTool
 				{
 					// static library passed in, pass it along but make path absolute, because FixDependencies script may be executed in a different directory
 					string AbsoluteAdditionalLibrary = Path.GetFullPath(AdditionalLibrary);
-					if (AbsoluteAdditionalLibrary.Contains(" "))
+					if (AbsoluteAdditionalLibrary.Contains(' '))
 					{
 						AbsoluteAdditionalLibrary = String.Format("\"{0}\"", AbsoluteAdditionalLibrary);
 					}
diff --git a/Engine/Source/Programs/UnrealBuildTool/Platform/Mac/MacToolChain.cs b/Engine/Source/Programs/UnrealBuildTool/Platform/Mac/MacToolChain.cs
--- a/Engine/Source/Programs/UnrealBuildTool/Platform/Mac/MacToolChain.cs
+++ b/Engine/Source/Programs/UnrealBuildTool/Platform/Mac/MacToolChain.cs
@@ -230,7 +230,6 @@ namespace UnrealBuildTool
 			// This is also added to the x64 builds in case someone ever adds -mfma.
 			Arguments.Add("-ffp-contract=off");
 
-
 			if (CompileEnvironment.bEnableOSX109Support)
 			{
 				Arguments.Add("-faligned-new"); // aligned operator new is supported only on macOS 10.14 and above
@@ -472,7 +471,6 @@ namespace UnrealBuildTool
 			}
 		}
 
-
 		public override FileItem[] LinkImportLibrary(LinkEnvironment LinkEnvironment, IActionGraphBuilder Graph)
 		{
 			// we actually create Actions for every dylib, but we never return a FileItem. Instead depend on a Link action
@@ -679,7 +677,6 @@ namespace UnrealBuildTool
 				AppendMacLine(FinalizeAppBundleScript, FormatCopyCommand(TempInfoPlist, XcodeInputPListFile));
 				AppendMacLine(FinalizeAppBundleScript, "chmod 644 \"{0}\"", XcodeInputPListFile);
 
-
 				// Generate PkgInfo file
 				string TempPkgInfo = "$TMPDIR/TempPkgInfo";
 				AppendMacLine(FinalizeAppBundleScript, "echo 'echo -n \"APPL????\"' | bash > \"{0}\"", TempPkgInfo);
@@ -688,7 +685,6 @@ namespace UnrealBuildTool
 				// Make sure OS X knows the bundle was updated
 				AppendMacLine(FinalizeAppBundleScript, "touch -c \"{0}.app\"", ExeName);
 
-
 				// codesign with ad-hoc signature for when building outside of Xcode the there will be at least some signature
 				// (and it can be a BuildProduct down in ModifyBuildProducts)
 				AppendMacLine(FinalizeAppBundleScript, "codesign -f -s - \"{0}.app\"", ExeName);
@@ -698,7 +694,6 @@ namespace UnrealBuildTool
 				FinalizeAppBundleScript.Close();
 			}
 
-
 			return OutputFiles.ToArray();
 		}
 
@@ -787,7 +782,7 @@ namespace UnrealBuildTool
 				IEnumerable<string> AdditionalLibraries = Enumerable.Concat(LinkEnvironment.SystemLibraries, LinkEnvironment.Libraries.Select(x => x.FullName));
 				foreach (string AdditionalLibrary in AdditionalLibraries)
 				{
-					if (Path.GetDirectoryName(AdditionalLibrary) != "" &&
+					if (!String.IsNullOrEmpty(Path.GetDirectoryName(AdditionalLibrary)) &&
 							 (Path.GetDirectoryName(AdditionalLibrary)!.Contains("Binaries/Mac") ||
 							 Path.GetDirectoryName(AdditionalLibrary)!.Contains("Binaries\\Mac")))
 					{
diff --git a/Engine/Source/Programs/UnrealBuildTool/Platform/Windows/AppXManifestGeneratorBase.cs b/Engine/Source/Programs/UnrealBuildTool/Platform/Windows/AppXManifestGeneratorBase.cs
--- a/Engine/Source/Programs/UnrealBuildTool/Platform/Windows/AppXManifestGeneratorBase.cs
+++ b/Engine/Source/Programs/UnrealBuildTool/Platform/Windows/AppXManifestGeneratorBase.cs
@@ -68,7 +68,6 @@ namespace UnrealBuildTool
 			Logger = InLogger;
 		}
 
-
 		/// <summary>
 		/// Returns a valid version of the given package version string
 		/// </summary>
@@ -180,8 +179,8 @@ namespace UnrealBuildTool
 		/// </summary>
 		protected bool GetConfigBool(string PlatformKey, string? GenericKey, bool DefaultValue = false)
 		{
-			var GenericValue = ReadIniString(GenericKey, IniSection_GeneralProjectSettings, null);
-			var ResultStr = ReadIniString(PlatformKey, IniSection_PlatformTargetSettings, GenericValue);
+			string? GenericValue = ReadIniString(GenericKey, IniSection_GeneralProjectSettings, null);
+			string? ResultStr = ReadIniString(PlatformKey, IniSection_PlatformTargetSettings, GenericValue);
 
 			if (ResultStr == null)
 			{
@@ -198,7 +197,7 @@ namespace UnrealBuildTool
 		/// </summary>
 		protected string GetConfigColor(string PlatformConfigKey, string DefaultValue)
 		{
-			var ConfigValue = GetConfigString(PlatformConfigKey, null, null);
+			string? ConfigValue = GetConfigString(PlatformConfigKey, null, null);
 			if (ConfigValue == null)
 			{
 				return DefaultValue;
@@ -218,8 +217,6 @@ namespace UnrealBuildTool
 			return DefaultValue;
 		}
 
-
-
 		/// <summary>
 		/// Create all the localization data. Returns whether there is any per-culture data set up
 		/// </summary>
@@ -280,11 +277,9 @@ namespace UnrealBuildTool
 				}
 			}
 
-
 			return bHasPerCultureResources;
 		}
 
-
 		/// <summary>
 		/// Register the locations where resource binary files can be found
 		/// </summary>
@@ -300,7 +295,6 @@ namespace UnrealBuildTool
 			AppXResources!.EngineFallbackBinaryResourceDirectories.Add(DirectoryReference.Combine(Unreal.EngineDirectory, "Platforms", Platform.ToString(), "Build", EngineResourceSubPath));
 		}
 
-
 		/// <summary>
 		/// Get the resources element
 		/// </summary>
@@ -319,7 +313,7 @@ namespace UnrealBuildTool
 			}
 
 			// Create the culture list. This list is unordered except that the default language must be first which we already took care of above.
-			var CultureElements = ResourceCulturesList.Select(c =>
+			IEnumerable<XElement> CultureElements = ResourceCulturesList.Select(c =>
 				new XElement("Resource", new XAttribute("Language", c)));
 
 			return new XElement("Resources", CultureElements);
@@ -331,8 +325,8 @@ namespace UnrealBuildTool
 		protected virtual string GetIdentityPackageName()
 		{
 			// Read the PackageName from config
-			var DefaultName = (ProjectFile != null) ? ProjectFile.GetFileNameWithoutAnyExtensions() : (TargetName ?? "DefaultUEProject");
-			var PackageName = Regex.Replace(GetConfigString("PackageName", "ProjectName", DefaultName), "[^-.A-Za-z0-9]", "");
+			string DefaultName = (ProjectFile != null) ? ProjectFile.GetFileNameWithoutAnyExtensions() : (TargetName ?? "DefaultUEProject");
+			string PackageName = Regex.Replace(GetConfigString("PackageName", "ProjectName", DefaultName), "[^-.A-Za-z0-9]", "");
 			if (String.IsNullOrWhiteSpace(PackageName))
 			{
 				Logger.LogError("Invalid package name {Name}. Package names must only contain letters, numbers, dash, and period and must be at least one character long.", PackageName);
@@ -343,7 +337,7 @@ namespace UnrealBuildTool
 			bool bPackageNameUseMachineName;
 			if (EngineIni!.GetBool(IniSection_PlatformTargetSettings, "bPackageNameUseMachineName", out bPackageNameUseMachineName) && bPackageNameUseMachineName)
 			{
-				var MachineName = Regex.Replace(Environment.MachineName.ToString(), "[^-.A-Za-z0-9]", "");
+				string MachineName = Regex.Replace(Environment.MachineName.ToString(), "[^-.A-Za-z0-9]", "");
 				PackageName = PackageName + ".NOT.SHIPPABLE." + MachineName;
 			}
 
@@ -355,7 +349,7 @@ namespace UnrealBuildTool
 		/// </summary>
 		protected virtual string GetIdentityPublisherName()
 		{
-			var PublisherName = GetConfigString("PublisherName", "CompanyDistinguishedName", "CN=NoPublisher");
+			string PublisherName = GetConfigString("PublisherName", "CompanyDistinguishedName", "CN=NoPublisher");
 			return PublisherName;
 		}
 
@@ -364,7 +358,7 @@ namespace UnrealBuildTool
 		/// </summary>
 		protected virtual string? GetIdentityVersionNumber()
 		{
-			var VersionNumber = GetConfigString("PackageVersion", "ProjectVersion", "1.0.0.0");
+			string? VersionNumber = GetConfigString("PackageVersion", "ProjectVersion", "1.0.0.0");
 			VersionNumber = ValidatePackageVersion(VersionNumber);
 
 			// If specified in the project settings attempt to retrieve the current build number and increment the version number by that amount, accounting for overflows
@@ -511,7 +505,6 @@ namespace UnrealBuildTool
 					UEStageIdToAppXCultureId[UEStageId] = AppXCultureId;
 					AppXResources.AddCulture(UEStageId);
 				}
-
 			}
 
 			// look up the default AppX culture
@@ -522,7 +515,6 @@ namespace UnrealBuildTool
 				UEStageIdToAppXCultureId[DefaultUECultureId] = DefaultAppXCultureId;
 			}
 
-
 			// Warn in shipping, we can run without translated cultures they're just needed for cert
 			if (!bHasLocalizationData && InExecutablePairs.ContainsKey(UnrealTargetConfiguration.Shipping))
 			{
@@ -531,7 +523,7 @@ namespace UnrealBuildTool
 
 			// Create the manifest document
 			string? IdentityName = null;
-			var ManifestXmlDocument = new XDocument(GetManifest(InExecutablePairs, out IdentityName));
+			XDocument ManifestXmlDocument = new XDocument(GetManifest(InExecutablePairs, out IdentityName));
 
 			// Export manifest to the intermediate directory and add it to the manifest resources files for copying
 			FileReference ManifestIntermediateFile = FileReference.Combine(IntermediateDirectory, InManifestName);
diff --git a/Engine/Source/Programs/UnrealBuildTool/Platform/Windows/MicrosoftPlatformSDK.cs b/Engine/Source/Programs/UnrealBuildTool/Platform/Windows/MicrosoftPlatformSDK.cs
--- a/Engine/Source/Programs/UnrealBuildTool/Platform/Windows/MicrosoftPlatformSDK.cs
+++ b/Engine/Source/Programs/UnrealBuildTool/Platform/Windows/MicrosoftPlatformSDK.cs
@@ -82,8 +82,6 @@ namespace UnrealBuildTool
 			return false;
 		}
 
-
-
 		#region Windows Specific SDK discovery
 
 		#region Cached Locations
@@ -508,14 +506,14 @@ namespace UnrealBuildTool
 		/// <returns>True if the given compiler is installed</returns>
 		public static bool HasCompiler(WindowsCompiler Compiler, UnrealArch Architecture, ILogger Logger)
 		{
-			return FindToolChainInstallations(Compiler, Logger).Where(x => (x.Architecture == Architecture)).Count() > 0;
+			return FindToolChainInstallations(Compiler, Logger).Where(x => (x.Architecture == Architecture)).Any();
 		}
 
 		public static bool HasValidCompiler(WindowsCompiler Compiler, UnrealArch Architecture, ILogger Logger)
 		{
 			// since this is static, we get the instance for it
 			MicrosoftPlatformSDK SDK = (MicrosoftPlatformSDK)GetSDKForPlatform("Win64")!;
-			return FindToolChainInstallations(Compiler, Logger).Where(x => (x.Version >= MinimumVisualCppVersion && x.Architecture == Architecture)).Count() > 0;
+			return FindToolChainInstallations(Compiler, Logger).Where(x => (x.Version >= MinimumVisualCppVersion && x.Architecture == Architecture)).Any();
 		}
 
 		/// <summary>
@@ -654,7 +652,7 @@ namespace UnrealBuildTool
 				.ThenByDescending(x => x.IsAutoSdk)
 				.ThenByDescending(x => x.Version);
 
-			if (SortedToolChains.Count() > 0)
+			if (SortedToolChains.Any())
 			{
 				IEnumerable<ToolChainInstallation> AvailableToolChains = SortedToolChains.Where(x => x.Error == null);
 				IEnumerable<ToolChainInstallation> UnavailableToolChains = SortedToolChains.Where(x => x.Error != null);
@@ -937,7 +935,6 @@ namespace UnrealBuildTool
 			}
 		}
 
-
 		/// <summary>
 		/// Finds the most appropriate redist directory for the given toolchain version
 		/// </summary>
@@ -1056,7 +1053,6 @@ namespace UnrealBuildTool
 			}
 		}
 
-
 		/// <summary>
 		/// Add an Intel OneAPI toolchain
 		/// </summary>
@@ -1153,7 +1149,6 @@ namespace UnrealBuildTool
 			}
 		}
 
-
 		/// <summary>
 		/// Determines if the given path is a valid Visual C++ version number
 		/// </summary>
diff --git a/Engine/Source/Programs/UnrealBuildTool/Platform/Windows/PVSToolChain.cs b/Engine/Source/Programs/UnrealBuildTool/Platform/Windows/PVSToolChain.cs
--- a/Engine/Source/Programs/UnrealBuildTool/Platform/Windows/PVSToolChain.cs
+++ b/Engine/Source/Programs/UnrealBuildTool/Platform/Windows/PVSToolChain.cs
@@ -706,9 +706,9 @@ namespace UnrealBuildTool
 				{
 					foreach (string PathMask in ApplicationSettings.PathMasks)
 					{
-						if (PathMask.Contains(":") || PathMask.Contains("\\") || PathMask.Contains("/"))
+						if (PathMask.Contains(':') || PathMask.Contains('\\') || PathMask.Contains('/'))
 						{
-							if (Path.IsPathRooted(PathMask) && !PathMask.Contains(":"))
+							if (Path.IsPathRooted(PathMask) && !PathMask.Contains(':'))
 							{
 								ConfigFileContents.AppendFormat("exclude-path=*{0}*\n", PathMask);
 							}
diff --git a/Engine/Source/Programs/UnrealBuildTool/Platform/Windows/ResXResource/UEAppXResources.cs b/Engine/Source/Programs/UnrealBuildTool/Platform/Windows/ResXResource/UEAppXResources.cs
--- a/Engine/Source/Programs/UnrealBuildTool/Platform/Windows/ResXResource/UEAppXResources.cs
+++ b/Engine/Source/Programs/UnrealBuildTool/Platform/Windows/ResXResource/UEAppXResources.cs
@@ -69,7 +69,6 @@ namespace UnrealBuildTool
 			AddCulture(DefaultResources);
 		}
 
-
 		/// <summary>
 		/// Adds a reference to the given culture
 		/// </summary>
@@ -93,17 +92,15 @@ namespace UnrealBuildTool
 			}
 		}
 
-
 		/// <summary>
 		/// Returns a collection of all registered cultures
 		/// </summary>
 		/// <returns></returns>
 		public IEnumerable<string> GetAllCultureIds()
 		{
-			return PerCultureResources.Keys.Where(X => X != DefaultResources);
+			return PerCultureResources.Keys.Where(X => !String.IsNullOrEmpty(X));
 		}
 
-
 		/// <summary>
 		/// Add the given string to the default culture
 		/// </summary>
@@ -148,7 +145,6 @@ namespace UnrealBuildTool
 			return GetDefaultResources().SourceStrings.ContainsKey(Key);
 		}
 
-
 		/// <summary>
 		/// Clear all stored strings for all cultures
 		/// </summary>
@@ -160,7 +156,6 @@ namespace UnrealBuildTool
 			}
 		}
 
-
 		Resources GetCultureResources(string AppXCultureId)
 		{
 			if (PerCultureResources.TryGetValue(AppXCultureId, out Resources? Result) && Result != null)
@@ -178,9 +173,6 @@ namespace UnrealBuildTool
 			return GetCultureResources(DefaultResources);
 		}
 
-
-
-
 		bool RunMakePri(string CommandLine)
 		{
 			StringBuilder ProcessOutput = new StringBuilder();
@@ -218,9 +210,6 @@ namespace UnrealBuildTool
 			}
 		}
 
-
-
-
 		private bool RemoveStaleResourceFiles(IEnumerable<string> RequiredFileFragments, DirectoryReference OutputDirectory)
 		{
 			DirectoryReference ResourceDirectory = DirectoryReference.Combine(OutputDirectory, BuildResourceSubPath);
@@ -263,7 +252,6 @@ namespace UnrealBuildTool
 			return true;
 		}
 
-
 		/// <summary>
 		/// Attempts to locate the given resource binary file in the previously-specified resource folders
 		/// </summary>
@@ -315,7 +303,6 @@ namespace UnrealBuildTool
 			return DoesCultureResourceBinaryFileExist(ResourceFileName, DefaultResources, AllowEngineFallback);
 		}
 
-
 		/// <summary>
 		/// Adds the given resource binary file(s) to the manifest files
 		/// </summary>
@@ -340,7 +327,6 @@ namespace UnrealBuildTool
 			return true;
 		}
 
-
 		/// <summary>
 		/// Adds the given file to the manifest files
 		/// </summary>
@@ -363,8 +349,6 @@ namespace UnrealBuildTool
 			return !Enumerable.SequenceEqual(FileContents1, FileContents2);
 		}
 
-
-
 		/// <summary>
 		/// Copies all of the generated files to the output folder
 		/// </summary>
@@ -377,7 +361,7 @@ namespace UnrealBuildTool
 			}
 
 			Logger.LogDebug("Updating manifest resource files...");
-			foreach (var FileToCopy in FilesToCopy)
+			foreach (KeyValuePair<FileReference, string> FileToCopy in FilesToCopy)
 			{
 				FileReference SourceFile = FileToCopy.Key;
 				FileReference TargetFile = FileReference.Combine(OutputDirectory, FileToCopy.Value);
@@ -433,8 +417,6 @@ namespace UnrealBuildTool
 			return UpdatedFiles;
 		}
 
-
-
 		/// <summary>
 		/// Adds the given string to the culture string writers
 		/// </summary>
@@ -453,7 +435,6 @@ namespace UnrealBuildTool
 			}
 			DefaultCultureResources.StringResources.Add(ResourceEntryName, DefaultCultureString + ValueSuffix);
 
-
 			// Get the localized culture values
 			foreach (string CultureId in GetAllCultureIds())
 			{
@@ -467,9 +448,6 @@ namespace UnrealBuildTool
 			return "ms-resource:" + ResourceEntryName;
 		}
 
-
-
-
 		/// <summary>
 		/// Generate the package resource index and copies all resources files to the output
 		/// </summary>
diff --git a/Engine/Source/Programs/UnrealBuildTool/Platform/Windows/VCCompileAction.cs b/Engine/Source/Programs/UnrealBuildTool/Platform/Windows/VCCompileAction.cs
--- a/Engine/Source/Programs/UnrealBuildTool/Platform/Windows/VCCompileAction.cs
+++ b/Engine/Source/Programs/UnrealBuildTool/Platform/Windows/VCCompileAction.cs
@@ -475,7 +475,7 @@ namespace UnrealBuildTool
 			foreach (string Definition in Definitions)
 			{
 				// Escape all quotation marks so that they get properly passed with the command line.
-				string DefinitionArgument = Definition.Contains("\"") ? Definition.Replace("\"", "\\\"") : Definition;
+				string DefinitionArgument = Definition.Contains('"') ? Definition.Replace("\"", "\\\"") : Definition;
 				VCToolChain.AddDefinition(Arguments, DefinitionArgument);
 			}
 
diff --git a/Engine/Source/Programs/UnrealBuildTool/Platform/Windows/VCToolChain.cs b/Engine/Source/Programs/UnrealBuildTool/Platform/Windows/VCToolChain.cs
--- a/Engine/Source/Programs/UnrealBuildTool/Platform/Windows/VCToolChain.cs
+++ b/Engine/Source/Programs/UnrealBuildTool/Platform/Windows/VCToolChain.cs
@@ -122,7 +122,7 @@ namespace UnrealBuildTool
 		public static void AddDefinition(List<string> Arguments, string Variable, string? Value)
 		{
 			// If the value has a space in it and isn't wrapped in quotes, do that now
-			if (Value != null && !Value.StartsWith("\"") && (Value.Contains(" ") || Value.Contains("$")))
+			if (Value != null && !Value.StartsWith("\"") && (Value.Contains(' ') || Value.Contains('$')))
 			{
 				Value = "\"" + Value + "\"";
 			}
@@ -450,7 +450,6 @@ namespace UnrealBuildTool
 			// Separate functions for linker.
 			Arguments.Add("/Gy");
 
-
 			// Microsoft recommends not passing /Zm except in very limited circumstances, please see:
 			// https://learn.microsoft.com/en-us/cpp/build/reference/zm-specify-precompiled-header-memory-allocation-limit
 			if (Target.WindowsPlatform.PCHMemoryAllocationFactor > 0)
@@ -666,7 +665,6 @@ namespace UnrealBuildTool
 						Arguments.Add("/Oy-");
 					}
 				}
-
 			}
 
 			//
@@ -927,7 +925,6 @@ namespace UnrealBuildTool
 			// Restrictive during regular compilation.
 			Arguments.Add("/W4");
 
-
 			// Treat warnings as errors
 			if (CompileEnvironment.bWarningsAsErrors)
 			{
@@ -1261,7 +1258,6 @@ namespace UnrealBuildTool
 				Arguments.Add(Target.WindowsPlatform.AdditionalLinkerOptions);
 			}
 
-
 			// Disable
 			//LINK : warning LNK4199: /DELAYLOAD:nvtt_64.dll ignored; no imports found from nvtt_64.dll
 			// type warning as we leverage the DelayLoad option to put third-party DLLs into a
@@ -1824,7 +1820,6 @@ namespace UnrealBuildTool
 					throw new BuildException("Unsupported build architecture for Address Sanitizer");
 				}
 
-
 				string ASanRuntimeDLL = $"clang_rt.asan_dynamic-{ASanArchSuffix}.dll";
 				string ASanDebugRuntimeDLL = $"clang_rt.asan_dbg_dynamic-{ASanArchSuffix}.dll";
 
@@ -2204,7 +2199,6 @@ namespace UnrealBuildTool
 				Arguments.Add("/IGNORE:4221");
 			}
 
-
 			if (!bIsBuildingLibraryOrImportLibrary)
 			{
 				// Delay-load these DLLs.
@@ -2516,7 +2510,6 @@ namespace UnrealBuildTool
 					PGDFile = PGDFiles.First();
 				}
 
-
 				string[] PGCFiles = Directory.GetFiles(LinkEnvironment.PGODirectory!, "*.pgc");
 				if (PGCFiles.Length == 0)
 				{
diff --git a/Engine/Source/Programs/UnrealBuildTool/ProjectFiles/CodeLite/CodeLiteProject.cs b/Engine/Source/Programs/UnrealBuildTool/ProjectFiles/CodeLite/CodeLiteProject.cs
--- a/Engine/Source/Programs/UnrealBuildTool/ProjectFiles/CodeLite/CodeLiteProject.cs
+++ b/Engine/Source/Programs/UnrealBuildTool/ProjectFiles/CodeLite/CodeLiteProject.cs
@@ -158,7 +158,7 @@ namespace UnrealBuildTool
 					//
 					foreach (string FolderName in SplitFolders)
 					{
-						if (FolderName.Equals(""))
+						if (String.IsNullOrEmpty(FolderName))
 						{
 							continue;
 						}
@@ -267,7 +267,6 @@ namespace UnrealBuildTool
 								ExecutableToRun += PlatformConfiguration;
 							}
 						}
-
 					}
 					else if (BuildHostPlatform.Current.Platform == UnrealTargetPlatform.Mac)
 					{
@@ -283,7 +282,6 @@ namespace UnrealBuildTool
 							{
 								ExecutableToRun += PlatformConfiguration;
 							}
-
 						}
 						else if (ProjectTargetType == TargetType.Editor)
 						{
@@ -300,7 +298,6 @@ namespace UnrealBuildTool
 								ExecutableToRun += PlatformConfiguration;
 							}
 						}
-
 					}
 					else if (BuildHostPlatform.Current.Platform == UnrealTargetPlatform.Win64)
 					{
@@ -331,7 +328,6 @@ namespace UnrealBuildTool
 						throw new BuildException("Unsupported platform.");
 					}
 
-
 					// Is this project a Game type?
 					XAttribute GeneralExecutableToRun = new XAttribute("Command", ExecutableToRun);
 					if (ProjectTargetType == TargetType.Game)
@@ -355,7 +351,7 @@ namespace UnrealBuildTool
 					}
 					else if (ProjectTargetType == TargetType.Editor)
 					{
-						if (TargetName != "UnrealEditor" && GameProjectFile != "")
+						if (TargetName != "UnrealEditor" && !String.IsNullOrEmpty(GameProjectFile))
 						{
 							string commandArguments = "\"" + GameProjectFile + "\"" + " -game";
 							XAttribute CommandArguments = new XAttribute("CommandArguments", commandArguments);
@@ -407,8 +403,6 @@ namespace UnrealBuildTool
 					// End of Add the working directory for the custom build commands.
 					//
 
-
-
 					//
 					// Make Build Target.
 					//
@@ -454,7 +448,6 @@ namespace UnrealBuildTool
 
 					CustomCleanCommand.Add(CleanCommandLine);
 
-
 					//
 					// End of Clean Build Target.
 					//
@@ -474,7 +467,6 @@ namespace UnrealBuildTool
 					// End of Clean Build Target.
 					//
 
-
 					//
 					// Some other fun Custom Targets.
 					//
diff --git a/Engine/Source/Programs/UnrealBuildTool/ProjectFiles/Eddie/EddieProject.cs b/Engine/Source/Programs/UnrealBuildTool/ProjectFiles/Eddie/EddieProject.cs
--- a/Engine/Source/Programs/UnrealBuildTool/ProjectFiles/Eddie/EddieProject.cs
+++ b/Engine/Source/Programs/UnrealBuildTool/ProjectFiles/Eddie/EddieProject.cs
@@ -77,7 +77,7 @@ namespace UnrealBuildTool
 			{
 				EddieFolder CurrentGroup;
 
-				if (CurrentPath != "")
+				if (!System.String.IsNullOrEmpty(CurrentPath))
 				{
 					CurrentPath += Path.DirectorySeparatorChar;
 				}
@@ -166,7 +166,6 @@ namespace UnrealBuildTool
 							ProjectFileContent.Append("AddFile \"" + File.Reference.GetFileName() + "\" \"" + File.Reference.FullName + "\"" + ProjectFileGenerator.NewLine);
 						}
 
-
 						ProjectFileContent.Append("EndFileGroup \"" + Path.GetFileName(CurGroup.Key) + "\"" + ProjectFileGenerator.NewLine);
 
 						ProjectFileGenerator.WriteFileIfChanged(CurGroup.Value.WorksetPath!, ProjectFileContent.ToString(), Logger, new UTF8Encoding());
@@ -217,5 +216,4 @@ namespace UnrealBuildTool
 			return bSuccess;
 		}
 	}
-
 }
diff --git a/Engine/Source/Programs/UnrealBuildTool/ProjectFiles/ProjectFileGenerator.cs b/Engine/Source/Programs/UnrealBuildTool/ProjectFiles/ProjectFileGenerator.cs
--- a/Engine/Source/Programs/UnrealBuildTool/ProjectFiles/ProjectFileGenerator.cs
+++ b/Engine/Source/Programs/UnrealBuildTool/ProjectFiles/ProjectFileGenerator.cs
@@ -51,7 +51,6 @@ namespace UnrealBuildTool
 			private set;
 		}
 
-
 		/// <summary>
 		/// Adds a new sub-folder to this folder
 		/// </summary>
@@ -90,7 +89,6 @@ namespace UnrealBuildTool
 			return ResultFolder!;
 		}
 
-
 		/// <summary>
 		/// Recursively searches for the specified project and returns the folder that it lives in, or null if not found
 		/// </summary>
@@ -679,7 +677,7 @@ namespace UnrealBuildTool
 				FileReference.Combine(Unreal.EngineDirectory, "Source", "Programs", "UnrealBuildTool", "UnrealBuildTool.csproj")
 			};
 
-			var CurrentParent = RulesAssembly.Parent;
+			RulesAssembly? CurrentParent = RulesAssembly.Parent;
 			while (CurrentParent != null)
 			{
 				if (RulesAssemblies.TryGetValue(CurrentParent, out DirectoryReference? FSParentPathBase))
@@ -843,7 +841,6 @@ namespace UnrealBuildTool
 				}
 			}
 
-
 			Format = ProjectFileFormat.VisualStudio;
 			return false;
 		}
@@ -1165,7 +1162,6 @@ namespace UnrealBuildTool
 					{
 						DebugProjectFiles.AddRange(NewProjectFiles);
 					}
-
 				}
 
 				//Related Debug Project Files - Tuple has the related Debug Project, SolutionFolder
@@ -1579,7 +1575,6 @@ namespace UnrealBuildTool
 			}
 		}
 
-
 		/// <summary>
 		/// Adds all game project files, including target projects and config files
 		/// </summary>
@@ -1601,7 +1596,6 @@ namespace UnrealBuildTool
 				{
 					// @todo projectfiles: We have engine localization files, but should we also add GAME localization files?
 
-
 					// Game restricted source files, since they won't be added via a module FileReference
 					foreach (DirectoryReference GameRestrictedSourceDirectory in Unreal.GetExtensionDirs(GameProjectDirectory, "Source", bIncludePlatformDirectories: false, bIncludeBaseDirectory: false))
 					{
@@ -1641,7 +1635,6 @@ namespace UnrealBuildTool
 			}
 		}
 
-
 		/// Adds all engine localization text files to the specified project
 		private void AddEngineLocalizationFiles(ProjectFile EngineProject)
 		{
@@ -1652,7 +1645,6 @@ namespace UnrealBuildTool
 			}
 		}
 
-
 		/// Adds all engine template text files to the specified project
 		private void AddEngineTemplateFiles(ProjectFile EngineProject)
 		{
@@ -1663,7 +1655,6 @@ namespace UnrealBuildTool
 			}
 		}
 
-
 		/// Adds all engine config files to the specified project
 		private void AddEngineConfigFiles(ProjectFile EngineProject)
 		{
@@ -1842,7 +1833,6 @@ namespace UnrealBuildTool
 			}
 		}
 
-
 		/// <summary>
 		/// Recursively collapses all sub-folders that are redundant.  Should only be called after we're done adding
 		/// files and projects to the primary project.
@@ -1863,7 +1853,6 @@ namespace UnrealBuildTool
 			//
 			// Also, we can kill folders that are completely empty.
 
-
 			foreach (PrimaryProjectFolder SubFolder in Folder.SubFolders)
 			{
 				// Recurse
@@ -2218,7 +2207,6 @@ namespace UnrealBuildTool
 
 		}
 
-
 		/// <summary>
 		/// Selects which platforms and build configurations we want in the project file
 		/// </summary>
@@ -2435,7 +2423,6 @@ namespace UnrealBuildTool
 						Logger.LogInformation("Searching for third-party source files...");
 					}
 
-
 					// Find all of the source files (and other files) and add them to the project
 					List<FileReference> FoundFiles = SourceFileSearch.FindModuleSourceFiles(CurModuleFile, SearchSubdirectories: SearchSubdirectories);
 					// remove any target files, they are technically not part of the module and are explicitly added when the project is created
@@ -2523,7 +2510,6 @@ namespace UnrealBuildTool
 			}
 		}
 
-
 		private ProjectFile FindOrAddProjectHelper(string InProjectFileNameBase, DirectoryReference InBaseFolder)
 		{
 			// Setup a project file entry for this module's project.  Remember, some projects may host multiple modules!
@@ -2903,7 +2889,6 @@ namespace UnrealBuildTool
 						ProjectFile.AddFilesToProject(AllSubTargetFilesPerTarget[TargetName], BaseFolder);
 					}
 
-
 					Logger.LogDebug("Generating target {Target} for {Project}", TargetRulesObject.Type.ToString(), ProjectFilePath);
 				}
 			}
@@ -2921,10 +2906,10 @@ namespace UnrealBuildTool
 					}
 
 					// hook up to the engine target(s)
-					foreach (var EngineProject in EngineProjects)
+					foreach (ProjectFile EngineProject in EngineProjects)
 					{
 						ProjectFile? ProjectFile = null;
-						foreach (var EngineTarget in EngineProject.ProjectTargets)
+						foreach (Project EngineTarget in EngineProject.ProjectTargets)
 						{
 							if (bMakeProjectPerTarget)
 							{
@@ -3010,7 +2995,6 @@ namespace UnrealBuildTool
 			}
 		}
 
-
 		/// Adds engine build infrastructure files to the specified project
 		protected void AddEngineBuildFiles(ProjectFile EngineProject)
 		{
@@ -3022,8 +3006,6 @@ namespace UnrealBuildTool
 			EngineProject.AddFilesToProject(SourceFileSearch.FindFiles(BuildDirectory, SubdirectoryNamesToExclude), Unreal.EngineDirectory);
 		}
 
-
-
 		/// Adds engine documentation to the specified project
 		protected void AddEngineDocumentation(ProjectFile EngineProject, ILogger Logger)
 		{
@@ -3083,9 +3065,6 @@ namespace UnrealBuildTool
 			}
 		}
 
-
-
-
 		/// <summary>
 		/// Adds a new project file and returns an object that represents that project file (or if the project file is already known, returns that instead.)
 		/// </summary>
@@ -3122,7 +3101,6 @@ namespace UnrealBuildTool
 			return NewProjectFile;
 		}
 
-
 		/// <summary>
 		/// Allocates a generator-specific project file object
 		/// </summary>
@@ -3131,7 +3109,6 @@ namespace UnrealBuildTool
 		/// <returns>The newly allocated project file object</returns>
 		protected abstract ProjectFile AllocateProjectFile(FileReference InitFilePath, DirectoryReference BaseDir);
 
-
 		/// <summary>
 		/// Returns a list of all the known project files
 		/// </summary>
@@ -3147,7 +3124,6 @@ namespace UnrealBuildTool
 			}
 		}
 
-
 		/// <summary>
 		/// Writes the project files to disk
 		/// </summary>
@@ -3200,7 +3176,6 @@ namespace UnrealBuildTool
 		/// <returns>True if successful</returns>
 		protected abstract bool WritePrimaryProjectFile(ProjectFile? UBTProject, PlatformProjectGeneratorCollection PlatformProjectGenerators, ILogger Logger);
 
-
 		/// <summary>
 		/// Writes any additional solution-wide debug files (e.g. UnrealVS hints)
 		/// </summary>
@@ -3221,8 +3196,6 @@ namespace UnrealBuildTool
 		{
 		}
 
-
-
 		/// <summary>
 		/// Writes the specified string content to a file.  Before writing to the file, it loads the existing file (if present) to see if the contents have changed
 		/// </summary>
@@ -3248,7 +3221,6 @@ namespace UnrealBuildTool
 				}
 			}
 
-
 			// Don't bother saving anything out if the new file content is the same as the old file's content
 			bool FileNeedsSave = true;
 			if (LoadedFileContent != null)
@@ -3408,7 +3380,6 @@ namespace UnrealBuildTool
 								bFoundDebugConfig = true;
 							}
 						}
-
 					}
 				}
 				catch
diff --git a/Engine/Source/Programs/UnrealBuildTool/ProjectFiles/QMake/QMakefileGenerator.cs b/Engine/Source/Programs/UnrealBuildTool/ProjectFiles/QMake/QMakefileGenerator.cs
--- a/Engine/Source/Programs/UnrealBuildTool/ProjectFiles/QMake/QMakefileGenerator.cs
+++ b/Engine/Source/Programs/UnrealBuildTool/ProjectFiles/QMake/QMakefileGenerator.cs
@@ -147,7 +147,6 @@ namespace UnrealBuildTool
 				{
 					AddIncludeDirectory(ref SystemIncludeDirectories, CurPath, Path.GetDirectoryName(QMakeProject.ProjectFilePath.FullName)!);
 				}
-
 			}
 
 			// Remove duplicate paths from include dir and system include dir list
@@ -189,7 +188,7 @@ namespace UnrealBuildTool
 							// If the value contains spaces, we must put the whole value inside quotes
 							value = value.Replace("\"", "\\\\\\\"");
 							bool hasSpaces = false;
-							if (value.Contains(" "))
+							if (value.Contains(' '))
 							{
 								value = value.Replace(" ", "\\ ");
 								hasSpaces = true;
@@ -372,7 +371,6 @@ namespace UnrealBuildTool
 						}
 					}
 				}
-
 			}
 
 			// Add section end to section strings;
diff --git a/Engine/Source/Programs/UnrealBuildTool/ProjectFiles/Rider/RiderProjectFile.cs b/Engine/Source/Programs/UnrealBuildTool/ProjectFiles/Rider/RiderProjectFile.cs
--- a/Engine/Source/Programs/UnrealBuildTool/ProjectFiles/Rider/RiderProjectFile.cs
+++ b/Engine/Source/Programs/UnrealBuildTool/ProjectFiles/Rider/RiderProjectFile.cs
@@ -145,7 +145,7 @@ namespace UnrealBuildTool
 
 		private bool IsPlatformInHostGroup(UnrealTargetPlatform Platform)
 		{
-			var Groups = UEBuildPlatform.GetPlatformGroups(BuildHostPlatform.Current.Platform);
+			IEnumerable<UnrealPlatformGroup> Groups = UEBuildPlatform.GetPlatformGroups(BuildHostPlatform.Current.Platform);
 			foreach (UnrealPlatformGroup Group in Groups)
 			{
 				// Desktop includes Linux, Mac and Windows.
@@ -539,7 +539,7 @@ namespace UnrealBuildTool
 			else if (UEBuildPlatform.IsPlatformInGroup(Target.Platform, UnrealPlatformGroup.Linux) ||
 					UEBuildPlatform.IsPlatformInGroup(Target.Platform, UnrealPlatformGroup.Unix))
 			{
-				var EngineDirectory = Unreal.EngineDirectory.ToString();
+				string EngineDirectory = Unreal.EngineDirectory.ToString();
 
 				string? UseLibcxxEnvVarOverride = Environment.GetEnvironmentVariable("UE_LINUX_USE_LIBCXX");
 				// assumes a single architecture
@@ -571,7 +571,7 @@ namespace UnrealBuildTool
 				}
 
 				string PlatformSdkVersionString = UEBuildPlatformSDK.GetSDKForPlatform(BuildPlatform.GetPlatformName())!.GetInstalledVersion()!;
-				var Version = GetLinuxToolchainVersionFromFullString(PlatformSdkVersionString);
+				string Version = GetLinuxToolchainVersionFromFullString(PlatformSdkVersionString);
 
 				string? InternalSdkPath = UEBuildPlatform.GetSDK(UnrealTargetPlatform.Linux)!.GetInternalSDKPath();
 				if (InternalSdkPath != null)
@@ -706,25 +706,24 @@ namespace UnrealBuildTool
 			}
 
 			string Name = Field.Item1;
-			if (Field.Item2 is bool)
+			if (Field.Item2 is bool vbool)
 			{
-				Writer.WriteValue(Name, (bool)Field.Item2);
+				Writer.WriteValue(Name, vbool);
 			}
-			else if (Field.Item2 is string)
+			else if (Field.Item2 is string vstring)
 			{
-				string FieldValue = (string)Field.Item2;
-				if (FieldValue != "")
+				if (!String.IsNullOrEmpty(vstring))
 				{
-					Writer.WriteValue(Name, (string)Field.Item2);
+					Writer.WriteValue(Name, vstring);
 				}
 			}
-			else if (Field.Item2 is int)
+			else if (Field.Item2 is int vint)
 			{
-				Writer.WriteValue(Name, (int)Field.Item2);
+				Writer.WriteValue(Name, vint);
 			}
-			else if (Field.Item2 is double)
+			else if (Field.Item2 is double vdouble)
 			{
-				Writer.WriteValue(Name, (double)Field.Item2);
+				Writer.WriteValue(Name, vdouble);
 			}
 			else if (Field.Item2 is CppStandardVersion version)
 			{
@@ -752,9 +751,8 @@ namespace UnrealBuildTool
 			{
 				Writer.WriteValue(Name, Field.Item2.ToString());
 			}
-			else if (Field.Item2 is IEnumerable<string>)
+			else if (Field.Item2 is IEnumerable<string> FieldValue)
 			{
-				IEnumerable<string> FieldValue = (IEnumerable<string>)Field.Item2;
 				if (FieldValue.Any())
 				{
 					Writer.WriteStringArrayField(Name, FieldValue);
@@ -865,7 +863,7 @@ namespace UnrealBuildTool
 
 			private void CalculateSystemIncludePaths(string SDKPath)
 			{
-				if (CurrentlyProcessedSDK != String.Empty)
+				if (!String.IsNullOrEmpty(CurrentlyProcessedSDK))
 				{
 					throw new InvalidOperationException("Cannot calculate include paths for several platforms at once");
 				}
diff --git a/Engine/Source/Programs/UnrealBuildTool/ProjectFiles/Rider/ToolchainInfo.cs b/Engine/Source/Programs/UnrealBuildTool/ProjectFiles/Rider/ToolchainInfo.cs
--- a/Engine/Source/Programs/UnrealBuildTool/ProjectFiles/Rider/ToolchainInfo.cs
+++ b/Engine/Source/Programs/UnrealBuildTool/ProjectFiles/Rider/ToolchainInfo.cs
@@ -100,14 +100,14 @@ namespace UnrealBuildTool
 				return true;
 			}
 
-			return Obj is ToolchainInfo && Equals((ToolchainInfo)Obj);
+			return Obj is ToolchainInfo info && Equals(info);
 		}
 
 		public override int GetHashCode()
 		{
 			unchecked
 			{
-				var HashCode = (int)CppStandard;
+				int HashCode = (int)CppStandard;
 				HashCode = (HashCode * 397) ^ bUseRTTI.GetHashCode();
 				HashCode = (HashCode * 397) ^ bEnableExceptions.GetHashCode();
 				HashCode = (HashCode * 397) ^ bIsBuildingLibrary.GetHashCode();
diff --git a/Engine/Source/Programs/UnrealBuildTool/ProjectFiles/VisualStudio/VCProject.cs b/Engine/Source/Programs/UnrealBuildTool/ProjectFiles/VisualStudio/VCProject.cs
--- a/Engine/Source/Programs/UnrealBuildTool/ProjectFiles/VisualStudio/VCProject.cs
+++ b/Engine/Source/Programs/UnrealBuildTool/ProjectFiles/VisualStudio/VCProject.cs
@@ -124,7 +124,6 @@ namespace UnrealBuildTool
 			}
 		}
 
-
 		/// <summary>
 		/// Attempts to load the project's GUID from an existing project file on disk
 		/// </summary>
@@ -199,7 +198,6 @@ namespace UnrealBuildTool
 				}
 			}
 
-
 			UEBuildPlatform? BuildPlatform;
 			if (!UEBuildPlatform.TryGetBuildPlatform(Platform, out BuildPlatform))
 			{
@@ -211,7 +209,6 @@ namespace UnrealBuildTool
 				return false;
 			}
 
-
 			List<UnrealTargetConfiguration> SupportedConfigurations = new List<UnrealTargetConfiguration>();
 			List<UnrealTargetPlatform> SupportedPlatforms = new List<UnrealTargetPlatform>();
 			if (ProjectTarget.TargetRules != null)
@@ -752,7 +749,6 @@ namespace UnrealBuildTool
 			}
 		}
 
-
 		/// <summary>
 		/// If found writes a debug project file to disk
 		/// </summary>
@@ -766,7 +762,6 @@ namespace UnrealBuildTool
 
 			BuildProjectConfigAndTargetCombinations(InPlatforms, InConfigurations, PlatformProjectGenerators, Logger);
 
-
 			foreach (ProjectConfigAndTargetCombination Combination in ProjectConfigAndTargetCombinations!)
 			{
 				if (Combination.Platform != null && !ProjectPlatforms.Contains(Combination.Platform.Value))
@@ -1384,9 +1379,9 @@ namespace UnrealBuildTool
 					{
 						int ValueIndex = 0;
 						Dictionary<string, string> UpdatedValues = new();
-						var KVPList = DirectoryToIncludeSearchPaths.ToList();
+						List<KeyValuePair<DirectoryReference, string>> KVPList = DirectoryToIncludeSearchPaths.ToList();
 						KVPList.SortBy(kvp => kvp.Key);
-						foreach (var DirectoryKVP in KVPList)
+						foreach (KeyValuePair<DirectoryReference, string> DirectoryKVP in KVPList)
 						{
 							string? Value = DirectoryKVP.Value;
 							if (!String.IsNullOrEmpty(Value))
@@ -1410,9 +1405,9 @@ namespace UnrealBuildTool
 					{
 						int ValueIndex = 0;
 						Dictionary<string, string> UpdatedValues = new();
-						var KVPList = DirectoryToForceIncludePaths.ToList();
+						List<KeyValuePair<DirectoryReference, string>> KVPList = DirectoryToForceIncludePaths.ToList();
 						KVPList.SortBy(kvp => kvp.Key);
-						foreach (var DirectoryKVP in KVPList)
+						foreach (KeyValuePair<DirectoryReference, string> DirectoryKVP in KVPList)
 						{
 							string? Value = DirectoryKVP.Value;
 							if (!String.IsNullOrEmpty(Value))
@@ -1436,9 +1431,9 @@ namespace UnrealBuildTool
 					{
 						int ValueIndex = 0;
 						Dictionary<string, string> UpdatedValues = new();
-						var KVPList = DirectoryToPchFile.ToList();
+						List<KeyValuePair<DirectoryReference, string?>> KVPList = DirectoryToPchFile.ToList();
 						KVPList.SortBy(kvp => kvp.Key);
-						foreach (var DirectoryKVP in KVPList)
+						foreach (KeyValuePair<DirectoryReference, string?> DirectoryKVP in KVPList)
 						{
 							string? Value = DirectoryKVP.Value;
 							if (!String.IsNullOrEmpty(Value))
@@ -1654,7 +1649,6 @@ namespace UnrealBuildTool
 				bSuccess = ProjectFileGenerator.WriteFileIfChanged(ProjectFilePath.FullName, VCProjectFileContent.ToString(), Logger);
 			}
 
-
 			// Save the filters file
 			if (bSuccess)
 			{
@@ -2096,7 +2090,6 @@ namespace UnrealBuildTool
 						VCProjectFileContent.Append(PathsStringBuilder.ToString());
 					}
 
-
 					// This is the standard UE based project NMake build line:
 					//	..\..\Build\BatchFiles\Build.bat <TARGETNAME> <PLATFORM> <CONFIGURATION>
 					//	ie ..\..\Build\BatchFiles\Build.bat BlankProgram Win64 Debug
@@ -2139,7 +2132,6 @@ namespace UnrealBuildTool
 		}
 	}
 
-
 	/// <summary>
 	/// A Visual C# project.
 	/// </summary>
@@ -2352,5 +2344,4 @@ namespace UnrealBuildTool
 		/// Cache of parsed info about this project
 		protected readonly Dictionary<UnrealTargetConfiguration, CsProjectInfo> CachedProjectInfo = new Dictionary<UnrealTargetConfiguration, CsProjectInfo>();
 	}
-
 }
diff --git a/Engine/Source/Programs/UnrealBuildTool/ProjectFiles/VisualStudioCode/VSCodeProjectFileGenerator.cs b/Engine/Source/Programs/UnrealBuildTool/ProjectFiles/VisualStudioCode/VSCodeProjectFileGenerator.cs
--- a/Engine/Source/Programs/UnrealBuildTool/ProjectFiles/VisualStudioCode/VSCodeProjectFileGenerator.cs
+++ b/Engine/Source/Programs/UnrealBuildTool/ProjectFiles/VisualStudioCode/VSCodeProjectFileGenerator.cs
@@ -123,7 +123,7 @@ namespace UnrealBuildTool
 		{
 			string Processed = CommonMakePathString(InRef, InPathType, InRelativeRoot);
 
-			if (Processed.Contains(" "))
+			if (Processed.Contains(' '))
 			{
 				if (HostPlatform == UnrealTargetPlatform.Win64 && InQuoteType == EQuoteType.Double)
 				{
@@ -267,7 +267,7 @@ namespace UnrealBuildTool
 			if (bForeignProject)
 			{
 				Projects = new List<ProjectFile>();
-				foreach (var Project in AllProjectFiles)
+				foreach (ProjectFile Project in AllProjectFiles)
 				{
 					if (GameProjectName == Project.ProjectFilePath.GetFileNameWithoutAnyExtensions())
 					{
@@ -476,7 +476,6 @@ namespace UnrealBuildTool
 			public List<Project> AllProjects = new List<Project>();
 		}
 
-
 		private ProjectData GatherProjectData(List<ProjectFile> InProjects, ILogger Logger)
 		{
 			ProjectData ProjectData = new ProjectData();
@@ -532,7 +531,6 @@ namespace UnrealBuildTool
 				{
 					VCSharpProjectFile VCSharpProject = (VCSharpProjectFile)Project;
 
-
 					string ProjectName = Project.ProjectFilePath.GetFileNameWithoutExtension();
 
 					ProjectData.Target Target = new ProjectData.Target(NewProject, ProjectName, TargetType.Program);
@@ -847,7 +845,6 @@ namespace UnrealBuildTool
 				}
 				OutFile.EndObject();
 			}
-
 		}
 
 		private void WriteCompileCommands(FileReference CompileCommandsFile, IEnumerable<FileReference> SourceFiles,
@@ -944,7 +941,6 @@ namespace UnrealBuildTool
 								}
 							}
 
-
 							OutFile.BeginArray("args");
 							{
 								OutFile.AddUnnamedField(Target.Name);
@@ -1174,7 +1170,7 @@ namespace UnrealBuildTool
 
 				string ProjectName = Target.Name;
 				// hyphens not allowed so change them to underscores in project name
-				if (ProjectName.Contains("-"))
+				if (ProjectName.Contains('-'))
 				{
 					Trace.TraceWarning("Project name contained hyphens, converted to underscore");
 					ProjectName = ProjectName.Replace("-", "_");
@@ -1317,7 +1313,6 @@ namespace UnrealBuildTool
 									}
 								}
 							}
-
 						}
 						OutFile.EndArray();
 					}
diff --git a/Engine/Source/Programs/UnrealBuildTool/ProjectFiles/Xcode/XcodeProject.File.cs b/Engine/Source/Programs/UnrealBuildTool/ProjectFiles/Xcode/XcodeProject.File.cs
--- a/Engine/Source/Programs/UnrealBuildTool/ProjectFiles/Xcode/XcodeProject.File.cs
+++ b/Engine/Source/Programs/UnrealBuildTool/ProjectFiles/Xcode/XcodeProject.File.cs
@@ -113,7 +113,6 @@ namespace UnrealBuildTool.XcodeProjectXcconfig
 		}
 	}
 
-
 	class XcodeFileCollection
 	{
 		static readonly string[] ResourceExtensions =
@@ -161,7 +160,6 @@ namespace UnrealBuildTool.XcodeProjectXcconfig
 
 		const string FileTypeDefault = "file.txt";
 
-
 		class ManualFileReference
 		{
 			public string Guid = XcodeProjectFileGenerator.MakeXcodeGuid();
@@ -189,7 +187,6 @@ namespace UnrealBuildTool.XcodeProjectXcconfig
 		private List<ManualFileReference> ManualFiles = new();
 		private Dictionary<string, string> GuidsForGroups = new();
 
-
 		public XcodeFileCollection(XcodeProjectFile ProjectFile)
 		{
 			ProjectDirectory = ProjectFile.ProjectFilePath.Directory;
@@ -213,7 +210,6 @@ namespace UnrealBuildTool.XcodeProjectXcconfig
 			{
 				Groups.Add(Key, new XcodeFileGroup(Other.Groups[Key]));
 			}
-
 		}
 
 		public void SetUProjectLocation(FileReference? UProjectLocation)
@@ -236,7 +232,7 @@ namespace UnrealBuildTool.XcodeProjectXcconfig
 					{
 						if (!SourceToBuildFileMap.TryGetValue(SourceDir, out BuildFileOffset))
 						{
-							while (SourceDir != null && DirectoryReference.EnumerateFiles(SourceDir, "*.Build.cs").Count() == 0)
+							while (SourceDir != null && !DirectoryReference.EnumerateFiles(SourceDir, "*.Build.cs").Any())
 							{
 								SourceDir = SourceDir.ParentDirectory!;
 							}
@@ -272,7 +268,7 @@ namespace UnrealBuildTool.XcodeProjectXcconfig
 				}
 			}
 
-			if (GroupName != "")
+			if (!String.IsNullOrEmpty(GroupName))
 			{
 				// group the files by path
 				XcodeFileGroup? Group = GroupName == null ? FindGroupByAbsolutePath(File.Reference.Directory) : null;
@@ -343,7 +339,6 @@ namespace UnrealBuildTool.XcodeProjectXcconfig
 			}
 		}
 
-
 		public void Write(StringBuilder Content)
 		{
 			Content.WriteLine("/* Begin PBXBuildFile section */");
@@ -365,7 +360,6 @@ namespace UnrealBuildTool.XcodeProjectXcconfig
 			Content.WriteLine("/* End PBXBuildFile section */");
 			Content.WriteLine();
 
-
 			Content.WriteLine("/* Begin PBXFileReference section */");
 			foreach (XcodeSourceFile File in AllFiles)
 			{
@@ -401,12 +395,11 @@ namespace UnrealBuildTool.XcodeProjectXcconfig
 
 			Content.WriteLine("/* Begin PBXGroup section */");
 
-
 			// write main/root group and it's children (adding the manual groups to the main group, which is the true param here0
 			WriteGroup(Content, RootGroup, bIsRootGroup: true);
 
 			// write some manual groups
-			foreach (var Pair in GuidsForGroups)
+			foreach (KeyValuePair<string, string> Pair in GuidsForGroups)
 			{
 				Content.WriteLine($"\t\t{Pair.Value} /* {Pair.Key} */ = {{");
 				Content.WriteLine("\t\t\tisa = PBXGroup;");
@@ -425,7 +418,6 @@ namespace UnrealBuildTool.XcodeProjectXcconfig
 			Content.WriteLine("/* End PBXGroup section */");
 		}
 
-
 		private void WriteGroup(StringBuilder Content, XcodeFileGroup Group, bool bIsRootGroup)
 		{
 			if (!Group.bIsReference)
@@ -444,7 +436,7 @@ namespace UnrealBuildTool.XcodeProjectXcconfig
 
 				if (bIsRootGroup)
 				{
-					foreach (var Pair in GuidsForGroups)
+					foreach (KeyValuePair<string, string> Pair in GuidsForGroups)
 					{
 						Content.WriteLine($"\t\t\t\t{Pair.Value} /* {Pair.Key} */,");
 					}
@@ -494,7 +486,6 @@ namespace UnrealBuildTool.XcodeProjectXcconfig
 			return null;
 		}
 
-
 		/// <summary>
 		/// Gets Xcode file type based on its extension
 		/// </summary>
@@ -572,7 +563,6 @@ namespace UnrealBuildTool.XcodeProjectXcconfig
 						return null;
 					}
 
-
 					CurrentGroup = new XcodeFileGroup(Path.GetFileName(CurrentPath), CurrentPath, GetFileType(Path.GetExtension(CurrentPath)).StartsWith("folder"));
 					if (Groups.Count == 0)
 					{
@@ -606,5 +596,4 @@ namespace UnrealBuildTool.XcodeProjectXcconfig
 			return InPath.Replace("\\", "/");
 		}
 	}
-
 }
diff --git a/Engine/Source/Programs/UnrealBuildTool/ProjectFiles/Xcode/XcodeProject.Scheme.cs b/Engine/Source/Programs/UnrealBuildTool/ProjectFiles/Xcode/XcodeProject.Scheme.cs
--- a/Engine/Source/Programs/UnrealBuildTool/ProjectFiles/Xcode/XcodeProject.Scheme.cs
+++ b/Engine/Source/Programs/UnrealBuildTool/ProjectFiles/Xcode/XcodeProject.Scheme.cs
@@ -7,7 +7,6 @@ using System.Linq;
 using System.Text;
 using EpicGames.Core;
 
-
 namespace UnrealBuildTool.XcodeProjectXcconfig
 {
 	static class XcodeSchemeFile
@@ -166,7 +165,7 @@ namespace UnrealBuildTool.XcodeProjectXcconfig
 				Content.WriteLine("      </BuildableProductRunnable>");
 				if (String.IsNullOrEmpty(OldCommandLineArguments))
 				{
-					if (bHasEditorConfiguration && GameProjectPath != "")
+					if (bHasEditorConfiguration && !String.IsNullOrEmpty(GameProjectPath))
 					{
 						Content.WriteLine("      <CommandLineArguments>");
 						if (UnrealData.bIsForeignProject)
diff --git a/Engine/Source/Programs/UnrealBuildTool/ProjectFiles/Xcode/XcodeProject.Utils.cs b/Engine/Source/Programs/UnrealBuildTool/ProjectFiles/Xcode/XcodeProject.Utils.cs
--- a/Engine/Source/Programs/UnrealBuildTool/ProjectFiles/Xcode/XcodeProject.Utils.cs
+++ b/Engine/Source/Programs/UnrealBuildTool/ProjectFiles/Xcode/XcodeProject.Utils.cs
@@ -9,7 +9,6 @@ using System.Text.RegularExpressions;
 using EpicGames.Core;
 using Microsoft.Extensions.Logging;
 
-
 namespace UnrealBuildTool.XcodeProjectXcconfig
 {
 	static class StringBuilderExtensions
@@ -41,7 +40,6 @@ namespace UnrealBuildTool.XcodeProjectXcconfig
 			return DirectoryReference.Combine(ProjectAsFile.Directory, $"{ProjectAsFile.GetFileNameWithoutExtension()}{Suffix(Platform)}{ProjectAsFile.GetExtension()}");
 		}
 
-
 		private static IEnumerable<UnrealTargetConfiguration> GetSupportedConfigurations()
 		{
 			return new UnrealTargetConfiguration[] {
@@ -78,7 +76,6 @@ namespace UnrealBuildTool.XcodeProjectXcconfig
 			return false;
 		}
 
-
 		public static string MakeExecutableFileName(string BinaryName, UnrealTargetPlatform Platform, UnrealTargetConfiguration Configuration, UnrealArchitectures Architectures, UnrealTargetConfiguration UndecoratedConfiguration)
 		{
 			StringBuilder Result = new StringBuilder();
@@ -98,7 +95,6 @@ namespace UnrealBuildTool.XcodeProjectXcconfig
 			return Result.ToString();
 		}
 
-
 		// cache for the below function
 		static Dictionary<string, UnrealArchitectures> CachedMacProjectArchitectures = new();
 
@@ -140,7 +136,6 @@ namespace UnrealBuildTool.XcodeProjectXcconfig
 			return Arches;
 		}
 
-
 		public static void FindPlistId(MetadataItem PlistItem, string Key, ref string? BundleId)
 		{
 			if (PlistItem.File == null || !FileReference.Exists(PlistItem.File))
@@ -151,7 +146,7 @@ namespace UnrealBuildTool.XcodeProjectXcconfig
 			string Identifier = Plist($"Print :{Key}", PlistItem.File.FullName);
 
 			// handle error
-			if (Identifier == "" || Identifier.StartsWith("Print:"))
+			if (String.IsNullOrEmpty(Identifier) || Identifier.StartsWith("Print:"))
 			{
 				if (PlistItem.Mode == MetadataMode.UsePremade)
 				{
@@ -186,7 +181,7 @@ namespace UnrealBuildTool.XcodeProjectXcconfig
 		{
 			string AddOutput = Plist($"Add {Entry} {Type} {Value}");
 			// error will be non-empty string
-			if (AddOutput != "")
+			if (!String.IsNullOrEmpty(AddOutput))
 			{
 				Plist($"Set {Entry} {Value}");
 			}
@@ -256,9 +251,6 @@ namespace UnrealBuildTool.XcodeProjectXcconfig
 			return null;
 		}
 
-
-
-
 		public static List<string> GetSupportedOrientations(ConfigHierarchy Ini)
 		{
 			List<string> Orientations = new();
diff --git a/Engine/Source/Programs/UnrealBuildTool/ProjectFiles/Xcode/XcodeProject.cs b/Engine/Source/Programs/UnrealBuildTool/ProjectFiles/Xcode/XcodeProject.cs
--- a/Engine/Source/Programs/UnrealBuildTool/ProjectFiles/Xcode/XcodeProject.cs
+++ b/Engine/Source/Programs/UnrealBuildTool/ProjectFiles/Xcode/XcodeProject.cs
@@ -155,16 +155,15 @@ namespace UnrealBuildTool.XcodeProjectXcconfig
 		//	XcodeProjectRelative = File.MakeRelativeTo(XcodeProjectFile.ParentDirectory!);
 		//}
 
-
 		// Location: Can be key of setting entry in .ini, or the full file path
 		// CopyFromFolderIfNotFound: If the file at "Location" does not exist, try to copy the same named file from this folder
 		public MetadataItem(DirectoryReference ProductDirectory, DirectoryReference XcodeProject, ConfigHierarchy Ini, string Location, MetadataMode InMode, DirectoryReference CopyFromFolderIfNotFound)
 		{
 			// no extension means it's a .ini entry
-			if (Path.GetExtension(Location) == "")
+			if (String.IsNullOrEmpty(Path.GetExtension(Location)))
 			{
 				string? FileLocation;
-				if (Ini.TryGetValue("/Script/MacTargetPlatform.XcodeProjectSettings", Location, out FileLocation) && FileLocation != "")
+				if (Ini.TryGetValue("/Script/MacTargetPlatform.XcodeProjectSettings", Location, out FileLocation) && !String.IsNullOrEmpty(FileLocation))
 				{
 					if (FileLocation.StartsWith("(FilePath=", StringComparison.OrdinalIgnoreCase))
 					{
@@ -242,7 +241,6 @@ namespace UnrealBuildTool.XcodeProjectXcconfig
 		public Dictionary<MetadataPlatform, MetadataItem> EntitlementsFiles = new();
 		public Dictionary<MetadataPlatform, MetadataItem> ShippingEntitlementsFiles = new();
 
-
 		public Metadata(DirectoryReference ProductDirectory, DirectoryReference XcodeProject, ConfigHierarchy Ini, bool bSupportsMac, bool bSupportsIOSOrTVOS, ILogger Logger)
 		{
 			DirectoryReference ResourceFolder = DirectoryReference.Combine(Unreal.EngineDirectory, "Build/Mac/Resources/");
@@ -327,7 +325,6 @@ namespace UnrealBuildTool.XcodeProjectXcconfig
 		}
 	}
 
-
 	class UnrealData
 	{
 		public bool bIsStubProject;
@@ -406,7 +403,6 @@ namespace UnrealBuildTool.XcodeProjectXcconfig
 			return Platform == null || XcodeProjectFileGenerator.XcodePlatforms.Contains((UnrealTargetPlatform)Platform);
 		}
 
-
 		public UnrealData(FileReference XcodeProjectFileLocation, bool bIsForDistribution, string BundleID, string AppName, bool bMakeProjectPerTarget)
 		{
 			// the .xcodeproj is actually a directory
@@ -747,7 +743,6 @@ namespace UnrealBuildTool.XcodeProjectXcconfig
 		/// <param name="Content"></param>
 		public abstract void Write(StringBuilder Content);
 
-
 		/// <summary>
 		/// Walks the references of the given node to find all nodes of the given type. 
 		/// </summary>
@@ -776,7 +771,6 @@ namespace UnrealBuildTool.XcodeProjectXcconfig
 			return Return.OfType<T>();
 		}
 
-
 		public void CreateXcconfigFile(XcodeProject Project, UnrealTargetPlatform? Platform, string Name)
 		{
 			DirectoryReference XcodeProjectDirectory = Project.UnrealData.XcodeProjectFileLocation.ParentDirectory!;
@@ -789,7 +783,6 @@ namespace UnrealBuildTool.XcodeProjectXcconfig
 
 		}
 
-
 		/// <summary>
 		/// THhis will walk the node reference tree and call WRite on each node to add all needed nodes to the xcode poject file
 		/// </summary>
@@ -818,7 +811,6 @@ namespace UnrealBuildTool.XcodeProjectXcconfig
 		}
 	}
 
-
 	class XcodeDependency : XcodeProjectNode
 	{
 		public XcodeTarget Target;
@@ -826,7 +818,6 @@ namespace UnrealBuildTool.XcodeProjectXcconfig
 		public string ProxyGuid = XcodeProjectFileGenerator.MakeXcodeGuid();
 		public string ProjectGuid;
 
-
 		public XcodeDependency(XcodeTarget Target, string ProjectGuid)
 		{
 			this.Target = Target;
@@ -1009,8 +1000,6 @@ namespace UnrealBuildTool.XcodeProjectXcconfig
 		}
 	}
 
-
-
 	class XcodeBuildConfig : XcodeProjectNode
 	{
 		public string Guid = XcodeProjectFileGenerator.MakeXcodeGuid();
@@ -1284,7 +1273,6 @@ namespace UnrealBuildTool.XcodeProjectXcconfig
 			}
 		}
 
-
 		/// <summary>
 		/// Write some scripts to do some fixup with how UBT links files. There is currently a difference between Mac and IOS/TVOS:
 		/// All:
@@ -1468,7 +1456,6 @@ namespace UnrealBuildTool.XcodeProjectXcconfig
 			{
 				ResourcesBuildPhase.AddResource(new FileReference(StoryboardPath));
 			}
-
 		}
 		protected void ProcessFrameworks(XcodeResourcesBuildPhase ResourcesBuildPhase, XcodeProject Project, XcodeProjectFile ProjectFile, ILogger Logger)
 		{
@@ -1623,7 +1610,6 @@ namespace UnrealBuildTool.XcodeProjectXcconfig
 			PlatformIni.TryGetValue("/Script/MacTargetPlatform.XcodeProjectSettings", "ModernSigningPrefix", out SigningPrefix);
 			PlatformIni.TryGetValue("/Script/MacTargetPlatform.XcodeProjectSettings", "AppCategory", out AppCategory);
 
-
 			if (Platform == UnrealTargetPlatform.Mac)
 			{
 				// editor vs game metadata
@@ -1686,7 +1672,6 @@ namespace UnrealBuildTool.XcodeProjectXcconfig
 				}
 			}
 
-
 			// get metadata for the platform set above
 			MetadataItem PlistMetadata = UnrealData.Metadata!.PlistFiles[MetadataPlatform];
 			// now pull the bundle id's out, as xcode will warn if they don't match (this has to happen after each platform set bundle id above)
@@ -1694,7 +1679,6 @@ namespace UnrealBuildTool.XcodeProjectXcconfig
 			// if the user had a ini setting, but also set it in the template plist, the plist one should win (this is only used if we are updating a template plist)
 			XcodeUtils.FindPlistId(PlistMetadata, "LSApplicationCategoryType", ref AppCategory);
 
-
 			// include another xcconfig for versions that UBT writes out
 			Xcconfig!.AppendLine($"#include? \"{ProjectOrEngineDir}/Intermediate/Build/Versions.xcconfig\"");
 			if (UnrealData.UProjectFileLocation == null)
@@ -1756,7 +1740,6 @@ namespace UnrealBuildTool.XcodeProjectXcconfig
 			Xcconfig.AppendLine($"PRODUCT_BUNDLE_IDENTIFIER = {BundleIdentifier}");
 			Xcconfig.AppendLine($"{DeploymentTargetKey} = {DeploymentTarget}");
 
-
 			Xcconfig.AppendLine("");
 			Xcconfig.AppendLine("// Plist settings");
 			Xcconfig.AppendLine($"INFOPLIST_FILE = {PlistMetadata.XcodeProjectRelative}");
@@ -1791,8 +1774,6 @@ namespace UnrealBuildTool.XcodeProjectXcconfig
 
 			Xcconfig.Write();
 
-
-
 			// Now for each config write out the specific settings
 			DirectoryReference? GameDir = UnrealData.UProjectFileLocation?.Directory;
 			string? GamePath = GameDir != null ? XcodeFileCollection.ConvertPath(GameDir.FullName) : null;
@@ -1823,7 +1804,6 @@ namespace UnrealBuildTool.XcodeProjectXcconfig
 				string ProductName = ExecutableName;
 				string ExecutableKey = $"UE_{Platform.ToString().ToUpper()}_EXECUTABLE_NAME";
 
-
 				MetadataItem EntitlementsMetadata = UnrealData.Metadata!.EntitlementsFiles[MetadataPlatform];
 
 				// if we have shipping specified, use it instead
@@ -1887,7 +1867,6 @@ namespace UnrealBuildTool.XcodeProjectXcconfig
 			// hook up each buildconfig to this Xcconfig
 			BuildConfigList.BuildConfigs.ForEach(x => x.Xcconfig = Xcconfig);
 
-
 			// add all of the files to be natively compiled by this target
 			XcodeSourcesBuildPhase SourcesBuildPhase = new XcodeSourcesBuildPhase();
 			BuildPhases.Add(SourcesBuildPhase);
@@ -2052,7 +2031,6 @@ namespace UnrealBuildTool.XcodeProjectXcconfig
 		/// </summary>
 		public Dictionary<XcodeProject, UnrealTargetPlatform?> RootProjects = new();
 
-
 		private XcodeProjectLegacy.XcodeProjectFile? LegacyProjectFile = null;
 		private bool bHasCheckedForLegacy = false;
 		public bool bHasLegacyProject => LegacyProjectFile != null;
@@ -2160,7 +2138,6 @@ namespace UnrealBuildTool.XcodeProjectXcconfig
 			}
 		}
 
-
 		public FileReference ProjectFilePathForPlatform(UnrealTargetPlatform? Platform)
 		{
 			return new FileReference(XcodeUtils.ProjectDirPathForPlatform(UnrealData.XcodeProjectFileLocation, Platform).FullName);
@@ -2180,7 +2157,6 @@ namespace UnrealBuildTool.XcodeProjectXcconfig
 				return true;
 			}
 
-
 			ConditionalCreateLegacyProject();
 
 			if (LegacyProjectFile != null)
@@ -2203,7 +2179,6 @@ namespace UnrealBuildTool.XcodeProjectXcconfig
 			UnrealData.bIsMergingProjects = FileReference.Exists(TemplateProject);
 			UnrealData.bWriteCodeSigningSettings = !UnrealData.bIsMergingProjects;
 
-
 			// turn all UE files into internal representation
 			ProcessSourceFiles();
 
@@ -2234,7 +2209,6 @@ namespace UnrealBuildTool.XcodeProjectXcconfig
 					// write metadata now so we can add them to the FileCollection
 					ConditionalWriteMetadataFiles(UnrealTargetPlatform.IOS);
 
-
 					StringBuilder Content = new StringBuilder();
 
 					Content.WriteLine(0, "// !$*UTF8*$!");
@@ -2341,7 +2315,6 @@ namespace UnrealBuildTool.XcodeProjectXcconfig
 			FileReference ImportFile = FileReference.Combine(PBXProjFilePath.Directory, "import.plist");
 			File.WriteAllText(ImportFile.FullName, Content.ToString());
 
-
 			// cache some standard guids from the template project
 			string ProjectGuid = Plist($"Print :rootObject");
 			string TemplateMainGroupGuid = Plist($"Print :objects:{ProjectGuid}:mainGroup");
@@ -2371,7 +2344,6 @@ namespace UnrealBuildTool.XcodeProjectXcconfig
 			// and import it into the template
 			Plist($"Merge \"{ImportFile.FullName}\" :objects");
 
-
 			// get all the targets in the template that are application types
 			IEnumerable<string> AppTargetGuids = XcodeUtils.PlistArray($":objects:{ProjectGuid}:targets")
 				.Where(TargetGuid => (Plist($"Print :objects:{TargetGuid}:productType") == "com.apple.product-type.application"));
@@ -2391,7 +2363,6 @@ namespace UnrealBuildTool.XcodeProjectXcconfig
 				Plist($"Add :objects:{ProjectGuid}:targets:0 string {Target.Guid}");
 			}
 
-
 			// hook up Xcconfig files to the project and the project configs
 			// @todo how to manage with conflicts already present...
 			//PlistSetAdd($":objects:{ProjectGuid}:baseConfigurationReference", RootProject.Xcconfig!.Guid, "string");
@@ -2462,7 +2433,7 @@ namespace UnrealBuildTool.XcodeProjectXcconfig
 				string Output = Plist($"Copy :objects:{GeneratedMainGroupGuid}:children:{Index} :objects:{TemplateMainGroupGuid}:children:100000000");
 
 				// loop until error
-				if (Output != "")
+				if (!String.IsNullOrEmpty(Output))
 				{
 					break;
 				}
diff --git a/Engine/Source/Programs/UnrealBuildTool/ProjectFiles/Xcode/XcodeProjectFileGenerator.cs b/Engine/Source/Programs/UnrealBuildTool/ProjectFiles/Xcode/XcodeProjectFileGenerator.cs
--- a/Engine/Source/Programs/UnrealBuildTool/ProjectFiles/Xcode/XcodeProjectFileGenerator.cs
+++ b/Engine/Source/Programs/UnrealBuildTool/ProjectFiles/Xcode/XcodeProjectFileGenerator.cs
@@ -67,7 +67,6 @@ namespace UnrealBuildTool
 		/// </summary>
 		public static FileReference ContentOnlySettingsFile = FileReference.Combine(Unreal.EngineDirectory, "Build/Xcode/ContentOnlySettings.xcconfig");
 
-
 		/// <summary>
 		/// A static copy of ProjectPlatforms from the base class
 		/// </summary>
@@ -79,20 +78,16 @@ namespace UnrealBuildTool
 		internal static Dictionary<Tuple<ProjectFile, UnrealTargetPlatform>, IEnumerable<UEBuildFramework>> TargetFrameworks = new();
 		internal static Dictionary<Tuple<ProjectFile, UnrealTargetPlatform>, IEnumerable<UEBuildBundleResource>> TargetBundles = new();
 
-
-
 		/// <summary>
 		/// Should we generate only a run project (no build/index targets)
 		/// </summary>
 		public static bool bGenerateRunOnlyProject = false;
 
-
 		public XcodeProjectFileGenerator(FileReference? InOnlyGameProject, CommandLineArguments CommandLine)
 			: base(InOnlyGameProject)
 		{
 			SingleGameProject = InOnlyGameProject;
 
-
 			if (CommandLine.HasOption("-distribution"))
 			{
 				bForDistribution = true;
@@ -266,7 +261,7 @@ namespace UnrealBuildTool
 			{
 				foreach (PrimaryProjectFolder CurFolder in FolderList)
 				{
-					var Modern = CurFolder.ChildProjects.FirstOrDefault(P =>
+					ProjectFile? Modern = CurFolder.ChildProjects.FirstOrDefault(P =>
 						P.GetType() == typeof(XcodeProjectXcconfig.XcodeProjectFile) &&
 						!((XcodeProjectXcconfig.XcodeProjectFile)P).bHasLegacyProject);
 					if (Modern != null)
@@ -406,7 +401,6 @@ namespace UnrealBuildTool
 			return WriteXcodeWorkspace(Logger);
 		}
 
-
 		/// <inheritdoc/>
 		protected override void ConfigureProjectFileGeneration(string[] Arguments, ref bool IncludeAllPlatforms, ILogger Logger)
 		{
diff --git a/Engine/Source/Programs/UnrealBuildTool/ProjectFiles/Xcode/XcodeProjectLegacy.cs b/Engine/Source/Programs/UnrealBuildTool/ProjectFiles/Xcode/XcodeProjectLegacy.cs
--- a/Engine/Source/Programs/UnrealBuildTool/ProjectFiles/Xcode/XcodeProjectLegacy.cs
+++ b/Engine/Source/Programs/UnrealBuildTool/ProjectFiles/Xcode/XcodeProjectLegacy.cs
@@ -900,7 +900,7 @@ namespace UnrealBuildTool.XcodeProjectLegacy
 			Content.Append("\t\t\t\tHEADER_SEARCH_PATHS = (" + ProjectFileGenerator.NewLine);
 			foreach (string SearchPath in IntelliSenseSystemIncludeSearchPaths)
 			{
-				string Path = SearchPath.Contains(" ") ? "\\\"" + SearchPath + "\\\"" : SearchPath;
+				string Path = SearchPath.Contains(' ') ? "\\\"" + SearchPath + "\\\"" : SearchPath;
 				Content.Append("\t\t\t\t\t\"" + Path + "\"," + ProjectFileGenerator.NewLine);
 			}
 			Content.Append("\t\t\t\t);" + ProjectFileGenerator.NewLine);
@@ -908,7 +908,7 @@ namespace UnrealBuildTool.XcodeProjectLegacy
 			Content.Append("\t\t\t\tUSER_HEADER_SEARCH_PATHS = (" + ProjectFileGenerator.NewLine);
 			foreach (string SearchPath in IntelliSenseIncludeSearchPaths)
 			{
-				string Path = SearchPath.Contains(" ") ? "\\\"" + SearchPath + "\\\"" : SearchPath;
+				string Path = SearchPath.Contains(' ') ? "\\\"" + SearchPath + "\\\"" : SearchPath;
 				Content.Append("\t\t\t\t\t\"" + Path + "\"," + ProjectFileGenerator.NewLine);
 			}
 			Content.Append("\t\t\t\t);" + ProjectFileGenerator.NewLine);
@@ -1353,7 +1353,6 @@ namespace UnrealBuildTool.XcodeProjectLegacy
 			Content.Append("\t\t};" + ProjectFileGenerator.NewLine);
 		}
 
-
 		private void AppendXCBuildConfigurationSection(StringBuilder Content, Dictionary<string, XcodeBuildConfig> ProjectBuildConfigs, Dictionary<string, XcodeBuildConfig> TargetBuildConfigs,
 				Dictionary<string, XcodeBuildConfig> BuildTargetBuildConfigs, Dictionary<string, XcodeBuildConfig> IndexTargetBuildConfigs, FileReference? GameProjectPath,
 				List<XcodeExtensionInfo> AllExtensions, ILogger Logger)
@@ -1662,7 +1661,7 @@ namespace UnrealBuildTool.XcodeProjectLegacy
 			Content.Append("      </BuildableProductRunnable>" + ProjectFileGenerator.NewLine);
 			if (String.IsNullOrEmpty(OldCommandLineArguments))
 			{
-				if (bHasEditorConfiguration && GameProjectPath != "")
+				if (bHasEditorConfiguration && !String.IsNullOrEmpty(GameProjectPath))
 				{
 					Content.Append("      <CommandLineArguments>" + ProjectFileGenerator.NewLine);
 					if (IsForeignProject)
@@ -1826,7 +1825,6 @@ namespace UnrealBuildTool.XcodeProjectLegacy
 			string CopyExtensionsBuildPhaseGuid = XcodeProjectFileGenerator.MakeXcodeGuid();
 			string ShellScriptSectionGuid = XcodeProjectFileGenerator.MakeXcodeGuid();
 
-
 			// Figure out all the desired configurations
 			List<XcodeBuildConfig> BuildConfigs = GetSupportedBuildConfigs(InPlatforms, InConfigurations, Logger);
 			if (BuildConfigs.Count == 0)
diff --git a/Engine/Source/Programs/UnrealBuildTool/System/ActionGraph.cs b/Engine/Source/Programs/UnrealBuildTool/System/ActionGraph.cs
--- a/Engine/Source/Programs/UnrealBuildTool/System/ActionGraph.cs
+++ b/Engine/Source/Programs/UnrealBuildTool/System/ActionGraph.cs
@@ -463,7 +463,7 @@ namespace UnrealBuildTool
 				if ((!Action.bCanExecuteRemotely || Action.ActionType == ActionType.Link) && Action.PrerequisiteActions.Count > 0)
 				{
 					int LastSortIndex = 0;
-					foreach (var Prereq in Action.PrerequisiteActions)
+					foreach (LinkedAction Prereq in Action.PrerequisiteActions)
 					{
 						LastSortIndex = Math.Max(LastSortIndex, Prereq.SortIndex);
 					}
@@ -889,7 +889,7 @@ namespace UnrealBuildTool
 
 			using (GlobalTracer.Instance.BuildSpan("Cache outdated actions based on recursive prerequisites").StartActive())
 			{
-				foreach (var Action in Actions)
+				foreach (LinkedAction Action in Actions)
 				{
 					IsActionOutdatedDueToPrerequisites(Action, OutdatedActions, bIgnoreOutdatedImportLibraries, Logger);
 				}
diff --git a/Engine/Source/Programs/UnrealBuildTool/System/ActionSerializer.cs b/Engine/Source/Programs/UnrealBuildTool/System/ActionSerializer.cs
--- a/Engine/Source/Programs/UnrealBuildTool/System/ActionSerializer.cs
+++ b/Engine/Source/Programs/UnrealBuildTool/System/ActionSerializer.cs
@@ -83,8 +83,8 @@ namespace UnrealBuildTool
 		/// <returns></returns>
 		static ActionSerialization()
 		{
-			var TypeToSerializerDict = new Dictionary<Type, IActionSerializer>();
-			var NameToSerializerDict = new Dictionary<string, IActionSerializer>(StringComparer.Ordinal);
+			Dictionary<Type, IActionSerializer> TypeToSerializerDict = new Dictionary<Type, IActionSerializer>();
+			Dictionary<string, IActionSerializer> NameToSerializerDict = new Dictionary<string, IActionSerializer>(StringComparer.Ordinal);
 
 			Type[] Types = Assembly.GetExecutingAssembly().GetTypes();
 			foreach (Type Type in Types)
diff --git a/Engine/Source/Programs/UnrealBuildTool/System/ConfigValueTracker.cs b/Engine/Source/Programs/UnrealBuildTool/System/ConfigValueTracker.cs
--- a/Engine/Source/Programs/UnrealBuildTool/System/ConfigValueTracker.cs
+++ b/Engine/Source/Programs/UnrealBuildTool/System/ConfigValueTracker.cs
@@ -89,7 +89,7 @@ namespace UnrealBuildTool
 		/// <returns>True if the keys are equal, false otherwise</returns>
 		public override bool Equals(object? Other)
 		{
-			return (Other is ConfigDependencyKey) && Equals((ConfigDependencyKey)Other);
+			return (Other is ConfigDependencyKey key) && Equals(key);
 		}
 
 		/// <summary>
diff --git a/Engine/Source/Programs/UnrealBuildTool/System/DataDrivenPlatformInfo.cs b/Engine/Source/Programs/UnrealBuildTool/System/DataDrivenPlatformInfo.cs
--- a/Engine/Source/Programs/UnrealBuildTool/System/DataDrivenPlatformInfo.cs
+++ b/Engine/Source/Programs/UnrealBuildTool/System/DataDrivenPlatformInfo.cs
@@ -74,7 +74,7 @@ namespace UnrealBuildTool
 			{
 				if (PlatformInfos == null)
 				{
-					var PlatformInfosDict = new Dictionary<string, ConfigDataDrivenPlatformInfo>(StringComparer.OrdinalIgnoreCase);
+					Dictionary<string, ConfigDataDrivenPlatformInfo> PlatformInfosDict = new Dictionary<string, ConfigDataDrivenPlatformInfo>(StringComparer.OrdinalIgnoreCase);
 					PlatformInfos = PlatformInfosDict;
 					Dictionary<string, string> IniParents = new Dictionary<string, string>(StringComparer.OrdinalIgnoreCase);
 
@@ -103,7 +103,6 @@ namespace UnrealBuildTool
 							ConfigFile Config = new ConfigFile(FileRef, ConfigLineAction.Add);
 							ConfigDataDrivenPlatformInfo NewInfo = new ConfigDataDrivenPlatformInfo();
 
-
 							// we must have the key section 
 							ConfigFileSection? Section = null;
 							if (Config.TryGetSection("DataDrivenPlatformInfo", out Section))
diff --git a/Engine/Source/Programs/UnrealBuildTool/System/EncryptionAndSigning.cs b/Engine/Source/Programs/UnrealBuildTool/System/EncryptionAndSigning.cs
--- a/Engine/Source/Programs/UnrealBuildTool/System/EncryptionAndSigning.cs
+++ b/Engine/Source/Programs/UnrealBuildTool/System/EncryptionAndSigning.cs
@@ -320,7 +320,7 @@ namespace UnrealBuildTool
 			ConfigHierarchySection CryptoSection = Ini.FindSection(SectionName);
 
 			// If we have new format crypto keys, read them in over the top of the legacy settings
-			if (CryptoSection != null && CryptoSection.KeyNames.Count() > 0)
+			if (CryptoSection != null && CryptoSection.KeyNames.Any())
 			{
 				Ini.GetBool(SectionName, "bEnablePakSigning", out bool bEnablePakSigning);
 				Settings.bEnablePakSigning = bEnablePakSigning;
diff --git a/Engine/Source/Programs/UnrealBuildTool/System/ExternalExecution.cs b/Engine/Source/Programs/UnrealBuildTool/System/ExternalExecution.cs
--- a/Engine/Source/Programs/UnrealBuildTool/System/ExternalExecution.cs
+++ b/Engine/Source/Programs/UnrealBuildTool/System/ExternalExecution.cs
@@ -502,7 +502,7 @@ namespace UnrealBuildTool
 				using (ThreadPoolWorkQueue Queue = new ThreadPoolWorkQueue())
 				{
 					IReadOnlySet<string> ExcludedFolders = UEBuildPlatform.GetBuildPlatform(Platform).GetExcludedFolderNames();
-					foreach (var Module in IncludePathModules)
+					foreach (UEBuildModuleCPP Module in IncludePathModules)
 					{
 						Queue.Enqueue(() =>
 						{
@@ -695,7 +695,6 @@ namespace UnrealBuildTool
 				}
 			}
 
-
 			foreach (UHTModuleInfo Module in UObjectModules)
 			{
 				// If we're using a precompiled engine, skip checking timestamps for modules that are under the engine directory
diff --git a/Engine/Source/Programs/UnrealBuildTool/System/GraphVisualization.cs b/Engine/Source/Programs/UnrealBuildTool/System/GraphVisualization.cs
--- a/Engine/Source/Programs/UnrealBuildTool/System/GraphVisualization.cs
+++ b/Engine/Source/Programs/UnrealBuildTool/System/GraphVisualization.cs
@@ -149,7 +149,7 @@ namespace UnrealBuildTool
 									// @todo: Add support for edge attributes, not just node attributes
 									Writer.WriteAttributeString("class", "node");   // Node attributes, not edges!
 
-									foreach (var Attribute in AllAttributes.Values)
+									foreach (GraphAttribute Attribute in AllAttributes.Values)
 									{
 										Writer.WriteStartElement("attribute");
 										{
@@ -202,12 +202,12 @@ namespace UnrealBuildTool
 										{
 											Writer.WriteStartElement("attvalues");
 											{
-												foreach (var AttributeHashEntry in GraphNode.Attributes)
+												foreach (KeyValuePair<string, object> AttributeHashEntry in GraphNode.Attributes)
 												{
-													var AttributeName = AttributeHashEntry.Key;
-													var AttributeValue = AttributeHashEntry.Value;
+													string AttributeName = AttributeHashEntry.Key;
+													object AttributeValue = AttributeHashEntry.Value;
 
-													var Attribute = AllAttributes[AttributeName];
+													GraphAttribute Attribute = AllAttributes[AttributeName];
 
 													Writer.WriteStartElement("attvalue");
 													{
@@ -225,7 +225,6 @@ namespace UnrealBuildTool
 							}
 							Writer.WriteEndElement();   // nodes
 
-
 							Writer.WriteStartElement("edges");
 							{
 								for (int EdgeId = 0; EdgeId < GraphEdges.Count; EdgeId++)
@@ -272,7 +271,6 @@ namespace UnrealBuildTool
 
 					Writer.Flush();
 				}
-
 			}
 			finally
 			{
diff --git a/Engine/Source/Programs/UnrealBuildTool/System/NativeProjects.cs b/Engine/Source/Programs/UnrealBuildTool/System/NativeProjects.cs
--- a/Engine/Source/Programs/UnrealBuildTool/System/NativeProjects.cs
+++ b/Engine/Source/Programs/UnrealBuildTool/System/NativeProjects.cs
@@ -77,8 +77,6 @@ namespace UnrealBuildTool
 			return CachedTargetNameToProjectFile.TryGetValue(InTargetName, out OutProjectFileName);
 		}
 
-
-
 		#region Hybrid content-only/code-based project
 
 		/// <summary>
@@ -226,10 +224,6 @@ namespace UnrealBuildTool
 			return true;
 		}
 
-
-
-
-
 		/// <summary>
 		/// Determines if a project (given a .uproject file) has source code. This is determined by finding at least one .Target.cs file
 		/// </summary>
@@ -242,18 +236,17 @@ namespace UnrealBuildTool
 			DirectoryReference TempSourceDir = DirectoryReference.Combine(UProjectFile.Directory, "Intermediate", "Source");
 
 			// check to see if we have a Target.cs file in Source or Intermediate/Source
-			if (DirectoryReference.Exists(SourceDir) && DirectoryReference.EnumerateFiles(SourceDir, "*.Target.cs", SearchOption.TopDirectoryOnly).Count() > 0)
+			if (DirectoryReference.Exists(SourceDir) && DirectoryReference.EnumerateFiles(SourceDir, "*.Target.cs", SearchOption.TopDirectoryOnly).Any())
 			{
 				return true;
 			}
-			if (bCheckForTempTargets && DirectoryReference.Exists(TempSourceDir) && DirectoryReference.EnumerateFiles(TempSourceDir, "*.Target.cs", SearchOption.TopDirectoryOnly).Count() > 0)
+			if (bCheckForTempTargets && DirectoryReference.Exists(TempSourceDir) && DirectoryReference.EnumerateFiles(TempSourceDir, "*.Target.cs", SearchOption.TopDirectoryOnly).Any())
 			{
 				return true;
 			}
 			return false;
 		}
 
-
 		private static bool RequiresTempTarget(FileReference UProjectFile, List<UnrealTargetPlatform> Platforms, List<UnrealTargetConfiguration> Configurations, [NotNullWhen(true)] out string? Reason, ILogger Logger)
 		{
 			// no reason by default
diff --git a/Engine/Source/Programs/UnrealBuildTool/System/Plugins.cs b/Engine/Source/Programs/UnrealBuildTool/System/Plugins.cs
--- a/Engine/Source/Programs/UnrealBuildTool/System/Plugins.cs
+++ b/Engine/Source/Programs/UnrealBuildTool/System/Plugins.cs
@@ -255,9 +255,9 @@ namespace UnrealBuildTool
 		/// <returns>The cached plugin info, or null if not found.</returns>
 		public static PluginInfo? GetPlugin(string PluginName)
 		{
-			foreach (var Pair in PluginInfoCache)
+			foreach (KeyValuePair<DirectoryReference, List<PluginInfo>> Pair in PluginInfoCache)
 			{
-				foreach (var Plugin in Pair.Value)
+				foreach (PluginInfo Plugin in Pair.Value)
 				{
 					if (String.Equals(Plugin.Name, PluginName, StringComparison.InvariantCultureIgnoreCase))
 					{
@@ -560,7 +560,6 @@ namespace UnrealBuildTool
 			// @todo platplug: what else do we want to support merging?!?
 		}
 
-
 		/// <summary>
 		/// Read all the plugin descriptors under the given directory
 		/// </summary>
diff --git a/Engine/Source/Programs/UnrealBuildTool/System/RestrictedFolders.cs b/Engine/Source/Programs/UnrealBuildTool/System/RestrictedFolders.cs
--- a/Engine/Source/Programs/UnrealBuildTool/System/RestrictedFolders.cs
+++ b/Engine/Source/Programs/UnrealBuildTool/System/RestrictedFolders.cs
@@ -84,7 +84,7 @@ namespace UnrealBuildTool
 		/// <returns>True if the restricted folder is equal to the other instance</returns>
 		public override bool Equals(object? Other)
 		{
-			return Other is RestrictedFolder && Id == ((RestrictedFolder)Other).Id;
+			return Other is RestrictedFolder folder && Id == folder.Id;
 		}
 
 		/// <summary>
diff --git a/Engine/Source/Programs/UnrealBuildTool/System/RulesAssembly.cs b/Engine/Source/Programs/UnrealBuildTool/System/RulesAssembly.cs
--- a/Engine/Source/Programs/UnrealBuildTool/System/RulesAssembly.cs
+++ b/Engine/Source/Programs/UnrealBuildTool/System/RulesAssembly.cs
@@ -137,7 +137,7 @@ namespace UnrealBuildTool
 			{
 				// Add the parent assemblies as references so they can be used
 				List<string>? ReferencedAssembies = null;
-				var CurrentParent = Parent;
+				RulesAssembly? CurrentParent = Parent;
 				while (CurrentParent != null && CurrentParent.CompiledAssembly != null)
 				{
 					if (ReferencedAssembies == null)
@@ -473,7 +473,6 @@ namespace UnrealBuildTool
 				}
 			}
 
-
 			// Figure out the best rules object to use
 			Type? RulesObjectType = PlatformRulesObjectType != null ? PlatformRulesObjectType : BaseRulesObjectType;
 			if (RulesObjectType == null)
@@ -533,7 +532,7 @@ namespace UnrealBuildTool
 				}
 
 				// Add the parent assemblies to the assembly cache so the types in them can be used when the constructor is called
-				var CurrentParent = Parent;
+				RulesAssembly? CurrentParent = Parent;
 				while (CurrentParent != null && CurrentParent.CompiledAssembly != null)
 				{
 					EpicGames.Core.AssemblyUtils.AddFileToAssemblyCache(CurrentParent.CompiledAssembly.Location);
@@ -874,7 +873,7 @@ namespace UnrealBuildTool
 		{
 			for (RulesAssembly? Assembly = this; Assembly != null; Assembly = Assembly.Parent)
 			{
-				if (Assembly.ModuleFileToContext.TryGetValue(ModuleFile, out var Context))
+				if (Assembly.ModuleFileToContext.TryGetValue(ModuleFile, out ModuleRulesContext? Context))
 				{
 					return Context;
 				}
diff --git a/Engine/Source/Programs/UnrealBuildTool/System/RulesCompiler.cs b/Engine/Source/Programs/UnrealBuildTool/System/RulesCompiler.cs
--- a/Engine/Source/Programs/UnrealBuildTool/System/RulesCompiler.cs
+++ b/Engine/Source/Programs/UnrealBuildTool/System/RulesCompiler.cs
@@ -294,7 +294,7 @@ namespace UnrealBuildTool
 				FindModuleRulesForPlugins(ProjectPlugins, DefaultModuleContext, ModuleFiles);
 
 				Logger.LogTrace(" Found {Count} Modules:", ModuleFiles.Count);
-				foreach (var Item in ModuleFiles)
+				foreach (KeyValuePair<FileReference, ModuleRulesContext> Item in ModuleFiles)
 				{
 					Logger.LogTrace("  {Module}", Item.Key);
 				}
@@ -375,7 +375,7 @@ namespace UnrealBuildTool
 				FindModuleRulesForPlugins(ForeignPlugins, PluginModuleContext, ModuleFiles);
 
 				Logger.LogTrace(" Found {Count} Modules:", ModuleFiles.Count);
-				foreach (var Item in ModuleFiles)
+				foreach (KeyValuePair<FileReference, ModuleRulesContext> Item in ModuleFiles)
 				{
 					Logger.LogTrace("  {Module}", Item.Key);
 				}
diff --git a/Engine/Source/Programs/UnrealBuildTool/System/SourceFileMetadataCache.cs b/Engine/Source/Programs/UnrealBuildTool/System/SourceFileMetadataCache.cs
--- a/Engine/Source/Programs/UnrealBuildTool/System/SourceFileMetadataCache.cs
+++ b/Engine/Source/Programs/UnrealBuildTool/System/SourceFileMetadataCache.cs
@@ -186,11 +186,11 @@ namespace UnrealBuildTool
 			}
 			else
 			{
-				var UpdateSourceFileInfo = (FileItem SourceFile) =>
+				Func<FileItem, SourceFileInfo> UpdateSourceFileInfo = (FileItem SourceFile) =>
 				{
-					var SourceFileInfo = new SourceFileInfo();
-					var FileText = FileReference.ReadAllText(SourceFile.Location);
-					var FileTextLines = FileText.Split('\n');
+					SourceFileInfo SourceFileInfo = new SourceFileInfo();
+					string FileText = FileReference.ReadAllText(SourceFile.Location);
+					string[] FileTextLines = FileText.Split('\n');
 
 					SourceFileInfo.LastWriteTimeUtc = SourceFile.LastWriteTimeUtc.Ticks;
 
@@ -265,11 +265,11 @@ namespace UnrealBuildTool
 			}
 			else
 			{
-				var UpdateHeaderFileInfo = (FileItem HeaderFile) =>
+				Func<FileItem, HeaderFileInfo> UpdateHeaderFileInfo = (FileItem HeaderFile) =>
 				{
-					var HeaderFileInfo = new HeaderFileInfo();
-					var FileText = FileReference.ReadAllText(HeaderFile.Location);
-					var FileTextLines = FileText.Split('\n');
+					HeaderFileInfo HeaderFileInfo = new HeaderFileInfo();
+					string FileText = FileReference.ReadAllText(HeaderFile.Location);
+					string[] FileTextLines = FileText.Split('\n');
 
 					HeaderFileInfo.LastWriteTimeUtc = HeaderFile.LastWriteTimeUtc.Ticks;
 
diff --git a/Engine/Source/Programs/UnrealBuildTool/System/TopologicalSorter.cs b/Engine/Source/Programs/UnrealBuildTool/System/TopologicalSorter.cs
--- a/Engine/Source/Programs/UnrealBuildTool/System/TopologicalSorter.cs
+++ b/Engine/Source/Programs/UnrealBuildTool/System/TopologicalSorter.cs
@@ -172,7 +172,7 @@ namespace UnrealBuildTool
 			HashSet<T> AllNodes = new HashSet<T>();
 
 			// Create a collection of all nodes based on the ones used in Edges.
-			foreach (var Edge in Edges)
+			foreach (Tuple<T, T> Edge in Edges)
 			{
 				AllNodes.Add(Edge.Item1);
 				AllNodes.Add(Edge.Item2);
@@ -190,7 +190,7 @@ namespace UnrealBuildTool
 			}
 
 			// Add edges to the graph.
-			foreach (var Edge in Edges)
+			foreach (Tuple<T, T> Edge in Edges)
 			{
 				GraphNode? NodeSrc;
 				GraphNode? NodeDst;
diff --git a/Engine/Source/Programs/UnrealBuildTool/System/Unity.cs b/Engine/Source/Programs/UnrealBuildTool/System/Unity.cs
--- a/Engine/Source/Programs/UnrealBuildTool/System/Unity.cs
+++ b/Engine/Source/Programs/UnrealBuildTool/System/Unity.cs
@@ -109,7 +109,6 @@ namespace UnrealBuildTool
 				}
 			}
 
-
 			/// <summary>
 			/// Starts a new unity file.  If the current unity file contains no files, this function has no effect, i.e. you will not get an empty unity file.
 			/// </summary>
@@ -187,7 +186,7 @@ namespace UnrealBuildTool
 			// Every single file in the module appears in the working set. Don't bother using adaptive unity for this module.
 			// Otherwise it would make full builds really slow.
 			GetAdaptiveFiles(Target, CPPFiles, HeaderFiles, CompileEnvironment, WorkingSet, BaseName, IntermediateDirectory, Graph, out NormalFiles, out AdaptiveFiles);
-			if (NormalFiles.Where(file => !file.HasExtension(".gen.cpp")).Count() == 0)
+			if (!NormalFiles.Where(file => !file.HasExtension(".gen.cpp")).Any())
 			{
 				NormalFiles = CPPFiles;
 				AdaptiveFiles.RemoveAll(new HashSet<FileItem>(NormalFiles).Contains);
diff --git a/Engine/Source/Programs/UnrealBuildTool/System/UnrealPluginLanguage.cs b/Engine/Source/Programs/UnrealBuildTool/System/UnrealPluginLanguage.cs
--- a/Engine/Source/Programs/UnrealBuildTool/System/UnrealPluginLanguage.cs
+++ b/Engine/Source/Programs/UnrealBuildTool/System/UnrealPluginLanguage.cs
@@ -530,7 +530,7 @@ namespace UnrealBuildTool
 				StringVariables = new Dictionary<string, string>();
 				ElementVariables = new Dictionary<string, XElement>();
 
-				if (PluginDir == null || PluginDir == "")
+				if (PluginDir == null || String.IsNullOrEmpty(PluginDir))
 				{
 					PluginDir = ".";
 				}
@@ -1816,7 +1816,7 @@ namespace UnrealBuildTool
 							string? Filespec = GetAttribute(CurrentContext, Node, "filespec");
 							if (Filespec != null)
 							{
-								if (Filespec.Contains(":") || Filespec.Contains(".."))
+								if (Filespec.Contains(':') || Filespec.Contains(".."))
 								{
 									Logger.LogInformation("\nFilespec {FileSpec} not allowed; ignored.", Filespec);
 								}
@@ -2728,7 +2728,6 @@ namespace UnrealBuildTool
 			{ }
 		}
 
-
 		/// <summary>
 		/// command class for being able to create config caches over and over without needing to read the ini files
 		/// </summary>
@@ -3454,7 +3453,6 @@ namespace UnrealBuildTool
 			yield return FileReference.Combine(EngineDirectory, "Restricted", "NotForLicensees", "Config", "Base" + BaseIniName + ".ini");
 		}
 
-
 		/// <summary>
 		/// Returns a list of INI filenames for the given project
 		/// </summary>
@@ -3534,5 +3532,4 @@ namespace UnrealBuildTool
 			}
 		}
 	}
-
 }
diff --git a/Engine/Source/Programs/UnrealBuildTool/System/Utils.cs b/Engine/Source/Programs/UnrealBuildTool/System/Utils.cs
--- a/Engine/Source/Programs/UnrealBuildTool/System/Utils.cs
+++ b/Engine/Source/Programs/UnrealBuildTool/System/Utils.cs
@@ -350,7 +350,6 @@ namespace UnrealBuildTool
 				}
 			}
 
-
 			DataReceivedEventHandler Output = (object sender, DataReceivedEventArgs Args) =>
 			{
 				if (Args != null && Args.Data != null)
@@ -375,7 +374,7 @@ namespace UnrealBuildTool
 			LocalProcess.StartInfo = StartInfo;
 			LocalProcess.OutputDataReceived += Output;
 			LocalProcess.ErrorDataReceived += Output;
-			var ExitCode = RunLocalProcess(LocalProcess);
+			int ExitCode = RunLocalProcess(LocalProcess);
 			if (ExitCode != 0 && !String.IsNullOrEmpty(LogFilename))
 			{
 				Logger.LogError("Process \'{AppName}\' failed. Details are in \'{LogFilename}\'", AppName, LogFilename);
@@ -384,7 +383,6 @@ namespace UnrealBuildTool
 			return ExitCode;
 		}
 
-
 		/// <summary>
 		/// Runs a local process and pipes the output to the log
 		/// </summary>
@@ -732,7 +730,6 @@ namespace UnrealBuildTool
 			return LocalString;
 		}
 
-
 		/// <summary>
 		/// Given a file path and a directory, returns a file path that is relative to the specified directory
 		/// </summary>
@@ -790,7 +787,6 @@ namespace UnrealBuildTool
 			return RelativePath;
 		}
 
-
 		/// <summary>
 		/// Backspaces the specified number of characters, then displays a progress percentage value to the console
 		/// </summary>
@@ -812,7 +808,6 @@ namespace UnrealBuildTool
 			Console.Write(ProgressString);
 		}
 
-
 		/*
 		 * Read and write classes with xml specifiers
 		 */
@@ -980,7 +975,6 @@ namespace UnrealBuildTool
 			}
 		}
 
-
 		/// <summary>
 		/// Returns true if the specified file's path is located under the specified directory, or any of that directory's sub-folders.  Does not care whether the file or directory exist or not.  This is a simple string-based check.
 		/// </summary>
@@ -1082,7 +1076,6 @@ namespace UnrealBuildTool
 			return Environment.ProcessorCount;
 		}
 
-
 		// int sysctlbyname(const char *name, void *oldp, size_t *oldlenp, void *newp, size_t newlen); // from man page
 		[DllImport("libc")]
 		static extern int sysctlbyname(string name, out int oldp, ref UInt64 oldlenp, IntPtr newp, UInt64 newlen);
@@ -1174,7 +1167,6 @@ namespace UnrealBuildTool
 			return MemoryInfo.TotalAvailableMemoryBytes != 0 ? MemoryInfo.TotalAvailableMemoryBytes : -1;
 		}
 
-
 		// vm_statistics64, based on the definition in <mach/vm_statistics.h>
 		[StructLayout(LayoutKind.Sequential)]
 		struct vm_statistics64
@@ -1639,7 +1631,6 @@ namespace UnrealBuildTool
 									Logger.Log(OverrideType,
 										"{Prefix}Not writing a file that was previously not written: \"{File}\"{Context}", Prefix, File, Context);
 								}
-
 							}
 						}
 					}
@@ -1786,7 +1777,7 @@ namespace UnrealBuildTool
 		/// <param name="Logger">Logger for output</param>
 		internal static void WriteFileIfChanged(FileItem FileItem, IEnumerable<string> ContentLines, StringComparison Comparison, ILogger Logger)
 		{
-			var Contents = String.Join(Environment.NewLine, ContentLines);
+			string Contents = String.Join(Environment.NewLine, ContentLines);
 			WriteFileIfChanged(FileItem, Contents, Comparison, Logger);
 		}
 
diff --git a/Engine/Source/Programs/UnrealBuildTool/ToolChain/AppleToolChain.cs b/Engine/Source/Programs/UnrealBuildTool/ToolChain/AppleToolChain.cs
--- a/Engine/Source/Programs/UnrealBuildTool/ToolChain/AppleToolChain.cs
+++ b/Engine/Source/Programs/UnrealBuildTool/ToolChain/AppleToolChain.cs
@@ -229,8 +229,6 @@ namespace UnrealBuildTool
 			return Result;
 		}
 
-
-
 		protected void StripSymbolsWithXcode(FileReference SourceFile, FileReference TargetFile, string ToolchainDir)
 		{
 			if (SourceFile != TargetFile)
@@ -320,7 +318,7 @@ namespace UnrealBuildTool
 		/// <inheritdoc/>
 		protected override string EscapePreprocessorDefinition(string Definition)
 		{
-			return Definition.Contains("\"") ? Definition.Replace("\"", "\\\"") : Definition;
+			return Definition.Contains('"') ? Definition.Replace("\"", "\\\"") : Definition;
 		}
 
 		protected override void GetCppStandardCompileArgument(CppCompileEnvironment CompileEnvironment, List<string> Arguments)
@@ -422,7 +420,6 @@ namespace UnrealBuildTool
 			}
 		}
 
-
 		/// <inheritdoc/>
 		protected override void GetCompileArguments_Debugging(CppCompileEnvironment CompileEnvironment, List<string> Arguments)
 		{
@@ -476,7 +473,7 @@ namespace UnrealBuildTool
 			string[] Tokens = DsymutilVersionString.Split(" ".ToCharArray());
 
 			// sanity check
-			if (Tokens.Length < 4 || Tokens[3].Contains(".") == false)
+			if (Tokens.Length < 4 || Tokens[3].Contains('.') == false)
 			{
 				Log.TraceInformationOnce("Unable to parse dsymutil version out of: {0}", DsymutilVersionString);
 			}
@@ -522,7 +519,6 @@ namespace UnrealBuildTool
 			return DsymutilLocation.FullName;
 		}
 
-
 		public override ICollection<FileItem> PostBuild(ReadOnlyTargetRules Target, FileItem Executable, LinkEnvironment BinaryLinkEnvironment, IActionGraphBuilder Graph)
 		{
 			List<FileItem> OutputFiles = new List<FileItem>(base.PostBuild(Target, Executable, BinaryLinkEnvironment, Graph));
@@ -539,7 +535,6 @@ namespace UnrealBuildTool
 			return OutputFiles;
 		}
 
-
 		#region Stub Xcode Projects
 
 		internal static bool GenerateProjectFiles(FileReference? ProjectFile, string[] Arguments, ILogger Logger, out DirectoryReference? XcodeProjectFile)
@@ -727,7 +722,6 @@ namespace UnrealBuildTool
 			FileReference PostBuildSyncFile = FileReference.Combine(IntermediateDir!, "PostBuildSync.dat");
 			BinaryFormatterUtils.Save(PostBuildSyncFile, PostBuildSync);
 
-
 			string PostBuildSyncArguments = String.Format("-modernxcode -Input=\"{0}\" -XmlConfigCache=\"{1}\" -remoteini=\"{2}\"", PostBuildSyncFile, XmlConfig.CacheFile, UnrealBuildTool.GetRemoteIniPath());
 			Action PostBuildSyncAction = Graph.CreateRecursiveAction<ApplePostBuildSyncMode>(ActionType.CreateAppBundle, PostBuildSyncArguments);
 
diff --git a/Engine/Source/Programs/UnrealBuildTool/ToolChain/ISPCToolChain.cs b/Engine/Source/Programs/UnrealBuildTool/ToolChain/ISPCToolChain.cs
--- a/Engine/Source/Programs/UnrealBuildTool/ToolChain/ISPCToolChain.cs
+++ b/Engine/Source/Programs/UnrealBuildTool/ToolChain/ISPCToolChain.cs
@@ -648,7 +648,7 @@ namespace UnrealBuildTool
 				{
 					string ObjTarget = Target;
 
-					if (Target.Contains("-"))
+					if (Target.Contains('-'))
 					{
 						// Remove lane width and gang size from obj file name
 						ObjTarget = Target.Split('-')[0];
diff --git a/Engine/Source/Programs/UnrealBuildTool/ToolChain/IWYUToolChain.cs b/Engine/Source/Programs/UnrealBuildTool/ToolChain/IWYUToolChain.cs
--- a/Engine/Source/Programs/UnrealBuildTool/ToolChain/IWYUToolChain.cs
+++ b/Engine/Source/Programs/UnrealBuildTool/ToolChain/IWYUToolChain.cs
@@ -40,7 +40,7 @@ namespace UnrealBuildTool
 
 		protected override ClangToolChainInfo GetToolChainInfo()
 		{
-			var PlatformSDK = new LinuxPlatformSDK(Logger);
+			LinuxPlatformSDK PlatformSDK = new LinuxPlatformSDK(Logger);
 			DirectoryReference? BaseLinuxPath = PlatformSDK.GetBaseLinuxPathForArchitecture(LinuxPlatform.DefaultHostArchitecture);
 
 			CrossCompilingArguments.Add($"--sysroot=\"{NormalizeCommandLinePath(BaseLinuxPath!)}\"");
