Build fix.

diff --git a/Engine/Source/Programs/UnrealBuildTool/Artifacts/HordeStorageArtifactCache.cs b/Engine/Source/Programs/UnrealBuildTool/Artifacts/HordeStorageArtifactCache.cs
--- a/Engine/Source/Programs/UnrealBuildTool/Artifacts/HordeStorageArtifactCache.cs
+++ b/Engine/Source/Programs/UnrealBuildTool/Artifacts/HordeStorageArtifactCache.cs
@@ -348,7 +348,7 @@ namespace UnrealBuildTool.Artifacts
 							{
 								string outputName = hordeArtifactMapping.ArtifactMapping.Outputs[refIndex++].GetFullPath(mapping.DirectoryMapping);
 								using FileStream stream = new(outputName, FileMode.Create, FileAccess.Write, FileShare.ReadWrite);
-								await ChunkedDataNode.CopyToStreamAsync(reader, artifactRef.Handle.Locator, stream, cancellationToken);
+								await ChunkedDataNode.CopyToStreamAsync(reader, artifactRef.Handle.Handle.Locator, stream, cancellationToken);
 							}
 							catch (Exception)
 							{
@@ -511,7 +511,7 @@ namespace UnrealBuildTool.Artifacts
 					await hordeArtifactMapping.WriteFilesAsync(writer, cancellationToken);
 
 					// Save the collection
-					NodeHandle _ = await writer.WriteAsync(refName, node, null, cancellationToken);
+					HashedNodeHandle _ = await writer.WriteAsync(refName, node, null, cancellationToken);
 				}, cancellationToken));
 			}
 			return tasks;
