#UE Unconditionally disabling a couple ASAN annotations to workaround an issue related to building third party libraries with different asan options than the main binary. This may be fixed in a future version of MSVC, but for now we will disable them.

diff --git a/Engine/Source/Programs/UnrealBuildTool/Platform/Windows/VCToolChain.cs b/Engine/Source/Programs/UnrealBuildTool/Platform/Windows/VCToolChain.cs
--- a/Engine/Source/Programs/UnrealBuildTool/Platform/Windows/VCToolChain.cs
+++ b/Engine/Source/Programs/UnrealBuildTool/Platform/Windows/VCToolChain.cs
@@ -589,12 +589,10 @@ namespace UnrealBuildTool
 					// MSVC has no support for __has_feature(address_sanitizer)
 					AddDefinition(Arguments, "USING_ADDRESS_SANITISER=1");
 
-					// Re-evalulate the necessity of these workarounds later
-					if (EnvVars.CompilerVersion < new VersionNumber(14, 36))
-					{
-						AddDefinition(Arguments, "_DISABLE_STRING_ANNOTATION=1");
-						AddDefinition(Arguments, "_DISABLE_VECTOR_ANNOTATION=1");
-					}
+					// Disabling a couple annotations to workaround an issue related to building third party libraries with different options than the main binary.
+					// This fixes an error that looks like this: error LNK2038: mismatch detected for 'annotate_string': value '0' doesn't match value '1' in Module.Core.XX_of_YY.cpp.obj
+					AddDefinition(Arguments, "_DISABLE_STRING_ANNOTATION=1");
+					AddDefinition(Arguments, "_DISABLE_VECTOR_ANNOTATION=1");
 				}
 
 				// Currently the ASan headers are not default around. They can be found at this location so lets use this until this is resolved in the toolchain
