Updated the shared pch logic to traverse the dependency tree to find the best shared pch for each module. The logic before would only look at the direct dependencies which if changed could result in possibly worse compile times because a less optimal pch could be picked.
Note: There however are occasions where the old logic can give better shared pchs because it doesn't check whether there is a circular dependency. The new logic does check for circular dependencies so it can sometimes return a safer but less optimal shared PCH. If we find the old logic gives a better shared pch then we go ahead and use it. This will be removed after some of the circular dependencies are resolved.

diff --git a/Engine/Source/Programs/UnrealBuildTool/Configuration/UEBuildModule.cs b/Engine/Source/Programs/UnrealBuildTool/Configuration/UEBuildModule.cs
--- a/Engine/Source/Programs/UnrealBuildTool/Configuration/UEBuildModule.cs
+++ b/Engine/Source/Programs/UnrealBuildTool/Configuration/UEBuildModule.cs
@@ -1013,6 +1013,58 @@ namespace UnrealBuildTool
 			}
 		}
 
+		/// <summary>
+		/// Gathers all the module dependencies a PCH would have
+		/// </summary>
+		public HashSet<UEBuildModule> GetAllDependencyModulesForPCH()
+		{
+			HashSet<UEBuildModule> ReferencedModules = new HashSet<UEBuildModule>();
+			InternalGetAllDependencyModulesForPCH(ReferencedModules, new HashSet<UEBuildModule>(), true);
+			return ReferencedModules;
+		}
+
+		/// <summary>
+		/// Internal function that gathers all the module dependencies a PCH would have
+		/// </summary>
+		/// <param name="ReferencedModules">Hash of all referenced modules with their addition index.</param>
+		/// <param name="IgnoreReferencedModules">Hashset used to ignore modules which are already added to the list</param>
+		/// <param name="bIncludePrivateDependencyModules">Whether to include private dependencies.</param>
+		private void InternalGetAllDependencyModulesForPCH(HashSet<UEBuildModule> ReferencedModules, HashSet<UEBuildModule> IgnoreReferencedModules, bool bIncludePrivateDependencyModules)
+		{
+			List<UEBuildModule> AllDependencyModules = new List<UEBuildModule>(
+				((bIncludePrivateDependencyModules && PrivateDependencyModules != null) ? PrivateDependencyModules.Count : 0) +
+				(PublicDependencyModules != null ? PublicDependencyModules.Count : 0) +
+				(PublicIncludePathModules != null ? PublicIncludePathModules!.Count : 0)
+				);
+			if (bIncludePrivateDependencyModules && PrivateDependencyModules != null)
+			{
+				AllDependencyModules.AddRange(PrivateDependencyModules);
+			}
+			if (PublicDependencyModules != null)
+			{
+				AllDependencyModules.AddRange(PublicDependencyModules!);
+			}
+			if (PublicIncludePathModules != null)
+			{
+				AllDependencyModules.AddRange(PublicIncludePathModules);
+			}
+
+			foreach (UEBuildModule DependencyModule in AllDependencyModules)
+			{
+				// Don't follow circular back-references!
+				if (!HasCircularDependencyOn(DependencyModule.Name))
+				{
+					if (IgnoreReferencedModules.Add(DependencyModule))
+					{
+						// Recurse into dependent modules first
+						DependencyModule.InternalGetAllDependencyModulesForPCH(ReferencedModules, IgnoreReferencedModules, false);
+
+						ReferencedModules.Add(DependencyModule);
+					}
+				}
+			}
+		}
+
 		public delegate UEBuildModule CreateModuleDelegate(string Name, string ReferenceChain);
 
 		/// <summary>
diff --git a/Engine/Source/Programs/UnrealBuildTool/Configuration/UEBuildModuleCPP.cs b/Engine/Source/Programs/UnrealBuildTool/Configuration/UEBuildModuleCPP.cs
--- a/Engine/Source/Programs/UnrealBuildTool/Configuration/UEBuildModuleCPP.cs
+++ b/Engine/Source/Programs/UnrealBuildTool/Configuration/UEBuildModuleCPP.cs
@@ -449,7 +449,7 @@ namespace UnrealBuildTool
 			}
 
 			// Configure the precompiled headers for this module
-			CompileEnvironment = SetupPrecompiledHeaders(Target, ToolChain, CompileEnvironment, LinkInputFiles, Graph);
+			CompileEnvironment = SetupPrecompiledHeaders(Target, ToolChain, CompileEnvironment, LinkInputFiles, Graph, Logger);
 			if (CompileEnvironment.PerArchPrecompiledHeaderFiles != null)
 			{
 				foreach (UnrealArch Arch in CompileEnvironment.PerArchPrecompiledHeaderFiles.Keys)
@@ -841,7 +841,7 @@ namespace UnrealBuildTool
 				PrecompiledHeaderDir = IntermediateDirectory;
 			}
 
-			return new PrecompiledHeaderTemplate(this, CompileEnvironment, HeaderFile, PrecompiledHeaderDir);
+			return new PrecompiledHeaderTemplate(this, CompileEnvironment, HeaderFile, PrecompiledHeaderDir, GetAllDependencyModulesForPCH());
 		}
 
 		static HashSet<string> GetImmutableDefinitions(List<string> Definitions)
@@ -1404,7 +1404,8 @@ namespace UnrealBuildTool
 		/// <param name="CompileEnvironment">The current compile environment</param>
 		/// <param name="LinkInputFiles">List of files that will be linked for the target</param>
 		/// <param name="Graph">List of build actions</param>
-		CppCompileEnvironment SetupPrecompiledHeaders(ReadOnlyTargetRules Target, UEToolChain? ToolChain, CppCompileEnvironment CompileEnvironment, List<FileItem> LinkInputFiles, IActionGraphBuilder Graph)
+		/// <param name="Logger">Logger for output</param>
+		CppCompileEnvironment SetupPrecompiledHeaders(ReadOnlyTargetRules Target, UEToolChain? ToolChain, CppCompileEnvironment CompileEnvironment, List<FileItem> LinkInputFiles, IActionGraphBuilder Graph, ILogger Logger)
 		{
 			if (Target.bUsePCHFiles && Rules.PCHUsage != ModuleRules.PCHUsageMode.NoPCHs)
 			{
@@ -1425,13 +1426,35 @@ namespace UnrealBuildTool
 				// Try to find a suitable shared PCH for this module
 				if (CompileEnvironment.bHasPrecompiledHeader == false && CompileEnvironment.SharedPCHs.Count > 0 && !CompileEnvironment.bIsBuildingLibrary && Rules.PCHUsage != ModuleRules.PCHUsageMode.NoSharedPCHs)
 				{
-					// Find all the dependencies of this module
-					HashSet<UEBuildModule> ReferencedModules = new HashSet<UEBuildModule>();
-					GetAllDependencyModules(new List<UEBuildModule>(), ReferencedModules, bIncludeDynamicallyLoaded: false, bForceCircular: false, bOnlyDirectDependencies: true);
+					// Find only the direct dependencies of this module
+					// Note that this is the old shared chooser logic that needs to be removed in the future
+					HashSet<UEBuildModule> OnlyDirectDependencies = new HashSet<UEBuildModule>();
+					GetAllDependencyModules(new List<UEBuildModule>(), OnlyDirectDependencies, bIncludeDynamicallyLoaded: false, bForceCircular: false, bOnlyDirectDependencies: true);
+					PrecompiledHeaderTemplate? OldTemplateValue = CompileEnvironment.SharedPCHs.FirstOrDefault(x => OnlyDirectDependencies.Contains(x.Module) && x.IsValidFor(CompileEnvironment));
+
+					// Find the first shared PCH module we can use that doesn't reference this module
+					HashSet<UEBuildModule> AllDependencies = GetAllDependencyModulesForPCH();
+					PrecompiledHeaderTemplate? Template = CompileEnvironment.SharedPCHs.FirstOrDefault(x => !x.ModuleDependencies.Contains(this) && AllDependencies.Contains(x.Module) && x.IsValidFor(CompileEnvironment));
+
+					if (Template != OldTemplateValue)
+					{
+						if (OldTemplateValue != null && Template != null)
+						{
+							/*
+							If GetAllDependencyModules returns a better shared pch than GetAllDependencyModulesForPCH then use that pch instead. This can happen if a circular dependency was 
+							found in the GetAllDependencyModulesForPCH code path. We shouldn't need to do this but there are too many circular dependencies not being detected by UBT ATM. We 
+							should be able to remove this when the circular dependencies have been addressed in the engine.
+							*/
+							var NewIndex = CompileEnvironment.SharedPCHs.IndexOf(Template);
+							var OldIndex = CompileEnvironment.SharedPCHs.IndexOf(OldTemplateValue);
+							if (NewIndex > OldIndex)
+							{
+								Template = OldTemplateValue;
+							}
+						}
+					}
 
-					// Find the first shared PCH module we can use
-					PrecompiledHeaderTemplate? Template = CompileEnvironment.SharedPCHs.FirstOrDefault(x => ReferencedModules.Contains(x.Module));
-					if (Template != null && Template.IsValidFor(CompileEnvironment))
+					if (Template != null)
 					{
 						PrecompiledHeaderInstance Instance = FindOrCreateSharedPCH(ToolChain, Template, CompileEnvironment, Graph);
 
@@ -1697,7 +1720,7 @@ namespace UnrealBuildTool
 		public CppCompileEnvironment CreateCompileEnvironmentForIntellisense(ReadOnlyTargetRules Target, CppCompileEnvironment BaseCompileEnvironment, ILogger Logger)
 		{
 			CppCompileEnvironment CompileEnvironment = CreateModuleCompileEnvironment(Target, BaseCompileEnvironment, Logger);
-			CompileEnvironment = SetupPrecompiledHeaders(Target, null, CompileEnvironment, new List<FileItem>(), new NullActionGraphBuilder(Logger));
+			CompileEnvironment = SetupPrecompiledHeaders(Target, null, CompileEnvironment, new List<FileItem>(), new NullActionGraphBuilder(Logger), Logger);
 			CreateHeaderForDefinitions(CompileEnvironment, IntermediateDirectory, null, new NullActionGraphBuilder(Logger));
 			return CompileEnvironment;
 		}
diff --git a/Engine/Source/Programs/UnrealBuildTool/System/PrecompiledHeaderTemplate.cs b/Engine/Source/Programs/UnrealBuildTool/System/PrecompiledHeaderTemplate.cs
--- a/Engine/Source/Programs/UnrealBuildTool/System/PrecompiledHeaderTemplate.cs
+++ b/Engine/Source/Programs/UnrealBuildTool/System/PrecompiledHeaderTemplate.cs
@@ -36,6 +36,11 @@ namespace UnrealBuildTool
 		/// </summary>
 		public List<PrecompiledHeaderInstance> Instances = new List<PrecompiledHeaderInstance>();
 
+		/// <summary>
+		/// All the module dependencies this template has
+		/// </summary>
+		public HashSet<UEBuildModule> ModuleDependencies = new HashSet<UEBuildModule>();
+
 		/// <summary>
 		/// Constructor
 		/// </summary>
@@ -43,12 +48,14 @@ namespace UnrealBuildTool
 		/// <param name="BaseCompileEnvironment">The compile environment to use</param>
 		/// <param name="HeaderFile">The header file to generate a PCH from</param>
 		/// <param name="OutputDir">Output directory for instances of this PCH</param>
-		public PrecompiledHeaderTemplate(UEBuildModuleCPP Module, CppCompileEnvironment BaseCompileEnvironment, FileItem HeaderFile, DirectoryReference OutputDir)
+		/// <param name="ModuleDependencies">All the module dependencies this template has</param>
+		public PrecompiledHeaderTemplate(UEBuildModuleCPP Module, CppCompileEnvironment BaseCompileEnvironment, FileItem HeaderFile, DirectoryReference OutputDir, HashSet<UEBuildModule> ModuleDependencies)
 		{
 			this.Module = Module;
 			this.BaseCompileEnvironment = BaseCompileEnvironment;
 			this.HeaderFile = HeaderFile;
 			this.OutputDir = OutputDir;
+			this.ModuleDependencies = ModuleDependencies;
 		}
 
 		/// <summary>
