Add an option to embed an alternate PDB search path into a Windows executable. This can be used to store PDB files in e.g. subdirectories. The default behaviour remains that the PDB file is searched in the same directory as the executable.

#rnx

diff --git a/Engine/Source/Programs/UnrealBuildTool/Platform/Windows/UEBuildWindows.cs b/Engine/Source/Programs/UnrealBuildTool/Platform/Windows/UEBuildWindows.cs
--- a/Engine/Source/Programs/UnrealBuildTool/Platform/Windows/UEBuildWindows.cs
+++ b/Engine/Source/Programs/UnrealBuildTool/Platform/Windows/UEBuildWindows.cs
@@ -393,6 +393,15 @@ namespace UnrealBuildTool
 		[CommandLine("-StrictEnumTypes")]
 		public bool bStrictEnumTypesConformance = false;
 
+		/// <summary>
+		/// Specify an alternate location for the PDB file. This option does not change the location of the generated PDB file,
+		/// it changes the name that is embedded into the executable. Path can contain %_PDB% which will be expanded to the original
+		/// PDB file name of the target, without the directory.
+		/// See https://learn.microsoft.com/en-us/cpp/build/reference/pdbaltpath-use-alternate-pdb-path
+		/// </summary>
+		[CommandLine("-PdbAltPath")]
+		public string? PdbAlternatePath = null;
+
 		/// VS2015 updated some of the CRT definitions but not all of the Windows SDK has been updated to match.
 		/// Microsoft provides legacy_stdio_definitions library to enable building with VS2015 until they fix everything up.
 		public bool bNeedsLegacyStdioDefinitionsLib => Compiler.IsMSVC() || Compiler.IsClang();
@@ -676,6 +685,8 @@ namespace UnrealBuildTool
 
 		public bool bStrictEnumTypesConformance => Inner.bStrictEnumTypesConformance;
 
+		public string? PdbAlternatePath => Inner.PdbAlternatePath;
+
 		public int DefaultStackSize => Inner.DefaultStackSize;
 
 		public int DefaultStackSizeCommit => Inner.DefaultStackSizeCommit;
diff --git a/Engine/Source/Programs/UnrealBuildTool/Platform/Windows/VCToolChain.cs b/Engine/Source/Programs/UnrealBuildTool/Platform/Windows/VCToolChain.cs
--- a/Engine/Source/Programs/UnrealBuildTool/Platform/Windows/VCToolChain.cs
+++ b/Engine/Source/Programs/UnrealBuildTool/Platform/Windows/VCToolChain.cs
@@ -1163,8 +1163,16 @@ namespace UnrealBuildTool
 				Arguments.Add("/DLL");
 			}
 
-			// Don't embed the full PDB path; we want to be able to move binaries elsewhere. They will always be side by side.
-			Arguments.Add("/PDBALTPATH:%_PDB%");
+			if (String.IsNullOrEmpty(Target.WindowsPlatform.PdbAlternatePath))
+			{
+				// Don't embed the full PDB path; we want to be able to move binaries elsewhere. They will always be side by side.
+				Arguments.Add("/PDBALTPATH:%_PDB%");
+			}
+			else
+			{
+				// Embed an alternate PDB path into the executable
+				Arguments.Add($"/PDBALTPATH:\"{Target.WindowsPlatform.PdbAlternatePath}\"");
+			}
 
 			// Deterministic link support
 			if (LinkEnvironment.bDeterministic)
