Added a common UECommon.props file that contains common properties all VS projects need.
Split the common paths into separate properties to save some memory.

Before: 82.6 MB (86,654,951 bytes)
After: 72.3 MB (75,873,141 bytes)

diff --git a/Engine/Source/Programs/UnrealBuildTool/ProjectFiles/VisualStudio/VCProject.cs b/Engine/Source/Programs/UnrealBuildTool/ProjectFiles/VisualStudio/VCProject.cs
--- a/Engine/Source/Programs/UnrealBuildTool/ProjectFiles/VisualStudio/VCProject.cs
+++ b/Engine/Source/Programs/UnrealBuildTool/ProjectFiles/VisualStudio/VCProject.cs
@@ -14,6 +14,7 @@ using System.Xml;
 using EpicGames.Core;
 using Microsoft.Extensions.Logging;
 using UnrealBuildBase;
+using Newtonsoft.Json.Linq;
 
 namespace UnrealBuildTool
 {
@@ -1066,13 +1067,7 @@ namespace UnrealBuildTool
 			{
 				VCProjectFileContent.AppendLine("  <PropertyGroup Label=\"Globals\">");
 				VCProjectFileContent.AppendLine("    <ProjectGuid>{0}</ProjectGuid>", ProjectGUID.ToString("B").ToUpperInvariant());
-				VCProjectFileContent.AppendLine("    <Keyword>MakeFileProj</Keyword>");
 				VCProjectFileContent.AppendLine("    <RootNamespace>{0}</RootNamespace>", ProjectName);
-				VCProjectFileGenerator.AppendPlatformToolsetProperty(VCProjectFileContent, ProjectFileFormat);
-				VCProjectFileContent.AppendLine("    <MinimumVisualStudioVersion>{0}</MinimumVisualStudioVersion>", VCProjectFileGenerator.GetProjectFileToolVersionString(ProjectFileFormat));
-				VCProjectFileContent.AppendLine("    <VCProjectVersion>{0}</VCProjectVersion>", VCProjectFileGenerator.GetProjectFileToolVersionString(ProjectFileFormat));
-				VCProjectFileContent.AppendLine("    <NMakeUseOemCodePage>true</NMakeUseOemCodePage>"); // Fixes mojibake with non-Latin character sets (UE-102825)
-				VCProjectFileContent.AppendLine("    <TargetRuntime>Native</TargetRuntime>");
 				VCProjectFileContent.AppendLine("  </PropertyGroup>");
 			}
 
@@ -1142,14 +1137,6 @@ namespace UnrealBuildTool
 				}
 			}
 
-			VCProjectFileContent.AppendLine("  <Import Project=\"$(VCTargetsPath)\\Microsoft.Cpp.Default.props\" />");
-
-			// Write the default configuration info
-			VCProjectFileContent.AppendLine("  <PropertyGroup Label=\"Configuration\">");
-			VCProjectFileContent.AppendLine($"    <ConfigurationType>{PlatformProjectGenerator.DefaultPlatformConfigurationType}</ConfigurationType>");
-			VCProjectFileGenerator.AppendPlatformToolsetProperty(VCProjectFileContent, ProjectFileFormat);
-			VCProjectFileContent.AppendLine("  </PropertyGroup>");
-
 			// Write the per platform/config configuration info
 			foreach (Tuple<string, UnrealTargetConfiguration> ConfigurationTuple in ProjectConfigurationNameAndConfigurations)
 			{
@@ -1194,42 +1181,10 @@ namespace UnrealBuildTool
 				}
 			}
 
-			VCProjectFileContent.AppendLine("  <Import Project=\"$(VCTargetsPath)\\Microsoft.Cpp.props\" />");
+			VCProjectFileContent.AppendLine("  <Import Project=\"UECommon.props\" />");
 			VCProjectFileContent.AppendLine("  <ImportGroup Label=\"ExtensionSettings\" />");
 			VCProjectFileContent.AppendLine("  <PropertyGroup Label=\"UserMacros\" />");
 
-			// Write the common and invalid configuration values
-			{
-				const string InvalidMessage = "echo The selected platform/configuration is not valid for this target.";
-
-				string ProjectRelativeUnusedDirectory = NormalizeProjectPath(DirectoryReference.Combine(Unreal.EngineDirectory, "Intermediate", "Build", "Unused"));
-
-				VCProjectFileContent.AppendLine("  <PropertyGroup>");
-				VCProjectFileContent.AppendLine("    <NMakeBuildCommandLine>{0}</NMakeBuildCommandLine>", InvalidMessage);
-				VCProjectFileContent.AppendLine("    <NMakeReBuildCommandLine>{0}</NMakeReBuildCommandLine>", InvalidMessage);
-				VCProjectFileContent.AppendLine("    <NMakeCleanCommandLine>{0}</NMakeCleanCommandLine>", InvalidMessage);
-				VCProjectFileContent.AppendLine("    <NMakeOutput>Invalid Output</NMakeOutput>", InvalidMessage);
-				VCProjectFileContent.AppendLine("    <OutDir>{0}{1}</OutDir>", ProjectRelativeUnusedDirectory, Path.DirectorySeparatorChar);
-				VCProjectFileContent.AppendLine("    <IntDir>{0}{1}</IntDir>", ProjectRelativeUnusedDirectory, Path.DirectorySeparatorChar);
-				// NOTE: We are intentionally overriding defaults for these paths with empty strings.  We never want Visual Studio's
-				//       defaults for these fields to be propagated, since they are version-sensitive paths that may not reflect
-				//       the environment that UBT is building in.  We'll set these environment variables ourselves!
-				// NOTE: We don't touch 'ExecutablePath' because that would result in Visual Studio clobbering the system "Path"
-				//       environment variable
-				VCProjectFileContent.AppendLine("    <IncludePath />");
-				VCProjectFileContent.AppendLine("    <ReferencePath />");
-				VCProjectFileContent.AppendLine("    <LibraryPath />");
-				VCProjectFileContent.AppendLine("    <LibraryWPath />");
-				VCProjectFileContent.AppendLine("    <SourcePath />");
-				VCProjectFileContent.AppendLine("    <ExcludePath />");
-				VCProjectFileContent.AppendLine("  </PropertyGroup>");
-			}
-
-			// Write default import group
-			VCProjectFileContent.AppendLine("  <ImportGroup Label=\"PropertySheets\">");
-			VCProjectFileContent.AppendLine("    <Import Project=\"$(UserRootDir)\\Microsoft.Cpp.$(Platform).user.props\" Condition=\"exists('$(UserRootDir)\\Microsoft.Cpp.$(Platform).user.props')\" Label=\"LocalAppDataPlatform\" />");
-			VCProjectFileContent.AppendLine("  </ImportGroup>");
-
 			// Write each project configuration
 			foreach (ProjectConfigAndTargetCombination Combination in ProjectConfigAndTargetCombinations)
 			{
@@ -1375,58 +1330,98 @@ namespace UnrealBuildTool
 				{
 					StringBuilder CommonProjectFileContent = new StringBuilder();
 
-					//AdditionalIncludeDirectories
+					var AddProperties = (IDictionary<DirectoryReference, string> DirectoryToStringDict, string CommonPropertyPrefix, string PropertyNamePrefix, string PropertyValuePrefix) =>
 					{
-						int ValueIndex = 0;
 						Dictionary<string, string> UpdatedValues = new();
-						List<KeyValuePair<DirectoryReference, string>> KVPList = DirectoryToIncludeSearchPaths.ToList();
+						List<KeyValuePair<DirectoryReference, string>> KVPList = DirectoryToStringDict.ToList();
 						KVPList.SortBy(kvp => kvp.Key);
-						foreach (KeyValuePair<DirectoryReference, string> DirectoryKVP in KVPList)
+
+						// Map each directory to a property
 						{
-							string? Value = DirectoryKVP.Value;
-							if (!String.IsNullOrEmpty(Value))
+							int PropertyIndex = 0;
+							foreach (KeyValuePair<DirectoryReference, string> DirectoryKVP in KVPList)
 							{
-								if (!UpdatedValues.ContainsKey(Value))
+								string? Value = DirectoryKVP.Value;
+								if (!String.IsNullOrEmpty(Value))
 								{
-									string PropertyName = "ClCompile_AdditionalIncludeDirectories";
-									if (ValueIndex > 0)
+									if (!UpdatedValues.ContainsKey(Value))
 									{
-										PropertyName += "_" + ValueIndex;
+										string PropertyName = PropertyNamePrefix;
+										if (PropertyIndex > 0)
+										{
+											PropertyName += "_" + PropertyIndex;
+										}
+										PropertyIndex++;
+										UpdatedValues.Add(Value, PropertyName);
 									}
-									ValueIndex++;
-									UpdatedValues.Add(Value, PropertyName);
-									CommonProjectFileContent.AppendLine($"    <{PropertyName}>$(NMakeIncludeSearchPath);{Value}</{PropertyName}>");
+									DirectoryToStringDict[DirectoryKVP.Key] = UpdatedValues[Value];
 								}
-								DirectoryToIncludeSearchPaths[DirectoryKVP.Key] = UpdatedValues[Value];
 							}
 						}
-					}
-					//ForcedIncludeFiles
-					{
-						int ValueIndex = 0;
-						Dictionary<string, string> UpdatedValues = new();
-						List<KeyValuePair<DirectoryReference, string>> KVPList = DirectoryToForceIncludePaths.ToList();
-						KVPList.SortBy(kvp => kvp.Key);
-						foreach (KeyValuePair<DirectoryReference, string> DirectoryKVP in KVPList)
+
+						// Find the common property values
+						Dictionary<string, string> ValueToCommonPropertyDict = new();
 						{
-							string? Value = DirectoryKVP.Value;
-							if (!String.IsNullOrEmpty(Value))
+							Dictionary<string, int> ValueAndCount = new();
+							foreach (var PropertyValue in UpdatedValues.Keys)
 							{
-								if (!UpdatedValues.ContainsKey(Value))
+								if (!String.IsNullOrEmpty(PropertyValue))
 								{
-									string PropertyName = "ClCompile_ForcedIncludeFiles";
-									if (ValueIndex > 0)
+									IEnumerable<string> PathList = PropertyValue.Split(new[] { ';' }, StringSplitOptions.RemoveEmptyEntries);
+									foreach (string Path in PathList)
 									{
-										PropertyName += "_" + ValueIndex;
+										if (ValueAndCount.TryGetValue(Path, out int PathCount))
+										{
+											ValueAndCount[Path] = ++PathCount;
+										}
+										else
+										{
+											ValueAndCount[Path] = 1;
+										}
 									}
-									ValueIndex++;
-									UpdatedValues.Add(Value, PropertyName);
-									CommonProjectFileContent.AppendLine($"    <{PropertyName}>{Value}</{PropertyName}>");
 								}
-								DirectoryToForceIncludePaths[DirectoryKVP.Key] = UpdatedValues[Value];
+							}
+
+							var CommonProperties = ValueAndCount.Where(kvp => kvp.Value > 1).Select(kvp => kvp.Key).ToList();
+							CommonProperties.Sort();
+
+							// Write out the common property values
+							int CommonPropertyValueIndex = 0;
+							foreach (var CommonProperty in CommonProperties)
+							{
+								string PropertyName = CommonPropertyPrefix;
+								if (CommonPropertyValueIndex > 0)
+								{
+									PropertyName += "_" + CommonPropertyValueIndex;
+								}
+								CommonPropertyValueIndex++;
+								ValueToCommonPropertyDict.Add(CommonProperty, $"$({PropertyName})");
+								CommonProjectFileContent.AppendLine($"    <{PropertyName}>{CommonProperty}</{PropertyName}>");
 							}
 						}
-					}
+
+						// Write out the updated properties
+						foreach (KeyValuePair<string, string> kvp in UpdatedValues)
+						{
+							string[] PathArray = kvp.Key.Split(new[] { ';' }, StringSplitOptions.RemoveEmptyEntries);
+							for (int Index = 0; Index < PathArray.Length; Index++)
+							{
+								if (ValueToCommonPropertyDict.TryGetValue(PathArray[Index], out string? PropertyValue) && !String.IsNullOrEmpty(PropertyValue))
+								{
+									PathArray[Index] = PropertyValue;
+								}
+							}
+							string NewValue = String.Join(";", PathArray);
+							CommonProjectFileContent.AppendLine($"    <{kvp.Value}>{PropertyValuePrefix}{NewValue}</{kvp.Value}>");
+						}
+					};
+
+					//AdditionalIncludeDirectories
+					AddProperties(DirectoryToIncludeSearchPaths, "ProjectAdditionalIncludeDirectories", "ClCompile_AdditionalIncludeDirectories", "$(NMakeIncludeSearchPath);");
+
+					//ForcedIncludeFiles
+					AddProperties(DirectoryToForceIncludePaths, "ProjectForcedIncludeFiles", "ClCompile_ForcedIncludeFiles", String.Empty);
+
 					// AdditionalOptions
 					{
 						int ValueIndex = 0;
@@ -2098,9 +2093,9 @@ namespace UnrealBuildTool
 					string BuildArguments = Builder.GetBuildArguments();
 
 					// NMake Build command line
-					VCProjectFileContent.AppendLine("    <NMakeBuildCommandLine>{0} {1}</NMakeBuildCommandLine>", EscapePath(NormalizeProjectPath(Builder.BuildScript)), BuildArguments);
-					VCProjectFileContent.AppendLine("    <NMakeReBuildCommandLine>{0} {1}</NMakeReBuildCommandLine>", EscapePath(NormalizeProjectPath(Builder.RebuildScript)), BuildArguments);
-					VCProjectFileContent.AppendLine("    <NMakeCleanCommandLine>{0} {1}</NMakeCleanCommandLine>", EscapePath(NormalizeProjectPath(Builder.CleanScript)), BuildArguments);
+					VCProjectFileContent.AppendLine("    <NMakeBuildCommandLine>$(BuildBatchScript) {0}</NMakeBuildCommandLine>", BuildArguments);
+					VCProjectFileContent.AppendLine("    <NMakeReBuildCommandLine>$(RebuildBatchScript) {0}</NMakeReBuildCommandLine>", BuildArguments);
+					VCProjectFileContent.AppendLine("    <NMakeCleanCommandLine>$(CleanBatchScript) {0}</NMakeCleanCommandLine>", BuildArguments);
 					VCProjectFileContent.AppendLine("    <NMakeOutput>{0}</NMakeOutput>", NormalizeProjectPath(NMakePath.FullName));
 					VCProjectFileContent.AppendLine("    <AdditionalOptions>{0} {1} {2}</AdditionalOptions>", GetCppStandardCompileArgument(TargetRulesObject.CppStandard), GetEnableCoroutinesArgument(), GetConformanceCompileArguments(TargetRulesObject));
 
@@ -2344,4 +2339,5 @@ namespace UnrealBuildTool
 		/// Cache of parsed info about this project
 		protected readonly Dictionary<UnrealTargetConfiguration, CsProjectInfo> CachedProjectInfo = new Dictionary<UnrealTargetConfiguration, CsProjectInfo>();
 	}
+
 }
diff --git a/Engine/Source/Programs/UnrealBuildTool/ProjectFiles/VisualStudio/VCProjectFileGenerator.cs b/Engine/Source/Programs/UnrealBuildTool/ProjectFiles/VisualStudio/VCProjectFileGenerator.cs
--- a/Engine/Source/Programs/UnrealBuildTool/ProjectFiles/VisualStudio/VCProjectFileGenerator.cs
+++ b/Engine/Source/Programs/UnrealBuildTool/ProjectFiles/VisualStudio/VCProjectFileGenerator.cs
@@ -11,6 +11,7 @@ using System.Xml.Linq;
 using EpicGames.Core;
 using Microsoft.Extensions.Logging;
 using UnrealBuildBase;
+using Newtonsoft.Json.Linq;
 
 namespace UnrealBuildTool
 {
@@ -472,12 +473,87 @@ namespace UnrealBuildTool
 			return MakeMd5Guid(Input);
 		}
 
+		private void WriteCommonPropsFile(ILogger Logger)
+		{
+			using (ProgressWriter Progress = new ProgressWriter("Creating common properties file...", true, Logger))
+			{
+				StringBuilder VCCommonTargetFileContent = new StringBuilder();
+				VCCommonTargetFileContent.AppendLine("<?xml version=\"1.0\" encoding=\"utf-8\"?>");
+				VCCommonTargetFileContent.AppendLine("<Project xmlns=\"http://schemas.microsoft.com/developer/msbuild/2003\">");
+
+				// Project globals (project GUID, project type, SCC bindings, etc)
+				{
+					string ToolVersionString = GetProjectFileToolVersionString(Settings.ProjectFileFormat);
+					VCCommonTargetFileContent.AppendLine("  <PropertyGroup Label=\"Globals\">");
+					VCCommonTargetFileContent.AppendLine("    <Keyword>MakeFileProj</Keyword>");
+					AppendPlatformToolsetProperty(VCCommonTargetFileContent, Settings.ProjectFileFormat);
+					VCCommonTargetFileContent.AppendLine("    <MinimumVisualStudioVersion>{0}</MinimumVisualStudioVersion>", ToolVersionString);
+					VCCommonTargetFileContent.AppendLine("    <VCProjectVersion>{0}</VCProjectVersion>", ToolVersionString);
+					VCCommonTargetFileContent.AppendLine("    <NMakeUseOemCodePage>true</NMakeUseOemCodePage>"); // Fixes mojibake with non-Latin character sets (UE-102825)
+					VCCommonTargetFileContent.AppendLine("    <TargetRuntime>Native</TargetRuntime>");
+					VCCommonTargetFileContent.AppendLine("  </PropertyGroup>");
+				}
+
+				// Write the default configuration info
+				VCCommonTargetFileContent.AppendLine("  <PropertyGroup Label=\"Configuration\">");
+				VCCommonTargetFileContent.AppendLine($"    <ConfigurationType>{PlatformProjectGenerator.DefaultPlatformConfigurationType}</ConfigurationType>");
+				AppendPlatformToolsetProperty(VCCommonTargetFileContent, Settings.ProjectFileFormat);
+				VCCommonTargetFileContent.AppendLine("  </PropertyGroup>");
+
+				VCCommonTargetFileContent.AppendLine("  <Import Project=\"$(VCTargetsPath)\\Microsoft.Cpp.Default.props\" />");
+				VCCommonTargetFileContent.AppendLine("  <Import Project=\"$(VCTargetsPath)\\Microsoft.Cpp.props\" />");
+
+				// Write the common and invalid configuration values
+				{
+					const string InvalidMessage = "echo The selected platform/configuration is not valid for this target.";
+
+					string ProjectRelativeUnusedDirectory = ProjectFile.NormalizeProjectPath(DirectoryReference.Combine(Unreal.EngineDirectory, "Intermediate", "Build", "Unused"));
+
+					VCCommonTargetFileContent.AppendLine("  <PropertyGroup>");
+
+					DirectoryReference BatchFilesDirectory = DirectoryReference.Combine(Unreal.EngineDirectory, "Build", "BatchFiles");
+					VCCommonTargetFileContent.AppendLine("    <BuildBatchScript>{0}</BuildBatchScript>", ProjectFile.EscapePath(ProjectFile.NormalizeProjectPath(FileReference.Combine(BatchFilesDirectory, "Build.bat"))));
+					VCCommonTargetFileContent.AppendLine("    <RebuildBatchScript>{0}</RebuildBatchScript>", ProjectFile.EscapePath(ProjectFile.NormalizeProjectPath(FileReference.Combine(BatchFilesDirectory, "Rebuild.bat"))));
+					VCCommonTargetFileContent.AppendLine("    <CleanBatchScript>{0}</CleanBatchScript>", ProjectFile.EscapePath(ProjectFile.NormalizeProjectPath(FileReference.Combine(BatchFilesDirectory, "Clean.bat"))));
+					VCCommonTargetFileContent.AppendLine("    <NMakeBuildCommandLine>{0}</NMakeBuildCommandLine>", InvalidMessage);
+					VCCommonTargetFileContent.AppendLine("    <NMakeReBuildCommandLine>{0}</NMakeReBuildCommandLine>", InvalidMessage);
+					VCCommonTargetFileContent.AppendLine("    <NMakeCleanCommandLine>{0}</NMakeCleanCommandLine>", InvalidMessage);
+					VCCommonTargetFileContent.AppendLine("    <NMakeOutput>Invalid Output</NMakeOutput>", InvalidMessage);
+					VCCommonTargetFileContent.AppendLine("    <OutDir>{0}{1}</OutDir>", ProjectRelativeUnusedDirectory, Path.DirectorySeparatorChar);
+					VCCommonTargetFileContent.AppendLine("    <IntDir>{0}{1}</IntDir>", ProjectRelativeUnusedDirectory, Path.DirectorySeparatorChar);
+					// NOTE: We are intentionally overriding defaults for these paths with empty strings.  We never want Visual Studio's
+					//       defaults for these fields to be propagated, since they are version-sensitive paths that may not reflect
+					//       the environment that UBT is building in.  We'll set these environment variables ourselves!
+					// NOTE: We don't touch 'ExecutablePath' because that would result in Visual Studio clobbering the system "Path"
+					//       environment variable
+					VCCommonTargetFileContent.AppendLine("    <IncludePath />");
+					VCCommonTargetFileContent.AppendLine("    <ReferencePath />");
+					VCCommonTargetFileContent.AppendLine("    <LibraryPath />");
+					VCCommonTargetFileContent.AppendLine("    <LibraryWPath />");
+					VCCommonTargetFileContent.AppendLine("    <SourcePath />");
+					VCCommonTargetFileContent.AppendLine("    <ExcludePath />");
+					VCCommonTargetFileContent.AppendLine("  </PropertyGroup>");
+				}
+
+				// Write default import group
+				VCCommonTargetFileContent.AppendLine("  <ImportGroup Label=\"PropertySheets\">");
+				VCCommonTargetFileContent.AppendLine("    <Import Project=\"$(UserRootDir)\\Microsoft.Cpp.$(Platform).user.props\" Condition=\"exists('$(UserRootDir)\\Microsoft.Cpp.$(Platform).user.props')\" Label=\"LocalAppDataPlatform\" />");
+				VCCommonTargetFileContent.AppendLine("  </ImportGroup>");
+
+				VCCommonTargetFileContent.AppendLine("</Project>");
+
+				Utils.WriteFileIfChanged(FileReference.Combine(IntermediateProjectFilesPath, "UECommon.props"), VCCommonTargetFileContent.ToString(), Logger);
+			}
+		}
+
 		/// <summary>
 		/// Writes the project files to disk
 		/// </summary>
 		/// <returns>True if successful</returns>
 		protected override bool WriteProjectFiles(PlatformProjectGeneratorCollection PlatformProjectGenerators, ILogger Logger)
 		{
+			WriteCommonPropsFile(Logger);
+
 			if (!base.WriteProjectFiles(PlatformProjectGenerators, Logger))
 			{
 				return false;
