- Fixed a bug caused by the C# auto-cleanup of code. (foo != "") is not the same as (!string.IsNullOrEmpty(foo)) , it is (foo == null || foo.Length() > 0)

diff --git a/Engine/Source/Programs/UnrealBuildTool/ProjectFiles/Xcode/XcodeProject.File.cs b/Engine/Source/Programs/UnrealBuildTool/ProjectFiles/Xcode/XcodeProject.File.cs
--- a/Engine/Source/Programs/UnrealBuildTool/ProjectFiles/Xcode/XcodeProject.File.cs
+++ b/Engine/Source/Programs/UnrealBuildTool/ProjectFiles/Xcode/XcodeProject.File.cs
@@ -268,7 +268,7 @@ namespace UnrealBuildTool.XcodeProjectXcconfig
 				}
 			}
 
-			if (!String.IsNullOrEmpty(GroupName))
+			if (GroupName == null || GroupName.Length > 0)
 			{
 				// group the files by path
 				XcodeFileGroup? Group = GroupName == null ? FindGroupByAbsolutePath(File.Reference.Directory) : null;
