Regularize support for LTO/ThinLTO for Linux.

- Issue a warning when ThinLTO flag is specified on the command line but has no effect.

diff --git a/Engine/Source/Programs/UnrealBuildTool/Platform/Linux/LinuxToolChain.cs b/Engine/Source/Programs/UnrealBuildTool/Platform/Linux/LinuxToolChain.cs
--- a/Engine/Source/Programs/UnrealBuildTool/Platform/Linux/LinuxToolChain.cs
+++ b/Engine/Source/Programs/UnrealBuildTool/Platform/Linux/LinuxToolChain.cs
@@ -302,23 +302,44 @@ namespace UnrealBuildTool
 			Arguments.Add("-Wno-undefined-bool-conversion"); // hides checking if 'this' pointer is null
 		}
 
-		/// <inheritdoc/>
-		protected override void GetCompileArguments_Optimizations(CppCompileEnvironment CompileEnvironment, List<string> Arguments)
+		private void AddLTOFlags(List<string> Arguments, bool bForLinker)
 		{
-			base.GetCompileArguments_Optimizations(CompileEnvironment, Arguments);
-
-			// Unlike on other platforms, allow LTO be specified independently of PGO
-			if (CompileEnvironment.bAllowLTCG)
+			if (Options.HasFlag(ClangToolChainOptions.EnableLinkTimeOptimization))
 			{
-				if ((Options & ClangToolChainOptions.EnableThinLTO) != 0)
+				if (Options.HasFlag(ClangToolChainOptions.EnableThinLTO))
 				{
-					Arguments.Add("-flto=thin");
+					if (bForLinker)
+					{
+						Arguments.Add(String.Format("-flto=thin -Wl,--thinlto-jobs={0}", Utils.GetPhysicalProcessorCount()));
+					}
+					else
+					{
+						Arguments.Add("-flto=thin");
+					}
 				}
 				else
 				{
 					Arguments.Add("-flto");
 				}
 			}
+		}
+
+		private void AddCompilerLTOFlags(List<string> Arguments)
+		{
+			 AddLTOFlags(Arguments, false);
+		}
+
+		private void AddLinkerLTOFlags(List<string> Arguments)
+		{
+			AddLTOFlags(Arguments, true);
+		}
+
+		/// <inheritdoc/>
+		protected override void GetCompileArguments_Optimizations(CppCompileEnvironment CompileEnvironment, List<string> Arguments)
+		{
+			base.GetCompileArguments_Optimizations(CompileEnvironment, Arguments);
+
+			AddCompilerLTOFlags(Arguments);
 
 			// optimization level
 			if (!CompileEnvironment.bOptimizeCode)
@@ -718,17 +739,7 @@ namespace UnrealBuildTool
 			}
 
 			// whether we actually can do that is checked in CanUseAdvancedLinkerFeatures() earlier
-			if (LinkEnvironment.bAllowLTCG)
-			{
-				if ((Options & ClangToolChainOptions.EnableThinLTO) != 0)
-				{
-					Arguments.Add(String.Format(" -flto=thin -Wl,--thinlto-jobs={0}", Utils.GetPhysicalProcessorCount()));
-				}
-				else
-				{
-					Arguments.Add("-flto");
-				}
-			}
+			AddLinkerLTOFlags(Arguments);
 
 			if (CrossCompiling())
 			{
diff --git a/Engine/Source/Programs/UnrealBuildTool/Platform/Linux/UEBuildLinux.cs b/Engine/Source/Programs/UnrealBuildTool/Platform/Linux/UEBuildLinux.cs
--- a/Engine/Source/Programs/UnrealBuildTool/Platform/Linux/UEBuildLinux.cs
+++ b/Engine/Source/Programs/UnrealBuildTool/Platform/Linux/UEBuildLinux.cs
@@ -195,6 +195,17 @@ namespace UnrealBuildTool
 			ValidateTarget(Target);
 		}
 
+		public bool IsLTOEnabled(ReadOnlyTargetRules Target)
+		{
+			// Force LTO on if using PGO, or if specified in the target rules.
+			if (Target.bAllowLTCG || Target.bPGOOptimize || Target.bPGOProfile)
+			{
+				return true;
+			}
+
+			return false;
+		}
+
 		public override void ValidateTarget(TargetRules Target)
 		{
 			if (!String.IsNullOrWhiteSpace(Environment.GetEnvironmentVariable("CLANG_STATIC_ANALYZER_MODE")))
@@ -603,9 +614,20 @@ namespace UnrealBuildTool
 					throw new BuildException("LibFuzzer is unsupported for non-monolithic builds.");
 				}
 			}
-			if (Target.bAllowLTCG && Target.bPreferThinLTO)
+
+			if (IsLTOEnabled(Target))
+			{
+				Options |= ClangToolChainOptions.EnableLinkTimeOptimization;
+
+				if (Target.bPreferThinLTO)
+				{
+					Options |= ClangToolChainOptions.EnableThinLTO;
+				}
+			}
+			else if (Target.bPreferThinLTO)
 			{
-				Options |= ClangToolChainOptions.EnableThinLTO;
+				// warn about ThinLTO not being useful on its own
+				Logger.LogWarning("Warning: bPreferThinLTO is set, but LTO is disabled. Flag will have no effect");
 			}
 
 			if (Target.bUseAutoRTFMCompiler)
