- Allow a project to override the Main SDK version for a platform (the platform's project Engine.ini would set
  [OverrideSDK]
  SDKVersion=x.y.z
and AutoSDK, or anything using GetMainVersion(), would use that version string instead
- If multiple targets are built in one run of UBT/UAT, it will error out because there is no support to switch SDKs (and builds happen in parallel, etc)
- This is not a complete solution, because it can cause problems with shared tools like ShaderCompileWorker, when different projects are on different SDKs and they have one SCW to share
- Renamed GetMainVersion to GetMainVersionInternal(), and wrapped that in a new non-vitual GetMainVersion() that calls GetMainVersionInternal() (and handles the ini overrides)

diff --git a/Engine/Source/Developer/TargetPlatform/Private/TargetPlatformManagerModule.cpp b/Engine/Source/Developer/TargetPlatform/Private/TargetPlatformManagerModule.cpp
--- a/Engine/Source/Developer/TargetPlatform/Private/TargetPlatformManagerModule.cpp
+++ b/Engine/Source/Developer/TargetPlatform/Private/TargetPlatformManagerModule.cpp
@@ -50,6 +50,23 @@ static bool IsAutoSDKsEnabled()
 	return false;
 }
 
+static FString GetProjectPathForUBT()
+{
+	if (FPaths::IsProjectFilePathSet())
+	{
+		return FPaths::ConvertRelativePathToFull(FPaths::GetProjectFilePath());
+	}
+	if (FApp::HasProjectName())
+	{
+		FString ProjectPath = FPaths::ProjectDir() / FApp::GetProjectName() + TEXT(".uproject");
+		if (FPaths::FileExists(ProjectPath))
+		{
+			return ProjectPath;
+		}
+	}
+	return FString();
+}
+
 // kick off a call to UBT nice and early so that it's results are hopefully ready when needed
 static FProcHandle AutoSDKSetupUBTProc;
 FDelayedAutoRegisterHelper GAutoSDKInit(EDelayedRegisterRunPhase::FileSystemReady, []
@@ -58,6 +75,12 @@ FDelayedAutoRegisterHelper GAutoSDKInit(EDelayedRegisterRunPhase::FileSystemRead
 		if (IsAutoSDKsEnabled() && FParse::Param(FCommandLine::Get(), TEXT("Multiprocess")) == false)
 		{
 			FString UBTParams(TEXT("-Mode=SetupPlatforms"));
+			FString Project = GetProjectPathForUBT();
+			if (Project.Len() > 0)
+			{
+				UBTParams += FString::Printf(TEXT(" -project=%s"), *Project);
+			}
+				
 			int32 UBTReturnCode = -1;
 			FString UBTOutput;
 
diff --git a/Engine/Source/Programs/AutomationTool/AutomationUtils/Automation.cs b/Engine/Source/Programs/AutomationTool/AutomationUtils/Automation.cs
--- a/Engine/Source/Programs/AutomationTool/AutomationUtils/Automation.cs
+++ b/Engine/Source/Programs/AutomationTool/AutomationUtils/Automation.cs
@@ -116,7 +116,7 @@ namespace AutomationTool
 				}
 				
 				// Initialize UBT
-				if (!UnrealBuildTool.PlatformExports.Initialize(Log.Logger))
+				if (!UnrealBuildTool.PlatformExports.Initialize(Environment.GetCommandLineArgs(), Log.Logger))
 				{
 					Logger.LogInformation("Failed to initialize UBT");
 					return ExitCode.Error_Unknown;
diff --git a/Engine/Source/Programs/Shared/EpicGames.Build/UEBuildPlatformSDK.cs b/Engine/Source/Programs/Shared/EpicGames.Build/UEBuildPlatformSDK.cs
--- a/Engine/Source/Programs/Shared/EpicGames.Build/UEBuildPlatformSDK.cs
+++ b/Engine/Source/Programs/Shared/EpicGames.Build/UEBuildPlatformSDK.cs
@@ -311,6 +311,20 @@ namespace EpicGames.Core
 				throw new Exception(string.Format("Re-registering SDK for {0}. All Platforms must have a unique SDK object", PlatformName));
 			}
 
+			// check that none of our override versions conflict with default version, by using MainVersion unless overridden and we aren't _also_
+			// requiring default version
+			if (!PlatformSDKMainVersions.ContainsKey(PlatformName))
+			{
+				PlatformSDKMainVersions.Add(PlatformName, SDK.GetMainVersionInternal());
+			}
+			else if (PlatformsThatNeedDefaultSDK != null && PlatformsThatNeedDefaultSDK.Contains(PlatformName))
+			{
+				if (!PlatformSDKMainVersions[PlatformName].Equals(SDK.GetMainVersionInternal(), StringComparison.OrdinalIgnoreCase))
+				{
+					throw new Exception($"A target overrode SDK for {PlatformName} to version {PlatformSDKMainVersions[PlatformName]}, but another target uses the default SDK version {SDK.GetMainVersionInternal()}");
+				}
+			}
+
 			SDKRegistry.Add(PlatformName, SDK);
 
 			SDK.Init(PlatformName, bIsSdkAllowedOnHost);
@@ -382,6 +396,13 @@ namespace EpicGames.Core
 		// True if this Sdk is allowed to be used by this host - if not, we can skip a lot 
 		public bool bIsSdkAllowedOnHost;
 
+		// Retrieves the version of the default SDK that is preferred choice, as well as the single AutoSDK version that will be used
+		public string GetMainVersion()
+		{
+			// at this point, the PlatformSDKMainVersions map is fully set up with default version, or with overrides that have been verified to not conflict
+			return PlatformSDKMainVersions[PlatformName!];
+		}
+
 		public SDKCollection GetAllSDKInfo()
 		{
 			SDKCollection AllSdks = new SDKCollection(this);
@@ -521,7 +542,7 @@ namespace EpicGames.Core
 		/// Return the SDK version that the platform wants to use (AutoSDK dir must match this, full SDKs can be in a valid range)
 		/// </summary>
 		/// <returns></returns>
-		public abstract string GetMainVersion();
+		protected abstract string GetMainVersionInternal();
 
 		/// <summary>
 		/// Gets the valid string range of Sdk versions. TryConvertVersionToInt() will need to succeed to make this usable for range checks
@@ -804,6 +825,26 @@ namespace EpicGames.Core
 
 		#endregion
 
+		#region Per-project SDK support
+
+		private static Dictionary<string, string> PlatformSDKMainVersions = new();
+		private static HashSet<string>? PlatformsThatNeedDefaultSDK = new();
+		public static void InitializePerProjectSDKVersions(Dictionary<string, string> PerPlatformSDKOverrides, HashSet<string> InPlatformsThatNeedDefaultSDK)
+		{
+			if (PlatformSDKMainVersions.Count > 0)
+			{
+				throw new Exception("InitializePerProjectSDKVersions() was called multiple times, which is not allowed");
+			}
+
+			// cache off the overrides - Register will fill in the rest
+			PlatformSDKMainVersions = PerPlatformSDKOverrides;
+
+			// remember the platforms that need the default version, which we will use later to detect conflicting versions in Register			
+			PlatformsThatNeedDefaultSDK = InPlatformsThatNeedDefaultSDK;
+		}
+
+		#endregion
+
 		#region Print SDK Info
 		private static bool bHasShownTurnkey = false;
 		private SDKStatus? SDKInfoValidity = null;
diff --git a/Engine/Source/Programs/UnrealBuildTool/Configuration/TargetDescriptor.cs b/Engine/Source/Programs/UnrealBuildTool/Configuration/TargetDescriptor.cs
--- a/Engine/Source/Programs/UnrealBuildTool/Configuration/TargetDescriptor.cs
+++ b/Engine/Source/Programs/UnrealBuildTool/Configuration/TargetDescriptor.cs
@@ -356,7 +356,7 @@ namespace UnrealBuildTool
 			}
 
 			// Expand all the platforms, architectures and configurations
-			foreach (UnrealTargetPlatform Platform in Platforms)
+			foreach (UnrealTargetPlatform Platform in Platforms.Distinct())
 			{
 				// Make sure the platform is valid
 				if (!InstalledPlatformInfo.IsValid(null, Platform, null, EProjectType.Code, InstalledPlatformState.Downloaded))
@@ -371,13 +371,16 @@ namespace UnrealBuildTool
 					}
 				}
 
-				UEBuildPlatform BuildPlatform = UEBuildPlatform.GetBuildPlatform(Platform);
+				// there are times that this can be run before we have setup build platforms (ie Per-Project AutoSDK version), in which case we
+				// cannot get Architecture info, but it's not needed yet
+				UEBuildPlatform? BuildPlatform;
+				UEBuildPlatform.TryGetBuildPlatform(Platform, out BuildPlatform);
 
 				// Parse the architecture parameter, or use null to look up platform defaults later
 				string ParamArchitectureList = Arguments.GetStringOrDefault("-Architecture=", "") + Arguments.GetStringOrDefault("-Architectures=", "");
 				UnrealArchitectures? ParamArchitectures = UnrealArchitectures.FromString(ParamArchitectureList, Platform);
 
-				foreach (UnrealTargetConfiguration Configuration in Configurations)
+				foreach (UnrealTargetConfiguration Configuration in Configurations.Distinct())
 				{
 					// Create all the target descriptors for targets specified by type
 					foreach (string TargetTypeString in Arguments.GetValues("-TargetType="))
@@ -430,18 +433,23 @@ namespace UnrealBuildTool
 							NativeProjects.ConditionalMakeTempTargetForHybridProject(TargetProjectFile, Logger);
 						}
 
-						if (ParamArchitectures == null)
+						if (ParamArchitectures != null)
 						{
-							// ask the platform what achitectures it wants for this project
-							Architectures = BuildPlatform.ArchitectureConfig.ActiveArchitectures(TargetProjectFile, TargetName);
+							Architectures = ParamArchitectures;
+						}
+						else if (BuildPlatform == null)
+						{
+							// if we can't use BuildPlatform yet, we just use any random architecture, since we need _something_
+							Architectures = new UnrealArchitectures(UnrealArch.X64);
 						}
 						else
 						{
-							Architectures = ParamArchitectures;
+							// ask the platform what achitectures it wants for this project
+							Architectures = BuildPlatform.ArchitectureConfig.ActiveArchitectures(TargetProjectFile, TargetName);
 						}
 
 						// If the platform wants a target for each architecture, make a target descriptor for each architecture, otherwise one target for all architectures
-						if (BuildPlatform.ArchitectureConfig.Mode == UnrealArchitectureMode.OneTargetPerArchitecture)
+						if (BuildPlatform != null && BuildPlatform.ArchitectureConfig.Mode == UnrealArchitectureMode.OneTargetPerArchitecture)
 						{
 							foreach (UnrealArch Architecture in Architectures.Architectures)
 							{
diff --git a/Engine/Source/Programs/UnrealBuildTool/Configuration/UEBuildPlatform.cs b/Engine/Source/Programs/UnrealBuildTool/Configuration/UEBuildPlatform.cs
--- a/Engine/Source/Programs/UnrealBuildTool/Configuration/UEBuildPlatform.cs
+++ b/Engine/Source/Programs/UnrealBuildTool/Configuration/UEBuildPlatform.cs
@@ -249,13 +249,109 @@ namespace UnrealBuildTool
 			}
 		}
 
+		private static string[] UATProjectParams = { "-project=", "-scriptsforproject=" };
+		// Before we setup AutoSDK, we check to see if any projects need to override the Main version so that AutoSDK
+		// will set up an alternate SDK
+		private static void InitializePerPlatformSDKs(string[] Args, bool bArgumentsAreForUBT, ILogger Logger)
+		{
+			Dictionary<string, string> PlatformToVersionMap = new();
+			HashSet<string> PlatformsThatNeedDefaultSDK = new();
+
+			IEnumerable<FileReference?> ProjectFiles;
+
+			if (bArgumentsAreForUBT == false)
+			{
+				ProjectFiles = Args
+						// find arguments that start with one of the hard-coded parameters
+						.Where(x => UATProjectParams.Any(y => x.StartsWith(y, StringComparison.OrdinalIgnoreCase)))
+						// treat the part after the = as a path to a uproject, and ask the NativeProjects class to find a .uproject
+						.Select(x => NativeProjects.FindProjectFile(x.Substring(x.IndexOf('=') + 1), Logger))
+						// only use existant projects
+						.Where(x => x != null && FileReference.Exists(x));
+			}
+			else
+			{
+				CommandLineArguments CommandLine = new(Args);
+				BuildConfiguration BuildConfiguration = new();
+				XmlConfig.ApplyTo(BuildConfiguration);
+				CommandLine.ApplyTo(BuildConfiguration);
+
+				// get the project files for all targets - we allow null uproject files which means to use the defaults
+				// (same as a uproject that has no override SDK versions set)
+				List<TargetDescriptor> AllTargets = TargetDescriptor.ParseCommandLine(CommandLine, BuildConfiguration.bUsePrecompiled, bSkipRulesCompile: true, bForceRulesCompile: false, Logger);
+				ProjectFiles = AllTargets.Select(x => x.ProjectFile);
+			}
+
+			foreach (FileReference? ProjectFile in ProjectFiles)
+			{
+				// because we can't safely tell, at this point, if any platform SDKs will be used for a given target, we have to check all of them
+				foreach (UnrealTargetPlatform Platform in UnrealTargetPlatform.GetValidPlatforms())
+				{
+					string PlatformName = Platform.ToString();
+					string? OverrideSDKVersion = null;
+
+					if (ProjectFile != null)
+					{
+						ConfigHierarchy Ini = ConfigCache.ReadHierarchy(ConfigHierarchyType.Engine, ProjectFile.Directory, Platform);
+						if (Ini.TryGetValue("OverrideSDK", "SDKVersion", out OverrideSDKVersion))
+						{
+							Logger.LogWarning("Project {Project} is overriding the active main SDK for {Platform} to {Override}", ProjectFile.GetFileNameWithoutAnyExtensions(), Platform, OverrideSDKVersion);
+
+							if (PlatformToVersionMap.ContainsKey(PlatformName) && PlatformToVersionMap[PlatformName] != OverrideSDKVersion)
+							{
+								throw new BuildException("Attempting to build with mismatching SDK versions! [MORE INFO HERE, BUT ALSO CHECK THAT ALL PROJECTS USE SAME VERSION]");
+							}
+						}
+					}
+					
+					// if we have a null version here, that means to use the default, later when the PlatformSDK is created
+					if (OverrideSDKVersion == null)
+					{
+						// mark _something_ in the map so we can detect mismatches
+						PlatformsThatNeedDefaultSDK.Add(PlatformName);
+					}
+					else
+					{
+						PlatformToVersionMap[PlatformName] = OverrideSDKVersion;
+					}
+				}
+			}
+
+			UEBuildPlatformSDK.InitializePerProjectSDKVersions(PlatformToVersionMap, PlatformsThatNeedDefaultSDK);
+		}
+
+		/// <summary>
+		/// Finds all the UEBuildPlatformFactory types in this assembly and uses them to register all the available platforms, and uses a UBT commandline to check for per-project SDKs
+		/// </summary>
+		/// <param name="bIncludeNonInstalledPlatforms">Whether to register platforms that are not installed</param>
+		/// <param name="bHostPlatformOnly">Only register the host platform</param>
+		/// <param name="ArgumentsForPerPlatform">Commandline args to look through for finding uprojects, to look up per-project SDK versions</param>
+		/// <param name="UBTModeType">The UBT mode (usually Build, but </param>
+		/// <param name="Logger">Logger for output</param>
+		internal static void RegisterPlatforms(bool bIncludeNonInstalledPlatforms, bool bHostPlatformOnly, Type UBTModeType, string[] ArgumentsForPerPlatform, ILogger Logger)
+		{
+			bool bUseTargetTripleParams = false;
+			// @todo : add a ToolModeOptions.UseTargetTripleCommandLine or something
+			if (UBTModeType.Name == "BuildMode")
+			{
+				bUseTargetTripleParams = true;
+			}
+			RegisterPlatforms(bIncludeNonInstalledPlatforms, bHostPlatformOnly, ArgumentsForPerPlatform, bArgumentsAreForUBT: bUseTargetTripleParams, Logger);
+		}
+
 		/// <summary>
-		/// Finds all the UEBuildPlatformFactory types in this assembly and uses them to register all the available platforms
+		/// Finds all the UEBuildPlatformFactory types in this assembly and uses them to register all the available platforms, and uses a UAT commandline to check for per-project SDKs
 		/// </summary>
 		/// <param name="bIncludeNonInstalledPlatforms">Whether to register platforms that are not installed</param>
 		/// <param name="bHostPlatformOnly">Only register the host platform</param>
+		/// <param name="ArgumentsForPerPlatform">Commandline args to look through for finding uprojects, to look up per-project SDK versions</param>
 		/// <param name="Logger">Logger for output</param>
-		public static void RegisterPlatforms(bool bIncludeNonInstalledPlatforms, bool bHostPlatformOnly, ILogger Logger)
+		internal static void RegisterPlatforms(bool bIncludeNonInstalledPlatforms, bool bHostPlatformOnly, string[] ArgumentsForPerPlatform, ILogger Logger)
+		{
+			RegisterPlatforms(bIncludeNonInstalledPlatforms, bHostPlatformOnly, ArgumentsForPerPlatform, bArgumentsAreForUBT:false, Logger);
+		}
+
+		private static void RegisterPlatforms(bool bIncludeNonInstalledPlatforms, bool bHostPlatformOnly, string[] ArgumentsForPerPlatform, bool bArgumentsAreForUBT, ILogger Logger)
 		{
 			// Initialize the installed platform info
 			using (GlobalTracer.Instance.BuildSpan("Initializing InstalledPlatformInfo").StartActive())
@@ -263,6 +359,11 @@ namespace UnrealBuildTool
 				InstalledPlatformInfo.Initialize();
 			}
 
+			using (GlobalTracer.Instance.BuildSpan("Initializing PerPlatformSDKs").StartActive())
+			{
+				InitializePerPlatformSDKs(ArgumentsForPerPlatform, bArgumentsAreForUBT, Logger);
+			}
+
 			// Find and register all tool chains and build platforms that are present
 			Type[] AllTypes;
 			using (GlobalTracer.Instance.BuildSpan("Querying types").StartActive())
diff --git a/Engine/Source/Programs/UnrealBuildTool/Modes/SetupPlatformsMode.cs b/Engine/Source/Programs/UnrealBuildTool/Modes/SetupPlatformsMode.cs
--- a/Engine/Source/Programs/UnrealBuildTool/Modes/SetupPlatformsMode.cs
+++ b/Engine/Source/Programs/UnrealBuildTool/Modes/SetupPlatformsMode.cs
@@ -20,8 +20,6 @@ namespace UnrealBuildTool
 		/// <param name="Logger"></param>
 		public override Task<int> ExecuteAsync(CommandLineArguments Arguments, ILogger Logger)
 		{
-			// Output a warning if there are any arguments that are still unused
-			Arguments.CheckAllArgumentsUsed();
 			return Task.FromResult(0);
 		}
 	}
diff --git a/Engine/Source/Programs/UnrealBuildTool/Platform/Android/AndroidPlatformSDK.Versions.cs b/Engine/Source/Programs/UnrealBuildTool/Platform/Android/AndroidPlatformSDK.Versions.cs
--- a/Engine/Source/Programs/UnrealBuildTool/Platform/Android/AndroidPlatformSDK.Versions.cs
+++ b/Engine/Source/Programs/UnrealBuildTool/Platform/Android/AndroidPlatformSDK.Versions.cs
@@ -6,7 +6,7 @@ namespace UnrealBuildTool
 {
 	partial class AndroidPlatformSDK : UEBuildPlatformSDK
 	{
-		public override string GetMainVersion()
+		protected override string GetMainVersionInternal()
 		{
 			return "r25b";
 		}
diff --git a/Engine/Source/Programs/UnrealBuildTool/Platform/Linux/LinuxPlatformSDK.Versions.cs b/Engine/Source/Programs/UnrealBuildTool/Platform/Linux/LinuxPlatformSDK.Versions.cs
--- a/Engine/Source/Programs/UnrealBuildTool/Platform/Linux/LinuxPlatformSDK.Versions.cs
+++ b/Engine/Source/Programs/UnrealBuildTool/Platform/Linux/LinuxPlatformSDK.Versions.cs
@@ -6,7 +6,7 @@ namespace UnrealBuildTool
 {
 	partial class LinuxPlatformSDK : UEBuildPlatformSDK
 	{
-		public override string GetMainVersion()
+		protected override string GetMainVersionInternal()
 		{
 			return "v21_clang-15.0.1-centos7";
 		}
diff --git a/Engine/Source/Programs/UnrealBuildTool/Platform/Mac/ApplePlatformSDK.Versions.cs b/Engine/Source/Programs/UnrealBuildTool/Platform/Mac/ApplePlatformSDK.Versions.cs
--- a/Engine/Source/Programs/UnrealBuildTool/Platform/Mac/ApplePlatformSDK.Versions.cs
+++ b/Engine/Source/Programs/UnrealBuildTool/Platform/Mac/ApplePlatformSDK.Versions.cs
@@ -7,7 +7,7 @@ namespace UnrealBuildTool
 {
 	internal partial class ApplePlatformSDK : UEBuildPlatformSDK
 	{
-		public override string GetMainVersion()
+		protected override string GetMainVersionInternal()
 		{
 			// Xcode prefered version?
 			return "14.1";
diff --git a/Engine/Source/Programs/UnrealBuildTool/Platform/Windows/MicrosoftPlatformSDK.Versions.cs b/Engine/Source/Programs/UnrealBuildTool/Platform/Windows/MicrosoftPlatformSDK.Versions.cs
--- a/Engine/Source/Programs/UnrealBuildTool/Platform/Windows/MicrosoftPlatformSDK.Versions.cs
+++ b/Engine/Source/Programs/UnrealBuildTool/Platform/Windows/MicrosoftPlatformSDK.Versions.cs
@@ -71,7 +71,7 @@ namespace UnrealBuildTool
 		static readonly VersionNumber MinimumIntelOneApiVersion = new VersionNumber(2023, 0, 0);
 
 		/// <inheritdoc/>
-		public override string GetMainVersion()
+		protected override string GetMainVersionInternal()
 		{
 			// preferred/main version is the top of the Preferred list - 
 			return PreferredWindowsSdkVersions.First().ToString();
diff --git a/Engine/Source/Programs/UnrealBuildTool/System/PlatformExports.cs b/Engine/Source/Programs/UnrealBuildTool/System/PlatformExports.cs
--- a/Engine/Source/Programs/UnrealBuildTool/System/PlatformExports.cs
+++ b/Engine/Source/Programs/UnrealBuildTool/System/PlatformExports.cs
@@ -146,15 +146,16 @@ namespace UnrealBuildTool
 		/// <summary>
 		/// Initialize UBT in the context of another host process (presumably UAT)
 		/// </summary>
+		/// <param name="CommandLineArgs">Command Line arguments that UBT may need access to for initializing platforms</param>
 		/// <param name="Logger">Logger for output</param>
 		/// <returns>True if initialization was successful</returns>
-		public static bool Initialize(ILogger Logger)
+		public static bool Initialize(string[] CommandLineArgs, ILogger Logger)
 		{
 			// Read the XML configuration files
 			XmlConfig.ReadConfigFiles(null, null, Logger);
 
 			// Register all the platform classes
-			UEBuildPlatform.RegisterPlatforms(false, false, Logger);
+			UEBuildPlatform.RegisterPlatforms(false, false, CommandLineArgs, Logger);
 			return true;
 		}
 	}
diff --git a/Engine/Source/Programs/UnrealBuildTool/UnrealBuildTool.cs b/Engine/Source/Programs/UnrealBuildTool/UnrealBuildTool.cs
--- a/Engine/Source/Programs/UnrealBuildTool/UnrealBuildTool.cs
+++ b/Engine/Source/Programs/UnrealBuildTool/UnrealBuildTool.cs
@@ -635,21 +635,21 @@ namespace UnrealBuildTool
 				{
 					using (GlobalTracer.Instance.BuildSpan("UEBuildPlatform.RegisterPlatforms()").StartActive())
 					{
-						UEBuildPlatform.RegisterPlatforms(false, false, Logger);
+						UEBuildPlatform.RegisterPlatforms(false, false, ModeType, ArgumentsArray, Logger);
 					}
 				}
 				if ((ModeOptions & ToolModeOptions.BuildPlatformsHostOnly) != 0)
 				{
 					using (GlobalTracer.Instance.BuildSpan("UEBuildPlatform.RegisterPlatforms()").StartActive())
 					{
-						UEBuildPlatform.RegisterPlatforms(false, true, Logger);
+						UEBuildPlatform.RegisterPlatforms(false, true, ModeType, ArgumentsArray, Logger);
 					}
 				}
 				if ((ModeOptions & ToolModeOptions.BuildPlatformsForValidation) != 0)
 				{
 					using (GlobalTracer.Instance.BuildSpan("UEBuildPlatform.RegisterPlatforms()").StartActive())
 					{
-						UEBuildPlatform.RegisterPlatforms(true, false, Logger);
+						UEBuildPlatform.RegisterPlatforms(true, false, ModeType, ArgumentsArray, Logger);
 					}
 				}
 
