UnrealBuildTool: If a target fails to vaildate, print any upgrade warnings in case those are the reason

#rnx

diff --git a/Engine/Source/Programs/UnrealBuildTool/Configuration/UEBuildTarget.cs b/Engine/Source/Programs/UnrealBuildTool/Configuration/UEBuildTarget.cs
--- a/Engine/Source/Programs/UnrealBuildTool/Configuration/UEBuildTarget.cs
+++ b/Engine/Source/Programs/UnrealBuildTool/Configuration/UEBuildTarget.cs
@@ -1201,7 +1201,20 @@ namespace UnrealBuildTool
 			// If we're using the shared build environment, make sure all the settings are valid
 			if (RulesObject.BuildEnvironment == TargetBuildEnvironment.Shared)
 			{
-				ValidateSharedEnvironment(RulesAssembly, Descriptor.Name, Descriptor.AdditionalArguments, RulesObject, Logger);
+				try
+				{
+					ValidateSharedEnvironment(RulesAssembly, Descriptor.Name, Descriptor.AdditionalArguments, RulesObject, Logger);
+				}
+				finally
+				{
+					// If there is a validation error, print out any upgrade diagnostics in case that was the reason for the error
+					List<string> diagnostics = new();
+					RulesObject.GetBuildSettingsInfo(diagnostics);
+					foreach (string diagnostic in diagnostics)
+					{
+						Logger.LogWarning("{Diagnostic}", diagnostic);
+					}
+				}
 			}
 
 			// If we're precompiling, generate a list of all the files that we depend on
