UnrealBuildTool: Print upgrade warning when a validate exception happens, not always

#rnx

diff --git a/Engine/Source/Programs/UnrealBuildTool/Configuration/UEBuildTarget.cs b/Engine/Source/Programs/UnrealBuildTool/Configuration/UEBuildTarget.cs
--- a/Engine/Source/Programs/UnrealBuildTool/Configuration/UEBuildTarget.cs
+++ b/Engine/Source/Programs/UnrealBuildTool/Configuration/UEBuildTarget.cs
@@ -1205,7 +1205,7 @@ namespace UnrealBuildTool
 				{
 					ValidateSharedEnvironment(RulesAssembly, Descriptor.Name, Descriptor.AdditionalArguments, RulesObject, Logger);
 				}
-				finally
+				catch (Exception)
 				{
 					// If there is a validation error, print out any upgrade diagnostics in case that was the reason for the error
 					List<string> diagnostics = new();
@@ -1214,6 +1214,7 @@ namespace UnrealBuildTool
 					{
 						Logger.LogWarning("{Diagnostic}", diagnostic);
 					}
+					throw;
 				}
 			}
 
