UnrealBuildTool: Revert change to print out diagnostics on an exception

#rnx

diff --git a/Engine/Source/Programs/UnrealBuildTool/Configuration/UEBuildTarget.cs b/Engine/Source/Programs/UnrealBuildTool/Configuration/UEBuildTarget.cs
--- a/Engine/Source/Programs/UnrealBuildTool/Configuration/UEBuildTarget.cs
+++ b/Engine/Source/Programs/UnrealBuildTool/Configuration/UEBuildTarget.cs
@@ -1201,21 +1201,7 @@ namespace UnrealBuildTool
 			// If we're using the shared build environment, make sure all the settings are valid
 			if (RulesObject.BuildEnvironment == TargetBuildEnvironment.Shared)
 			{
-				try
-				{
-					ValidateSharedEnvironment(RulesAssembly, Descriptor.Name, Descriptor.AdditionalArguments, RulesObject, Logger);
-				}
-				catch (Exception)
-				{
-					// If there is a validation error, print out any upgrade diagnostics in case that was the reason for the error
-					List<string> diagnostics = new();
-					RulesObject.GetBuildSettingsInfo(diagnostics);
-					foreach (string diagnostic in diagnostics)
-					{
-						Logger.LogWarning("{Diagnostic}", diagnostic);
-					}
-					throw;
-				}
+				ValidateSharedEnvironment(RulesAssembly, Descriptor.Name, Descriptor.AdditionalArguments, RulesObject, Logger);
 			}
 
 			// If we're precompiling, generate a list of all the files that we depend on
