Updated the shared pch logic to not follow the private module dependencies/includes when getting all the possible modules referenced in the pch header.

diff --git a/Engine/Source/Programs/UnrealBuildTool/Configuration/UEBuildModule.cs b/Engine/Source/Programs/UnrealBuildTool/Configuration/UEBuildModule.cs
--- a/Engine/Source/Programs/UnrealBuildTool/Configuration/UEBuildModule.cs
+++ b/Engine/Source/Programs/UnrealBuildTool/Configuration/UEBuildModule.cs
@@ -1016,10 +1016,11 @@ namespace UnrealBuildTool
 		/// <summary>
 		/// Gathers all the module dependencies a PCH would have
 		/// </summary>
-		public HashSet<UEBuildModule> GetAllDependencyModulesForPCH()
+		/// <param name="bIncludePrivateModules">Whether to include private modules.</param>
+		public HashSet<UEBuildModule> GetAllDependencyModulesForPCH(bool bIncludePrivateModules)
 		{
 			HashSet<UEBuildModule> ReferencedModules = new HashSet<UEBuildModule>();
-			InternalGetAllDependencyModulesForPCH(ReferencedModules, new HashSet<UEBuildModule>(), true);
+			InternalGetAllDependencyModulesForPCH(ReferencedModules, new HashSet<UEBuildModule>(), bIncludePrivateModules);
 			return ReferencedModules;
 		}
 
@@ -1028,15 +1029,15 @@ namespace UnrealBuildTool
 		/// </summary>
 		/// <param name="ReferencedModules">Hash of all referenced modules with their addition index.</param>
 		/// <param name="IgnoreReferencedModules">Hashset used to ignore modules which are already added to the list</param>
-		/// <param name="bIncludePrivateDependencyModules">Whether to include private dependencies.</param>
-		private void InternalGetAllDependencyModulesForPCH(HashSet<UEBuildModule> ReferencedModules, HashSet<UEBuildModule> IgnoreReferencedModules, bool bIncludePrivateDependencyModules)
+		/// <param name="bIncludePrivateModules">Whether to include private modules.</param>
+		private void InternalGetAllDependencyModulesForPCH(HashSet<UEBuildModule> ReferencedModules, HashSet<UEBuildModule> IgnoreReferencedModules, bool bIncludePrivateModules)
 		{
 			List<UEBuildModule> AllDependencyModules = new List<UEBuildModule>(
-				((bIncludePrivateDependencyModules && PrivateDependencyModules != null) ? PrivateDependencyModules.Count : 0) +
+				((bIncludePrivateModules && PrivateDependencyModules != null) ? PrivateDependencyModules.Count : 0) +
 				(PublicDependencyModules != null ? PublicDependencyModules.Count : 0) +
 				(PublicIncludePathModules != null ? PublicIncludePathModules!.Count : 0)
 				);
-			if (bIncludePrivateDependencyModules && PrivateDependencyModules != null)
+			if (bIncludePrivateModules && PrivateDependencyModules != null)
 			{
 				AllDependencyModules.AddRange(PrivateDependencyModules);
 			}
diff --git a/Engine/Source/Programs/UnrealBuildTool/Configuration/UEBuildModuleCPP.cs b/Engine/Source/Programs/UnrealBuildTool/Configuration/UEBuildModuleCPP.cs
--- a/Engine/Source/Programs/UnrealBuildTool/Configuration/UEBuildModuleCPP.cs
+++ b/Engine/Source/Programs/UnrealBuildTool/Configuration/UEBuildModuleCPP.cs
@@ -841,7 +841,7 @@ namespace UnrealBuildTool
 				PrecompiledHeaderDir = IntermediateDirectory;
 			}
 
-			return new PrecompiledHeaderTemplate(this, CompileEnvironment, HeaderFile, PrecompiledHeaderDir, GetAllDependencyModulesForPCH());
+			return new PrecompiledHeaderTemplate(this, CompileEnvironment, HeaderFile, PrecompiledHeaderDir, GetAllDependencyModulesForPCH(false));
 		}
 
 		static HashSet<string> GetImmutableDefinitions(List<string> Definitions)
@@ -1433,7 +1433,7 @@ namespace UnrealBuildTool
 					PrecompiledHeaderTemplate? OldTemplateValue = CompileEnvironment.SharedPCHs.FirstOrDefault(x => OnlyDirectDependencies.Contains(x.Module) && x.IsValidFor(CompileEnvironment));
 
 					// Find the first shared PCH module we can use that doesn't reference this module
-					HashSet<UEBuildModule> AllDependencies = GetAllDependencyModulesForPCH();
+					HashSet<UEBuildModule> AllDependencies = GetAllDependencyModulesForPCH(true);
 					PrecompiledHeaderTemplate? Template = CompileEnvironment.SharedPCHs.FirstOrDefault(x => !x.ModuleDependencies.Contains(this) && AllDependencies.Contains(x.Module) && x.IsValidFor(CompileEnvironment));
 
 					if (Template != OldTemplateValue)
