- Fix for per-project SDK version commandline parsing with -architecture, which was trying to use UEBuildPlatform before it had been initialized

diff --git a/Engine/Source/Programs/UnrealBuildTool/Configuration/UEBuildTarget.cs b/Engine/Source/Programs/UnrealBuildTool/Configuration/UEBuildTarget.cs
--- a/Engine/Source/Programs/UnrealBuildTool/Configuration/UEBuildTarget.cs
+++ b/Engine/Source/Programs/UnrealBuildTool/Configuration/UEBuildTarget.cs
@@ -968,6 +968,14 @@ namespace UnrealBuildTool
 				return null;
 			}
 
+			// if this validation platform was passed in, but was too early for the build platform to be registered, we can return
+			// null here. we don't want to catch all exceptions and return null, because if ValidationPlatform is good, but the 
+			// architecture is invalid, we want to keep the exception 
+			if (ValidationPlatform != null && !UEBuildPlatform.TryGetBuildPlatform(ValidationPlatform.Value, out _))
+			{
+				return null;
+			}
+
 			return new UnrealArchitectures(ArchString, ValidationPlatform);
 		}
 
