[UBT] Add support for Code coverage on Linux

diff --git a/Engine/Source/Developer/AutomationController/Private/AutomationCommandline.cpp b/Engine/Source/Developer/AutomationController/Private/AutomationCommandline.cpp
--- a/Engine/Source/Developer/AutomationController/Private/AutomationCommandline.cpp
+++ b/Engine/Source/Developer/AutomationController/Private/AutomationCommandline.cpp
@@ -36,7 +36,8 @@ enum class EAutomationCommand : uint8
 	RunCommandLineTests,	//Run only tests that are listed on the commandline
 	RunAll,					//Run all the tests that are supported
 	RunFilter,              //Run only tests that are tagged with this filter
-	Quit					//quit the app when tests are done
+	Quit,					//quit the app when tests are done, uses forced exit
+	SoftQuit				//quit the app when tests are done without forced exit
 };
 
 
@@ -420,12 +421,12 @@ public:
 				{
 					AutomationCommand = AutomationCommandQueue[0];
 					AutomationCommandQueue.RemoveAt(0);
-					if (AutomationCommand == EAutomationCommand::Quit)
+					if (AutomationCommand == EAutomationCommand::Quit || AutomationCommand == EAutomationCommand::SoftQuit)
 					{
 						if (AutomationCommandQueue.IsValidIndex(0))
 						{
-							// Add Quit back to the end of the array.
-							AutomationCommandQueue.Add(EAutomationCommand::Quit);
+							// Add Quit and SoftQuit commands back to the end of the array.
+							AutomationCommandQueue.Add(AutomationCommand);
 							break;
 						}
 					}
@@ -433,7 +434,7 @@ public:
 				}
 
 				// Only quit if Quit is the actual last element in the array.
-				if (AutomationCommand == EAutomationCommand::Quit)
+				if (AutomationCommand == EAutomationCommand::Quit || AutomationCommand == EAutomationCommand::SoftQuit)
 				{
 					if (!GIsCriticalError)
 					{
@@ -446,7 +447,7 @@ public:
 					UE_LOG(LogAutomationCommandLine, Log, TEXT("Shutting down. GIsCriticalError=%d"), GIsCriticalError);
 					// some tools parse this.
 					UE_LOG(LogAutomationCommandLine, Display, TEXT("**** TEST COMPLETE. EXIT CODE: %d ****"), GIsCriticalError ? -1 : 0);
-					FPlatformMisc::RequestExitWithStatus(true, GIsCriticalError ? -1 : 0);
+					FPlatformMisc::RequestExitWithStatus(AutomationCommand == EAutomationCommand::SoftQuit ? false : true, GIsCriticalError ? -1 : 0);
 					// We have finished the testing, and results are available
 					AutomationTestState = EAutomationTestState::Complete;
 				}
@@ -639,6 +640,11 @@ protected:
 					AutomationCommandQueue.Add(EAutomationCommand::Quit);
 					Ar.Logf(TEXT("Automation: Quit Command Queued."));
 				}
+				else if (FParse::Command(&TempCmd, TEXT("SoftQuit")))
+				{
+					AutomationCommandQueue.Add(EAutomationCommand::SoftQuit);
+					Ar.Logf(TEXT("Automation: SoftQuit Command Queued."));
+				}
 				else if (FParse::Command(&TempCmd, TEXT("IgnoreLogEvents")))
 				{
 					if (IConsoleVariable* CVar = IConsoleManager::Get().FindConsoleVariable(TEXT("Automation.CaptureLogEvents")))
@@ -659,6 +665,7 @@ protected:
 					Ar.Logf(TEXT("\tAutomation SetMinimumPriority <minimum priority>"));
 					Ar.Logf(TEXT("\tAutomation SetPriority <priority>"));
 					Ar.Logf(TEXT("\tAutomation Quit"));
+					Ar.Logf(TEXT("\tAutomation SoftQuit"));
 					bHandled = false;
 				}
 			}
diff --git a/Engine/Source/Programs/UnrealBuildTool/Configuration/ModuleRules.cs b/Engine/Source/Programs/UnrealBuildTool/Configuration/ModuleRules.cs
--- a/Engine/Source/Programs/UnrealBuildTool/Configuration/ModuleRules.cs
+++ b/Engine/Source/Programs/UnrealBuildTool/Configuration/ModuleRules.cs
@@ -642,6 +642,10 @@ namespace UnrealBuildTool
 		{
 			get
 			{
+				if (bCodeCoverage) {
+					return CodeOptimization.Never;
+				}
+
 				if (OptimizeCodeOverride.HasValue)
 				{
 					return OptimizeCodeOverride.Value;
@@ -791,6 +795,11 @@ namespace UnrealBuildTool
 		/// </summary>
 		public bool bUseRTTI = false;
 
+		/// <summary>
+		/// Enable code coverage compilation/linking support.
+		/// </summary>
+		public bool bCodeCoverage = false;
+
 		/// <summary>
 		/// Obsolete: Direct the compiler to generate AVX instructions wherever SSE or AVX intrinsics are used, on the platforms that support it.
 		/// Note that by enabling this you are changing the minspec for the PC platform, and the resultant executable will crash on machines without AVX support.
diff --git a/Engine/Source/Programs/UnrealBuildTool/Configuration/TargetRules.cs b/Engine/Source/Programs/UnrealBuildTool/Configuration/TargetRules.cs
--- a/Engine/Source/Programs/UnrealBuildTool/Configuration/TargetRules.cs
+++ b/Engine/Source/Programs/UnrealBuildTool/Configuration/TargetRules.cs
@@ -1812,6 +1812,13 @@ namespace UnrealBuildTool
 		[XmlConfigFile(Category = "BuildConfiguration")]
 		public bool bPGOOptimize = false;
 
+		/// <summary>
+		/// Whether the target requires code coverage compilation and linking.
+		/// </summary>
+		[CommandLine("-CodeCoverage", Value = "true")]
+		[XmlConfigFile(Category = "BuildConfiguration")]
+		public bool bCodeCoverage;
+
 		/// <summary>
 		/// Whether to support edit and continue.
 		/// </summary>
@@ -3101,6 +3108,8 @@ namespace UnrealBuildTool
 
 		public bool bSupportEditAndContinue => Inner.bSupportEditAndContinue;
 
+		public bool bCodeCoverage => Inner.bCodeCoverage;
+
 		public bool bOmitFramePointers => Inner.bOmitFramePointers;
 
 		public bool bEnableCppModules => Inner.bEnableCppModules;
diff --git a/Engine/Source/Programs/UnrealBuildTool/Configuration/UEBuildBinary.cs b/Engine/Source/Programs/UnrealBuildTool/Configuration/UEBuildBinary.cs
--- a/Engine/Source/Programs/UnrealBuildTool/Configuration/UEBuildBinary.cs
+++ b/Engine/Source/Programs/UnrealBuildTool/Configuration/UEBuildBinary.cs
@@ -228,6 +228,7 @@ namespace UnrealBuildTool
 				// Produce additional binary but link it as a console app
 				LinkEnvironment ConsoleAppLinkEvironment = new LinkEnvironment(BinaryLinkEnvironment);
 				ConsoleAppLinkEvironment.bIsBuildingConsoleApplication = true;
+				ConsoleAppLinkEvironment.bCodeCoverage = CompileEnvironment.bCodeCoverage;
 				ConsoleAppLinkEvironment.WindowsEntryPointOverride = "WinMainCRTStartup";       // For WinMain() instead of "main()" for Launch module
 				ConsoleAppLinkEvironment.OutputFilePaths = ConsoleAppLinkEvironment.OutputFilePaths.Select(Path => GetAdditionalConsoleAppPath(Path)).ToList();
 
@@ -778,6 +779,9 @@ namespace UnrealBuildTool
 			BinaryLinkEnvironment.bIsBuildingDLL = IsBuildingDll(Type);
 			BinaryLinkEnvironment.bIsBuildingLibrary = IsBuildingLibrary(Type);
 
+			// Code coverage inherited from compile environment
+			BinaryLinkEnvironment.bCodeCoverage = CompileEnvironment.bCodeCoverage;
+
 			// If we don't have any resource file, use the default or compile a custom one for this module
 			if (BinaryLinkEnvironment.Platform.IsInGroup(UnrealPlatformGroup.Windows))
 			{
diff --git a/Engine/Source/Programs/UnrealBuildTool/Configuration/UEBuildModuleCPP.cs b/Engine/Source/Programs/UnrealBuildTool/Configuration/UEBuildModuleCPP.cs
--- a/Engine/Source/Programs/UnrealBuildTool/Configuration/UEBuildModuleCPP.cs
+++ b/Engine/Source/Programs/UnrealBuildTool/Configuration/UEBuildModuleCPP.cs
@@ -882,6 +882,7 @@ namespace UnrealBuildTool
 			CompileEnvironment.PrecompiledHeaderAction = PrecompiledHeaderAction.Create;
 			CompileEnvironment.PrecompiledHeaderIncludeFilename = WrapperFile.Location;
 			CompileEnvironment.bOptimizeCode = ModuleCompileEnvironment.bOptimizeCode;
+			CompileEnvironment.bCodeCoverage = ModuleCompileEnvironment.bCodeCoverage;
 
 			// Create the action to compile the PCH file.
 			CPPOutput Output;
@@ -1202,6 +1203,7 @@ namespace UnrealBuildTool
 		private static void CopySettingsForSharedPCH(CppCompileEnvironment ModuleCompileEnvironment, CppCompileEnvironment CompileEnvironment)
 		{
 			CompileEnvironment.bOptimizeCode = ModuleCompileEnvironment.bOptimizeCode;
+			CompileEnvironment.bCodeCoverage = ModuleCompileEnvironment.bCodeCoverage;
 			CompileEnvironment.bUseRTTI = ModuleCompileEnvironment.bUseRTTI;
 			CompileEnvironment.bEnableExceptions = ModuleCompileEnvironment.bEnableExceptions;
 			CompileEnvironment.ShadowVariableWarningLevel = ModuleCompileEnvironment.ShadowVariableWarningLevel;
@@ -1700,9 +1702,13 @@ namespace UnrealBuildTool
 		/// <param name="Setting">The optimization setting from the rules file</param>
 		/// <param name="Configuration">The active target configuration</param>
 		/// <param name="bIsEngineModule">Whether the current module is an engine module</param>
+		/// <param name="bCodeCoverage">Whether the current module should be compiled with code coverage support</param>
 		/// <returns>True if optimization should be enabled</returns>
-		public static bool ShouldEnableOptimization(ModuleRules.CodeOptimization Setting, UnrealTargetConfiguration Configuration, bool bIsEngineModule)
+		public static bool ShouldEnableOptimization(ModuleRules.CodeOptimization Setting, UnrealTargetConfiguration Configuration, bool bIsEngineModule, bool bCodeCoverage)
 		{
+			if (bCodeCoverage) {
+				return false;
+			}
 			switch (Setting)
 			{
 				case ModuleRules.CodeOptimization.Never:
@@ -1738,7 +1744,8 @@ namespace UnrealBuildTool
 
 			// Override compile environment
 			Result.bUseUnity = Rules.bUseUnity;
-			Result.bOptimizeCode = ShouldEnableOptimization(Rules.OptimizeCode, Target.Configuration, Rules.bTreatAsEngineModule);
+			Result.bCodeCoverage = Target.bCodeCoverage;
+			Result.bOptimizeCode = ShouldEnableOptimization(Rules.OptimizeCode, Target.Configuration, Rules.bTreatAsEngineModule, Result.bCodeCoverage);
 			Result.bUseRTTI |= Rules.bUseRTTI;
 			Result.bEnableBufferSecurityChecks = Rules.bEnableBufferSecurityChecks;
 			Result.MinSourceFilesForUnityBuildOverride = Rules.MinSourceFilesForUnityBuildOverride;
@@ -1867,8 +1874,9 @@ namespace UnrealBuildTool
 		{
 			CppCompileEnvironment CompileEnvironment = new CppCompileEnvironment(BaseCompileEnvironment);
 
-			// Use the default optimization setting for 
-			CompileEnvironment.bOptimizeCode = ShouldEnableOptimization(ModuleRules.CodeOptimization.Default, Target.Configuration, Rules.bTreatAsEngineModule);
+			// Use the default optimization setting for
+			CompileEnvironment.bOptimizeCode = ShouldEnableOptimization(ModuleRules.CodeOptimization.Default, Target.Configuration, Rules.bTreatAsEngineModule, Rules.bCodeCoverage);
+			CompileEnvironment.bCodeCoverage = Rules.bCodeCoverage;
 
 			// Override compile environment
 			CompileEnvironment.bIsBuildingDLL = !Target.ShouldCompileMonolithic();
diff --git a/Engine/Source/Programs/UnrealBuildTool/Configuration/UEBuildTarget.cs b/Engine/Source/Programs/UnrealBuildTool/Configuration/UEBuildTarget.cs
--- a/Engine/Source/Programs/UnrealBuildTool/Configuration/UEBuildTarget.cs
+++ b/Engine/Source/Programs/UnrealBuildTool/Configuration/UEBuildTarget.cs
@@ -4574,6 +4574,7 @@ namespace UnrealBuildTool
 			GlobalCompileEnvironment.AdditionalArguments = Rules.AdditionalCompilerArguments ?? String.Empty;
 			GlobalCompileEnvironment.bDeterministic = Rules.bDeterministic;
 			GlobalCompileEnvironment.CrashDiagnosticDirectory = Rules.CrashDiagnosticDirectory;
+			GlobalCompileEnvironment.bCodeCoverage = Rules.bCodeCoverage;
 
 			GlobalLinkEnvironment.bUseDebugCRT = bUseDebugCRT;
 			GlobalLinkEnvironment.bUseStaticCRT = Rules.bUseStaticCRT;
@@ -4598,6 +4599,7 @@ namespace UnrealBuildTool
 			GlobalLinkEnvironment.bPrintTimingInfo = Rules.bPrintToolChainTimingInfo;
 			GlobalLinkEnvironment.bUsePIE = Rules.bEnablePIE;
 			GlobalLinkEnvironment.AdditionalArguments = Rules.AdditionalLinkerArguments ?? String.Empty;
+			GlobalLinkEnvironment.bCodeCoverage = Rules.bCodeCoverage;
 
 			if (Rules.bPGOOptimize && Rules.bPGOProfile)
 			{
diff --git a/Engine/Source/Programs/UnrealBuildTool/Platform/Linux/LinuxToolChain.cs b/Engine/Source/Programs/UnrealBuildTool/Platform/Linux/LinuxToolChain.cs
--- a/Engine/Source/Programs/UnrealBuildTool/Platform/Linux/LinuxToolChain.cs
+++ b/Engine/Source/Programs/UnrealBuildTool/Platform/Linux/LinuxToolChain.cs
@@ -341,8 +341,19 @@ namespace UnrealBuildTool
 
 			AddCompilerLTOFlags(Arguments);
 
-			// optimization level
-			if (!CompileEnvironment.bOptimizeCode)
+			if (CompileEnvironment.bCodeCoverage)
+			{
+				Arguments.Add("-O0");
+				if (ShouldUseLibcxx())
+				{
+					Arguments.Add("--coverage"); // gcov
+				}
+				else
+				{
+					Arguments.Add("-fprofile-instr-generate -fcoverage-mapping"); // llvm-cov
+				}
+			}
+			else if (!CompileEnvironment.bOptimizeCode) // optimization level
 			{
 				Arguments.Add("-O0");
 			}
@@ -671,6 +682,21 @@ namespace UnrealBuildTool
 				Arguments.Add("-Wl,--gdb-index");
 			}
 
+			if (LinkEnvironment.bCodeCoverage)
+			{
+				// Unreal Separates the linking phase and the compilation phase.
+				// We pass to clang the flag `--coverage` during the compile time
+				// And we link the correct compiler-rt library (shipped by UE, and part of the LLVM toolchain)
+				// to every binary produced.
+				if (ShouldUseLibcxx())
+				{
+					Arguments.Add("--coverage"); // gcov
+				}
+				else
+				{
+					Arguments.Add("-fprofile-instr-generate"); // llvm-cov
+				}
+			}
 			// RPATH for third party libs
 			Arguments.Add("-Wl,-rpath=${ORIGIN}");
 			Arguments.Add("-Wl,-rpath-link=${ORIGIN}");
@@ -832,6 +858,11 @@ namespace UnrealBuildTool
 				Logger.LogInformation("  Prefix for PGO data files='{CompileEnvironmentPGOFilenamePrefix}'", CompileEnvironment.PGOFilenamePrefix);
 			}
 
+			if (CompileEnvironment.bCodeCoverage)
+			{
+				Logger.LogInformation("Using --coverage build flag");
+			}
+
 			if (CompileEnvironment.bPGOProfile)
 			{
 				Logger.LogInformation("Using PGI (profile guided instrumentation).");
diff --git a/Engine/Source/Programs/UnrealBuildTool/Platform/Linux/UEBuildLinux.cs b/Engine/Source/Programs/UnrealBuildTool/Platform/Linux/UEBuildLinux.cs
--- a/Engine/Source/Programs/UnrealBuildTool/Platform/Linux/UEBuildLinux.cs
+++ b/Engine/Source/Programs/UnrealBuildTool/Platform/Linux/UEBuildLinux.cs
@@ -503,6 +503,8 @@ namespace UnrealBuildTool
 				LinkEnvironment.PGOFilenamePrefix = CompileEnvironment.PGOFilenamePrefix;
 			}
 
+			LinkEnvironment.bCodeCoverage = CompileEnvironment.bCodeCoverage;
+
 			// For consistency with other platforms, also enable LTO whenever doing profile-guided optimizations.
 			// Obviously both PGI (instrumented) and PGO (optimized) binaries need to have that
 			if (CompileEnvironment.bPGOProfile || CompileEnvironment.bPGOOptimize)
diff --git a/Engine/Source/Programs/UnrealBuildTool/System/CppCompileEnvironment.cs b/Engine/Source/Programs/UnrealBuildTool/System/CppCompileEnvironment.cs
--- a/Engine/Source/Programs/UnrealBuildTool/System/CppCompileEnvironment.cs
+++ b/Engine/Source/Programs/UnrealBuildTool/System/CppCompileEnvironment.cs
@@ -386,6 +386,11 @@ namespace UnrealBuildTool
 		/// </summary>
 		public bool bOptimizeCode = false;
 
+		/// <summary>
+		/// True if the compilation should produce tracing output for code coverage.
+		/// </summary>
+		public bool bCodeCoverage = false;
+
 		/// <summary>
 		/// Allows to fine tune optimizations level for speed and\or code size
 		/// </summary>
@@ -709,6 +714,7 @@ namespace UnrealBuildTool
 			StaticAnalyzerAdditionalCheckers = new HashSet<string>(Other.StaticAnalyzerAdditionalCheckers);
 			bOptimizeCode = Other.bOptimizeCode;
 			bAllowAutoRTFMInstrumentation = Other.bAllowAutoRTFMInstrumentation;
+			bCodeCoverage = Other.bCodeCoverage;
 			OptimizationLevel = Other.OptimizationLevel;
 			bCreateDebugInfo = Other.bCreateDebugInfo;
 			bIsBuildingLibrary = Other.bIsBuildingLibrary;
diff --git a/Engine/Source/Programs/UnrealBuildTool/System/LinkEnvironment.cs b/Engine/Source/Programs/UnrealBuildTool/System/LinkEnvironment.cs
--- a/Engine/Source/Programs/UnrealBuildTool/System/LinkEnvironment.cs
+++ b/Engine/Source/Programs/UnrealBuildTool/System/LinkEnvironment.cs
@@ -209,6 +209,11 @@ namespace UnrealBuildTool
 		/// </summary>
 		public int DefaultStackSizeCommit = 0;
 
+		/// <summary>
+		/// Wether to link code coverage / tracing libs
+		/// </summary>
+		public bool bCodeCoverage = false;
+
 		/// <summary>
 		/// Whether to omit frame pointers or not. Disabling is useful for e.g. memory profiling on the PC
 		/// </summary>
@@ -400,6 +405,7 @@ namespace UnrealBuildTool
 			bIsBuildingDotNetAssembly = Other.bIsBuildingDotNetAssembly;
 			DefaultStackSize = Other.DefaultStackSize;
 			DefaultStackSizeCommit = Other.DefaultStackSizeCommit;
+			bCodeCoverage = Other.bCodeCoverage;
 			bOmitFramePointers = Other.bOmitFramePointers;
 			bSupportEditAndContinue = Other.bSupportEditAndContinue;
 			bUseIncrementalLinking = Other.bUseIncrementalLinking;
