- Fixed location of where EnsureXcodeProjectExists() looks for BP projects
- Fixed warning about SDKROOT[sdk=macosx] type warnings with legacy Xcode

diff --git a/Engine/Source/Programs/AutomationTool/IOS/IOSPlatform.Automation.cs b/Engine/Source/Programs/AutomationTool/IOS/IOSPlatform.Automation.cs
--- a/Engine/Source/Programs/AutomationTool/IOS/IOSPlatform.Automation.cs
+++ b/Engine/Source/Programs/AutomationTool/IOS/IOSPlatform.Automation.cs
@@ -1079,10 +1079,6 @@ public class IOSPlatform : ApplePlatform
 		bWasGenerated = false;
 		DirectoryReference RawProjectDir = RawProjectPath.Directory;
 		DirectoryReference XcodeProj = DirectoryReference.Combine(RawProjectDir, "Intermediate/ProjectFiles", $"{RawProjectPath.GetFileNameWithoutAnyExtensions()}_{PlatformName}.xcworkspace");
-		if (!DirectoryReference.Exists(DirectoryReference.Combine(RawProjectDir, "Source")) && !DirectoryReference.Exists(DirectoryReference.Combine(RawProjectDir, "/Intermediate/Source")))
-		{
-			XcodeProj = DirectoryReference.Combine(Unreal.EngineDirectory, "Intermediate/ProjectFiles/", XcodeProj.GetDirectoryName());
-		}
 		Console.WriteLine("Project: " + XcodeProj.FullName);
 		{
 			// project.xcodeproj doesn't exist, so generate temp project
diff --git a/Engine/Source/Programs/UnrealBuildTool/ProjectFiles/Xcode/XcodeProjectLegacy.cs b/Engine/Source/Programs/UnrealBuildTool/ProjectFiles/Xcode/XcodeProjectLegacy.cs
--- a/Engine/Source/Programs/UnrealBuildTool/ProjectFiles/Xcode/XcodeProjectLegacy.cs
+++ b/Engine/Source/Programs/UnrealBuildTool/ProjectFiles/Xcode/XcodeProjectLegacy.cs
@@ -1043,7 +1043,6 @@ namespace UnrealBuildTool.XcodeProjectLegacy
 				Content.Append("\t\t\t\tIPHONEOS_DEPLOYMENT_TARGET = " + IOSRunTimeVersion + ";" + ProjectFileGenerator.NewLine);
 				Content.Append("\t\t\t\t\"PRODUCT_NAME[sdk=iphoneos*]\" = \"" + TargetName + "\";" + ProjectFileGenerator.NewLine); // @todo: change to Path.GetFileName(Config.IOSExecutablePath) when we stop using payload
 				Content.Append("\t\t\t\t\"TARGETED_DEVICE_FAMILY[sdk=iphoneos*]\" = \"" + IOSRunTimeDevices + "\";" + ProjectFileGenerator.NewLine);
-				Content.Append("\t\t\t\t\"SDKROOT[sdk=iphoneos]\" = iphoneos;" + ProjectFileGenerator.NewLine);
 				if (!String.IsNullOrEmpty(TEAM_IOS))
 				{
 					Content.Append("\t\t\t\t\"DEVELOPMENT_TEAM[sdk=iphoneos*]\" = " + TEAM_IOS + ";" + ProjectFileGenerator.NewLine);
@@ -1064,7 +1063,6 @@ namespace UnrealBuildTool.XcodeProjectLegacy
 				Content.Append("\t\t\t\tTVOS_DEPLOYMENT_TARGET = " + TVOSRunTimeVersion + ";" + ProjectFileGenerator.NewLine);
 				Content.Append("\t\t\t\t\"PRODUCT_NAME[sdk=appletvos*]\" = \"" + TargetName + "\";" + ProjectFileGenerator.NewLine); // @todo: change to Path.GetFileName(Config.TVOSExecutablePath) when we stop using payload
 				Content.Append("\t\t\t\t\"TARGETED_DEVICE_FAMILY[sdk=appletvos*]\" = \"" + TVOSRunTimeDevices + "\";" + ProjectFileGenerator.NewLine);
-				Content.Append("\t\t\t\t\"SDKROOT[sdk=appletvos]\" = appletvos;" + ProjectFileGenerator.NewLine);
 				if (!String.IsNullOrEmpty(TEAM_TVOS))
 				{
 					Content.Append("\t\t\t\t\"DEVELOPMENT_TEAM[sdk=appletvos*]\" = " + TEAM_TVOS + ";" + ProjectFileGenerator.NewLine);
@@ -1084,7 +1082,6 @@ namespace UnrealBuildTool.XcodeProjectLegacy
 				Content.Append("\t\t\t\t\"VALID_ARCHS[sdk=macosx*]\" = \"" + String.Join(" ", SupportedMacArchitectures.Architectures.Select(x => x.AppleName)) + "\";" + ProjectFileGenerator.NewLine);
 				Content.Append("\t\t\t\t\"PRODUCT_NAME[sdk=macosx*]\" = \"" + MacExecutableFileName + "\";" + ProjectFileGenerator.NewLine);
 				Content.Append("\t\t\t\t\"CONFIGURATION_BUILD_DIR[sdk=macosx*]\" = \"" + MacExecutableDir + "\";" + ProjectFileGenerator.NewLine);
-				Content.Append("\t\t\t\t\"SDKROOT[sdk=macosx*]\" = macosx;" + ProjectFileGenerator.NewLine);
 
 				// MacToolchain uses the IOS bundle identifier
 				ConfigHierarchy IOSIni = ConfigCache.ReadHierarchy(ConfigHierarchyType.Engine, ProjectFile?.Directory, UnrealTargetPlatform.IOS);
