Don't use the Box Executor if compiling with -SingleFile=

diff --git a/Engine/Source/Programs/UnrealBuildTool/System/ActionGraph.cs b/Engine/Source/Programs/UnrealBuildTool/System/ActionGraph.cs
--- a/Engine/Source/Programs/UnrealBuildTool/System/ActionGraph.cs
+++ b/Engine/Source/Programs/UnrealBuildTool/System/ActionGraph.cs
@@ -353,7 +353,8 @@ namespace UnrealBuildTool
 		private static ActionExecutor SelectExecutor(BuildConfiguration BuildConfiguration, int ActionCount, List<TargetDescriptor> TargetDescriptors, ILogger Logger)
 		{
 #if __BOXEXECUTOR_AVAILABLE__
-			if (BuildConfiguration.bAllowBoxExecutor && BoxExecutor.IsAvailable(Logger))
+			bool bAnySingleFile = TargetDescriptors.Any(Descriptor => Descriptor.SpecificFilesToCompile.Count == 1);
+			if (!bAnySingleFile && BuildConfiguration.bAllowBoxExecutor && BoxExecutor.IsAvailable(Logger))
 			{
 				return new BoxExecutor(BuildConfiguration.MaxParallelActions, BuildConfiguration.bAllCores, BuildConfiguration.bCompactOutput, Logger);
 			}
