- Disabled the per-project SDK initialization, as it's causing TargetRules to get created early when using -TargetType, which causes problems

diff --git a/Engine/Source/Programs/UnrealBuildTool/Configuration/UEBuildPlatform.cs b/Engine/Source/Programs/UnrealBuildTool/Configuration/UEBuildPlatform.cs
--- a/Engine/Source/Programs/UnrealBuildTool/Configuration/UEBuildPlatform.cs
+++ b/Engine/Source/Programs/UnrealBuildTool/Configuration/UEBuildPlatform.cs
@@ -359,10 +359,11 @@ namespace UnrealBuildTool
 				InstalledPlatformInfo.Initialize();
 			}
 
-			using (GlobalTracer.Instance.BuildSpan("Initializing PerPlatformSDKs").StartActive())
-			{
-				InitializePerPlatformSDKs(ArgumentsForPerPlatform, bArgumentsAreForUBT, Logger);
-			}
+// disabling _for now_ as it is causing issues when using -TargetType param, causing too much stuff to be created early
+//			using (GlobalTracer.Instance.BuildSpan("Initializing PerPlatformSDKs").StartActive())
+//			{
+//				InitializePerPlatformSDKs(ArgumentsForPerPlatform, bArgumentsAreForUBT, Logger);
+//			}
 
 			// Find and register all tool chains and build platforms that are present
 			Type[] AllTypes;
