Cleaned up naming of structures and classes to reduce the overuse of the word Mapping.

diff --git a/Engine/Source/Programs/UnrealBuildTool.Tests/ArtifactsTest.cs b/Engine/Source/Programs/UnrealBuildTool.Tests/ArtifactsTest.cs
--- a/Engine/Source/Programs/UnrealBuildTool.Tests/ArtifactsTest.cs
+++ b/Engine/Source/Programs/UnrealBuildTool.Tests/ArtifactsTest.cs
@@ -18,38 +18,38 @@ namespace UnrealBuildToolTests
 		static readonly Utf8String s_input2Data = new Utf8String("This is some more sample input");
 		static readonly Utf8String s_output1Data = new Utf8String("This is the sample output");
 
-		private static Artifact MakeInput1()
+		private static ArtifactFile MakeInput1()
 		{
-			return new Artifact(ArtifactDirectoryTree.Absolute, "Input1", IoHash.Compute(s_input1Data.Span));
+			return new ArtifactFile(ArtifactDirectoryTree.Absolute, "Input1", IoHash.Compute(s_input1Data.Span));
 		}
 
-		private static Artifact MakeInput2()
+		private static ArtifactFile MakeInput2()
 		{
-			return new Artifact(ArtifactDirectoryTree.Absolute, "Input2", IoHash.Compute(s_input2Data.Span));
+			return new ArtifactFile(ArtifactDirectoryTree.Absolute, "Input2", IoHash.Compute(s_input2Data.Span));
 		}
 
-		private static Artifact MakeOutput1()
+		private static ArtifactFile MakeOutput1()
 		{
-			return new Artifact(ArtifactDirectoryTree.Absolute, "Output1", IoHash.Compute(s_output1Data.Span));
+			return new ArtifactFile(ArtifactDirectoryTree.Absolute, "Output1", IoHash.Compute(s_output1Data.Span));
 		}
 
-		public static ArtifactMapping MakeBundle1()
+		public static ArtifactAction MakeBundle1()
 		{
-			return new ArtifactMapping(
+			return new ArtifactAction(
 				IoHash.Compute(new Utf8String("SampleKey")),
 				IoHash.Compute(new Utf8String("SampleBundleKey")),
-				new Artifact[] { MakeInput1(), MakeInput2() },
-				new Artifact[] { MakeOutput1() }
+				new ArtifactFile[] { MakeInput1(), MakeInput2() },
+				new ArtifactFile[] { MakeOutput1() }
 			);
 		}
 
-		public static ArtifactMapping MakeBundle2()
+		public static ArtifactAction MakeBundle2()
 		{
-			return new ArtifactMapping(
+			return new ArtifactAction(
 				IoHash.Compute(new Utf8String("SampleKey")),
 				IoHash.Compute(new Utf8String("SampleBundleKeyV2")),
-				new Artifact[] { MakeInput1(), MakeInput2() },
-				new Artifact[] { MakeOutput1() }
+				new ArtifactFile[] { MakeInput1(), MakeInput2() },
+				new ArtifactFile[] { MakeOutput1() }
 			);
 		}
 
@@ -63,16 +63,16 @@ namespace UnrealBuildToolTests
 			_ = cache.WaitForReadyAsync().Result;
 			Assert.AreEqual(ArtifactCacheState.Available, cache.State);
 
-			ArtifactMapping bundle1 = MakeBundle1();
-			cache.SaveArtifactMappingsAsync(new ArtifactMapping[] { bundle1 }, cancellationToken).Wait();
+			ArtifactAction bundle1 = MakeBundle1();
+			cache.SaveArtifactActionsAsync(new ArtifactAction[] { bundle1 }, cancellationToken).Wait();
 
-			ArtifactMapping[] readBack1 = cache.QueryArtifactMappingsAsync(new IoHash[] { bundle1.Key }, cancellationToken).Result;
+			ArtifactAction[] readBack1 = cache.QueryArtifactActionsAsync(new IoHash[] { bundle1.Key }, cancellationToken).Result;
 			Assert.AreEqual(1, readBack1.Length);
 
-			ArtifactMapping bundle2 = MakeBundle2();
-			cache.SaveArtifactMappingsAsync(new ArtifactMapping[] { bundle2 }, cancellationToken).Wait();
+			ArtifactAction bundle2 = MakeBundle2();
+			cache.SaveArtifactActionsAsync(new ArtifactAction[] { bundle2 }, cancellationToken).Wait();
 
-			ArtifactMapping[] readBack2 = cache.QueryArtifactMappingsAsync(new IoHash[] { bundle1.Key }, cancellationToken).Result;
+			ArtifactAction[] readBack2 = cache.QueryArtifactActionsAsync(new IoHash[] { bundle1.Key }, cancellationToken).Result;
 			Assert.AreEqual(2, readBack2.Length);
 
 			cache.FlushChangesAsync(cancellationToken).Wait();
diff --git a/Engine/Source/Programs/UnrealBuildTool/Artifacts/ActionArtifactCache.cs b/Engine/Source/Programs/UnrealBuildTool/Artifacts/ActionArtifactCache.cs
--- a/Engine/Source/Programs/UnrealBuildTool/Artifacts/ActionArtifactCache.cs
+++ b/Engine/Source/Programs/UnrealBuildTool/Artifacts/ActionArtifactCache.cs
@@ -62,7 +62,7 @@ namespace UnrealBuildTool.Artifacts
 		/// <param name="file">File in question</param>
 		/// <param name="hash">Hash value used to populate artifact</param>
 		/// <returns>Created artifact</returns>
-		public Artifact GetArtifact(LinkedAction action, FileItem file, IoHash hash)
+		public ArtifactFile GetArtifact(LinkedAction action, FileItem file, IoHash hash)
 		{
 			if (action.ArtifactMode.HasFlag(ArtifactMode.AbsolutePath))
 			{
@@ -89,9 +89,9 @@ namespace UnrealBuildTool.Artifacts
 		/// <param name="path">Path to artifact</param>
 		/// <param name="hash">Hash of the artifact</param>
 		/// <returns>The artifact</returns>
-		private Artifact CreateArtifact(ArtifactDirectoryTree tree, string path, IoHash hash)
+		private ArtifactFile CreateArtifact(ArtifactDirectoryTree tree, string path, IoHash hash)
 		{
-			return new(tree, new Utf8String(path), hash) { DirectoryMapping = this };
+			return new(tree, new Utf8String(path), hash);
 		}
 	}
 
@@ -201,7 +201,7 @@ namespace UnrealBuildTool.Artifacts
 			(List<FileItem> inputs, _) = CollectInputs(action, false);
 			IoHash key = await GetKeyAsync(directoryMapping, action, inputs);
 
-			ArtifactMapping[] mappings = await ArtifactCache.QueryArtifactMappingsAsync(new IoHash[] { key }, cancellationToken);
+			ArtifactAction[] artifactActions = await ArtifactCache.QueryArtifactActionsAsync(new IoHash[] { key }, cancellationToken);
 
 			string actionDescription = String.Empty;
 			if (LogCacheMisses)
@@ -209,20 +209,20 @@ namespace UnrealBuildTool.Artifacts
 				actionDescription = $"{(action.CommandDescription ?? action.CommandPath.GetFileNameWithoutExtension())} {action.StatusDescription}".Trim();
 			}
 
-			if (mappings.Length == 0)
+			if (artifactActions.Length == 0)
 			{
 				if (LogCacheMisses)
 				{
-					_logger.LogInformation("Artifact Cache Miss: No artifact mappings found for {ActionDescription}", actionDescription);
+					_logger.LogInformation("Artifact Cache Miss: No artifact actions found for {ActionDescription}", actionDescription);
 				}
 				return false;
 			}
 
-			foreach (ArtifactMapping mapping in mappings)
+			foreach (ArtifactAction artifactAction in artifactActions)
 			{
 				bool match = true;
 
-				foreach (Artifact input in mapping.Inputs)
+				foreach (ArtifactFile input in artifactAction.Inputs)
 				{
 					string name = input.GetFullPath(directoryMapping);
 					FileItem item = FileItem.GetItemByPath(name);
@@ -248,9 +248,9 @@ namespace UnrealBuildTool.Artifacts
 
 				if (match)
 				{
-					ArtifactMapping mappingCopy = mapping;
-					mappingCopy.DirectoryMapping = directoryMapping;
-					bool[]? readResults = await ArtifactCache.QueryArtifactOutputsAsync(new[] { mappingCopy }, cancellationToken);
+					ArtifactAction artifactActionCopy = artifactAction;
+					artifactActionCopy.DirectoryMapping = directoryMapping;
+					bool[]? readResults = await ArtifactCache.QueryArtifactOutputsAsync(new[] { artifactActionCopy }, cancellationToken);
 
 					if (readResults == null || readResults.Length == 0 || !readResults[0])
 					{
@@ -258,7 +258,7 @@ namespace UnrealBuildTool.Artifacts
 					}
 					else
 					{
-						foreach (Artifact output in mapping.Outputs)
+						foreach (ArtifactFile output in artifactAction.Outputs)
 						{
 							string outputName = output.GetFullPath(directoryMapping);
 							FileItem item = FileItem.GetItemByPath(outputName);
@@ -285,10 +285,10 @@ namespace UnrealBuildTool.Artifacts
 				return;
 			}
 
-			ArtifactMapping mapping = await CreateArtifactMappingAsync(action, cancellationToken);
-			if (mapping.Key != IoHash.Zero)
+			ArtifactAction artifactAction = await CreateArtifactActionAsync(action, cancellationToken);
+			if (artifactAction.Key != IoHash.Zero)
 			{
-				await ArtifactCache.SaveArtifactMappingsAsync(new ArtifactMapping[] { (ArtifactMapping)mapping }, cancellationToken);
+				await ArtifactCache.SaveArtifactActionsAsync(new ArtifactAction[] { (ArtifactAction)artifactAction }, cancellationToken);
 			}
 			return;
 		}
@@ -300,37 +300,37 @@ namespace UnrealBuildTool.Artifacts
 		}
 
 		/// <summary>
-		/// Create a new artifact mapping that represents the input and output of the action
+		/// Create a new artifact action that represents the input and output of the action
 		/// </summary>
 		/// <param name="action">Source action</param>
 		/// <param name="cancellationToken">Cancellation token</param>
-		/// <returns>Artifact mapping</returns>
-		private async Task<ArtifactMapping> CreateArtifactMappingAsync(LinkedAction action, CancellationToken cancellationToken)
+		/// <returns>Artifact action</returns>
+		private async Task<ArtifactAction> CreateArtifactActionAsync(LinkedAction action, CancellationToken cancellationToken)
 		{
 			ArtifactDirectoryMapping directoryMapping = GetDirectoryMapping(action);
 			(List<FileItem> inputs, List<FileItem>? dependencies) = CollectInputs(action, false);
-			(IoHash key, IoHash mappingKey) = await GetKeyAndMappingKeyAsync(directoryMapping, action, inputs, dependencies);
+			(IoHash key, IoHash actionKey) = await GetKeyAndActionKeyAsync(directoryMapping, action, inputs, dependencies);
 
 			// We gather the output files first to make sure that all the generated files (including the dependency file) gets
 			// their cached FileInfo reset.
-			List<Artifact> outputs = new();
+			List<ArtifactFile> outputs = new();
 			foreach (FileItem output in action.ProducedItems)
 			{
 				// Outputs can not be a directory.
 				if (output.Attributes.HasFlag(System.IO.FileAttributes.Directory))
 				{
-					return new(IoHash.Zero, IoHash.Zero, Array.Empty<Artifact>(), Array.Empty<Artifact>());
+					return new(IoHash.Zero, IoHash.Zero, Array.Empty<ArtifactFile>(), Array.Empty<ArtifactFile>());
 				}
 				IoHash hash = await _fileHasher.GetDigestAsync(output, cancellationToken);
-				Artifact artifact = directoryMapping.GetArtifact(action, output, hash);
+				ArtifactFile artifact = directoryMapping.GetArtifact(action, output, hash);
 				outputs.Add(artifact);
 			}
 
-			List<Artifact> inputArtifacts = new();
+			List<ArtifactFile> inputArtifacts = new();
 			foreach (FileItem input in inputs)
 			{
 				IoHash hash = await _fileHasher.GetDigestAsync(input, cancellationToken);
-				Artifact artifact = directoryMapping.GetArtifact(action, input, hash);
+				ArtifactFile artifact = directoryMapping.GetArtifact(action, input, hash);
 				inputArtifacts.Add(artifact);
 			}
 
@@ -339,12 +339,15 @@ namespace UnrealBuildTool.Artifacts
 				foreach (FileItem dependency in dependencies)
 				{
 					IoHash hash = await _fileHasher.GetDigestAsync(dependency, cancellationToken);
-					Artifact artifact = directoryMapping.GetArtifact(action, dependency, hash);
+					ArtifactFile artifact = directoryMapping.GetArtifact(action, dependency, hash);
 					inputArtifacts.Add(artifact);
 				}
 			}
 
-			return new(key, mappingKey, inputArtifacts.ToArray(), outputs.ToArray());
+			return new(key, actionKey, inputArtifacts.ToArray(), outputs.ToArray())
+			{
+				DirectoryMapping = directoryMapping
+			};
 		}
 
 		/// <summary>
@@ -363,21 +366,21 @@ namespace UnrealBuildTool.Artifacts
 		}
 
 		/// <summary>
-		/// Generate the key and mapping key hashes for the action
+		/// Generate the key and action key hashes for the action
 		/// </summary>
 		/// <param name="directoryMapping">Directory mapping object</param>
 		/// <param name="action">Source action</param>
 		/// <param name="inputs">Inputs used to construct the key</param>
 		/// <param name="dependencies">Dependencies used to construct the key</param>
-		/// <returns>Task object with the key and mapping key</returns>
-		private async Task<(IoHash, IoHash)> GetKeyAndMappingKeyAsync(ArtifactDirectoryMapping directoryMapping, LinkedAction action, List<FileItem> inputs, List<FileItem>? dependencies)
+		/// <returns>Task object with the key and action key</returns>
+		private async Task<(IoHash, IoHash)> GetKeyAndActionKeyAsync(ArtifactDirectoryMapping directoryMapping, LinkedAction action, List<FileItem> inputs, List<FileItem>? dependencies)
 		{
 			StringBuilder builder = new();
 			await AppendKeyAsync(builder, directoryMapping, action, inputs);
 			IoHash key = IoHash.Compute(new Utf8String(builder.ToString()));
-			await AppendMappingKeyAsync(builder, directoryMapping, action, dependencies);
-			IoHash mappingKey = IoHash.Compute(new Utf8String(builder.ToString()));
-			return (key, mappingKey);
+			await AppendActionKeyAsync(builder, directoryMapping, action, dependencies);
+			IoHash actionKey = IoHash.Compute(new Utf8String(builder.ToString()));
+			return (key, actionKey);
 		}
 
 		/// <summary>
@@ -396,14 +399,14 @@ namespace UnrealBuildTool.Artifacts
 		}
 
 		/// <summary>
-		/// Generate the full mapping key.  This contains the hashes for the action's inputs and the dependent files.
+		/// Generate the full action key.  This contains the hashes for the action's inputs and the dependent files.
 		/// </summary>
 		/// <param name="builder">Destination builder</param>
 		/// <param name="directoryMapping">Directory mapping object</param>
 		/// <param name="action">Source action</param>
 		/// <param name="dependencies">Dependencies used to construct the key</param>
 		/// <returns>Task object</returns>
-		private async Task AppendMappingKeyAsync(StringBuilder builder, ArtifactDirectoryMapping directoryMapping, LinkedAction action, List<FileItem>? dependencies)
+		private async Task AppendActionKeyAsync(StringBuilder builder, ArtifactDirectoryMapping directoryMapping, LinkedAction action, List<FileItem>? dependencies)
 		{
 			await AppendFiles(builder, directoryMapping, action, dependencies);
 		}
@@ -429,7 +432,7 @@ namespace UnrealBuildTool.Artifacts
 				string[] lines = new string[files.Count];
 				for (int index = 0; index < files.Count; index++)
 				{
-					Artifact artifact = directoryMapping.GetArtifact(action, files[index], waits[index].Result);
+					ArtifactFile artifact = directoryMapping.GetArtifact(action, files[index], waits[index].Result);
 					lines[index] = $"{GetArtifactTreeName(artifact)} {artifact.Tree} {artifact.Name} {artifact.ContentHash}";
 				}
 				Array.Sort(lines, StringComparer.Ordinal);
@@ -525,7 +528,7 @@ namespace UnrealBuildTool.Artifacts
 		/// </summary>
 		/// <param name="artifact">Artifact in question</param>
 		/// <returns>Tree name</returns>
-		private static string GetArtifactTreeName(Artifact artifact)
+		private static string GetArtifactTreeName(ArtifactFile artifact)
 		{
 			switch (artifact.Tree)
 			{
diff --git a/Engine/Source/Programs/UnrealBuildTool/Artifacts/Artifacts.cs b/Engine/Source/Programs/UnrealBuildTool/Artifacts/Artifacts.cs
--- a/Engine/Source/Programs/UnrealBuildTool/Artifacts/Artifacts.cs
+++ b/Engine/Source/Programs/UnrealBuildTool/Artifacts/Artifacts.cs
@@ -33,24 +33,13 @@ namespace UnrealBuildTool.Artifacts
 	}
 
 	/// <summary>
-	/// Represents a single artifact
+	/// Represents a single artifact file
 	/// </summary>
 	/// <param name="Tree">Directory tree containing the artifact</param>
 	/// <param name="Name">Name of the artifact</param>
 	/// <param name="ContentHash">Hash of the artifact contents</param>
-	public record struct Artifact(ArtifactDirectoryTree Tree, Utf8String Name, IoHash ContentHash)
+	public record struct ArtifactFile(ArtifactDirectoryTree Tree, Utf8String Name, IoHash ContentHash)
 	{
-		/// <summary>
-		/// Directory mapping object
-		/// </summary>
-		public IArtifactDirectoryMapping? DirectoryMapping { get; set; } = null;
-
-		/// <summary>
-		/// The full path of the artifact
-		/// </summary>
-		/// <returns>Full path of the artifact</returns>
-		public string GetFullPath() => GetFullPath(DirectoryMapping);
-
 		/// <summary>
 		/// The full path of the artifact
 		/// </summary>
@@ -72,13 +61,13 @@ namespace UnrealBuildTool.Artifacts
 	}
 
 	/// <summary>
-	/// Given a set of inputs, provide a list of outputs
+	/// Collection of inputs and outputs for an action
 	/// </summary>
 	/// <param name="Key">The hash of the primary input and the environment</param>
-	/// <param name="MappingKey">The unique hash for all inputs and the environment</param>
+	/// <param name="ActionKey">The unique hash for all inputs and the environment</param>
 	/// <param name="Inputs">Information about all inputs</param>
 	/// <param name="Outputs">Information about all outputs</param>
-	public record struct ArtifactMapping(IoHash Key, IoHash MappingKey, Artifact[] Inputs, Artifact[] Outputs)
+	public record struct ArtifactAction(IoHash Key, IoHash ActionKey, ArtifactFile[] Inputs, ArtifactFile[] Outputs)
 	{
 		/// <summary>
 		/// Directory mapping object
@@ -88,7 +77,7 @@ namespace UnrealBuildTool.Artifacts
 		/// <inheritdoc/>
 		public override int GetHashCode()
 		{
-			return HashCode.Combine(Key.GetHashCode(), MappingKey.GetHashCode(), Inputs.GetHashCode(), Outputs.GetHashCode());
+			return HashCode.Combine(Key.GetHashCode(), ActionKey.GetHashCode(), Inputs.GetHashCode(), Outputs.GetHashCode());
 		}
 	}
 
@@ -103,12 +92,12 @@ namespace UnrealBuildTool.Artifacts
 		/// </summary>
 		/// <param name="reader">Source reader</param>
 		/// <returns>Read artifact structure</returns>
-		public static Artifact ReadArtifact(this IMemoryReader reader)
+		public static ArtifactFile ReadArtifact(this IMemoryReader reader)
 		{
 			ArtifactDirectoryTree tree = (ArtifactDirectoryTree)reader.ReadUInt8();
 			Utf8String name = reader.ReadUtf8String();
 			IoHash contentHash = reader.ReadIoHash();
-			return new Artifact(tree, name, contentHash);
+			return new ArtifactFile(tree, name, contentHash);
 		}
 
 		/// <summary>
@@ -116,7 +105,7 @@ namespace UnrealBuildTool.Artifacts
 		/// </summary>
 		/// <param name="writer">Destination writer</param>
 		/// <param name="artifact">Artifact to be written</param>
-		public static void WriteArtifact(this IMemoryWriter writer, Artifact artifact)
+		public static void WriteArtifact(this IMemoryWriter writer, ArtifactFile artifact)
 		{
 			writer.WriteUInt8((byte)artifact.Tree);
 			writer.WriteUtf8String(artifact.Name);
@@ -124,30 +113,30 @@ namespace UnrealBuildTool.Artifacts
 		}
 
 		/// <summary>
-		/// Read an artifact mapping structure
+		/// Read an artifact action structure
 		/// </summary>
 		/// <param name="reader">Source reader</param>
-		/// <returns>Read artifact mapping structure</returns>
-		public static ArtifactMapping ReadArtifactMapping(this IMemoryReader reader)
+		/// <returns>Read artifact action structure</returns>
+		public static ArtifactAction ReadArtifactAction(this IMemoryReader reader)
 		{
 			IoHash key = reader.ReadIoHash();
-			IoHash mappingKey = reader.ReadIoHash();
-			Artifact[] inputs = reader.ReadVariableLengthArray(() => reader.ReadArtifact());
-			Artifact[] outputs = reader.ReadVariableLengthArray(() => reader.ReadArtifact());
-			return new ArtifactMapping(key, mappingKey, inputs, outputs);
+			IoHash actionKey = reader.ReadIoHash();
+			ArtifactFile[] inputs = reader.ReadVariableLengthArray(() => reader.ReadArtifact());
+			ArtifactFile[] outputs = reader.ReadVariableLengthArray(() => reader.ReadArtifact());
+			return new ArtifactAction(key, actionKey, inputs, outputs);
 		}
 
 		/// <summary>
-		/// Write an artifact mapping structure
+		/// Write an artifact action structure
 		/// </summary>
 		/// <param name="writer">Destination writer</param>
-		/// <param name="artifactMapping">Artifact mapping to be written</param>
-		public static void WriteArtifactMapping(this IMemoryWriter writer, ArtifactMapping artifactMapping)
+		/// <param name="artifactAction">Artifact action to be written</param>
+		public static void WriteArtifactAction(this IMemoryWriter writer, ArtifactAction artifactAction)
 		{
-			writer.WriteIoHash(artifactMapping.Key);
-			writer.WriteIoHash(artifactMapping.MappingKey);
-			writer.WriteVariableLengthArray(artifactMapping.Inputs, x => writer.WriteArtifact(x));
-			writer.WriteVariableLengthArray(artifactMapping.Outputs, x => writer.WriteArtifact(x));
+			writer.WriteIoHash(artifactAction.Key);
+			writer.WriteIoHash(artifactAction.ActionKey);
+			writer.WriteVariableLengthArray(artifactAction.Inputs, x => writer.WriteArtifact(x));
+			writer.WriteVariableLengthArray(artifactAction.Outputs, x => writer.WriteArtifact(x));
 		}
 	}
 
@@ -208,28 +197,28 @@ namespace UnrealBuildTool.Artifacts
 		public Task<ArtifactCacheState> WaitForReadyAsync();
 
 		/// <summary>
-		/// Given a collection of partial keys return all matching mappings
+		/// Given a collection of partial keys return all matching artifacts
 		/// </summary>
 		/// <param name="partialKeys">Source file key</param>
 		/// <param name="cancellationToken">Token to be used to cancel operations</param>
 		/// <returns>Collection of all known artifacts</returns>
-		public Task<ArtifactMapping[]> QueryArtifactMappingsAsync(IoHash[] partialKeys, CancellationToken cancellationToken);
+		public Task<ArtifactAction[]> QueryArtifactActionsAsync(IoHash[] partialKeys, CancellationToken cancellationToken);
 
 		/// <summary>
 		/// Query the actual contents of a group of artifacts
 		/// </summary>
-		/// <param name="artifactMappings">Mappings to be read</param>
+		/// <param name="artifactActions">Actions to be read</param>
 		/// <param name="cancellationToken">Token to be used to cancel operations</param>
 		/// <returns>Dictionary of the artifacts</returns>
-		public Task<bool[]?> QueryArtifactOutputsAsync(ArtifactMapping[] artifactMappings, CancellationToken cancellationToken);
+		public Task<bool[]?> QueryArtifactOutputsAsync(ArtifactAction[] artifactActions, CancellationToken cancellationToken);
 
 		/// <summary>
 		/// Save new artifact to the cache
 		/// </summary>
-		/// <param name="artifactMappings">Collection of artifacts to be saved</param>
+		/// <param name="artifactActions">Collection of artifacts to be saved</param>
 		/// <param name="cancellationToken">Token to be used to cancel operations</param>
 		/// <returns>Asynchronous task objects</returns>
-		public Task SaveArtifactMappingsAsync(ArtifactMapping[] artifactMappings, CancellationToken cancellationToken);
+		public Task SaveArtifactActionsAsync(ArtifactAction[] artifactActions, CancellationToken cancellationToken);
 
 		/// <summary>
 		/// Flush all updates in the cache asynchronously.
diff --git a/Engine/Source/Programs/UnrealBuildTool/Artifacts/HordeStorageArtifactCache.cs b/Engine/Source/Programs/UnrealBuildTool/Artifacts/HordeStorageArtifactCache.cs
--- a/Engine/Source/Programs/UnrealBuildTool/Artifacts/HordeStorageArtifactCache.cs
+++ b/Engine/Source/Programs/UnrealBuildTool/Artifacts/HordeStorageArtifactCache.cs
@@ -16,50 +16,50 @@ using Microsoft.Extensions.Logging.Abstractions;
 namespace UnrealBuildTool.Artifacts
 {
 	/// <summary>
-	/// Horde specific artifact mapping structure that also contains the file nodes for the outputs
+	/// Horde specific artifact action structure that also contains the file nodes for the outputs
 	/// </summary>
-	readonly struct HordeArtifactMapping
+	readonly struct HordeArtifactAction
 	{
 
 		/// <summary>
-		/// Artifact mapping
+		/// Artifact action
 		/// </summary>
-		public readonly ArtifactMapping ArtifactMapping;
+		public readonly ArtifactAction ArtifactAction;
 
 		/// <summary>
 		/// Collection of output file references.  There should be exactly the same number
-		/// of file references as outputs in the mapping
+		/// of file references as outputs in the action
 		/// </summary>
 		public readonly NodeRef<ChunkedDataNode>[] OutputRefs;
 
 		/// <summary>
 		/// Construct a new horde artifact number
 		/// </summary>
-		/// <param name="artifactMapping">Artifact mapping</param>
+		/// <param name="artifactAction">Artifact action</param>
 		/// <exception cref="ArgumentException"></exception>
-		public HordeArtifactMapping(ArtifactMapping artifactMapping)
+		public HordeArtifactAction(ArtifactAction artifactAction)
 		{
-			ArtifactMapping = artifactMapping;
-			OutputRefs = new NodeRef<ChunkedDataNode>[ArtifactMapping.Outputs.Length];
+			ArtifactAction = artifactAction;
+			OutputRefs = new NodeRef<ChunkedDataNode>[ArtifactAction.Outputs.Length];
 		}
 
 		/// <summary>
-		/// Construct a new artifact mapping from the reader
+		/// Construct a new artifact action from the reader
 		/// </summary>
 		/// <param name="reader">Source reader</param>
-		public HordeArtifactMapping(NodeReader reader)
+		public HordeArtifactAction(NodeReader reader)
 		{
-			ArtifactMapping = reader.ReadArtifactMapping();
+			ArtifactAction = reader.ReadArtifactAction();
 			OutputRefs = reader.ReadVariableLengthArray(() => new NodeRef<ChunkedDataNode>(reader));
 		}
 
 		/// <summary>
-		/// Serialize the artifact mapping 
+		/// Serialize the artifact action 
 		/// </summary>
 		/// <param name="writer">Destination writer</param>
 		public void Serialize(ITreeNodeWriter writer)
 		{
-			writer.WriteArtifactMapping(ArtifactMapping);
+			writer.WriteArtifactAction(ArtifactAction);
 			writer.WriteVariableLengthArray(OutputRefs, x => x.Serialize(writer));
 		}
 
@@ -82,9 +82,9 @@ namespace UnrealBuildTool.Artifacts
 
 			ChunkedDataWriter fileWriter = new(writer, options);
 			int index = 0;
-			foreach (Artifact artifact in ArtifactMapping.Outputs)
+			foreach (ArtifactFile artifact in ArtifactAction.Outputs)
 			{
-				string outputName = artifact.GetFullPath();
+				string outputName = artifact.GetFullPath(ArtifactAction.DirectoryMapping);
 				using FileStream stream = new(outputName, FileMode.Open, FileAccess.Read, FileShare.Read);
 				OutputRefs[index++] = new NodeRef<ChunkedDataNode>(await fileWriter.CreateAsync(stream, nodeTypeOptions.TargetSize, cancellationToken));
 			}
@@ -98,66 +98,66 @@ namespace UnrealBuildTool.Artifacts
 	{
 
 		/// <summary>
-		/// Read a horde artifact mapping
+		/// Read a horde artifact action
 		/// </summary>
 		/// <param name="reader">Source reader</param>
-		/// <returns>Created artifact mapping</returns>
-		public static HordeArtifactMapping ReadHordeArtifactMapping(this NodeReader reader)
+		/// <returns>Created artifact action</returns>
+		public static HordeArtifactAction ReadHordeArtifactAction(this NodeReader reader)
 		{
-			return new HordeArtifactMapping(reader);
+			return new HordeArtifactAction(reader);
 		}
 
 		/// <summary>
-		/// Write a horde artifact mapping
+		/// Write a horde artifact action
 		/// </summary>
 		/// <param name="writer">Destination writer</param>
-		/// <param name="artifactMapping">Artifact mapping to write</param>
-		public static void WriteHordeArtifactMapping(this ITreeNodeWriter writer, HordeArtifactMapping artifactMapping)
+		/// <param name="artifactAction">Artifact action to write</param>
+		public static void WriteHordeArtifactAction(this ITreeNodeWriter writer, HordeArtifactAction artifactAction)
 		{
-			artifactMapping.Serialize(writer);
+			artifactAction.Serialize(writer);
 		}
 	}
 
 	/// <summary>
-	/// Horde node that represents a collection of mapping nodes 
+	/// Horde node that represents a collection of action nodes 
 	/// </summary>
 	[NodeType("{E8DBCD77-861D-4CAE-B77F-5807D26E2533}")]
-	class ArtifactMappingCollectionNode : Node
+	class ArtifactActionCollectionNode : Node
 	{
 
 		/// <summary>
-		/// Collection of mappings
+		/// Collection of actions
 		/// </summary>
-		public Dictionary<IoHash, HordeArtifactMapping> Mappings = new();
+		public Dictionary<IoHash, HordeArtifactAction> ArtifactActions = new();
 
 		/// <summary>
 		/// Construct a new collection
 		/// </summary>
-		public ArtifactMappingCollectionNode()
+		public ArtifactActionCollectionNode()
 		{
 		}
 
 		/// <summary>
-		/// Construct a new mapping collection from the source reader
+		/// Construct a new artifact action collection from the source reader
 		/// </summary>
 		/// <param name="reader">Source reader</param>
-		public ArtifactMappingCollectionNode(NodeReader reader)
+		public ArtifactActionCollectionNode(NodeReader reader)
 		{
-			Mappings = reader.ReadDictionary<IoHash, HordeArtifactMapping>(() => reader.ReadIoHash(), () => reader.ReadHordeArtifactMapping());
+			ArtifactActions = reader.ReadDictionary<IoHash, HordeArtifactAction>(() => reader.ReadIoHash(), () => reader.ReadHordeArtifactAction());
 		}
 
 		/// <inheritdoc/>
 		public override void Serialize(ITreeNodeWriter writer)
 		{
-			writer.WriteDictionary<IoHash, HordeArtifactMapping>(Mappings, (x) => writer.WriteIoHash(x), (x) => writer.WriteHordeArtifactMapping(x));
+			writer.WriteDictionary<IoHash, HordeArtifactAction>(ArtifactActions, (x) => writer.WriteIoHash(x), (x) => writer.WriteHordeArtifactAction(x));
 		}
 
 		/// <inheritdoc/>
 		public override IEnumerable<NodeRef> EnumerateRefs()
 		{
-			foreach (HordeArtifactMapping artifactMapping in Mappings.Values)
+			foreach (HordeArtifactAction artifactAction in ArtifactActions.Values)
 			{
-				foreach (NodeRef outputRef in artifactMapping.OutputRefs)
+				foreach (NodeRef outputRef in artifactAction.OutputRefs)
 				{
 					yield return outputRef;
 				}
@@ -179,7 +179,7 @@ namespace UnrealBuildTool.Artifacts
 	public class HordeStorageArtifactCache : IArtifactCache
 	{
 		/// <summary>
-		/// Defines the theoretical max number of pending mappings to write
+		/// Defines the theoretical max number of pending actions to write
 		/// </summary>
 		const int MaxPendingSize = 128;
 
@@ -204,14 +204,14 @@ namespace UnrealBuildTool.Artifacts
 		private int _state = (int)ArtifactCacheState.Pending;
 
 		/// <summary>
-		/// Collection of mappings waiting to be written
+		/// Collection of actions waiting to be written
 		/// </summary>
-		private readonly List<ArtifactMapping> _pendingMappings;
+		private readonly List<ArtifactAction> _pendingWrites;
 
 		/// <summary>
 		/// Task for any pending flush
 		/// </summary>
-		private Task? _pendingFlushTask = null;
+		private Task? _pendingWritesFlushTask = null;
 
 		/// <summary>
 		/// Controls access to shared data structures
@@ -256,7 +256,7 @@ namespace UnrealBuildTool.Artifacts
 
 		static HordeStorageArtifactCache()
 		{
-			Node.RegisterType<ArtifactMappingCollectionNode>();
+			Node.RegisterType<ArtifactActionCollectionNode>();
 		}
 
 		/// <summary>
@@ -268,7 +268,7 @@ namespace UnrealBuildTool.Artifacts
 		{
 			_store = storage;
 			_logger = logger;
-			_pendingMappings = new(MaxPendingSize);
+			_pendingWrites = new(MaxPendingSize);
 		}
 
 		/// <inheritdoc/>
@@ -278,29 +278,29 @@ namespace UnrealBuildTool.Artifacts
 		}
 
 		/// <inheritdoc/>
-		public async Task<ArtifactMapping[]> QueryArtifactMappingsAsync(IoHash[] partialKeys, CancellationToken cancellationToken)
+		public async Task<ArtifactAction[]> QueryArtifactActionsAsync(IoHash[] partialKeys, CancellationToken cancellationToken)
 		{
 			if (State != ArtifactCacheState.Available || _store == null)
 			{
-				return Array.Empty<ArtifactMapping>();
+				return Array.Empty<ArtifactAction>();
 			}
 
-			List<ArtifactMapping> mappings = new();
+			List<ArtifactAction> artifactActions = new();
 			await _semaphore.WaitAsync(cancellationToken);
 			try
 			{
 				foreach (IoHash key in partialKeys)
 				{
-					lock (_pendingMappings)
+					lock (_pendingWrites)
 					{
-						mappings.AddRange(_pendingMappings.Where(x => x.Key == key));
+						artifactActions.AddRange(_pendingWrites.Where(x => x.Key == key));
 					}
-					ArtifactMappingCollectionNode? node = await _store.TryReadNodeAsync<ArtifactMappingCollectionNode>(GetRefName(key), default, cancellationToken);
+					ArtifactActionCollectionNode? node = await _store.TryReadNodeAsync<ArtifactActionCollectionNode>(GetRefName(key), default, cancellationToken);
 					if (node != null)
 					{
-						foreach (HordeArtifactMapping mapping in node.Mappings.Values)
+						foreach (HordeArtifactAction artifactAction in node.ArtifactActions.Values)
 						{
-							mappings.Add(mapping.ArtifactMapping);
+							artifactActions.Add(artifactAction.ArtifactAction);
 						}
 					}
 				}
@@ -309,33 +309,33 @@ namespace UnrealBuildTool.Artifacts
 			{
 				_semaphore.Release();
 			}
-			return mappings.ToArray();
+			return artifactActions.ToArray();
 		}
 
 		/// <inheritdoc/>
-		public async Task<bool[]?> QueryArtifactOutputsAsync(ArtifactMapping[] mappings, CancellationToken cancellationToken)
+		public async Task<bool[]?> QueryArtifactOutputsAsync(ArtifactAction[] artifactActions, CancellationToken cancellationToken)
 		{
 			if (State != ArtifactCacheState.Available || _store == null)
 			{
 				return null;
 			}
 
-			bool[] output = new bool[mappings.Length];
+			bool[] output = new bool[artifactActions.Length];
 			Array.Fill(output, false);
 
-			for (int index = 0; index < mappings.Length; index++)
+			for (int index = 0; index < artifactActions.Length; index++)
 			{
 				output[index] = false;
-				ArtifactMapping mapping = mappings[index];
-				ArtifactMappingCollectionNode? node = await _store.TryReadNodeAsync<ArtifactMappingCollectionNode>(GetRefName(mapping.Key), default, cancellationToken);
+				ArtifactAction artifactAction = artifactActions[index];
+				ArtifactActionCollectionNode? node = await _store.TryReadNodeAsync<ArtifactActionCollectionNode>(GetRefName(artifactAction.Key), default, cancellationToken);
 				if (node != null)
 				{
-					if (node.Mappings.TryGetValue(mapping.MappingKey, out HordeArtifactMapping hordeArtifactMapping))
+					if (node.ArtifactActions.TryGetValue(artifactAction.ActionKey, out HordeArtifactAction hordeArtifactAction))
 					{
 						output[index] = true;
 
 						int refIndex = 0;
-						foreach (NodeRef<ChunkedDataNode> artifactRef in hordeArtifactMapping.OutputRefs)
+						foreach (NodeRef<ChunkedDataNode> artifactRef in hordeArtifactAction.OutputRefs)
 						{
 							if (artifactRef.Handle == null)
 							{
@@ -344,7 +344,7 @@ namespace UnrealBuildTool.Artifacts
 							}
 							try
 							{
-								string outputName = hordeArtifactMapping.ArtifactMapping.Outputs[refIndex++].GetFullPath(mapping.DirectoryMapping);
+								string outputName = hordeArtifactAction.ArtifactAction.Outputs[refIndex++].GetFullPath(artifactAction.DirectoryMapping);
 								using FileStream stream = new(outputName, FileMode.Create, FileAccess.Write, FileShare.ReadWrite);
 								await ChunkedDataNode.CopyToStreamAsync(artifactRef.Handle, stream, cancellationToken);
 							}
@@ -357,9 +357,9 @@ namespace UnrealBuildTool.Artifacts
 
 						if (!output[index])
 						{
-							foreach (Artifact artifact in hordeArtifactMapping.ArtifactMapping.Outputs)
+							foreach (ArtifactFile artifact in hordeArtifactAction.ArtifactAction.Outputs)
 							{
-								string outputName = artifact.GetFullPath(mapping.DirectoryMapping);
+								string outputName = artifact.GetFullPath(artifactAction.DirectoryMapping);
 								if (File.Exists(outputName))
 								{
 									try
@@ -379,16 +379,16 @@ namespace UnrealBuildTool.Artifacts
 		}
 
 		/// <inheritdoc/>
-		public async Task SaveArtifactMappingsAsync(ArtifactMapping[] artifactMappings, CancellationToken cancellationToken)
+		public async Task SaveArtifactActionsAsync(ArtifactAction[] artifactActions, CancellationToken cancellationToken)
 		{
 			if (State != ArtifactCacheState.Available || _store == null)
 			{
 				return;
 			}
 
-			lock (_pendingMappings)
+			lock (_pendingWrites)
 			{
-				_pendingMappings.AddRange(artifactMappings);
+				_pendingWrites.AddRange(artifactActions);
 			}
 
 			Task? task = FlushChangesInternalAsync(false, cancellationToken);
@@ -414,32 +414,32 @@ namespace UnrealBuildTool.Artifacts
 		}
 
 		/// <summary>
-		/// Optionally flush all pending mappings
+		/// Optionally flush all pending writes
 		/// </summary>
 		/// <param name="force">If true, force a flush</param>
 		/// <param name="cancellationToken">Cancellation token</param>
 		private Task? FlushChangesInternalAsync(bool force, CancellationToken cancellationToken)
 		{
 			Task? pendingFlushTask = null;
-			lock (_pendingMappings)
+			lock (_pendingWrites)
 			{
 
 				// If any prior flush task has completed, then forget it
-				if (_pendingFlushTask != null && _pendingFlushTask.IsCompleted)
+				if (_pendingWritesFlushTask != null && _pendingWritesFlushTask.IsCompleted)
 				{
-					_pendingFlushTask = null;
+					_pendingWritesFlushTask = null;
 				}
 
 				// We start a new flush under the following condition
 				//
-				// 1) Mappings must be pending
+				// 1) Actions must be pending
 				// 2) Create a new task if force is specified
 				// 3) -OR- Create a new task if there is no current task and we have reached the limit 
-				if (_pendingMappings.Count > 0 && (force || (_pendingMappings.Count >= MaxPendingSize && _pendingFlushTask == null)))
+				if (_pendingWrites.Count > 0 && (force || (_pendingWrites.Count >= MaxPendingSize && _pendingWritesFlushTask == null)))
 				{
-					ArtifactMapping[] mappingsToFlush = _pendingMappings.ToArray();
-					Task? priorTask = _pendingFlushTask;
-					_pendingMappings.Clear();
+					ArtifactAction[] artifactActionsToFlush = _pendingWrites.ToArray();
+					Task? priorTask = _pendingWritesFlushTask;
+					_pendingWrites.Clear();
 					async Task action()
 					{
 
@@ -449,11 +449,11 @@ namespace UnrealBuildTool.Artifacts
 							await priorTask;
 						}
 
-						// Block reading while we update the mappings
+						// Block reading while we update the actions
 						await _semaphore.WaitAsync(cancellationToken);
 						try
 						{
-							List<Task> tasks = CommitMappings(mappingsToFlush, cancellationToken);
+							List<Task> tasks = CommitArtifactActions(artifactActionsToFlush, cancellationToken);
 							await Task.WhenAll(tasks);
 						}
 						finally
@@ -461,7 +461,7 @@ namespace UnrealBuildTool.Artifacts
 							_semaphore.Release();
 						}
 					}
-					pendingFlushTask = _pendingFlushTask = new(() => action().Wait(), cancellationToken);
+					pendingFlushTask = _pendingWritesFlushTask = new(() => action().Wait(), cancellationToken);
 				}
 			}
 
@@ -471,40 +471,40 @@ namespace UnrealBuildTool.Artifacts
 		}
 
 		/// <summary>
-		/// Add a group of mappings to a new or existing source
+		/// Add a group of artifact actions to a new or existing source
 		/// </summary>
-		/// <param name="mappings">New mappings to add</param>
+		/// <param name="artifactActions">New artifact actions to add</param>
 		/// <param name="cancellationToken">Token to be used to cancel operations</param>
 		/// <returns>List of tasks</returns>
-		private List<Task> CommitMappings(ArtifactMapping[] mappings, CancellationToken cancellationToken)
+		private List<Task> CommitArtifactActions(ArtifactAction[] artifactActions, CancellationToken cancellationToken)
 		{
 			List<Task> tasks = new();
-			if (mappings.Length == 0)
+			if (artifactActions.Length == 0)
 			{
 				return tasks;
 			}
 
-			// Loop through the mappings
-			foreach (ArtifactMapping mapping in mappings)
+			// Loop through the artifact actions
+			foreach (ArtifactAction artifactAction in artifactActions)
 			{
 
 				// Create the task to write the files
 				tasks.Add(Task.Run(async () =>
 				{
-					RefName refName = GetRefName(mapping.Key);
+					RefName refName = GetRefName(artifactAction.Key);
 
 					// Locate the destination collection for this key
-					ArtifactMappingCollectionNode? node = _store!.TryReadNodeAsync<ArtifactMappingCollectionNode>(refName, default, cancellationToken).Result;
-					node ??= new ArtifactMappingCollectionNode();
+					ArtifactActionCollectionNode? node = _store!.TryReadNodeAsync<ArtifactActionCollectionNode>(refName, default, cancellationToken).Result;
+					node ??= new ArtifactActionCollectionNode();
 
-					// Update the mapping collection
-					HordeArtifactMapping hordeArtifactMapping = new(mapping);
-					node.Mappings[mapping.MappingKey] = hordeArtifactMapping;
+					// Update the artifact action collection
+					HordeArtifactAction hordeArtifactAction = new(artifactAction);
+					node.ArtifactActions[artifactAction.ActionKey] = hordeArtifactAction;
 					node.MarkAsDirty();
 
-					// Save the mapping file
+					// Save the artifact action file
 					using IStorageWriter writer = _store!.CreateWriter();
-					await hordeArtifactMapping.WriteFilesAsync(writer, cancellationToken);
+					await hordeArtifactAction.WriteFilesAsync(writer, cancellationToken);
 
 					// Save the collection
 					NodeHandle _ = await writer.WriteAsync(refName, node, null, cancellationToken);
@@ -549,9 +549,9 @@ namespace UnrealBuildTool.Artifacts
 		}
 
 		/// <summary>
-		/// Return the ref name for horde storage given a mapping collection key
+		/// Return the ref name for horde storage given a artitfact action collection key
 		/// </summary>
-		/// <param name="key">Mapping collection key</param>
+		/// <param name="key">Artifact action collection key</param>
 		/// <returns>The reference name</returns>
 		private static RefName GetRefName(IoHash key)
 		{
