Build fix.

diff --git a/Engine/Source/Programs/UnrealBuildTool/Artifacts/HordeStorageArtifactCache.cs b/Engine/Source/Programs/UnrealBuildTool/Artifacts/HordeStorageArtifactCache.cs
--- a/Engine/Source/Programs/UnrealBuildTool/Artifacts/HordeStorageArtifactCache.cs
+++ b/Engine/Source/Programs/UnrealBuildTool/Artifacts/HordeStorageArtifactCache.cs
@@ -77,8 +77,9 @@ namespace UnrealBuildTool.Artifacts
 		/// <returns>Task</returns>
 		public async Task WriteFilesAsync(IStorageWriter writer, CancellationToken cancellationToken)
 		{
-			ChunkingOptionsForNodeType nodeTypeOptions = new(512 * 1024, 2 * 1024 * 1024, 1 * 1024 * 1024);
-			ChunkingOptions options = new() { LeafOptions = nodeTypeOptions, InteriorOptions = nodeTypeOptions };
+			LeafChunkedDataNodeOptions leafOptions = new(512 * 1024, 1 * 1024 * 1024, 2 * 1024 * 1024);
+			InteriorChunkedDataNodeOptions interiorOptions = new(1, 10, 20);
+			ChunkingOptions options = new() { LeafOptions = leafOptions, InteriorOptions = interiorOptions };
 
 			ChunkedDataWriter fileWriter = new(writer, options);
 			int index = 0;
@@ -86,7 +87,7 @@ namespace UnrealBuildTool.Artifacts
 			{
 				string outputName = artifact.GetFullPath(ArtifactAction.DirectoryMapping);
 				using FileStream stream = new(outputName, FileMode.Open, FileAccess.Read, FileShare.Read);
-				OutputRefs[index++] = new NodeRef<ChunkedDataNode>(await fileWriter.CreateAsync(stream, nodeTypeOptions.TargetSize, cancellationToken));
+				OutputRefs[index++] = new NodeRef<ChunkedDataNode>(await fileWriter.CreateAsync(stream, leafOptions.TargetSize, cancellationToken));
 			}
 		}
 	}
