BuildSettingsVersion.V4: Default CppStandard to Cpp20. Latest is not being updated to V4 with this change and will happen at a later date
* Add new setting CppStandardEngine so this can be controlled separately for engine module
* CppStandard no longer requires RequiresUniqueBuildEnvironment as it does not affect engine modules

diff --git a/Engine/Source/Programs/UnrealBuildTool/Configuration/TargetRules.cs b/Engine/Source/Programs/UnrealBuildTool/Configuration/TargetRules.cs
--- a/Engine/Source/Programs/UnrealBuildTool/Configuration/TargetRules.cs
+++ b/Engine/Source/Programs/UnrealBuildTool/Configuration/TargetRules.cs
@@ -129,10 +129,15 @@ namespace UnrealBuildTool
 		/// <summary>
 		/// New defaults for 5.2:
 		/// * ModuleRules.bLegacyParentIncludePaths = false
-		/// Work in progress, not ready to be enabled as Latest.
 		/// </summary>
 		V3,
 
+		/// <summary>
+		/// New defaults for 5.3:
+		/// * TargetRules.CppStandard = CppStandard.Default has changed from Cpp17 to Cpp20
+		/// </summary>
+		V4,
+
 		// *** When adding new entries here, be sure to update GameProjectUtils::GetDefaultBuildSettingsVersion() to ensure that new projects are created correctly. ***
 
 		/// <summary>
@@ -2070,12 +2075,29 @@ namespace UnrealBuildTool
 		private bool? bLegacyParentIncludePathsPrivate;
 
 		/// <summary>
-		/// Which C++ standard to use for compiling this target
+		/// Which C++ standard to use for compiling this target (for engine modules)
 		/// </summary>
 		[RequiresUniqueBuildEnvironment]
+		[CommandLine("-CppStdEngine")]
+		[XmlConfigFile(Category = "BuildConfiguration")]
+		public CppStandardVersion CppStandardEngine
+		{
+			get => CppStandardEnginePrivate ?? CppStandardVersion.EngineDefault;
+			set => CppStandardEnginePrivate = value;
+		}
+		private CppStandardVersion? CppStandardEnginePrivate;
+
+		/// <summary>
+		/// Which C++ standard to use for compiling this target (for non-engine modules)
+		/// </summary>
 		[CommandLine("-CppStd")]
 		[XmlConfigFile(Category = "BuildConfiguration")]
-		public CppStandardVersion CppStandard = CppStandardVersion.Default;
+		public CppStandardVersion CppStandard
+		{
+			get => CppStandardPrivate ?? (DefaultBuildSettings < BuildSettingsVersion.V4 ? CppStandardVersion.Cpp17 : CppStandardVersion.Default);
+			set => CppStandardPrivate = value;
+		}
+		private CppStandardVersion? CppStandardPrivate;
 
 		/// <summary>
 		/// Which C standard to use for compiling this target
@@ -2705,6 +2727,11 @@ namespace UnrealBuildTool
 					ModifiedSettings.Add(Tuple.Create(String.Format("{0} = false", nameof(bLegacyParentIncludePaths)), "Omits module parent folders from include paths to reduce compiler command line length. (Previously: true)."));
 				}
 
+				if (BuildSettingsVersion.V4 <= LatestVersion && DefaultBuildSettings < BuildSettingsVersion.V4)
+				{
+					ModifiedSettings.Add(Tuple.Create(String.Format("{0} = CppStandardVersion.Default", nameof(CppStandard)), "Updates C++ Standard to C++20 (Previously: CppStandardVersion.Cpp17)."));
+				}
+
 				if (ModifiedSettings.Count > 0)
 				{
 					string FormatString = String.Format("[Upgrade]     {{0,-{0}}}   => {{1}}", ModifiedSettings.Max(x => x.Item1.Length));
@@ -3181,6 +3208,8 @@ namespace UnrealBuildTool
 
 		public bool bLegacyParentIncludePaths => Inner.bLegacyParentIncludePaths;
 
+		public CppStandardVersion CppStandardEngine => Inner.CppStandardEngine;
+
 		public CppStandardVersion CppStandard => Inner.CppStandard;
 
 		public CStandardVersion CStandard => Inner.CStandard;
diff --git a/Engine/Source/Programs/UnrealBuildTool/Configuration/UEBuildModuleCPP.cs b/Engine/Source/Programs/UnrealBuildTool/Configuration/UEBuildModuleCPP.cs
--- a/Engine/Source/Programs/UnrealBuildTool/Configuration/UEBuildModuleCPP.cs
+++ b/Engine/Source/Programs/UnrealBuildTool/Configuration/UEBuildModuleCPP.cs
@@ -1209,6 +1209,7 @@ namespace UnrealBuildTool
 			CompileEnvironment.ShadowVariableWarningLevel = ModuleCompileEnvironment.ShadowVariableWarningLevel;
 			CompileEnvironment.UnsafeTypeCastWarningLevel = ModuleCompileEnvironment.UnsafeTypeCastWarningLevel;
 			CompileEnvironment.bEnableUndefinedIdentifierWarnings = ModuleCompileEnvironment.bEnableUndefinedIdentifierWarnings;
+			CompileEnvironment.CppStandardEngine = ModuleCompileEnvironment.CppStandardEngine;
 			CompileEnvironment.CppStandard = ModuleCompileEnvironment.CppStandard;
 			CompileEnvironment.CStandard = ModuleCompileEnvironment.CStandard;
 			CompileEnvironment.IncludeOrderVersion = ModuleCompileEnvironment.IncludeOrderVersion;
@@ -1797,10 +1798,10 @@ namespace UnrealBuildTool
 			{
 				Result.CppStandard = (CppStandardVersion)Rules.CppStandard;
 			}
-			// Otherwise, for engine modules use the EngineDefault standard.
+			// Otherwise, for engine modules use the engine standard (typically CppStandardVersion.EngineDefault).
 			else if (Rules.bTreatAsEngineModule)
 			{
-				Result.CppStandard = CppStandardVersion.EngineDefault;
+				Result.CppStandard = Result.CppStandardEngine;
 			}
 			// CppModules require at least Cpp20.
 			if (Target.bEnableCppModules && Result.CppStandard < CppStandardVersion.Cpp20)
diff --git a/Engine/Source/Programs/UnrealBuildTool/Configuration/UEBuildTarget.cs b/Engine/Source/Programs/UnrealBuildTool/Configuration/UEBuildTarget.cs
--- a/Engine/Source/Programs/UnrealBuildTool/Configuration/UEBuildTarget.cs
+++ b/Engine/Source/Programs/UnrealBuildTool/Configuration/UEBuildTarget.cs
@@ -4568,6 +4568,7 @@ namespace UnrealBuildTool
 			GlobalCompileEnvironment.bUseInlining = Rules.bUseInlining;
 			GlobalCompileEnvironment.bCompileISPC = Rules.bCompileISPC;
 			GlobalCompileEnvironment.bHideSymbolsByDefault = !Rules.bPublicSymbolsByDefault;
+			GlobalCompileEnvironment.CppStandardEngine = Rules.CppStandardEngine;
 			GlobalCompileEnvironment.CppStandard = Rules.CppStandard;
 			GlobalCompileEnvironment.CStandard = Rules.CStandard;
 			GlobalCompileEnvironment.MinCpuArchX64 = Rules.MinCpuArchX64;
diff --git a/Engine/Source/Programs/UnrealBuildTool/System/CppCompileEnvironment.cs b/Engine/Source/Programs/UnrealBuildTool/System/CppCompileEnvironment.cs
--- a/Engine/Source/Programs/UnrealBuildTool/System/CppCompileEnvironment.cs
+++ b/Engine/Source/Programs/UnrealBuildTool/System/CppCompileEnvironment.cs
@@ -44,9 +44,9 @@ namespace UnrealBuildTool
 		Latest,
 
 		/// <summary>
-		/// Use the default standard version
+		/// Use the default standard version (BuildSettingsVersion.V1-V3: Cpp17, V4: Cpp20)
 		/// </summary>
-		Default = Cpp17,
+		Default = Cpp20,
 
 		/// <summary>
 		/// Use the default standard version for engine modules
@@ -609,6 +609,11 @@ namespace UnrealBuildTool
 		/// </summary>
 		public bool bHideSymbolsByDefault = true;
 
+		/// <summary>
+		/// Which C++ standard to support for engine modules. CppStandard will be set to this for engine modules and CppStandardEngine should not be checked in any toolchain. May not be compatible with all platforms.
+		/// </summary>
+		public CppStandardVersion CppStandardEngine = CppStandardVersion.EngineDefault;
+
 		/// <summary>
 		/// Which C++ standard to support. May not be compatible with all platforms.
 		/// </summary>
@@ -753,6 +758,7 @@ namespace UnrealBuildTool
 			ParentPCHInstance = Other.ParentPCHInstance;
 			bHackHeaderGenerator = Other.bHackHeaderGenerator;
 			bHideSymbolsByDefault = Other.bHideSymbolsByDefault;
+			CppStandardEngine = Other.CppStandardEngine;
 			CppStandard = Other.CppStandard;
 			CStandard = Other.CStandard;
 			MinCpuArchX64 = Other.MinCpuArchX64;
