[Box]
* Added -BoxNoWrite command line option. When this is set, intermediate/output files are not written to disk (for now there are a few exceptions such as ispc h-files and pch-files)
* Binaries with bug fixes
Added -BoxPrintSummary command line option. By default summary is now not written out to log

diff --git a/Engine/Source/Programs/UnrealBuildTool/System/ActionExecutor.cs b/Engine/Source/Programs/UnrealBuildTool/System/ActionExecutor.cs
--- a/Engine/Source/Programs/UnrealBuildTool/System/ActionExecutor.cs
+++ b/Engine/Source/Programs/UnrealBuildTool/System/ActionExecutor.cs
@@ -64,6 +64,11 @@ namespace UnrealBuildTool
 		/// <returns>True if the build succeeded, false otherwise</returns>
 		public abstract Task<bool> ExecuteActionsAsync(IEnumerable<LinkedAction> ActionsToExecute, ILogger Logger, IActionArtifactCache? actionArtifactCache = null);
 
+		/// <summary>
+		/// Will verify that produced items exists on disk
+		/// </summary>
+		public virtual bool VerifyOutputs => true;
+
 		protected void WriteToolOutput(string Line)
 		{
 			lock (Parser)
diff --git a/Engine/Source/Programs/UnrealBuildTool/System/ActionGraph.cs b/Engine/Source/Programs/UnrealBuildTool/System/ActionGraph.cs
--- a/Engine/Source/Programs/UnrealBuildTool/System/ActionGraph.cs
+++ b/Engine/Source/Programs/UnrealBuildTool/System/ActionGraph.cs
@@ -418,15 +418,18 @@ namespace UnrealBuildTool
 				}
 
 				// Verify the link outputs were created (seems to happen with Win64 compiles)
-				foreach (LinkedAction BuildAction in ActionsToExecute)
+				if (Executor.VerifyOutputs)
 				{
-					if (BuildAction.ActionType == ActionType.Link)
+					foreach (LinkedAction BuildAction in ActionsToExecute)
 					{
-						foreach (FileItem Item in BuildAction.ProducedItems)
+						if (BuildAction.ActionType == ActionType.Link)
 						{
-							if (!Item.Exists)
+							foreach (FileItem Item in BuildAction.ProducedItems)
 							{
-								throw new BuildException($"Failed to produce item: {Item.AbsolutePath}");
+								if (!Item.Exists)
+								{
+									throw new BuildException($"Failed to produce item: {Item.AbsolutePath}");
+								}
 							}
 						}
 					}
