UnrealBuildTool: Always add CppStandard to PCH variant name, to prevent creating duplicate PCH when the engine's cpp standard is mismatched with a project

#rnx

diff --git a/Engine/Source/Programs/UnrealBuildTool/Configuration/UEBuildModuleCPP.cs b/Engine/Source/Programs/UnrealBuildTool/Configuration/UEBuildModuleCPP.cs
--- a/Engine/Source/Programs/UnrealBuildTool/Configuration/UEBuildModuleCPP.cs
+++ b/Engine/Source/Programs/UnrealBuildTool/Configuration/UEBuildModuleCPP.cs
@@ -1160,15 +1160,13 @@ namespace UnrealBuildTool
 				}
 			}
 
-			if (CompileEnvironment.CppStandard != BaseCompileEnvironment.CppStandard)
+			// Always add the cpp standard into the filename to prevent variants when the engine standard is mismatched with a project
+			switch (CompileEnvironment.CppStandard)
 			{
-				switch (CompileEnvironment.CppStandard)
-				{
-					case CppStandardVersion.Cpp14: Variant += ".Cpp14"; break;
-					case CppStandardVersion.Cpp17: Variant += ".Cpp17"; break;
-					case CppStandardVersion.Cpp20: Variant += ".Cpp20"; break;
-					case CppStandardVersion.Latest: Variant += ".CppLatest"; break;
-				}
+				case CppStandardVersion.Cpp14: Variant += ".Cpp14"; break;
+				case CppStandardVersion.Cpp17: Variant += ".Cpp17"; break;
+				case CppStandardVersion.Cpp20: Variant += ".Cpp20"; break;
+				case CppStandardVersion.Latest: Variant += ".CppLatest"; break;
 			}
 
 			if (CompileEnvironment.CStandard != BaseCompileEnvironment.CStandard)
