[UBT]
* IWYUToolChain - Fixed bug where ordering of produced files had changed and toolchain assumed .o file would be first

diff --git a/Engine/Source/Programs/UnrealBuildTool/ToolChain/IWYUToolChain.cs b/Engine/Source/Programs/UnrealBuildTool/ToolChain/IWYUToolChain.cs
--- a/Engine/Source/Programs/UnrealBuildTool/ToolChain/IWYUToolChain.cs
+++ b/Engine/Source/Programs/UnrealBuildTool/ToolChain/IWYUToolChain.cs
@@ -100,7 +100,9 @@ namespace UnrealBuildTool
 
 				Action Action = CompileCPPFile(CompileEnvironment, SourceFile, OutputDir, ModuleName, Graph, GlobalArguments, new CPPOutput());
 
-				string CommandLineArgs = CommonCommandLineArgs + " -Xiwyu --write_json_path=\"" + Action.ProducedItems.First() + "\" ";
+				FileItem IwyuItem = Action.ProducedItems.First(i => i.Name.EndsWith(".iwyu"));
+
+				string CommandLineArgs = CommonCommandLineArgs + " -Xiwyu --write_json_path=\"" + IwyuItem + "\" ";
 				if (SourceFile.HasExtension(".cpp"))
 				{
 					List<FileItem>? InlinedFiles;
@@ -119,7 +121,7 @@ namespace UnrealBuildTool
 				}
 
 				Action.CommandArguments = CommandLineArgs + Action.CommandArguments;
-				IwyuFiles.Add(Action.ProducedItems.First());
+				IwyuFiles.Add(IwyuItem);
 			}
 
 			return new CPPOutput() { ObjectFiles = IwyuFiles };
