UnrealBuildTool: Allow any engine module (bTreatAsEngineModule) to access Internal include paths from other modules

diff --git a/Engine/Source/Programs/UnrealBuildTool/Configuration/UEBuildModule.cs b/Engine/Source/Programs/UnrealBuildTool/Configuration/UEBuildModule.cs
--- a/Engine/Source/Programs/UnrealBuildTool/Configuration/UEBuildModule.cs
+++ b/Engine/Source/Programs/UnrealBuildTool/Configuration/UEBuildModule.cs
@@ -629,8 +629,8 @@ namespace UnrealBuildTool
 				AddIncludePaths(IncludePaths, LegacyPublicIncludePaths);
 			}
 
-			// Add this module's internal include paths, only if the scope contains the same as the SourceModule's scope
-			if (SourceModule != null && Rules.Context.Scope.Contains(SourceModule.Rules.Context.Scope))
+			// Add this module's internal include paths, only if the scope contains the same as the SourceModule's scope or if this module is an engine module
+			if (SourceModule != null && (Rules.Context.Scope.Contains(SourceModule.Rules.Context.Scope) || Rules.bTreatAsEngineModule))
 			{
 				AddIncludePaths(IncludePaths, InternalIncludePaths);
 			}
