Support to compile separate files in VS using the "Compile" menu item when you right click on the files.

diff --git a/Engine/Source/Programs/UnrealBuildTool/Configuration/TargetDescriptor.cs b/Engine/Source/Programs/UnrealBuildTool/Configuration/TargetDescriptor.cs
--- a/Engine/Source/Programs/UnrealBuildTool/Configuration/TargetDescriptor.cs
+++ b/Engine/Source/Programs/UnrealBuildTool/Configuration/TargetDescriptor.cs
@@ -66,6 +66,18 @@ namespace UnrealBuildTool
 		[CommandLine("-SingleFile=")]
 		public List<FileReference> SpecificFilesToCompile = new List<FileReference>();
 
+		/// <summary>
+		/// Relative path to file(s) to compile
+		/// </summary>
+		[CommandLine("-Files=", ListSeparator = ';')]
+		public List<string> RelativePathsToSpecificFilesToCompile = new List<string>();
+
+		/// <summary>
+		/// Working directory when compiling with RelativePathsToSpecificFilesToCompile
+		/// </summary>
+		[CommandLine("-WorkingDir=")]
+		public string? WorkingDir = null;
+
 		/// <summary>
 		/// Will build all files that directly include any of the files provided in -SingleFile
 		/// </summary>
@@ -156,6 +168,18 @@ namespace UnrealBuildTool
 					}
 				}
 
+				// Create the full path for the files specified in RelativePathsToSpecificFilesToCompile
+				DirectoryReference CurrentWorkingDir = Unreal.RootDirectory;
+				if (WorkingDir != null)
+				{
+					CurrentWorkingDir = new DirectoryReference(WorkingDir);
+				}
+
+				foreach (string RelativeFilePath in RelativePathsToSpecificFilesToCompile)
+				{
+					SpecificFilesToCompile.Add(FileReference.Combine(CurrentWorkingDir, RelativeFilePath));
+				}
+
 				// Parse all the hot-reload module names
 				foreach (string ModuleWithSuffix in Arguments.GetValues("-ModuleWithSuffix="))
 				{
diff --git a/Engine/Source/Programs/UnrealBuildTool/ProjectFiles/VisualStudio/VCProject.cs b/Engine/Source/Programs/UnrealBuildTool/ProjectFiles/VisualStudio/VCProject.cs
--- a/Engine/Source/Programs/UnrealBuildTool/ProjectFiles/VisualStudio/VCProject.cs
+++ b/Engine/Source/Programs/UnrealBuildTool/ProjectFiles/VisualStudio/VCProject.cs
@@ -2111,6 +2111,12 @@ namespace UnrealBuildTool
 					{
 						VCProjectFileContent.Append(ProjGenerator.GetVisualStudioLayoutDirSection(Platform, Configuration, ConditionString, Combination.ProjectTarget.TargetRules!.Type, Combination.ProjectTarget.TargetFilePath, ProjectFilePath, NMakePath, ProjectFileFormat));
 					}
+
+					VCProjectFileContent.AppendLine("  <ItemDefinitionGroup {0}>", ConditionString);
+					VCProjectFileContent.AppendLine("    <NMakeCompile>");
+					VCProjectFileContent.AppendLine("      <NMakeCompileFileCommandLine>$(BuildBatchScript) {0} -WorkingDir=$(MSBuildProjectDirectory) -Files=$(SelectedFiles)</NMakeCompileFileCommandLine>", BuildArguments);
+					VCProjectFileContent.AppendLine("    </NMakeCompile>");
+					VCProjectFileContent.AppendLine("  </ItemDefinitionGroup>");
 				}
 
 				if (VCUserFileContent != null && Combination.ProjectTarget != null)
