UnrealBuildTool: Set AVX defines if a min arch is set

#rnx

diff --git a/Engine/Source/Programs/UnrealBuildTool/Platform/Windows/VCToolChain.cs b/Engine/Source/Programs/UnrealBuildTool/Platform/Windows/VCToolChain.cs
--- a/Engine/Source/Programs/UnrealBuildTool/Platform/Windows/VCToolChain.cs
+++ b/Engine/Source/Programs/UnrealBuildTool/Platform/Windows/VCToolChain.cs
@@ -692,6 +692,14 @@ namespace UnrealBuildTool
 				AddDefinition(Arguments, "PLATFORM_ENABLE_VECTORINTRINSICS=1");
 				AddDefinition(Arguments, "PLATFORM_MAYBE_HAS_AVX=1");
 				AddDefinition(Arguments, "PLATFORM_ALWAYS_HAS_AVX=1");
+				if (CompileEnvironment.MinCpuArchX64 >= MinimumCpuArchitectureX64.AVX2)
+				{
+					AddDefinition(Arguments, "PLATFORM_ALWAYS_HAS_AVX_2=1");
+				}
+				if (CompileEnvironment.MinCpuArchX64 >= MinimumCpuArchitectureX64.AVX512)
+				{
+					AddDefinition(Arguments, "PLATFORM_ALWAYS_HAS_AVX_512=1");
+				}
 			}
 
 			// Prompt the user before reporting internal errors to Microsoft.
