Build fix.

diff --git a/Engine/Source/Programs/UnrealBuildTool/Artifacts/HordeStorageArtifactCache.cs b/Engine/Source/Programs/UnrealBuildTool/Artifacts/HordeStorageArtifactCache.cs
--- a/Engine/Source/Programs/UnrealBuildTool/Artifacts/HordeStorageArtifactCache.cs
+++ b/Engine/Source/Programs/UnrealBuildTool/Artifacts/HordeStorageArtifactCache.cs
@@ -57,7 +57,7 @@ namespace UnrealBuildTool.Artifacts
 		/// Serialize the artifact action 
 		/// </summary>
 		/// <param name="writer">Destination writer</param>
-		public void Serialize(ITreeNodeWriter writer)
+		public void Serialize(NodeWriter writer)
 		{
 			writer.WriteArtifactAction(ArtifactAction);
 			writer.WriteVariableLengthArray(OutputRefs, x => x.Serialize(writer));
@@ -113,7 +113,7 @@ namespace UnrealBuildTool.Artifacts
 		/// </summary>
 		/// <param name="writer">Destination writer</param>
 		/// <param name="artifactAction">Artifact action to write</param>
-		public static void WriteHordeArtifactAction(this ITreeNodeWriter writer, HordeArtifactAction artifactAction)
+		public static void WriteHordeArtifactAction(this NodeWriter writer, HordeArtifactAction artifactAction)
 		{
 			artifactAction.Serialize(writer);
 		}
@@ -148,7 +148,7 @@ namespace UnrealBuildTool.Artifacts
 		}
 
 		/// <inheritdoc/>
-		public override void Serialize(ITreeNodeWriter writer)
+		public override void Serialize(NodeWriter writer)
 		{
 			writer.WriteDictionary<IoHash, HordeArtifactAction>(ArtifactActions, (x) => writer.WriteIoHash(x), (x) => writer.WriteHordeArtifactAction(x));
 		}
