- Fixed a crash when generating some Xcode stub projects for commandline building. I could not repro, but this fix works for someone that could

diff --git a/Engine/Source/Programs/UnrealBuildTool/ProjectFiles/ProjectFileGenerator.cs b/Engine/Source/Programs/UnrealBuildTool/ProjectFiles/ProjectFileGenerator.cs
--- a/Engine/Source/Programs/UnrealBuildTool/ProjectFiles/ProjectFileGenerator.cs
+++ b/Engine/Source/Programs/UnrealBuildTool/ProjectFiles/ProjectFileGenerator.cs
@@ -1172,15 +1172,15 @@ namespace UnrealBuildTool
 
 				foreach (KeyValuePair<FileReference, ProjectFile> CurProgramProject in ProgramProjects)
 				{
-					FileReference UnrealProjectFile = CurProgramProject.Value.ProjectTargets.First().UnrealProjectFilePath!;
-					Project? Target = CurProgramProject.Value.ProjectTargets.FirstOrDefault(t => !String.IsNullOrEmpty(t.TargetRules!.SolutionDirectory));
-					if (!bIncludeEnginePrograms && UnrealProjectFile.IsUnderDirectory(Unreal.EngineDirectory) &&
+					FileReference? UnrealProjectFile = CurProgramProject.Value.ProjectTargets.First().UnrealProjectFilePath;
+					if (!bIncludeEnginePrograms && UnrealProjectFile != null && UnrealProjectFile.IsUnderDirectory(Unreal.EngineDirectory) &&
 						// allow a program passed in with -project= to be added, even if bIncludeEnginePrograms is false
 						!DoesProgramMatchOnlyGameProject(UnrealProjectFile))
 					{
 						continue;
 					}
 
+					Project? Target = CurProgramProject.Value.ProjectTargets.FirstOrDefault(t => !String.IsNullOrEmpty(t.TargetRules!.SolutionDirectory));
 					if (Target != null)
 					{
 						RootFolder.AddSubFolder(Target.TargetRules!.SolutionDirectory).ChildProjects.Add(CurProgramProject.Value);
