Disable new clang warning Wsingle-bit-bitfield-constant-conversion, resolving Xcode 15 compile errors

diff --git a/Engine/Source/Programs/UnrealBuildTool/ToolChain/AppleToolChain.cs b/Engine/Source/Programs/UnrealBuildTool/ToolChain/AppleToolChain.cs
--- a/Engine/Source/Programs/UnrealBuildTool/ToolChain/AppleToolChain.cs
+++ b/Engine/Source/Programs/UnrealBuildTool/ToolChain/AppleToolChain.cs
@@ -365,7 +365,9 @@ namespace UnrealBuildTool
 		{
 			base.GetCompileArguments_WarningsAndErrors(CompileEnvironment, Arguments);
 
+			Arguments.Add("-Wno-unknown-warning-option");
 			Arguments.Add("-Wno-range-loop-analysis");
+			Arguments.Add("-Wno-single-bit-bitfield-constant-conversion");
 		}
 
 		/// <inheritdoc/>
