UnrealBuildTool: Filter out programs when generating projects if they're not located under Source/Programs after the TargetRule is created

#rnx

diff --git a/Engine/Source/Programs/UnrealBuildTool/ProjectFiles/ProjectFileGenerator.cs b/Engine/Source/Programs/UnrealBuildTool/ProjectFiles/ProjectFileGenerator.cs
--- a/Engine/Source/Programs/UnrealBuildTool/ProjectFiles/ProjectFileGenerator.cs
+++ b/Engine/Source/Programs/UnrealBuildTool/ProjectFiles/ProjectFileGenerator.cs
@@ -2704,6 +2704,10 @@ namespace UnrealBuildTool
 					string ProjectFileNameBase;
 					if (TargetRulesObject.Type == TargetType.Program)
 					{
+						if (!bIncludeEnginePrograms && IsEngineTarget)
+						{
+							continue;
+						}
 						IsProgramTarget = true;
 						ProjectFileNameBase = TargetName;
 					}
