Fixed an issue where the linux editor would not create/use the unrealed shared PCH.

diff --git a/Engine/Source/Programs/UnrealBuildTool/Configuration/UEBuildModule.cs b/Engine/Source/Programs/UnrealBuildTool/Configuration/UEBuildModule.cs
--- a/Engine/Source/Programs/UnrealBuildTool/Configuration/UEBuildModule.cs
+++ b/Engine/Source/Programs/UnrealBuildTool/Configuration/UEBuildModule.cs
@@ -1017,10 +1017,12 @@ namespace UnrealBuildTool
 		/// Gathers all the module dependencies a PCH would have
 		/// </summary>
 		/// <param name="bIncludePrivateModules">Whether to include private modules.</param>
-		public HashSet<UEBuildModule> GetAllDependencyModulesForPCH(bool bIncludePrivateModules)
+		/// <param name="bForceCircular">True if circular dependencies should be processed</param>
+		public HashSet<UEBuildModule> GetAllDependencyModulesForPCH(bool bIncludePrivateModules, bool bForceCircular)
 		{
 			HashSet<UEBuildModule> ReferencedModules = new HashSet<UEBuildModule>();
-			InternalGetAllDependencyModulesForPCH(ReferencedModules, new HashSet<UEBuildModule>(), bIncludePrivateModules);
+			HashSet<UEBuildModule> IgnoreReferencedModules = new HashSet<UEBuildModule> { this };
+			InternalGetAllDependencyModulesForPCH(ReferencedModules, IgnoreReferencedModules, bIncludePrivateModules, bForceCircular);
 			return ReferencedModules;
 		}
 
@@ -1030,7 +1032,8 @@ namespace UnrealBuildTool
 		/// <param name="ReferencedModules">Hash of all referenced modules with their addition index.</param>
 		/// <param name="IgnoreReferencedModules">Hashset used to ignore modules which are already added to the list</param>
 		/// <param name="bIncludePrivateModules">Whether to include private modules.</param>
-		private void InternalGetAllDependencyModulesForPCH(HashSet<UEBuildModule> ReferencedModules, HashSet<UEBuildModule> IgnoreReferencedModules, bool bIncludePrivateModules)
+		/// <param name="bForceCircular">True if circular dependencies should be processed</param>
+		private void InternalGetAllDependencyModulesForPCH(HashSet<UEBuildModule> ReferencedModules, HashSet<UEBuildModule> IgnoreReferencedModules, bool bIncludePrivateModules, bool bForceCircular)
 		{
 			List<UEBuildModule> AllDependencyModules = new List<UEBuildModule>(
 				((bIncludePrivateModules && PrivateDependencyModules != null) ? PrivateDependencyModules.Count : 0) +
@@ -1050,15 +1053,15 @@ namespace UnrealBuildTool
 				AllDependencyModules.AddRange(PublicIncludePathModules);
 			}
 
-			foreach (UEBuildModule DependencyModule in AllDependencyModules)
+			foreach (UEBuildModule DependencyModule in AllDependencyModules.Distinct())
 			{
 				// Don't follow circular back-references!
-				if (!HasCircularDependencyOn(DependencyModule.Name))
+				if (bForceCircular || !HasCircularDependencyOn(DependencyModule.Name))
 				{
 					if (IgnoreReferencedModules.Add(DependencyModule))
 					{
 						// Recurse into dependent modules first
-						DependencyModule.InternalGetAllDependencyModulesForPCH(ReferencedModules, IgnoreReferencedModules, false);
+						DependencyModule.InternalGetAllDependencyModulesForPCH(ReferencedModules, IgnoreReferencedModules, false, bForceCircular);
 
 						ReferencedModules.Add(DependencyModule);
 					}
diff --git a/Engine/Source/Programs/UnrealBuildTool/Configuration/UEBuildModuleCPP.cs b/Engine/Source/Programs/UnrealBuildTool/Configuration/UEBuildModuleCPP.cs
--- a/Engine/Source/Programs/UnrealBuildTool/Configuration/UEBuildModuleCPP.cs
+++ b/Engine/Source/Programs/UnrealBuildTool/Configuration/UEBuildModuleCPP.cs
@@ -841,7 +841,7 @@ namespace UnrealBuildTool
 				PrecompiledHeaderDir = IntermediateDirectory;
 			}
 
-			return new PrecompiledHeaderTemplate(this, CompileEnvironment, HeaderFile, PrecompiledHeaderDir, GetAllDependencyModulesForPCH(false));
+			return new PrecompiledHeaderTemplate(this, CompileEnvironment, HeaderFile, PrecompiledHeaderDir, GetAllDependencyModulesForPCH(false, false));
 		}
 
 		static HashSet<string> GetImmutableDefinitions(List<string> Definitions)
@@ -1434,7 +1434,7 @@ namespace UnrealBuildTool
 					PrecompiledHeaderTemplate? OldTemplateValue = CompileEnvironment.SharedPCHs.FirstOrDefault(x => OnlyDirectDependencies.Contains(x.Module) && x.IsValidFor(CompileEnvironment));
 
 					// Find the first shared PCH module we can use that doesn't reference this module
-					HashSet<UEBuildModule> AllDependencies = GetAllDependencyModulesForPCH(true);
+					HashSet<UEBuildModule> AllDependencies = GetAllDependencyModulesForPCH(true, false);
 					PrecompiledHeaderTemplate? Template = CompileEnvironment.SharedPCHs.FirstOrDefault(x => !x.ModuleDependencies.Contains(this) && AllDependencies.Contains(x.Module) && x.IsValidFor(CompileEnvironment));
 
 					if (Template != OldTemplateValue)
diff --git a/Engine/Source/Programs/UnrealBuildTool/Configuration/UEBuildTarget.cs b/Engine/Source/Programs/UnrealBuildTool/Configuration/UEBuildTarget.cs
--- a/Engine/Source/Programs/UnrealBuildTool/Configuration/UEBuildTarget.cs
+++ b/Engine/Source/Programs/UnrealBuildTool/Configuration/UEBuildTarget.cs
@@ -3438,8 +3438,7 @@ namespace UnrealBuildTool
 			Dictionary<UEBuildModuleCPP, int> SharedPCHModuleToPriority = new Dictionary<UEBuildModuleCPP, int>();
 			foreach (UEBuildModuleCPP SharedPCHModule in SharedPCHModules)
 			{
-				List<UEBuildModule> Dependencies = new List<UEBuildModule>();
-				SharedPCHModule.GetAllDependencyModules(Dependencies, new HashSet<UEBuildModule>(), false, false, false);
+				HashSet<UEBuildModule> Dependencies = SharedPCHModule.GetAllDependencyModulesForPCH(false, false);
 				SharedPCHModuleToPriority.Add(SharedPCHModule, Dependencies.Count(x => SharedPCHModules.Contains(x)));
 			}
 
