Add a new TryParseAsMap helper function to read a serialized UPROPERTY TMap into a Dictionary.

#rnx

diff --git a/Engine/Source/Programs/UnrealBuildTool/System/ConfigHierarchy.cs b/Engine/Source/Programs/UnrealBuildTool/System/ConfigHierarchy.cs
--- a/Engine/Source/Programs/UnrealBuildTool/System/ConfigHierarchy.cs
+++ b/Engine/Source/Programs/UnrealBuildTool/System/ConfigHierarchy.cs
@@ -906,6 +906,38 @@ namespace UnrealBuildTool
 			return true;
 		}
 
+		/// <summary>
+		/// Attempts to parse the given line as a UE map (eg. (("key1","value1"), ("key2","value2")).
+		/// </summary>
+		/// <param name="Line">Line of text to parse</param>
+		/// <param name="Map">Receives dictionary for the config map</param>
+		/// <returns>True if a map was parsed, false otherwise</returns>
+		public static bool TryParseAsMap(string Line, [NotNullWhen(true)] out Dictionary<string,string>? Map)
+		{
+			// read outer array
+			if (!TryParse(Line, out string[]? Array))
+			{
+				Map = null;
+				return false;
+			}
+
+			// read each pair - they're stored in the same way as an array of 2
+			Dictionary<string,string> NewMap = new Dictionary<string, string>();
+			foreach (string ArrayItem in Array)
+			{
+				if (!TryParse( ArrayItem, out string[]? Pairs) || Pairs.Length != 2)
+				{
+					Map = null;
+					return false;
+				}
+
+				NewMap[Pairs[0]] = Pairs[1];
+			}
+
+			Map = NewMap;
+			return true;
+		}
+
 		class ConfigLayerExpansion
 		{
 			// a set of replacements from the source file to possible other files
