Updated UnrealHeaderTool to save the EpicGames.UHT assembly name as the path being used to generate code.

diff --git a/Engine/Source/Programs/UnrealBuildTool/System/ExternalExecution.cs b/Engine/Source/Programs/UnrealBuildTool/System/ExternalExecution.cs
--- a/Engine/Source/Programs/UnrealBuildTool/System/ExternalExecution.cs
+++ b/Engine/Source/Programs/UnrealBuildTool/System/ExternalExecution.cs
@@ -2,6 +2,7 @@
 
 using System;
 using System.Collections.Generic;
+using System.Data;
 using System.Diagnostics;
 using System.IO;
 using System.Linq;
@@ -574,18 +575,6 @@ namespace UnrealBuildTool
 			}
 		}
 
-		/// <summary>
-		/// Gets the path to the receipt for UHT
-		/// </summary>
-		/// <returns>Path to the UHT receipt</returns>
-		public static FileReference GetHeaderToolReceiptFile(FileReference? ProjectFile, FileReference[]? EnabledUhtPlugins)
-		{
-			UnrealArchitectures Architectures = UnrealArchitectureConfig.ForPlatform(BuildHostPlatform.Current.Platform).ActiveArchitectures(null, "UnrealHeaderTool");
-
-			DirectoryReference BaseDir = EnabledUhtPlugins != null && EnabledUhtPlugins.Length > 0 && ProjectFile != null ? ProjectFile.Directory : Unreal.EngineDirectory;
-			return TargetReceipt.GetDefaultPath(BaseDir, "UnrealHeaderTool", BuildHostPlatform.Current.Platform, UnrealTargetConfiguration.Development, Architectures);
-		}
-
 		/// <summary>
 		/// Gets UnrealHeaderTool.exe path. Does not care if UnrealheaderTool was build as a monolithic exe or not.
 		/// </summary>
@@ -1010,6 +999,17 @@ namespace UnrealBuildTool
 			System.IO.File.WriteAllText(ModuleInfoFileName.FullName, JsonSerializer.Serialize(Manifest, new JsonSerializerOptions { WriteIndented = true }));
 		}
 
+		/// <summary>
+		/// Return the latest of two date/times
+		/// </summary>
+		/// <param name="Lhs">First date/time to compare</param>
+		/// <param name="Rhs">Second date/time to compare</param>
+		/// <returns>Latest of the two dates</returns>
+		private static DateTime LatestDateTime(DateTime Lhs, DateTime Rhs)
+		{
+			return Lhs > Rhs ? Lhs : Rhs;
+		}
+
 		/// <summary>
 		/// Builds and runs the header tool and touches the header directories.
 		/// Performs any early outs if headers need no changes, given the UObject modules, tool path, game name, and configuration
@@ -1025,25 +1025,16 @@ namespace UnrealBuildTool
 			{
 				string RootLocalPath = Unreal.RootDirectory.FullName;
 
-				// Figure out the receipt path
-				FileReference HeaderToolReceipt = GetHeaderToolReceiptFile(ProjectFile, Makefile.EnabledUhtPlugins);
+				// Get the path to the assemblies we care about
+				string UhtAssemblyPath = typeof(UhtSession).Assembly.Location;
 
-				// Get UHT assembly timestamp
-				DateTime CompositeTimestamp = new FileInfo(Assembly.GetExecutingAssembly().Location).LastWriteTimeUtc;
-				DateTime Timestamp = new FileInfo(typeof(UhtSession).Assembly.Location).LastWriteTimeUtc;
-				if (CompositeTimestamp < Timestamp)
-				{
-					CompositeTimestamp = Timestamp;
-				}
+				// Get a composite date/time for cirtical assemblies being used
+				DateTime CompositeTimestamp = LatestDateTime(new FileInfo(Assembly.GetExecutingAssembly().Location).LastWriteTimeUtc, new FileInfo(UhtAssemblyPath).LastWriteTimeUtc);
 				if (Makefile.EnabledUhtPlugins != null)
 				{
 					foreach (FileReference Plugin in Makefile.EnabledUhtPlugins)
 					{
-						Timestamp = new FileInfo(Plugin.FullName).LastWriteTimeUtc;
-						if (CompositeTimestamp < Timestamp)
-						{
-							CompositeTimestamp = Timestamp;
-						}
+						CompositeTimestamp = LatestDateTime(CompositeTimestamp, new FileInfo(Plugin.FullName).LastWriteTimeUtc);
 					}
 				}
 
@@ -1067,7 +1058,7 @@ namespace UnrealBuildTool
 					{
 						bUHTNeedsToRun = true;
 					}
-					else if (FileReference.ReadAllText(ToolInfoFile) != HeaderToolReceipt.FullName)
+					else if (FileReference.ReadAllText(ToolInfoFile) != UhtAssemblyPath)
 					{
 						bUHTNeedsToRun = true;
 					}
@@ -1163,7 +1154,7 @@ namespace UnrealBuildTool
 
 					// Update the tool info file
 					DirectoryReference.CreateDirectory(ToolInfoFile.Directory);
-					FileReference.WriteAllText(ToolInfoFile, HeaderToolReceipt.FullName);
+					FileReference.WriteAllText(ToolInfoFile, UhtAssemblyPath);
 
 					// Now that UHT has successfully finished generating code, we need to update all cached FileItems in case their last write time has changed.
 					// Otherwise UBT might not detect changes UHT made.
