Adding per Target controls for ALLOW_PROFILEGPU_IN_TEST and ALLOW_PROFILEGPU_IN_SHIPPING.

diff --git a/Engine/Source/Programs/UnrealBuildTool/Configuration/TargetRules.cs b/Engine/Source/Programs/UnrealBuildTool/Configuration/TargetRules.cs
--- a/Engine/Source/Programs/UnrealBuildTool/Configuration/TargetRules.cs
+++ b/Engine/Source/Programs/UnrealBuildTool/Configuration/TargetRules.cs
@@ -1269,6 +1269,18 @@ namespace UnrealBuildTool
 		[RequiresUniqueBuildEnvironment]
 		public bool bUseChecksInShipping = false;
 
+		/// <summary>
+		/// Whether to turn on GPU markers for Test builds.
+		/// </summary>
+		[RequiresUniqueBuildEnvironment]
+		public bool bAllowProfileGPUInTest = false;
+
+		/// <summary>
+		/// Whether to turn on GPU markers for Shipping builds.
+		/// </summary>
+		[RequiresUniqueBuildEnvironment]
+		public bool bAllowProfileGPUInShipping = false;
+
 		/// <summary>
 		/// Whether to turn on UTF-8 mode, mapping TCHAR to UTF8CHAR.
 		/// </summary>
@@ -2995,6 +3007,10 @@ namespace UnrealBuildTool
 
 		public bool bUseChecksInShipping => Inner.bUseChecksInShipping;
 
+		public bool bAllowProfileGPUInTest => Inner.bAllowProfileGPUInTest;
+
+		public bool bAllowProfileGPUInShipping => Inner.bAllowProfileGPUInShipping;
+
 		public bool bTCHARIsUTF8 => Inner.bTCHARIsUTF8;
 
 		public bool bUseEstimatedUtcNow => Inner.bUseEstimatedUtcNow;
diff --git a/Engine/Source/Programs/UnrealBuildTool/Configuration/UEBuildTarget.cs b/Engine/Source/Programs/UnrealBuildTool/Configuration/UEBuildTarget.cs
--- a/Engine/Source/Programs/UnrealBuildTool/Configuration/UEBuildTarget.cs
+++ b/Engine/Source/Programs/UnrealBuildTool/Configuration/UEBuildTarget.cs
@@ -4781,6 +4781,24 @@ namespace UnrealBuildTool
 				GlobalCompileEnvironment.Definitions.Add("USE_LOGGING_IN_SHIPPING=0");
 			}
 
+			if (Rules.bAllowProfileGPUInTest)
+			{
+				GlobalCompileEnvironment.Definitions.Add("ALLOW_PROFILEGPU_IN_TEST=1");
+			}
+			else
+			{
+				GlobalCompileEnvironment.Definitions.Add("ALLOW_PROFILEGPU_IN_TEST=0");
+			}
+
+			if (Rules.bAllowProfileGPUInShipping)
+			{
+				GlobalCompileEnvironment.Definitions.Add("ALLOW_PROFILEGPU_IN_SHIPPING=1");
+			}
+			else
+			{
+				GlobalCompileEnvironment.Definitions.Add("ALLOW_PROFILEGPU_IN_SHIPPING=0");
+			}
+
 			if (Rules.bLoggingToMemoryEnabled)
 			{
 				GlobalCompileEnvironment.Definitions.Add("WITH_LOGGING_TO_MEMORY=1");
diff --git a/Engine/Source/Runtime/D3D12RHI/D3D12RHI.Build.cs b/Engine/Source/Runtime/D3D12RHI/D3D12RHI.Build.cs
--- a/Engine/Source/Runtime/D3D12RHI/D3D12RHI.Build.cs
+++ b/Engine/Source/Runtime/D3D12RHI/D3D12RHI.Build.cs
@@ -37,8 +37,8 @@ public class D3D12RHI : ModuleRules
 			PublicDefinitions.Add("D3D12RHI_PLATFORM_HAS_CUSTOM_INTERFACE=0");
 
 			if (Target.WindowsPlatform.bPixProfilingEnabled &&
-				Target.Configuration != UnrealTargetConfiguration.Shipping &&
-				Target.Configuration != UnrealTargetConfiguration.Test)
+				(Target.Configuration != UnrealTargetConfiguration.Shipping || Target.bAllowProfileGPUInShipping) &&
+				(Target.Configuration != UnrealTargetConfiguration.Test || Target.bAllowProfileGPUInTest))
 			{
 				PublicDefinitions.Add("PROFILE");
 				PublicDependencyModuleNames.Add("WinPixEventRuntime");
