remove unnecessary log message if PerCultureResources is not set up in a shipping build.

#rnx

diff --git a/Engine/Source/Programs/UnrealBuildTool/Platform/Windows/AppXManifestGeneratorBase.cs b/Engine/Source/Programs/UnrealBuildTool/Platform/Windows/AppXManifestGeneratorBase.cs
--- a/Engine/Source/Programs/UnrealBuildTool/Platform/Windows/AppXManifestGeneratorBase.cs
+++ b/Engine/Source/Programs/UnrealBuildTool/Platform/Windows/AppXManifestGeneratorBase.cs
@@ -493,7 +493,7 @@ namespace UnrealBuildTool
 				Logger.LogWarning("The default culture must be one of the staged cultures. Using {DefaultCulture}.", DefaultUECultureId);
 			}
 
-			bool bHasLocalizationData = BuildLocalizationData();
+			BuildLocalizationData();
 
 			// generate the list of AppX cultures to stage
 			foreach (string UEStageId in SelectedUECultureIds)
@@ -515,12 +515,6 @@ namespace UnrealBuildTool
 				UEStageIdToAppXCultureId[DefaultUECultureId] = DefaultAppXCultureId;
 			}
 
-			// Warn in shipping, we can run without translated cultures they're just needed for cert
-			if (!bHasLocalizationData && InExecutablePairs.ContainsKey(UnrealTargetConfiguration.Shipping))
-			{
-				Logger.LogInformation("Staged culture mappings not setup in the editor. See Per Culture Resources in the {Platform} Target Settings.", Platform.ToString());
-			}
-
 			// Create the manifest document
 			string? IdentityName = null;
 			XDocument ManifestXmlDocument = new XDocument(GetManifest(InExecutablePairs, out IdentityName));
