- Fixed ModernXcode when a program has a ExeBinariesSubFolder

diff --git a/Engine/Source/Programs/UnrealBuildTool/ProjectFiles/Xcode/XcodeProject.cs b/Engine/Source/Programs/UnrealBuildTool/ProjectFiles/Xcode/XcodeProject.cs
--- a/Engine/Source/Programs/UnrealBuildTool/ProjectFiles/Xcode/XcodeProject.cs
+++ b/Engine/Source/Programs/UnrealBuildTool/ProjectFiles/Xcode/XcodeProject.cs
@@ -1717,12 +1717,14 @@ namespace UnrealBuildTool.XcodeProjectXcconfig
 			Xcconfig.AppendLine($"SUPPORTED_PLATFORMS = {SupportedPlatforms}");
 			Xcconfig.AppendLine($"SDKROOT = {SDKRoot}");
 
-			// Xcode creates the Build Dir (where the .app is) by combining {SYMROOT}/{CONFIGURATION}{EFFECTIVE_PLATFORM_NAME}
+			// Xcode creates the Build Dir (where the .app is) by combining {SYMROOT}/{CONFIGURATION}{EFFECTIVE_PLATFORM_NAME}, so we set SYMROOT
+			// to the Parent directory of the diectory the binary is in, CONFIGURATION to nothing, and EFFECTIVE_PLATFORM_NAME to the directory
+			// the binary is in
 			Xcconfig.AppendLine("");
 			Xcconfig.AppendLine($"// These settings combined will tell Xcode to write to Binaries/{Platform} (instead of something like Binaries/Development-iphoneos)");
 			Xcconfig.AppendLine($"SYMROOT = {ConfigBuildDir.ParentDirectory}");
 			Xcconfig.AppendLine($"CONFIGURATION = ");
-			Xcconfig.AppendLine($"EFFECTIVE_PLATFORM_NAME = {Platform}");
+			Xcconfig.AppendLine($"EFFECTIVE_PLATFORM_NAME = {ConfigBuildDir.GetDirectoryName()}");
 
 			if (ExtraConfigLines.Count > 0)
 			{
