Adding header specific compile options
- -IncludeHeaders will add individual header compiles to the build.
- -HeadersOnly when used with -IncludeHeaders will compile only the headers. This replaces -IWYUHeadersOnly.
- Ported file collision support over from ClangToolChain into VCToolChain

diff --git a/Engine/Source/Programs/UnrealBuildTool/Configuration/TargetRules.cs b/Engine/Source/Programs/UnrealBuildTool/Configuration/TargetRules.cs
--- a/Engine/Source/Programs/UnrealBuildTool/Configuration/TargetRules.cs
+++ b/Engine/Source/Programs/UnrealBuildTool/Configuration/TargetRules.cs
@@ -1362,10 +1362,21 @@ namespace UnrealBuildTool
 		public bool bIWYU = false;
 
 		/// <summary>
-		/// Tells "include what you use" to only compile header files.
+		/// Adds header files in included modules to the build.
 		/// </summary>
-		[CommandLine("-IWYUHeadersOnly")]
-		public bool bIWYUHeadersOnly = false;
+		[CommandLine("-IncludeHeaders")]
+		public bool bIncludeHeaders
+		{
+			get => bIncludeHeadersPrivate || bIWYU;
+			set => bIncludeHeadersPrivate = value;
+		}
+		private bool bIncludeHeadersPrivate = false;
+
+		/// <summary>
+		/// When used with -IncludeHeaders, only header files will be compiled.
+		/// </summary>
+		[CommandLine("-HeadersOnly")]
+		public bool bHeadersOnly = false;
 
 		/// <summary>
 		/// Enforce "include what you use" rules; warns if monolithic headers (Engine.h, UnrealEd.h, etc...) are used, and checks that source files include their matching header first.
@@ -3040,7 +3051,9 @@ namespace UnrealBuildTool
 
 		public bool bIWYU => Inner.bIWYU;
 
-		public bool bIWYUHeadersOnly => Inner.bIWYUHeadersOnly;
+		public bool bIncludeHeaders => Inner.bIncludeHeaders;
+
+		public bool bHeadersOnly => Inner.bHeadersOnly;
 
 		public bool bEnforceIWYU => Inner.bEnforceIWYU;
 
diff --git a/Engine/Source/Programs/UnrealBuildTool/Configuration/UEBuildModuleCPP.cs b/Engine/Source/Programs/UnrealBuildTool/Configuration/UEBuildModuleCPP.cs
--- a/Engine/Source/Programs/UnrealBuildTool/Configuration/UEBuildModuleCPP.cs
+++ b/Engine/Source/Programs/UnrealBuildTool/Configuration/UEBuildModuleCPP.cs
@@ -749,13 +749,13 @@ namespace UnrealBuildTool
 			}
 
 			// IWYU needs to build all headers separate from cpp files to produce proper recommendations for includes
-			if (Target.bIWYU)
+			if (Target.bIncludeHeaders)
 			{
 				// Collect the headers that should be built
 				List<FileItem> HeaderFileItems = GetCompilableHeaders(InputFiles, CompileEnvironment);
 				if (HeaderFileItems.Count > 0)
 				{
-					if (Target.bIWYUHeadersOnly)
+					if (Target.bHeadersOnly)
 					{
 						LinkInputFiles.Clear();
 					}
diff --git a/Engine/Source/Programs/UnrealBuildTool/Platform/Windows/VCToolChain.cs b/Engine/Source/Programs/UnrealBuildTool/Platform/Windows/VCToolChain.cs
--- a/Engine/Source/Programs/UnrealBuildTool/Platform/Windows/VCToolChain.cs
+++ b/Engine/Source/Programs/UnrealBuildTool/Platform/Windows/VCToolChain.cs
@@ -1453,15 +1453,22 @@ namespace UnrealBuildTool
 					}
 				}
 
+				string FileName = SourceFile.Name;
+				if (CompileEnvironment.CollidingNames != null && CompileEnvironment.CollidingNames.Contains(SourceFile))
+				{
+					string HashString = ContentHash.MD5(SourceFile.AbsolutePath.Substring(Unreal.RootDirectory.FullName.Length)).GetHashCode().ToString("X4");
+					FileName = Path.GetFileNameWithoutExtension(FileName) + "_" + HashString + Path.GetExtension(FileName);
+				}
+
 				if (CompileEnvironment.bPreprocessOnly)
 				{
-					CompileAction.PreprocessedFile = FileItem.GetItemByFileReference(FileReference.Combine(OutputDir, SourceFile.Location.GetFileName() + ".i"));
+					CompileAction.PreprocessedFile = FileItem.GetItemByFileReference(FileReference.Combine(OutputDir, FileName + ".i"));
 					CompileAction.ResponseFile = FileItem.GetItemByFileReference(GetResponseFileName(CompileEnvironment, CompileAction.PreprocessedFile));
 				}
 				else
 				{
 					// Add the object file to the produced item list.
-					string ObjectLeafFilename = Path.GetFileName(SourceFile.AbsolutePath) + ".obj";
+					string ObjectLeafFilename = FileName + ".obj";
 					FileItem ObjectFile = FileItem.GetItemByFileReference(FileReference.Combine(OutputDir, ObjectLeafFilename));
 
 					CompileAction.ObjectFile = ObjectFile;
@@ -1477,7 +1484,7 @@ namespace UnrealBuildTool
 
 					if (Target.WindowsPlatform.Compiler.IsMSVC() && CompileEnvironment.bWithAssembly && SourceFile.HasExtension(".cpp"))
 					{
-						CompileAction.AssemblyFile = FileItem.GetItemByFileReference(FileReference.Combine(OutputDir, SourceFile.Location.GetFileName() + ".asm"));
+						CompileAction.AssemblyFile = FileItem.GetItemByFileReference(FileReference.Combine(OutputDir, FileName + ".asm"));
 					}
 
 					// Experimental: support for JSON output of timing data
@@ -1529,7 +1536,7 @@ namespace UnrealBuildTool
 					else if (!bIsPlainCFile)
 					{
 						// Ungrouped C++ files use a PDB per file.
-						PDBLocation = FileReference.Combine(OutputDir, SourceFile.Location.GetFileName() + ".pdb");
+						PDBLocation = FileReference.Combine(OutputDir, FileName + ".pdb");
 					}
 					else
 					{
@@ -1585,7 +1592,7 @@ namespace UnrealBuildTool
 					CompileAction.Arguments.Add($"/ifcOutput \"{IfcFile.Location}\"");
 					CompileAction.CompiledModuleInterfaceFile = IfcFile;
 
-					FileItem IfcDepsFile = FileItem.GetItemByFileReference(FileReference.Combine(OutputDir, SourceFile.Location.GetFileName() + ".md.json"));
+					FileItem IfcDepsFile = FileItem.GetItemByFileReference(FileReference.Combine(OutputDir, FileName + ".md.json"));
 
 					VCCompileAction CompileDepsAction = new VCCompileAction(CompileAction);
 					CompileDepsAction.ActionType = ActionType.GatherModuleDependencies;
@@ -1613,21 +1620,21 @@ namespace UnrealBuildTool
 				if ((Target.bPrintToolChainTimingInfo || Target.WindowsPlatform.bCompilerTrace) && Target.WindowsPlatform.Compiler.IsMSVC())
 				{
 					CompileAction.ForceClFilter = true;
-					CompileAction.TimingFile = FileItem.GetItemByFileReference(FileReference.Combine(OutputDir, $"{SourceFile.Location.GetFileName()}.timing"));
+					CompileAction.TimingFile = FileItem.GetItemByFileReference(FileReference.Combine(OutputDir, $"{FileName}.timing"));
 					GenerateParseTimingInfoAction(SourceFile, CompileAction.TimingFile, Graph);
 				}
 
 				if (Target.WindowsPlatform.Compiler.IsMSVC() && !CompileAction.ForceClFilter)
 				{
-					CompileAction.DependencyListFile = FileItem.GetItemByFileReference(FileReference.Combine(OutputDir, $"{SourceFile.Location.GetFileName()}.dep.json"));
+					CompileAction.DependencyListFile = FileItem.GetItemByFileReference(FileReference.Combine(OutputDir, $"{FileName}.dep.json"));
 				}
 				else if (Target.WindowsPlatform.Compiler.IsClang())
 				{
-					CompileAction.DependencyListFile = FileItem.GetItemByFileReference(FileReference.Combine(OutputDir, $"{SourceFile.Location.GetFileName()}.d"));
+					CompileAction.DependencyListFile = FileItem.GetItemByFileReference(FileReference.Combine(OutputDir, $"{FileName}.d"));
 				}
 				else
 				{
-					CompileAction.DependencyListFile = FileItem.GetItemByFileReference(FileReference.Combine(OutputDir, $"{SourceFile.Location.GetFileName()}.txt"));
+					CompileAction.DependencyListFile = FileItem.GetItemByFileReference(FileReference.Combine(OutputDir, $"{FileName}.txt"));
 					CompileAction.bShowIncludes = Target.WindowsPlatform.bShowIncludes;
 				}
 
@@ -1636,10 +1643,10 @@ namespace UnrealBuildTool
 				{
 					if (Target.StaticAnalyzer == StaticAnalyzer.Default && !CompileEnvironment.bDisableStaticAnalysis)
 					{
-						CompileAction.AnalyzeLogFile = FileItem.GetItemByFileReference(FileReference.Combine(OutputDir, $"{SourceFile.Location.GetFileName()}.sa.sarif"));
+						CompileAction.AnalyzeLogFile = FileItem.GetItemByFileReference(FileReference.Combine(OutputDir, $"{FileName}.sa.sarif"));
 					}
 
-					CompileAction.ExperimentalLogFile = FileItem.GetItemByFileReference(FileReference.Combine(OutputDir, $"{SourceFile.Location.GetFileName()}.sarif"));
+					CompileAction.ExperimentalLogFile = FileItem.GetItemByFileReference(FileReference.Combine(OutputDir, $"{FileName}.sarif"));
 				}
 
 				// Allow derived toolchains to make further changes
