Updated the logic that determines which dependencies a shared PCH has. This new logic goes through all the includes of the pch header and finds the actual headers and modules used. This creates a smaller list of dependencies than before, which is later used to determine whether it is safe or not for the PCH to be used by other modulres.

diff --git a/Engine/Source/Programs/UnrealBuildTool/Configuration/UEBuildModuleCPP.cs b/Engine/Source/Programs/UnrealBuildTool/Configuration/UEBuildModuleCPP.cs
--- a/Engine/Source/Programs/UnrealBuildTool/Configuration/UEBuildModuleCPP.cs
+++ b/Engine/Source/Programs/UnrealBuildTool/Configuration/UEBuildModuleCPP.cs
@@ -820,13 +820,124 @@ namespace UnrealBuildTool
 			return HeaderFileItems;
 		}
 
+		private void FindIncludedHeaders(HashSet<string> VisitedIncludes, HashSet<FileItem> FoundIncludeFileItems, HashSet<string> IncludesNotFound, CppCompileEnvironment CompileEnvironment, FileItem FileToSearch, ILogger Logger)
+		{
+			HashSet<DirectoryReference> SearchPaths = new(CompileEnvironment.UserIncludePaths);
+			SearchPaths.Add(FileToSearch.Directory.Location);
+
+			foreach (string HeaderInclude in CompileEnvironment.MetadataCache.GetHeaderIncludes(FileToSearch))
+			{
+				string TranformedHeaderInclude = HeaderInclude;
+				if (TranformedHeaderInclude.Contains("COMPILED_PLATFORM_HEADER("))
+				{
+					string PlatformName = UEBuildPlatform.GetBuildPlatform(CompileEnvironment.Platform).GetPlatformName();
+					TranformedHeaderInclude = TranformedHeaderInclude.Replace("COMPILED_PLATFORM_HEADER(", PlatformName + "/" + PlatformName);
+					TranformedHeaderInclude = TranformedHeaderInclude.TrimEnd(')');
+				}
+
+				if (VisitedIncludes.Add(TranformedHeaderInclude))
+				{
+					var FoundDir = SearchPaths.FirstOrDefault(dir => DirectoryLookupCache.FileExists(FileReference.Combine(dir, TranformedHeaderInclude)));
+					if (FoundDir != null)
+					{
+						FileItem IncludeFileItem = FileItem.GetItemByFileReference(FileReference.Combine(FoundDir, TranformedHeaderInclude));
+						FoundIncludeFileItems.Add(IncludeFileItem);
+						Logger.LogDebug("{0} SharedPCH - '{1}' included '{2}'.", Name, FileToSearch.Location, IncludeFileItem.Location);
+
+						FindIncludedHeaders(VisitedIncludes, FoundIncludeFileItems, IncludesNotFound, CompileEnvironment, IncludeFileItem, Logger);
+					}
+					else
+					{
+						if (!TranformedHeaderInclude.Contains('.'))
+						{
+							Logger.LogDebug("{0} SharedPCH - Skipping '{1}' found in '{2}' because it doesn't appear to be a module header.", Name, TranformedHeaderInclude, FileToSearch.Location);
+						}
+						else if (TranformedHeaderInclude.EndsWith(".generated.h"))
+						{
+							Logger.LogDebug("{0} SharedPCH - Skipping '{1}' found in '{2}' because it appears to be a generated UHT header.", Name, TranformedHeaderInclude, FileToSearch.Location);
+						}
+						else
+						{
+							Logger.LogDebug("{0} SharedPCH - Could not find include directory for '{1}' found in '{2}'.", Name, TranformedHeaderInclude, FileToSearch.Location);
+							IncludesNotFound.Add(TranformedHeaderInclude);
+						}
+					}
+				}
+			}
+		}
+
+		/// <summary>
+		/// Gets the list of module dependencies that this module declares and then tries to optimize the list to what is actually used by the header.
+		/// </summary>
+		/// <param name="CompileEnvironment">Compile environment for this PCH</param>
+		/// <param name="PCHHeaderFile">PCH header file</param>
+		/// <param name="Logger">Logger for output</param>
+		/// <returns>Returns an optimized list of modules this module's shared PCH uses</returns>
+		private HashSet<UEBuildModule> GetSharedPCHModuleDependencies(CppCompileEnvironment CompileEnvironment, FileItem PCHHeaderFile, ILogger Logger)
+		{
+			Dictionary<UEBuildModule, bool> ModuleToIncludePathsOnlyFlag = new Dictionary<UEBuildModule, bool>();
+			FindModulesInPublicCompileEnvironment(ModuleToIncludePathsOnlyFlag);
+
+			HashSet<UEBuildModule> AllModuleDeps = ModuleToIncludePathsOnlyFlag.Keys.ToHashSet();
+			HashSet<UEBuildModule> ModuleDeps = new HashSet<UEBuildModule>(AllModuleDeps) { this };
+
+			HashSet<string> VisitedIncludes = new();
+			HashSet<FileItem> FoundIncludeFileItems = new();
+			HashSet<string> IncludesNotFound = new();
+			FindIncludedHeaders(VisitedIncludes, FoundIncludeFileItems, IncludesNotFound, CompileEnvironment, PCHHeaderFile, Logger);
+			if (IncludesNotFound.Any())
+			{
+				// We are assuming that the includes not found aren't important. If this is found to be not true then we should return all the module dependencies
+				//return AllModuleDeps;
+			}
+
+			bool FoundAllModules = true;
+			HashSet<UEBuildModule> OptModules = new();
+			foreach (var IncludeFile in FoundIncludeFileItems)
+			{
+				if (CompileEnvironment.MetadataCache.UsesAPIDefine(IncludeFile) || CompileEnvironment.MetadataCache.ContainsReflectionMarkup(IncludeFile))
+				{
+					UEBuildModule? FoundModule = ModuleDeps.FirstOrDefault(Module => Module.ContainsFile(IncludeFile.Location));
+					if (FoundModule != null)
+					{
+						if (ModuleToIncludePathsOnlyFlag[FoundModule])
+						{
+							Logger.LogDebug("{0} SharedPCH - '{1}' is exporting types but '{2}' isn't declared as a public dependency.", Name, IncludeFile.Location, FoundModule);
+						}
+
+						OptModules.Add(FoundModule);
+					}
+					else
+					{
+						Logger.LogWarning("{0} SharedPCH - '{1}' is exporting types but the module this file belongs to isn't declared as a public dependency or include.", Name, IncludeFile.Location);
+						FoundAllModules = false;
+					}
+				}
+				else
+				{
+					Logger.LogDebug("{0} SharedPCH - '{1}' is not exporting types so we are ignoring the dependency", Name, IncludeFile.Location);
+				}	
+			}
+
+			if (!FoundAllModules)
+			{
+				Logger.LogWarning("{0} SharedPCH - Is missing public dependencies. To be safe, this shared PCH will fall back to use all the module dependencies. Note that this could affect compile times.", Name);
+				return AllModuleDeps;
+			}
+			else
+			{
+				return OptModules;
+			}
+		}
+
 		/// <summary>
 		/// Create a shared PCH template for this module, which allows constructing shared PCH instances in the future
 		/// </summary>
 		/// <param name="Target">The target which owns this module</param>
 		/// <param name="BaseCompileEnvironment">Base compile environment for this target</param>
+		/// <param name="Logger">Logger for output</param>
 		/// <returns>Template for shared PCHs</returns>
-		public PrecompiledHeaderTemplate CreateSharedPCHTemplate(UEBuildTarget Target, CppCompileEnvironment BaseCompileEnvironment)
+		public PrecompiledHeaderTemplate CreateSharedPCHTemplate(UEBuildTarget Target, CppCompileEnvironment BaseCompileEnvironment, ILogger Logger)
 		{
 			CppCompileEnvironment CompileEnvironment = CreateSharedPCHCompileEnvironment(Target, BaseCompileEnvironment);
 			FileItem HeaderFile = FileItem.GetItemByFileReference(FileReference.Combine(ModuleDirectory, Rules.SharedPCHHeaderFile!));
@@ -841,7 +952,7 @@ namespace UnrealBuildTool
 				PrecompiledHeaderDir = IntermediateDirectory;
 			}
 
-			return new PrecompiledHeaderTemplate(this, CompileEnvironment, HeaderFile, PrecompiledHeaderDir, GetAllDependencyModulesForPCH(false, false));
+			return new PrecompiledHeaderTemplate(this, CompileEnvironment, HeaderFile, PrecompiledHeaderDir, GetSharedPCHModuleDependencies(CompileEnvironment, HeaderFile, Logger));
 		}
 
 		static HashSet<string> GetImmutableDefinitions(List<string> Definitions)
diff --git a/Engine/Source/Programs/UnrealBuildTool/Configuration/UEBuildTarget.cs b/Engine/Source/Programs/UnrealBuildTool/Configuration/UEBuildTarget.cs
--- a/Engine/Source/Programs/UnrealBuildTool/Configuration/UEBuildTarget.cs
+++ b/Engine/Source/Programs/UnrealBuildTool/Configuration/UEBuildTarget.cs
@@ -3446,7 +3446,7 @@ namespace UnrealBuildTool
 			List<PrecompiledHeaderTemplate> OrderedSharedPCHModules = GlobalCompileEnvironment.SharedPCHs;
 			foreach (UEBuildModuleCPP Module in SharedPCHModuleToPriority.OrderByDescending(x => x.Value).Select(x => x.Key))
 			{
-				OrderedSharedPCHModules.Add(Module.CreateSharedPCHTemplate(this, GlobalCompileEnvironment));
+				OrderedSharedPCHModules.Add(Module.CreateSharedPCHTemplate(this, GlobalCompileEnvironment, Logger));
 			}
 
 			// Print the ordered list of shared PCHs
diff --git a/Engine/Source/Programs/UnrealBuildTool/System/SourceFileMetadataCache.cs b/Engine/Source/Programs/UnrealBuildTool/System/SourceFileMetadataCache.cs
--- a/Engine/Source/Programs/UnrealBuildTool/System/SourceFileMetadataCache.cs
+++ b/Engine/Source/Programs/UnrealBuildTool/System/SourceFileMetadataCache.cs
@@ -28,7 +28,8 @@ namespace UnrealBuildTool
 	{
 		Valid = 0,
 		Unsupported = 1,
-		Skip = 2
+		Skip = 2,
+		Ignore = 3
 	}
 
 	/// <summary>
@@ -72,6 +73,11 @@ namespace UnrealBuildTool
 			/// </summary>
 			public bool bContainsMarkup;
 
+			/// <summary>
+			/// Whether or not the file has types that use DLL export/import defines
+			/// </summary>
+			public bool bUsesAPIDefine;
+
 			/// <summary>
 			/// List of includes that header contains
 			/// </summary>
@@ -83,7 +89,7 @@ namespace UnrealBuildTool
 		/// <summary>
 		/// The current file version
 		/// </summary>
-		public const int CurrentVersion = 6;
+		public const int CurrentVersion = 7;
 
 		/// <summary>
 		/// Location of this dependency cache
@@ -250,7 +256,7 @@ namespace UnrealBuildTool
 					Match ImportMatch = ImportRegex.Match(Line);
 					if (ImportMatch.Success)
 					{
-						return IncludeMatch.Groups[1].Value;
+						return ImportMatch.Groups[1].Value;
 					}
 				}
 			}
@@ -295,6 +301,62 @@ namespace UnrealBuildTool
 			}
 		}
 
+		/// <summary>
+		/// Attempts to find the include in the passed in line
+		/// </summary>
+		/// <param name="Line">Line from a file</param>
+		/// <returns>Include found or null if not found</returns>
+		static private string? FindInclude(string Line)
+		{
+			ReadOnlySpan<char> IncludeSpan = Line.AsSpan().TrimStart();
+			if (IncludeSpan.StartsWith("#include"))
+			{
+				IncludeSpan = IncludeSpan.Slice("#include".Length).TrimStart();
+				if (IncludeSpan.IsEmpty)
+				{
+					return null;
+				}
+
+				char EndChar;
+				bool TrimQuotation = true;
+				if (IncludeSpan[0] == '"')
+				{
+					EndChar = '"';
+				}
+				else if (IncludeSpan[0] == '<')
+				{
+					EndChar = '>';
+				}
+				else
+				{
+					EndChar = ')';
+					TrimQuotation = false;
+				}
+
+				if (TrimQuotation)
+				{
+					IncludeSpan = IncludeSpan.Slice(1);
+				}
+
+				int EndIndex = IncludeSpan.IndexOf(EndChar);
+				if (EndIndex == -1)
+				{
+					return null;
+				}
+
+				// This will include the ')' at the end
+				if (!TrimQuotation)
+				{
+					EndIndex++;
+				}
+
+				IncludeSpan = IncludeSpan.Slice(0, EndIndex);
+
+				return IncludeSpan.ToString();
+			}
+			return null;
+		}
+
 		/// <summary>
 		/// Read entire header file to find markup and includes
 		/// </summary>
@@ -302,6 +364,7 @@ namespace UnrealBuildTool
 		private void ParseHeader(HeaderFileInfo HeaderFileInfo, string[] FileText)
 		{
 			bool bContainsMarkup = false;
+			bool bUsesAPIDefine = false;
 			SortedSet<string> Includes = new();
 			foreach (string Line in FileText)
 			{
@@ -310,50 +373,17 @@ namespace UnrealBuildTool
 					bContainsMarkup = ReflectionMarkupRegex.IsMatch(Line);
 				}
 
-				ReadOnlySpan<char> IncludeSpan = Line.AsSpan().TrimStart();
-				if (IncludeSpan.StartsWith("#include"))
+				if (!bUsesAPIDefine)
 				{
-					IncludeSpan = IncludeSpan.Slice("#include".Length).TrimStart();
-					if (IncludeSpan.IsEmpty)
-					{
-						continue;
-					}
-
-					char EndChar;
-					bool TrimQuotation = true;
-					if (IncludeSpan[0] == '"')
-					{
-						EndChar = '"';
-					}
-					else if (IncludeSpan[0] == '<')
-					{
-						EndChar = '>';
-					}
-					else
-					{
-						EndChar = ')';
-						TrimQuotation = false;
-					}
-					if (TrimQuotation)
-					{
-						IncludeSpan = IncludeSpan.Slice(1);
-					}
-
-					if (IncludeSpan.Contains("HEADER_UNIT_IGNORE", StringComparison.OrdinalIgnoreCase))
-					{
-						continue;
-					}
-
-					int EndIndex = IncludeSpan.IndexOf(EndChar);
-
-					if (EndIndex == -1)
-					{
-						continue;
-					}
+					bUsesAPIDefine = Line.Contains("_API", StringComparison.Ordinal);
+				}
 
-					IncludeSpan = IncludeSpan.Slice(0, EndIndex);
-					Includes.Add(IncludeSpan.ToString());
+				string? Include = FindInclude(Line);
+				if (Include == null)
+				{
+					continue;
 				}
+				Includes.Add(Include);
 
 				int HeaderUnitIndex = Line.IndexOf("HEADER_UNIT_");
 				if (HeaderUnitIndex != -1)
@@ -367,10 +397,15 @@ namespace UnrealBuildTool
 					{
 						HeaderFileInfo.UnitType = HeaderUnitType.Skip;
 					}
+					else if (Span.StartsWith("IGNORE"))
+					{
+						HeaderFileInfo.UnitType = HeaderUnitType.Ignore;
+					}
 				}
 			}
 
 			HeaderFileInfo.bContainsMarkup = bContainsMarkup;
+			HeaderFileInfo.bUsesAPIDefine = bUsesAPIDefine;
 			HeaderFileInfo.Includes = Includes.ToList();
 		}
 
@@ -432,6 +467,16 @@ namespace UnrealBuildTool
 			return GetHeaderFileInfo(HeaderFile).bContainsMarkup;
 		}
 
+		/// <summary>
+		/// Determines whether the given file uses the *_API define
+		/// </summary>
+		/// <param name="HeaderFile">The source file to parse</param>
+		/// <returns>True if the file uses the *_API define</returns>
+		public bool UsesAPIDefine(FileItem HeaderFile)
+		{
+			return GetHeaderFileInfo(HeaderFile).bUsesAPIDefine;
+		}
+
 		/// <summary>
 		/// Returns header unit type for a header file (and parse the file if not already cached)
 		/// </summary>
@@ -547,6 +592,7 @@ namespace UnrealBuildTool
 						HeaderFileInfo HeaderFileInfo = new HeaderFileInfo();
 						HeaderFileInfo.LastWriteTimeUtc = Reader.ReadLong();
 						HeaderFileInfo.bContainsMarkup = Reader.ReadBool();
+						HeaderFileInfo.bUsesAPIDefine = Reader.ReadBool();
 						HeaderFileInfo.UnitType = (HeaderUnitType)Reader.ReadByte();
 						HeaderFileInfo.Includes = Reader.ReadList(() => Reader.ReadString())!;
 
@@ -592,6 +638,7 @@ namespace UnrealBuildTool
 						Writer.WriteCompactFileItem(Pair.Key);
 						Writer.WriteLong(Pair.Value.LastWriteTimeUtc);
 						Writer.WriteBool(Pair.Value.bContainsMarkup);
+						Writer.WriteBool(Pair.Value.bUsesAPIDefine);
 						Writer.WriteByte((byte)Pair.Value.UnitType);
 						Writer.WriteList(Pair.Value.Includes, Item => Writer.WriteString(Item));
 					}
