Removed a hack in the sharedpch logic that is no longer needed.
Changed a warning to a debug message because it broke a branch that is behind in merges. I will bring back the warnings when the branch catches up.

diff --git a/Engine/Source/Programs/UnrealBuildTool/Configuration/UEBuildModuleCPP.cs b/Engine/Source/Programs/UnrealBuildTool/Configuration/UEBuildModuleCPP.cs
--- a/Engine/Source/Programs/UnrealBuildTool/Configuration/UEBuildModuleCPP.cs
+++ b/Engine/Source/Programs/UnrealBuildTool/Configuration/UEBuildModuleCPP.cs
@@ -909,7 +909,7 @@ namespace UnrealBuildTool
 					}
 					else
 					{
-						Logger.LogWarning("{0} SharedPCH - '{1}' is exporting types but the module this file belongs to isn't declared as a public dependency or include.", Name, IncludeFile.Location);
+						Logger.LogDebug("{0} SharedPCH - '{1}' is exporting types but the module this file belongs to isn't declared as a public dependency or include.", Name, IncludeFile.Location);
 						FoundAllModules = false;
 					}
 				}
@@ -921,7 +921,7 @@ namespace UnrealBuildTool
 
 			if (!FoundAllModules)
 			{
-				Logger.LogWarning("{0} SharedPCH - Is missing public dependencies. To be safe, this shared PCH will fall back to use all the module dependencies. Note that this could affect compile times.", Name);
+				Logger.LogDebug("{0} SharedPCH - Is missing public dependencies. To be safe, this shared PCH will fall back to use all the module dependencies. Note that this could affect compile times.", Name);
 				return AllModuleDeps;
 			}
 			else
@@ -1538,34 +1538,10 @@ namespace UnrealBuildTool
 				// Try to find a suitable shared PCH for this module
 				if (CompileEnvironment.bHasPrecompiledHeader == false && CompileEnvironment.SharedPCHs.Count > 0 && !CompileEnvironment.bIsBuildingLibrary && Rules.PCHUsage != ModuleRules.PCHUsageMode.NoSharedPCHs)
 				{
-					// Find only the direct dependencies of this module
-					// Note that this is the old shared chooser logic that needs to be removed in the future
-					HashSet<UEBuildModule> OnlyDirectDependencies = new HashSet<UEBuildModule>();
-					GetAllDependencyModules(new List<UEBuildModule>(), OnlyDirectDependencies, bIncludeDynamicallyLoaded: false, bForceCircular: false, bOnlyDirectDependencies: true);
-					PrecompiledHeaderTemplate? OldTemplateValue = CompileEnvironment.SharedPCHs.FirstOrDefault(x => OnlyDirectDependencies.Contains(x.Module) && x.IsValidFor(CompileEnvironment));
-
 					// Find the first shared PCH module we can use that doesn't reference this module
 					HashSet<UEBuildModule> AllDependencies = GetAllDependencyModulesForPCH(true, false);
 					PrecompiledHeaderTemplate? Template = CompileEnvironment.SharedPCHs.FirstOrDefault(x => !x.ModuleDependencies.Contains(this) && AllDependencies.Contains(x.Module) && x.IsValidFor(CompileEnvironment));
 
-					if (Template != OldTemplateValue)
-					{
-						if (OldTemplateValue != null && Template != null)
-						{
-							/*
-							If GetAllDependencyModules returns a better shared pch than GetAllDependencyModulesForPCH then use that pch instead. This can happen if a circular dependency was 
-							found in the GetAllDependencyModulesForPCH code path. We shouldn't need to do this but there are too many circular dependencies not being detected by UBT ATM. We 
-							should be able to remove this when the circular dependencies have been addressed in the engine.
-							*/
-							var NewIndex = CompileEnvironment.SharedPCHs.IndexOf(Template);
-							var OldIndex = CompileEnvironment.SharedPCHs.IndexOf(OldTemplateValue);
-							if (NewIndex > OldIndex)
-							{
-								Template = OldTemplateValue;
-							}
-						}
-					}
-
 					if (Template != null)
 					{
 						PrecompiledHeaderInstance Instance = FindOrCreateSharedPCH(ToolChain, Template, CompileEnvironment, Graph);
