UnrealBuildTool: Remove PCH variants for ShadowVariableWarningLevel, UnsafeTypeCastWarningLevel & bEnableUndefinedIdentifierWarnings

#rnx

diff --git a/Engine/Source/Programs/UnrealBuildTool/Configuration/UEBuildModuleCPP.cs b/Engine/Source/Programs/UnrealBuildTool/Configuration/UEBuildModuleCPP.cs
--- a/Engine/Source/Programs/UnrealBuildTool/Configuration/UEBuildModuleCPP.cs
+++ b/Engine/Source/Programs/UnrealBuildTool/Configuration/UEBuildModuleCPP.cs
@@ -1043,7 +1043,7 @@ namespace UnrealBuildTool
 				StringBuilder Writer = new StringBuilder();
 				Writer.AppendLine("#pragma once");
 				WriteDefinitions(Definitions, Writer);
-				FileItem SharedDefinitionsFileItem = Graph.CreateIntermediateTextFile(SharedDefinitionsLocation, Writer.ToString());
+				FileItem SharedDefinitionsFileItem = Graph.CreateIntermediateTextFile(SharedDefinitionsLocation, Writer.ToString(), AllowAsync: false);
 
 				// Create the wrapper file, which sets all the definitions needed to compile it
 				FileReference WrapperLocation = FileReference.Combine(Template.OutputDir, String.Format("SharedPCH.{0}{1}.h", Template.Module.Name, Variant));
@@ -1144,18 +1144,6 @@ namespace UnrealBuildTool
 			{
 				return false;
 			}
-			if (ModuleCompileEnvironment.ShadowVariableWarningLevel != CompileEnvironment.ShadowVariableWarningLevel)
-			{
-				return false;
-			}
-			if (ModuleCompileEnvironment.UnsafeTypeCastWarningLevel != CompileEnvironment.UnsafeTypeCastWarningLevel)
-			{
-				return false;
-			}
-			if (ModuleCompileEnvironment.bEnableUndefinedIdentifierWarnings != CompileEnvironment.bEnableUndefinedIdentifierWarnings)
-			{
-				return false;
-			}
 			if (ModuleCompileEnvironment.CppStandard != CompileEnvironment.CppStandard)
 			{
 				return false;
@@ -1215,49 +1203,6 @@ namespace UnrealBuildTool
 					Variant += ".NoExceptions";
 				}
 			}
-			if (CompileEnvironment.ShadowVariableWarningLevel != BaseCompileEnvironment.ShadowVariableWarningLevel)
-			{
-				if (CompileEnvironment.ShadowVariableWarningLevel == WarningLevel.Error)
-				{
-					Variant += ".ShadowErrors";
-				}
-				else if (CompileEnvironment.ShadowVariableWarningLevel == WarningLevel.Warning)
-				{
-					Variant += ".ShadowWarnings";
-				}
-				else
-				{
-					Variant += ".NoShadow";
-				}
-			}
-
-			if (CompileEnvironment.UnsafeTypeCastWarningLevel != BaseCompileEnvironment.UnsafeTypeCastWarningLevel)
-			{
-				if (CompileEnvironment.UnsafeTypeCastWarningLevel == WarningLevel.Error)
-				{
-					Variant += ".TypeCastErrors";
-				}
-				else if (CompileEnvironment.UnsafeTypeCastWarningLevel == WarningLevel.Warning)
-				{
-					Variant += ".TypeCastWarnings";
-				}
-				else
-				{
-					Variant += ".NoTypeCast";
-				}
-			}
-
-			if (CompileEnvironment.bEnableUndefinedIdentifierWarnings != BaseCompileEnvironment.bEnableUndefinedIdentifierWarnings)
-			{
-				if (CompileEnvironment.bEnableUndefinedIdentifierWarnings)
-				{
-					Variant += ".Undef";
-				}
-				else
-				{
-					Variant += ".NoUndef";
-				}
-			}
 
 			if (CompileEnvironment.bDeterministic != BaseCompileEnvironment.bDeterministic)
 			{
@@ -1315,9 +1260,6 @@ namespace UnrealBuildTool
 			CompileEnvironment.bCodeCoverage = ModuleCompileEnvironment.bCodeCoverage;
 			CompileEnvironment.bUseRTTI = ModuleCompileEnvironment.bUseRTTI;
 			CompileEnvironment.bEnableExceptions = ModuleCompileEnvironment.bEnableExceptions;
-			CompileEnvironment.ShadowVariableWarningLevel = ModuleCompileEnvironment.ShadowVariableWarningLevel;
-			CompileEnvironment.UnsafeTypeCastWarningLevel = ModuleCompileEnvironment.UnsafeTypeCastWarningLevel;
-			CompileEnvironment.bEnableUndefinedIdentifierWarnings = ModuleCompileEnvironment.bEnableUndefinedIdentifierWarnings;
 			CompileEnvironment.CppStandardEngine = ModuleCompileEnvironment.CppStandardEngine;
 			CompileEnvironment.CppStandard = ModuleCompileEnvironment.CppStandard;
 			CompileEnvironment.CStandard = ModuleCompileEnvironment.CStandard;
