Partial fix for modern Xcode would build with no code change
Remove outdated xcworkspace files when GPF with new modernXcode settings

diff --git a/Engine/Source/Programs/UnrealBuildTool/Configuration/UEBuildTarget.cs b/Engine/Source/Programs/UnrealBuildTool/Configuration/UEBuildTarget.cs
--- a/Engine/Source/Programs/UnrealBuildTool/Configuration/UEBuildTarget.cs
+++ b/Engine/Source/Programs/UnrealBuildTool/Configuration/UEBuildTarget.cs
@@ -2153,7 +2153,7 @@ namespace UnrealBuildTool
 				Receipt.AdditionalProperties.Add(new ReceiptProperty("CustomConfig", Rules.CustomConfig));
 			}
 
-			ToolChain.ModifyTargetReceipt(Receipt);
+			ToolChain.ModifyTargetReceipt(Rules, Receipt);
 
 			return Receipt;
 		}
diff --git a/Engine/Source/Programs/UnrealBuildTool/Platform/Mac/UEBuildMac.cs b/Engine/Source/Programs/UnrealBuildTool/Platform/Mac/UEBuildMac.cs
--- a/Engine/Source/Programs/UnrealBuildTool/Platform/Mac/UEBuildMac.cs
+++ b/Engine/Source/Programs/UnrealBuildTool/Platform/Mac/UEBuildMac.cs
@@ -446,7 +446,9 @@ namespace UnrealBuildTool
 		public override List<FileReference> FinalizeBinaryPaths(FileReference BinaryName, FileReference? ProjectFile, ReadOnlyTargetRules Target)
 		{
 			List<FileReference> BinaryPaths = new List<FileReference>();
-			if (Target.bIsBuildingConsoleApplication || !String.IsNullOrEmpty(BinaryName.GetExtension()))
+			// ModernXcode now builds binary outside of .app, instead Xcode will be responsible of generating .app
+			if (AppleExports.UseModernXcode(ProjectFile) ||
+				(Target.bIsBuildingConsoleApplication || !String.IsNullOrEmpty(BinaryName.GetExtension())))
 			{
 				BinaryPaths.Add(BinaryName);
 			}
diff --git a/Engine/Source/Programs/UnrealBuildTool/ProjectFiles/Xcode/XcodeProject.cs b/Engine/Source/Programs/UnrealBuildTool/ProjectFiles/Xcode/XcodeProject.cs
--- a/Engine/Source/Programs/UnrealBuildTool/ProjectFiles/Xcode/XcodeProject.cs
+++ b/Engine/Source/Programs/UnrealBuildTool/ProjectFiles/Xcode/XcodeProject.cs
@@ -1281,8 +1281,8 @@ namespace UnrealBuildTool.XcodeProjectXcconfig
 		///       ("BuildCookRun -build -cook -stage -package" can skip the copying until the -package step).
 		///       Also we don't copy anything if this is the engine, no-project, build, and an 
 		/// Mac:
-		///   - UBT will link directly to Foo.app/Contents/MacOS/Foo, for .app bundled apps
-		///   - During normal processing, that's all that is needed, so we don't need to copy anything
+		///   - UBT will link executable next to .app
+		///   - We will copy it into .app here, similar to iOS
 		///   - However, during Archiving, the .app is created in a intermediate location, so then here we copy from Binaries/Mac to the intermediate location
 		///     - Trying to have UBT link directly to the intermeidate location causes various issues, so we copy it like IOS does
 		/// IOS/TVOS:
@@ -1299,12 +1299,6 @@ namespace UnrealBuildTool.XcodeProjectXcconfig
 		/// <param name="Project"></param>
 		protected void ProcessScripts(XcodeResourcesBuildPhase ResourcesBuildPhase, XcodeProject Project)
 		{
-			// nothing to do for editors
-			if (Project.UnrealData.TargetRules.Type == TargetType.Editor)
-			{
-				return;
-			}
-
 			List<string> CopyScript = new();
 
 			// UBT no longer copies the executable into the .app directory in PostBuild, so we do it here
@@ -1328,6 +1322,15 @@ namespace UnrealBuildTool.XcodeProjectXcconfig
 				""
 			});
 
+			// Editor just need the above script to copy executable into .app
+			if (Project.UnrealData.TargetRules.Type == TargetType.Editor)
+			{
+				XcodeShellScriptBuildPhase EditorCopyScriptPhase = new("Copy Executable into .app", CopyScript, new string[] { }, new string[] { $"/dev/null" });
+				BuildPhases.Add(EditorCopyScriptPhase);
+				References.Add(EditorCopyScriptPhase);
+				return;
+			}
+
 			// rsync the Staged build into the .app, unless the UE_SKIP_STAGEDDATA_SYNC var is set to 1
 			// editor builds don't need staged content in them
 			if (TargetType != TargetType.Editor)
@@ -1790,19 +1793,8 @@ namespace UnrealBuildTool.XcodeProjectXcconfig
 				// hook up the Buildconfig that matches this info to this xcconfig file
 				XcconfigFile ConfigXcconfig = MatchedConfig.Xcconfig!;
 
-				FileReference PlatformExecutablePath;
 				string ExecutableName = XcodeUtils.MakeExecutableFileName(Config.ExeName, Platform, Config.BuildConfig, TargetRules.Architectures, TargetRules.UndecoratedConfiguration);
-				if (Platform == UnrealTargetPlatform.Mac)
-				{
-					// Mac builds right into the .app
-					PlatformExecutablePath = FileReference.Combine(ConfigBuildDir, ExecutableName + ".app", "Contents", "MacOS", ExecutableName);
-				}
-				else // IOS,TVOS
-				{
-					// IOS and TVOS build to outside of the .app
-					PlatformExecutablePath = FileReference.Combine(ConfigBuildDir, ExecutableName);
-				}
-				string ExetuableSubPath = PlatformExecutablePath.MakeRelativeTo(ConfigBuildDir);
+				string ExetuableSubPath = FileReference.Combine(ConfigBuildDir, ExecutableName).MakeRelativeTo(ConfigBuildDir);
 				string ProductName = ExecutableName;
 				string ExecutableKey = $"UE_{Platform.ToString().ToUpper()}_EXECUTABLE_NAME";
 
diff --git a/Engine/Source/Programs/UnrealBuildTool/ProjectFiles/Xcode/XcodeProjectFileGenerator.cs b/Engine/Source/Programs/UnrealBuildTool/ProjectFiles/Xcode/XcodeProjectFileGenerator.cs
--- a/Engine/Source/Programs/UnrealBuildTool/ProjectFiles/Xcode/XcodeProjectFileGenerator.cs
+++ b/Engine/Source/Programs/UnrealBuildTool/ProjectFiles/Xcode/XcodeProjectFileGenerator.cs
@@ -393,6 +393,27 @@ namespace UnrealBuildTool
 				}
 			}
 
+			// delete outdated workspace files, reduce confusion
+			if (bHasModernProjects)
+			{
+				DirectoryReference OutdatedWorkspaceDirectory = new DirectoryReference(PrimaryProjectPath + "/" + PrimaryProjectNameForPlatform(null) + ".xcworkspace");
+				if (DirectoryReference.Exists(OutdatedWorkspaceDirectory))
+				{
+					DirectoryReference.Delete(OutdatedWorkspaceDirectory, true);
+				}
+			}
+			else
+			{
+				foreach (UnrealTargetPlatform? Platform in WorkspacePlatforms)
+				{
+					DirectoryReference OutdatedWorkspaceDirectory = new DirectoryReference(PrimaryProjectPath + "/" + PrimaryProjectNameForPlatform(Platform) + ".xcworkspace");
+					if (DirectoryReference.Exists(OutdatedWorkspaceDirectory))
+					{
+						DirectoryReference.Delete(OutdatedWorkspaceDirectory, true);
+					}
+				}
+			}
+
 			return bSuccess;
 		}
 
diff --git a/Engine/Source/Programs/UnrealBuildTool/ToolChain/AppleToolChain.cs b/Engine/Source/Programs/UnrealBuildTool/ToolChain/AppleToolChain.cs
--- a/Engine/Source/Programs/UnrealBuildTool/ToolChain/AppleToolChain.cs
+++ b/Engine/Source/Programs/UnrealBuildTool/ToolChain/AppleToolChain.cs
@@ -537,6 +537,18 @@ namespace UnrealBuildTool
 			return OutputFiles;
 		}
 
+		public override void ModifyTargetReceipt(ReadOnlyTargetRules Target, TargetReceipt Receipt)
+		{
+			base.ModifyTargetReceipt(Target, Receipt);
+
+			if (Receipt.Platform == UnrealTargetPlatform.Mac && !Target.bIsBuildingConsoleApplication
+				&& Receipt.Launch != null && String.IsNullOrEmpty(Receipt.Launch!.GetExtension()))
+			{
+				// For Mac binary executables, we are building it outside of app, but expect it to be inside of .app after Xcode does its thing
+				Receipt.Launch = Receipt.Launch + ".app/Contents/MacOS/" + Receipt.Launch.GetFileName();
+			}
+		}
+
 		#region Stub Xcode Projects
 
 		internal static bool GenerateProjectFiles(FileReference? ProjectFile, string[] Arguments, ILogger Logger, out DirectoryReference? XcodeProjectFile)
@@ -709,7 +721,7 @@ namespace UnrealBuildTool
 			FileReference StagedExe;
 			if (Platform == UnrealTargetPlatform.Mac)
 			{
-				StagedExe = FileReference.Combine(Executable.Directory.ParentDirectory!, "PkgInfo");
+				StagedExe = FileReference.Combine(Executable.Directory, Executable.GetFileName() + ".app/Contents/PkgInfo");
 			}
 			else
 			{
diff --git a/Engine/Source/Programs/UnrealBuildTool/ToolChain/UEToolChain.cs b/Engine/Source/Programs/UnrealBuildTool/ToolChain/UEToolChain.cs
--- a/Engine/Source/Programs/UnrealBuildTool/ToolChain/UEToolChain.cs
+++ b/Engine/Source/Programs/UnrealBuildTool/ToolChain/UEToolChain.cs
@@ -214,7 +214,7 @@ namespace UnrealBuildTool
 		{
 		}
 
-		public virtual void ModifyTargetReceipt(TargetReceipt Receipt)
+		public virtual void ModifyTargetReceipt(ReadOnlyTargetRules Target, TargetReceipt Receipt)
 		{
 		}
 
