Write the is_unreal_test file to the correct location if the target lives under the project directory. Addresses an issue where the PostBuildStep would fail for test targets outside of the engine.

diff --git a/Engine/Source/Programs/UnrealBuildTool/Configuration/TestTargetRules.cs b/Engine/Source/Programs/UnrealBuildTool/Configuration/TestTargetRules.cs
--- a/Engine/Source/Programs/UnrealBuildTool/Configuration/TestTargetRules.cs
+++ b/Engine/Source/Programs/UnrealBuildTool/Configuration/TestTargetRules.cs
@@ -187,9 +187,17 @@ namespace UnrealBuildTool
 				}
 				if (File != null)
 				{
-					string OutputDirectory = UEBuildTarget.GetOutputDirectoryForExecutable(Unreal.EngineDirectory, File).FullName;
-					string OutputFileName = $"{OutputDirectory}\\Binaries\\{Target.Platform}\\{ExeBinariesSubFolder}\\{OutputName}.exe.is_unreal_test";
-					PostBuildSteps.Add("echo > " + OutputFileName);
+					DirectoryReference OutputDirectory;
+					if (ProjectFile != null && File.IsUnderDirectory(ProjectFile.Directory))
+					{
+						OutputDirectory = UEBuildTarget.GetOutputDirectoryForExecutable(ProjectFile.Directory, File);
+					}
+					else
+					{
+						OutputDirectory = UEBuildTarget.GetOutputDirectoryForExecutable(Unreal.EngineDirectory, File);
+					}
+					FileReference OutputTestFile = FileReference.Combine(OutputDirectory, "Binaries", Target.Platform.ToString(), ExeBinariesSubFolder, $"{OutputName}.exe.is_unreal_test");
+					PostBuildSteps.Add("echo > " + OutputTestFile.FullName);
 				}
 			}
 		}
