UnrealBuildTool: Fix failure to open/browse to source file in Visual Studio
* Write PrimaryProjectPath.txt with the location of the solution, this could be different than the root depending on project gen arguments
* Update SourceCodeAccessor accordingly to get the correct path to the solution, maintaining old defaults if the PrimaryProjectPath.txt file is not found

#rnx

diff --git a/Engine/Plugins/Developer/RiderSourceCodeAccess/Source/RiderSourceCodeAccess/Private/RiderSourceCodeAccessor.cpp b/Engine/Plugins/Developer/RiderSourceCodeAccess/Source/RiderSourceCodeAccess/Private/RiderSourceCodeAccessor.cpp
--- a/Engine/Plugins/Developer/RiderSourceCodeAccess/Source/RiderSourceCodeAccess/Private/RiderSourceCodeAccessor.cpp
+++ b/Engine/Plugins/Developer/RiderSourceCodeAccess/Source/RiderSourceCodeAccess/Private/RiderSourceCodeAccessor.cpp
@@ -325,31 +325,23 @@ void FRiderSourceCodeAccessor::CachePathToSln() const
 	}
 	else
 	{
-		CachedSolutionPath = FPaths::ConvertRelativePathToFull(FPaths::ProjectDir());
-
-		if (!FUProjectDictionary(FPaths::RootDir()).IsForeignProject(CachedSolutionPath))
-		{
-			FString PrimaryProjectName;
-			if (!FFileHelper::LoadFileToString(PrimaryProjectName, *(FPaths::EngineIntermediateDir() / TEXT("ProjectFiles/PrimaryProjectName.txt"))))
-			{
-				PrimaryProjectName = "UE5";
-			}
-			CachedSolutionPath = FPaths::Combine(FPaths::RootDir(), PrimaryProjectName + TEXT(".sln"));
-		}
-		else
-		{
-			const FProjectDescriptor* CurrentProject = IProjectManager::Get().GetCurrentProject();
-
-			if (CurrentProject == nullptr || CurrentProject->Modules.Num() == 0)
+			FString PrimaryProjectPath;
+			if (!FFileHelper::LoadFileToString(PrimaryProjectPath, *(FPaths::EngineIntermediateDir() / TEXT("ProjectFiles") / TEXT("PrimaryProjectPath.txt"))))
 			{
-				CachedSolutionPath = TEXT("");
+				const FProjectDescriptor* CurrentProject = IProjectManager::Get().GetCurrentProject();
+				const FString ProjectDir = FPaths::ConvertRelativePathToFull(FPaths::ProjectDir());
+
+				if ((CurrentProject == nullptr || CurrentProject->Modules.Num() == 0) || !FUProjectDictionary::GetDefault().IsForeignProject(ProjectDir))
+				{
+					PrimaryProjectPath = FPaths::RootDir() / TEXT("UE5");
+				}
+				else
+				{
+					const FString BaseName = FApp::HasProjectName() ? FApp::GetProjectName() : FPaths::GetBaseFilename(ProjectDir);
+					PrimaryProjectPath = ProjectDir / BaseName;
+				}
 			}
-			else
-			{
-				const FString BaseName = FApp::HasProjectName() ? FApp::GetProjectName() : FPaths::GetBaseFilename(CachedSolutionPath);
-				CachedSolutionPath = FPaths::Combine(CachedSolutionPath, BaseName + TEXT(".sln"));
-			}
-		}
+			CachedSolutionPath = PrimaryProjectPath + TEXT(".sln");
 	}
 }
 
diff --git a/Engine/Plugins/Developer/VisualStudioSourceCodeAccess/Source/VisualStudioSourceCodeAccess/Private/VisualStudioSourceCodeAccessor.cpp b/Engine/Plugins/Developer/VisualStudioSourceCodeAccess/Source/VisualStudioSourceCodeAccess/Private/VisualStudioSourceCodeAccessor.cpp
--- a/Engine/Plugins/Developer/VisualStudioSourceCodeAccess/Source/VisualStudioSourceCodeAccess/Private/VisualStudioSourceCodeAccessor.cpp
+++ b/Engine/Plugins/Developer/VisualStudioSourceCodeAccess/Source/VisualStudioSourceCodeAccess/Private/VisualStudioSourceCodeAccessor.cpp
@@ -1670,31 +1670,23 @@ FString FVisualStudioSourceCodeAccessor::GetSolutionPath() const
 		}
 		else
 		{
-			CachedSolutionPath = FPaths::ConvertRelativePathToFull(FPaths::ProjectDir());
-
-			if (!FUProjectDictionary::GetDefault().IsForeignProject(CachedSolutionPath))
-			{
-				FString PrimaryProjectName;
-				if (!FFileHelper::LoadFileToString(PrimaryProjectName, *(FPaths::EngineIntermediateDir() / TEXT("ProjectFiles/PrimaryProjectName.txt"))))
-				{
-					PrimaryProjectName = "UE5";
-				}
-				CachedSolutionPath = FPaths::Combine(FPaths::RootDir(), PrimaryProjectName + TEXT(".sln"));
-			}
-			else
+			FString PrimaryProjectPath;
+			if (!FFileHelper::LoadFileToString(PrimaryProjectPath, *(FPaths::EngineIntermediateDir() / TEXT("ProjectFiles") / TEXT("PrimaryProjectPath.txt"))))
 			{
 				const FProjectDescriptor* CurrentProject = IProjectManager::Get().GetCurrentProject();
+				const FString ProjectDir = FPaths::ConvertRelativePathToFull(FPaths::ProjectDir());
 
-				if (CurrentProject == nullptr || CurrentProject->Modules.Num() == 0)
+				if ((CurrentProject == nullptr || CurrentProject->Modules.Num() == 0) || !FUProjectDictionary::GetDefault().IsForeignProject(ProjectDir))
 				{
-					CachedSolutionPath = TEXT("");
+					PrimaryProjectPath = FPaths::RootDir() / TEXT("UE5");
 				}
 				else
 				{
-					FString BaseName = FApp::HasProjectName() ? FApp::GetProjectName() : FPaths::GetBaseFilename(CachedSolutionPath);
-					CachedSolutionPath = FPaths::Combine(CachedSolutionPath, BaseName + TEXT(".sln"));
+					const FString BaseName = FApp::HasProjectName() ? FApp::GetProjectName() : FPaths::GetBaseFilename(ProjectDir);
+					PrimaryProjectPath = ProjectDir / BaseName;
 				}
 			}
+			CachedSolutionPath = PrimaryProjectPath + TEXT(".sln");
 		}
 	}
 
diff --git a/Engine/Source/Programs/UnrealBuildTool/ProjectFiles/ProjectFileGenerator.cs b/Engine/Source/Programs/UnrealBuildTool/ProjectFiles/ProjectFileGenerator.cs
--- a/Engine/Source/Programs/UnrealBuildTool/ProjectFiles/ProjectFileGenerator.cs
+++ b/Engine/Source/Programs/UnrealBuildTool/ProjectFiles/ProjectFileGenerator.cs
@@ -857,8 +857,6 @@ namespace UnrealBuildTool
 		{
 			bool bSuccess = true;
 
-			bool WritePrimaryProjectNameTxt = false;
-
 			// Parse project generator options
 			bool IncludeAllPlatforms = true;
 			ConfigureProjectFileGeneration(Arguments, ref IncludeAllPlatforms, Logger);
@@ -904,9 +902,6 @@ namespace UnrealBuildTool
 						PrimaryProjectName = NewPrimaryProjectName;
 					}
 				}
-
-				// Write out the name of the primary project file, so the runtime knows to use it
-				WritePrimaryProjectNameTxt = true;
 			}
 
 			// if making a temp primary project, put it in with the rest of the project files
@@ -932,16 +927,13 @@ namespace UnrealBuildTool
 					PrimaryProjectName += SortedPlatform;
 					IntermediateProjectFilesPath = new DirectoryReference(IntermediateProjectFilesPath.FullName + SortedPlatform);
 				}
-
-				// the primary project name is always read from our intermediate directory and not the overriden one for this set of platforms
-				WritePrimaryProjectNameTxt = true;
 			}
 
-			if (WritePrimaryProjectNameTxt)
-			{
-				FileReference PrimaryProjectNameLocation = FileReference.Combine(Unreal.EngineDirectory, "Intermediate", "ProjectFiles", "PrimaryProjectName.txt");
-				Utils.WriteFileIfChanged(PrimaryProjectNameLocation, PrimaryProjectName, Logger);
-			}
+			// Write out the name of the primary project file and it's path, so the runtime knows to use it
+			FileReference PrimaryProjectNameLocation = FileReference.Combine(Unreal.EngineDirectory, "Intermediate", "ProjectFiles", "PrimaryProjectName.txt");
+			Utils.WriteFileIfChanged(PrimaryProjectNameLocation, PrimaryProjectName, Logger);
+			FileReference PrimaryProjectPathLocation = FileReference.Combine(Unreal.EngineDirectory, "Intermediate", "ProjectFiles", "PrimaryProjectPath.txt");
+			Utils.WriteFileIfChanged(PrimaryProjectPathLocation, DirectoryReference.Combine(PrimaryProjectPath, PrimaryProjectName).FullName.Replace("\\", "/", StringComparison.InvariantCultureIgnoreCase), Logger);
 
 			bool bCleanProjectFiles = Arguments.Any(x => x.Equals("-CleanProjects", StringComparison.InvariantCultureIgnoreCase));
 			if (bCleanProjectFiles)
