UnrealBuildTool: Clean up platform checks

#rnx

diff --git a/Engine/Source/Programs/UnrealBuildTool/UnrealBuildTool.csproj b/Engine/Source/Programs/UnrealBuildTool/UnrealBuildTool.csproj
--- a/Engine/Source/Programs/UnrealBuildTool/UnrealBuildTool.csproj
+++ b/Engine/Source/Programs/UnrealBuildTool/UnrealBuildTool.csproj
@@ -72,16 +72,12 @@
     </When>
   </Choose>
 
-  <Choose>
-    <When Condition="'$([System.Runtime.InteropServices.RuntimeInformation]::IsOSPlatform($([System.Runtime.InteropServices.OSPlatform]::Windows)))' == 'true' And '$([System.Runtime.InteropServices.RuntimeInformation]::OSArchitecture)' != 'Arm64'">
-      <PropertyGroup Condition="Exists('..\..\..\Restricted\NotForLicensees\Source\Programs\UnrealBuildTool\Executors\Experimental\BoxExecutor.cs')">
-        <DefineConstants>$(DefineConstants);__BOXEXECUTOR_AVAILABLE__</DefineConstants>
-      </PropertyGroup>
-      <PropertyGroup Condition="Exists('..\..\..\Restricted\NotForLicensees\Source\Programs\Shared\EpicGames.Box\BoxServer.cs')">
-        <DefineConstants>$(DefineConstants);__BOX_AVAILABLE__</DefineConstants>
-      </PropertyGroup>
-    </When>
-  </Choose>
+  <PropertyGroup Condition="Exists('..\..\..\Restricted\NotForLicensees\Source\Programs\UnrealBuildTool\Executors\Experimental\BoxExecutor.cs')">
+    <DefineConstants>$(DefineConstants);__BOXEXECUTOR_AVAILABLE__</DefineConstants>
+  </PropertyGroup>
+  <PropertyGroup Condition="Exists('..\..\..\Restricted\NotForLicensees\Source\Programs\Shared\EpicGames.Box\BoxServer.cs')">
+    <DefineConstants>$(DefineConstants);__BOX_AVAILABLE__</DefineConstants>
+  </PropertyGroup>
 
   <ItemGroup>
     <Compile Include="..\Shared\MetaData.cs">
