separate the AppX culture id from the PerCultureResources into its own UE StageId -> AppX CultureId override table.

 - Having two fields generally just caused confusion.
 - Most of the time, a simple 1:1 passthrough is sufficient.
 - DLC and Features both need the CultureId when configuring localized properties, but this was stored in a separate structure and required fiddly code to look it up.

diff --git a/Engine/Source/Programs/UnrealBuildTool/Platform/Windows/AppXManifestGeneratorBase.cs b/Engine/Source/Programs/UnrealBuildTool/Platform/Windows/AppXManifestGeneratorBase.cs
--- a/Engine/Source/Programs/UnrealBuildTool/Platform/Windows/AppXManifestGeneratorBase.cs
+++ b/Engine/Source/Programs/UnrealBuildTool/Platform/Windows/AppXManifestGeneratorBase.cs
@@ -58,6 +58,9 @@ namespace UnrealBuildTool
 
 		/// Logger for output
 		protected readonly ILogger Logger;
+		
+		/// Whether we have logged the deprecation warning for PerCultureResources CultureId being replaced by StageIdOverrides
+		protected static bool bHasWarnedAboutDeprecatedCultureId = false;
 
 		/// <summary>
 		/// Create a manifest generator for the given platform variant.
@@ -233,6 +236,16 @@ namespace UnrealBuildTool
 				AppXResources!.AddDefaultStrings(DefaultStrings);
 			}
 
+			// read StageId overrides
+			bool bHasStageIdOverrides = false;
+			if (EngineIni!.GetString(IniSection_PlatformTargetSettings, "StageIdOverrides", out string? StageIdOverridesString) && 
+				ConfigHierarchy.TryParseAsMap(StageIdOverridesString, out Dictionary<string,string>? StageIdOverrides))
+			{
+				bHasStageIdOverrides = true;
+				UEStageIdToAppXCultureId = StageIdOverrides;
+			}
+
+
 			// add per culture strings
 			if (EngineIni.GetArray(IniSection_PlatformTargetSettings, "PerCultureResources", out List<string>? PerCultureResources))
 			{
@@ -242,7 +255,6 @@ namespace UnrealBuildTool
 				{
 					if (!ConfigHierarchy.TryParse(CultureResources, out Dictionary<string, string>? CultureProperties)
 						|| !CultureProperties.ContainsKey("CultureStringResources")
-						|| !CultureProperties.ContainsKey("CultureId")
 						|| !CultureProperties.ContainsKey("StageId"))
 					{
 						Logger.LogWarning("Invalid per-culture resource value: {Culture}", CultureResources);
@@ -256,11 +268,33 @@ namespace UnrealBuildTool
 						continue;
 					}
 
-					string CultureId = CultureProperties["CultureId"];
-					if (String.IsNullOrEmpty(CultureId))
+					string CultureId = "";
+					if (bHasStageIdOverrides)
 					{
-						Logger.LogWarning("Missing CultureId value: {Culture}", CultureResources);
-						continue;
+						CultureId = UEStageIdToAppXCultureId.ContainsKey(StageId) ? UEStageIdToAppXCultureId[StageId] : StageId;
+					}
+					else if (!CultureProperties.ContainsKey("CultureId"))
+					{
+						Logger.LogWarning("Invalid per-culture resource value: {Culture}", CultureResources);
+					}
+					else
+					{
+						CultureId = CultureProperties["CultureId"];
+						if (String.IsNullOrEmpty(CultureId))
+						{
+							Logger.LogWarning("Missing CultureId value: {Culture}", CultureResources);
+							continue;
+						}
+
+						// PerCultureResources CultureId is deprecated in favor of new StageIdOverrides property
+						// if the CultureId field contains an overridden property, warn they need to update the data
+						if (CultureId != StageId && !bHasWarnedAboutDeprecatedCultureId)
+						{
+							Logger.LogWarning("PerCultureResources is out of date - please re-save this project's {Platform}Engine.ini in the editor to update StageIdOverrides. This must be done before UE5.5 to avoid losing data to deprecation", Platform);
+							bHasWarnedAboutDeprecatedCultureId = true;
+						}
+
+						UEStageIdToAppXCultureId[StageId] = CultureId;
 					}
 					AppXResources.AddCulture(CultureId);
 
@@ -271,9 +305,6 @@ namespace UnrealBuildTool
 						continue;
 					}
 					AppXResources!.AddCultureStrings(CultureId, CultureStringResources);
-
-					// build culture lookup table
-					UEStageIdToAppXCultureId[StageId] = CultureId;
 				}
 			}
 
