- Fixed project generation for code plugins under hybrid temp targets

diff --git a/Engine/Source/Programs/UnrealBuildTool/ProjectFiles/ProjectFileGenerator.cs b/Engine/Source/Programs/UnrealBuildTool/ProjectFiles/ProjectFileGenerator.cs
--- a/Engine/Source/Programs/UnrealBuildTool/ProjectFiles/ProjectFileGenerator.cs
+++ b/Engine/Source/Programs/UnrealBuildTool/ProjectFiles/ProjectFileGenerator.cs
@@ -2511,7 +2511,7 @@ namespace UnrealBuildTool
 			// Starting at the base directory of the module find a project which has the same directory as base, walking up the directory hierarchy until a match is found
 
 			DirectoryReference Path = CurModuleFile.Directory;
-			bool bIsTemporaryTarget = Path.ContainsName("Intermediate", 0);
+			bool bIsTemporaryModule = Path.ContainsName("Intermediate", 0);
 
 			while (!Path.IsRootDirectory())
 			{
@@ -2524,14 +2524,22 @@ namespace UnrealBuildTool
 						FileReference ProjectInfo = Game;
 						BaseFolder = ProjectInfo.Directory;
 
-						if (bMakeProjectPerTarget && !bIsTemporaryTarget)
+						if (bMakeProjectPerTarget && !bIsTemporaryModule)
 						{
 							// find the project that the module is under, and has a TargetType target (useful with bMakeProjectPerTarget)
 							foreach (KeyValuePair<FileReference, ProjectFile> Pair in ProjectFileMap)
 							{
-								if (((ProjectTarget)Pair.Value.ProjectTargets[0]).TargetFilePath.Directory.ParentDirectory == Path && Pair.Value.ProjectTargets.Any(x => x.TargetRules!.Type == TargetType))
+								FileReference TargetFile = ((ProjectTarget)Pair.Value.ProjectTargets[0]).TargetFilePath;
+								bool bIsTemporaryTarget = TargetFile.ContainsName("Intermediate", 0);
+
+								// check if the TargetFile is under <project>/Source, or under <project>/Intermediate/Source
+								if ((!bIsTemporaryTarget && TargetFile.Directory.ParentDirectory == Path) ||
+									(bIsTemporaryTarget && TargetFile.Directory.ParentDirectory!.ParentDirectory == Path))
 								{
-									return Pair.Value;
+									if (Pair.Value.ProjectTargets.Any(x => x.TargetRules!.Type == TargetType))
+									{
+										return Pair.Value;
+									}
 								}
 							}
 						}
