Updated MSVC weight to be 1.0f

Before:
Total time in Parallel executor: 2602.17 seconds
Total execution time: 2642.38 seconds

After:
Total time in Parallel executor: 2259.58 seconds
Total execution time: 2296.80 seconds

diff --git a/Engine/Source/Programs/UnrealBuildTool/Configuration/TargetRules.cs b/Engine/Source/Programs/UnrealBuildTool/Configuration/TargetRules.cs
--- a/Engine/Source/Programs/UnrealBuildTool/Configuration/TargetRules.cs
+++ b/Engine/Source/Programs/UnrealBuildTool/Configuration/TargetRules.cs
@@ -2054,7 +2054,7 @@ namespace UnrealBuildTool
 		/// The weight(cpu/memory utilization) of a MSVC compile action
 		/// </summary>
 		[XmlConfigFile(Category = "BuildConfiguration")]
-		public float MSVCCompileActionWeight = 1.5f;
+		public float MSVCCompileActionWeight = 1.0f;
 
 		/// <summary>
 		/// The weight(cpu/memory utilization) of a clang compile action
