fix for manifest.xml used wrong contents folder

diff --git a/Engine/Source/Programs/UnrealBuildTool/Platform/Mac/MacToolChain.cs b/Engine/Source/Programs/UnrealBuildTool/Platform/Mac/MacToolChain.cs
--- a/Engine/Source/Programs/UnrealBuildTool/Platform/Mac/MacToolChain.cs
+++ b/Engine/Source/Programs/UnrealBuildTool/Platform/Mac/MacToolChain.cs
@@ -1120,7 +1120,15 @@ namespace UnrealBuildTool
 
 			if (!BundleContentsDirectories.ContainsKey(Target) && Binary.Type == UEBuildBinaryType.Executable)
 			{
-				BundleContentsDirectories.Add(Target, Binary.OutputFilePath.Directory.ParentDirectory!);
+				if (!Binary.OutputFilePath.Directory.FullName.EndsWith(".app/Contents/MacOS", StringComparison.OrdinalIgnoreCase))
+				{
+					// The output binary is outside of the .app bundle (modern Xcode will move it later in the build proccess)
+					BundleContentsDirectories.Add(Target, DirectoryReference.Combine(Binary.OutputFilePath.Directory, Binary.OutputFilePath.GetFileName() + ".app", "Contents"));
+				}
+				else
+				{
+					BundleContentsDirectories.Add(Target, Binary.OutputFilePath.Directory.ParentDirectory!);
+				}
 			}
 			DirectoryReference? BundleContentsDirectory = BundleContentsDirectories.GetValueOrDefault(Target);
 
