[UBT]
* Added AdditionalDescription parameter to ParallelExecutor.RunAction
* Added bCanExecuteInBox for actions

diff --git a/Engine/Source/Programs/UnrealBuildTool/Executors/ParallelExecutor.cs b/Engine/Source/Programs/UnrealBuildTool/Executors/ParallelExecutor.cs
--- a/Engine/Source/Programs/UnrealBuildTool/Executors/ParallelExecutor.cs
+++ b/Engine/Source/Programs/UnrealBuildTool/Executors/ParallelExecutor.cs
@@ -210,7 +210,7 @@ namespace UnrealBuildTool
 			};
 		}
 
-		protected static async Task<ExecuteResults> RunAction(LinkedAction Action, ManagedProcessGroup ProcessGroup, CancellationToken CancellationToken)
+		protected static async Task<ExecuteResults> RunAction(LinkedAction Action, ManagedProcessGroup ProcessGroup, CancellationToken CancellationToken, string? AdditionalDescription = null)
 		{
 			CancellationToken.ThrowIfCancellationRequested();
 
@@ -227,7 +227,7 @@ namespace UnrealBuildTool
 			int ExitCode = Process.ExitCode;
 			TimeSpan ProcessorTime = Process.TotalProcessorTime;
 			TimeSpan ExecutionTime = Process.ExitTime - Process.StartTime;
-			return new ExecuteResults(LogLines, ExitCode, ExecutionTime, ProcessorTime);
+			return new ExecuteResults(LogLines, ExitCode, ExecutionTime, ProcessorTime, AdditionalDescription);
 		}
 	}
 
diff --git a/Engine/Source/Programs/UnrealBuildTool/Platform/Windows/VCCompileAction.cs b/Engine/Source/Programs/UnrealBuildTool/Platform/Windows/VCCompileAction.cs
--- a/Engine/Source/Programs/UnrealBuildTool/Platform/Windows/VCCompileAction.cs
+++ b/Engine/Source/Programs/UnrealBuildTool/Platform/Windows/VCCompileAction.cs
@@ -180,6 +180,9 @@ namespace UnrealBuildTool
 		/// <inheritdoc/>
 		public bool bCanExecuteRemotelyWithXGE { get; set; } = true;
 
+		/// <inheritdoc/>
+		public bool bCanExecuteInBox { get; set; } = true;
+
 		/// <inheritdoc/>
 		public bool bUseActionHistory => true;
 
diff --git a/Engine/Source/Programs/UnrealBuildTool/System/Action.cs b/Engine/Source/Programs/UnrealBuildTool/System/Action.cs
--- a/Engine/Source/Programs/UnrealBuildTool/System/Action.cs
+++ b/Engine/Source/Programs/UnrealBuildTool/System/Action.cs
@@ -145,6 +145,11 @@ namespace UnrealBuildTool
 		/// </summary>
 		bool bCanExecuteRemotelyWithXGE { get; }
 
+		/// <summary>
+		/// True if this action can be executed by box
+		/// </summary>
+		bool bCanExecuteInBox { get; }
+
 		/// <summary>
 		/// True if this action is using the GCC compiler.  Some build systems may be able to optimize for this case.
 		/// </summary>
@@ -262,6 +267,11 @@ namespace UnrealBuildTool
 		/// </summary>
 		public bool bCanExecuteRemotelyWithXGE { get; set; } = true;
 
+		/// <summary>
+		/// True if this action can be executed in box
+		/// </summary>
+		public bool bCanExecuteInBox { get; set; } = true;
+
 		/// <summary>
 		/// True if this action is using the GCC compiler.  Some build systems may be able to optimize for this case.
 		/// </summary>
@@ -319,6 +329,7 @@ namespace UnrealBuildTool
 			bCanExecuteRemotely = InOther.bCanExecuteRemotely;
 			bCanExecuteRemotelyWithSNDBS = InOther.bCanExecuteRemotelyWithSNDBS;
 			bCanExecuteRemotelyWithXGE = InOther.bCanExecuteRemotelyWithXGE;
+			bCanExecuteInBox = InOther.bCanExecuteInBox;
 			bIsGCCCompiler = InOther.bIsGCCCompiler;
 			bShouldOutputStatusDescription = InOther.bShouldOutputStatusDescription;
 			bProducesImportLibrary = InOther.bProducesImportLibrary;
@@ -340,6 +351,7 @@ namespace UnrealBuildTool
 			bCanExecuteRemotely = Reader.ReadBool();
 			bCanExecuteRemotelyWithSNDBS = Reader.ReadBool();
 			bCanExecuteRemotelyWithXGE = Reader.ReadBool();
+			bCanExecuteInBox = Reader.ReadBool();
 			bIsGCCCompiler = Reader.ReadBool();
 			bShouldOutputStatusDescription = Reader.ReadBool();
 			bProducesImportLibrary = Reader.ReadBool();
@@ -369,6 +381,7 @@ namespace UnrealBuildTool
 			Writer.WriteBool(bCanExecuteRemotely);
 			Writer.WriteBool(bCanExecuteRemotelyWithSNDBS);
 			Writer.WriteBool(bCanExecuteRemotelyWithXGE);
+			Writer.WriteBool(bCanExecuteInBox);
 			Writer.WriteBool(bIsGCCCompiler);
 			Writer.WriteBool(bShouldOutputStatusDescription);
 			Writer.WriteBool(bProducesImportLibrary);
@@ -449,6 +462,12 @@ namespace UnrealBuildTool
 				Action.bCanExecuteRemotelyWithXGE = bCanExecuteRemotelyWithXGE;
 			}
 
+			bool bCanExecuteInBox;
+			if (Object.TryGetBoolField("bCanExecuteInBox", out bCanExecuteInBox))
+			{
+				Action.bCanExecuteInBox = bCanExecuteInBox;
+			}
+
 			bool bIsGCCCompiler;
 			if (Object.TryGetBoolField("bIsGCCCompiler", out bIsGCCCompiler))
 			{
@@ -557,6 +576,7 @@ namespace UnrealBuildTool
 			Writer.WriteValue("bCanExecuteRemotely", Action.bCanExecuteRemotely);
 			Writer.WriteValue("bCanExecuteRemotelyWithSNDBS", Action.bCanExecuteRemotelyWithSNDBS);
 			Writer.WriteValue("bCanExecuteRemotelyWithXGE", Action.bCanExecuteRemotelyWithXGE);
+			Writer.WriteValue("bCanExecuteInBox", Action.bCanExecuteInBox);
 			Writer.WriteValue("bIsGCCCompiler", Action.bIsGCCCompiler);
 			Writer.WriteValue("bShouldOutputStatusDescription", Action.bShouldOutputStatusDescription);
 			Writer.WriteValue("bProducesImportLibrary", Action.bProducesImportLibrary);
@@ -678,6 +698,7 @@ namespace UnrealBuildTool
 		public bool bCanExecuteRemotely => Inner.bCanExecuteRemotely;
 		public bool bCanExecuteRemotelyWithSNDBS => Inner.bCanExecuteRemotelyWithSNDBS;
 		public bool bCanExecuteRemotelyWithXGE => Inner.bCanExecuteRemotelyWithXGE;
+		public bool bCanExecuteInBox => Inner.bCanExecuteInBox;
 		public bool bIsGCCCompiler => Inner.bIsGCCCompiler;
 		public bool bShouldOutputStatusDescription => Inner.bShouldOutputStatusDescription;
 		public bool bProducesImportLibrary => Inner.bProducesImportLibrary;
