Fix folder issues in Xcode for programs not under Engine folder
Also force being a bundled application for iOS and TVOS since they don't support not being bundled

diff --git a/Engine/Source/Programs/UnrealBuildTool/ProjectFiles/Xcode/XcodeProject.cs b/Engine/Source/Programs/UnrealBuildTool/ProjectFiles/Xcode/XcodeProject.cs
--- a/Engine/Source/Programs/UnrealBuildTool/ProjectFiles/Xcode/XcodeProject.cs
+++ b/Engine/Source/Programs/UnrealBuildTool/ProjectFiles/Xcode/XcodeProject.cs
@@ -48,13 +48,14 @@ namespace UnrealBuildTool.XcodeProjectXcconfig
 {
 	class UnrealBuildConfig
 	{
-		public UnrealBuildConfig(string InDisplayName, string InBuildTarget, string InExeName, ProjectTarget? InProjectTarget, UnrealTargetConfiguration InBuildConfig)
+		public UnrealBuildConfig(string InDisplayName, string InBuildTarget, string InExeName, ProjectTarget? InProjectTarget, UnrealTargetConfiguration InBuildConfig, DirectoryReference InRootDirectory)
 		{
 			DisplayName = InDisplayName;
 			BuildTarget = InBuildTarget;
 			ExeName = InExeName;
 			ProjectTarget = InProjectTarget;
 			BuildConfig = InBuildConfig;
+			RootDirectory = InRootDirectory;
 
 			if (BuildTarget != ProjectTarget?.Name)
 			{
@@ -67,6 +68,7 @@ namespace UnrealBuildTool.XcodeProjectXcconfig
 		public string ExeName;
 		public ProjectTarget? ProjectTarget;
 		public UnrealTargetConfiguration BuildConfig;
+		public DirectoryReference RootDirectory;
 
 		public bool bSupportsMac => Supports(UnrealTargetPlatform.Mac);
 		public bool bSupportsIOS => Supports(UnrealTargetPlatform.IOS);
@@ -333,7 +335,7 @@ namespace UnrealBuildTool.XcodeProjectXcconfig
 		public bool bIsContentOnlyProject;
 
 		public TargetRules TargetRules => _TargetRules!;
-		public bool bIsAppBundle;
+		bool bIsAppBundle;
 
 		public bool bUseAutomaticSigning = false;
 		public bool bIsMergingProjects = false;
@@ -403,6 +405,16 @@ namespace UnrealBuildTool.XcodeProjectXcconfig
 			return Platform == null || XcodeProjectFileGenerator.XcodePlatforms.Contains((UnrealTargetPlatform)Platform);
 		}
 
+		public bool IsAppBundle(UnrealTargetPlatform Platform)
+		{
+			if (Platform == UnrealTargetPlatform.IOS || Platform == UnrealTargetPlatform.TVOS)
+			{
+				// iOS and TvOS always need app bundles
+				return true;
+			}
+			return bIsAppBundle;
+		}
+
 		public UnrealData(FileReference XcodeProjectFileLocation, bool bIsForDistribution, string BundleID, string AppName, bool bMakeProjectPerTarget)
 		{
 			// the .xcodeproj is actually a directory
@@ -703,19 +715,7 @@ namespace UnrealBuildTool.XcodeProjectXcconfig
 												}
 											}
 
-											// Get the output directory
-											DirectoryReference OutputDirectory = DirectoryReference.Combine(RootDirectory, "Binaries");
-											DirectoryReference MacBinaryDir = DirectoryReference.Combine(OutputDirectory, "Mac");
-											DirectoryReference IOSBinaryDir = DirectoryReference.Combine(OutputDirectory, "IOS");
-											DirectoryReference TVOSBinaryDir = DirectoryReference.Combine(OutputDirectory, "TVOS");
-											if (!String.IsNullOrEmpty(ProjectTarget.TargetRules.ExeBinariesSubFolder))
-											{
-												MacBinaryDir = DirectoryReference.Combine(MacBinaryDir, ProjectTarget.TargetRules.ExeBinariesSubFolder);
-												IOSBinaryDir = DirectoryReference.Combine(IOSBinaryDir, ProjectTarget.TargetRules.ExeBinariesSubFolder);
-												TVOSBinaryDir = DirectoryReference.Combine(TVOSBinaryDir, ProjectTarget.TargetRules.ExeBinariesSubFolder);
-											}
-
-											BuildConfigs.Add(new UnrealBuildConfig(ConfigName, TargetName, ExeName, ProjectTarget, Configuration));
+											BuildConfigs.Add(new UnrealBuildConfig(ConfigName, TargetName, ExeName, ProjectTarget, Configuration, RootDirectory));
 										}
 									}
 								}
@@ -1236,7 +1236,7 @@ namespace UnrealBuildTool.XcodeProjectXcconfig
 		public TargetType TargetType;
 
 		public XcodeRunTarget(XcodeProject Project, string TargetName, TargetType TargetType, UnrealTargetPlatform Platform, XcodeBuildTarget? BuildTarget, XcodeProjectFile ProjectFile, ILogger Logger)
-			: base(Project.UnrealData.bIsAppBundle ? XcodeTarget.Type.Run_App : XcodeTarget.Type.Run_Tool, Project.UnrealData,
+			: base(Project.UnrealData.IsAppBundle(Platform) ? XcodeTarget.Type.Run_App : XcodeTarget.Type.Run_Tool, Project.UnrealData,
 				  TargetName + (XcodeProjectFileGenerator.PerPlatformMode == XcodePerPlatformMode.OneWorkspacePerPlatform ? "" : $"_{Platform}"))
 		{
 			this.TargetType = TargetType;
@@ -1247,9 +1247,9 @@ namespace UnrealBuildTool.XcodeProjectXcconfig
 			References.Add(BuildConfigList);
 
 			// add the Product item to the project to be visible in left pane
-			Project.FileCollection.AddFileReference(ProductGuid, UnrealData.ProductName, "explicitFileType", Project.UnrealData.bIsAppBundle ? "wrapper.application" : "\"compiled.mach-o.executable\"", "BUILT_PRODUCTS_DIR", "Products");
+			Project.FileCollection.AddFileReference(ProductGuid, UnrealData.ProductName, "explicitFileType", Project.UnrealData.IsAppBundle(Platform) ? "wrapper.application" : "\"compiled.mach-o.executable\"", "BUILT_PRODUCTS_DIR", "Products");
 
-			if (Project.UnrealData.bIsAppBundle)
+			if (Project.UnrealData.IsAppBundle(Platform))
 			{
 				XcodeResourcesBuildPhase ResourcesBuildPhase = new XcodeResourcesBuildPhase(Project.FileCollection);
 				BuildPhases.Add(ResourcesBuildPhase);
@@ -1575,7 +1575,7 @@ namespace UnrealBuildTool.XcodeProjectXcconfig
 			DirectoryReference ConfigBuildDir = (TargetRules.Type == TargetType.Editor || UnrealData.bIsContentOnlyProject) ? Unreal.EngineDirectory : ProjectOrEngineDir;
 			if (TargetRules.Type == TargetType.Program && TargetRules.File!.IsUnderDirectory(Unreal.EngineDirectory))
 			{
-				ConfigBuildDir = Unreal.EngineDirectory;
+				ConfigBuildDir = BuildConfig.RootDirectory;
 			}
 			string BinariesBaseDir = ConfigBuildDir.GetDirectoryName();
 
