Compile servers with WITH_RAYTRACING=0

- This reduces Shipping server binary by a couple megabytes, and also helps reduce work during the LTO/PGO builds.

diff --git a/Engine/Source/Programs/UnrealBuildTool/Platform/Linux/UEBuildLinux.cs b/Engine/Source/Programs/UnrealBuildTool/Platform/Linux/UEBuildLinux.cs
--- a/Engine/Source/Programs/UnrealBuildTool/Platform/Linux/UEBuildLinux.cs
+++ b/Engine/Source/Programs/UnrealBuildTool/Platform/Linux/UEBuildLinux.cs
@@ -523,7 +523,7 @@ namespace UnrealBuildTool
 
 			CompileEnvironment.Definitions.Add("INT64_T_TYPES_NOT_LONG_LONG=1");
 
-			if (Target.LinuxPlatform.bEnableRayTracing)
+			if (Target.LinuxPlatform.bEnableRayTracing && Target.Type != TargetType.Server)
 			{
 				CompileEnvironment.Definitions.Add("RHI_RAYTRACING=1");
 			}
diff --git a/Engine/Source/Programs/UnrealBuildTool/Platform/Mac/UEBuildMac.cs b/Engine/Source/Programs/UnrealBuildTool/Platform/Mac/UEBuildMac.cs
--- a/Engine/Source/Programs/UnrealBuildTool/Platform/Mac/UEBuildMac.cs
+++ b/Engine/Source/Programs/UnrealBuildTool/Platform/Mac/UEBuildMac.cs
@@ -510,7 +510,7 @@ namespace UnrealBuildTool
 			CompileEnvironment.Definitions.Add("PLATFORM_MAC=1");
 			CompileEnvironment.Definitions.Add("PLATFORM_APPLE=1");
 
-			if (Target.MacPlatform.bEnableRayTracing)
+			if (Target.MacPlatform.bEnableRayTracing && Target.Type != TargetType.Server)
 			{
 				CompileEnvironment.Definitions.Add("RHI_RAYTRACING=1");
 			}
diff --git a/Engine/Source/Programs/UnrealBuildTool/Platform/Windows/UEBuildWindows.cs b/Engine/Source/Programs/UnrealBuildTool/Platform/Windows/UEBuildWindows.cs
--- a/Engine/Source/Programs/UnrealBuildTool/Platform/Windows/UEBuildWindows.cs
+++ b/Engine/Source/Programs/UnrealBuildTool/Platform/Windows/UEBuildWindows.cs
@@ -1361,7 +1361,7 @@ namespace UnrealBuildTool
 				CompileEnvironment.Definitions.Add(String.Format("OVERRIDE_PLATFORM_HEADER_NAME={0}", OverridePlatformHeaderName));
 			}
 
-			if (Target.IsInPlatformGroup(UnrealPlatformGroup.Windows) && Target.WindowsPlatform.bEnableRayTracing)
+			if (Target.IsInPlatformGroup(UnrealPlatformGroup.Windows) && Target.WindowsPlatform.bEnableRayTracing && Target.Type != TargetType.Server)
 			{
 				CompileEnvironment.Definitions.Add("RHI_RAYTRACING=1");
 			}
