Allow plugins "Tests" folder that can include test target and module files.

diff --git a/Engine/Source/Programs/AutomationTool/Scripts/BuildPluginCommand.Automation.cs b/Engine/Source/Programs/AutomationTool/Scripts/BuildPluginCommand.Automation.cs
--- a/Engine/Source/Programs/AutomationTool/Scripts/BuildPluginCommand.Automation.cs
+++ b/Engine/Source/Programs/AutomationTool/Scripts/BuildPluginCommand.Automation.cs
@@ -432,6 +432,7 @@ public sealed class BuildPlugin : BuildCommand
 		Filter.Include("/Intermediate/Build/.../Inc/...");
 		Filter.Include("/Shaders/...");
 		Filter.Include("/Source/...");
+		Filter.Exclude("/Tests/...");
 
 		// Add custom rules for all platforms
 		AddRulesFromFileToFilter(Filter, FileReference.Combine(PluginFile.Directory, "Config", "FilterPlugin.ini"));
diff --git a/Engine/Source/Programs/Shared/EpicGames.Build/System/Rules.cs b/Engine/Source/Programs/Shared/EpicGames.Build/System/Rules.cs
--- a/Engine/Source/Programs/Shared/EpicGames.Build/System/Rules.cs
+++ b/Engine/Source/Programs/Shared/EpicGames.Build/System/Rules.cs
@@ -85,13 +85,14 @@ namespace UnrealBuildBase
 				RootFolders.AddRange(GameFolders.SelectMany(x => Unreal.GetExtensionDirs(x)));
 			}
 
-			// Find all the plugin source directories
+			// Find all the plugin source and tests directories
 			foreach (DirectoryReference RootFolder in RootFolders)
 			{
 				DirectoryReference PluginsFolder = DirectoryReference.Combine(RootFolder, "Plugins");
 				foreach (FileReference PluginFile in PluginsBase.EnumeratePlugins(PluginsFolder))
 				{
 					Folders.Add(DirectoryReference.Combine(PluginFile.Directory, "Source"));
+					Folders.Add(DirectoryReference.Combine(PluginFile.Directory, "Tests"));
 				}
 			}
 
diff --git a/Engine/Source/Programs/UnrealBuildTool/System/RulesCompiler.cs b/Engine/Source/Programs/UnrealBuildTool/System/RulesCompiler.cs
--- a/Engine/Source/Programs/UnrealBuildTool/System/RulesCompiler.cs
+++ b/Engine/Source/Programs/UnrealBuildTool/System/RulesCompiler.cs
@@ -290,8 +290,9 @@ namespace UnrealBuildTool
 				Logger.LogTrace(" Found {Count} Plugins:", ProjectPlugins.Count);
 				ProjectPlugins.ForEach(x => { Logger.LogTrace("  {Plugin}", x.File); });
 
-				// Find all the plugin module rules
+				// Find all the plugin module rules as well as plugin test target and module rules
 				FindModuleRulesForPlugins(ProjectPlugins, DefaultModuleContext, ModuleFiles);
+				FindTestRulesForPlugins(ProjectPlugins, DefaultModuleContext, ModuleFiles, TargetFiles);
 
 				Logger.LogTrace(" Found {Count} Modules:", ModuleFiles.Count);
 				foreach (KeyValuePair<FileReference, ModuleRulesContext> Item in ModuleFiles)
@@ -419,7 +420,7 @@ namespace UnrealBuildTool
 		/// <summary>
 		/// Finds all the module rules for plugins under the given directory.
 		/// </summary>
-		/// <param name="Plugins">The directory to search</param>
+		/// <param name="Plugins">The list of plugins to search modules for</param>
 		/// <param name="DefaultContext">The default context for any files that are enumerated</param>
 		/// <param name="ModuleFileToContext">Dictionary which is filled with mappings from the module file to its corresponding context</param>
 		private static void FindModuleRulesForPlugins(IReadOnlyList<PluginInfo> Plugins, ModuleRulesContext DefaultContext, Dictionary<FileReference, ModuleRulesContext> ModuleFileToContext)
@@ -428,22 +429,51 @@ namespace UnrealBuildTool
 
 			foreach (PluginInfo Plugin in Plugins)
 			{
-				List<FileReference> PluginModuleFiles = Rules.FindAllRulesFiles(DirectoryReference.Combine(Plugin.Directory, "Source"), Rules.RulesFileType.Module).ToList();
-				foreach (FileReference ChildFile in Plugin.ChildFiles)
-				{
-					PluginModuleFiles.AddRange(Rules.FindAllRulesFiles(DirectoryReference.Combine(ChildFile.Directory, "Source"), Rules.RulesFileType.Module));
-				}
+				FindModuleRulesForPluginInFolder(Plugin, "Source", DefaultContext, ModuleFileToContext);
+			}
+		}
 
-				foreach (FileReference ModuleFile in PluginModuleFiles)
-				{
-					ModuleRulesContext PluginContext = new ModuleRulesContext(DefaultContext);
-					PluginContext.DefaultOutputBaseDir = Plugin.Directory;
-					PluginContext.Plugin = Plugin;
-					ModuleFileToContext[ModuleFile] = PluginContext;
-				}
+		/// <summary>
+		/// Finds all the module and target rules for plugins' tests.
+		/// </summary>
+		/// <param name="Plugins">The list of plugins to search test rules for</param>
+		/// <param name="DefaultContext">The default context for any files that are enumerated</param>
+		/// <param name="ModuleFileToContext">Dictionary which is filled with mappings from the module file to its corresponding context</param>
+		/// <param name="TargetFiles">List of target files to add test target rules to</param>
+		private static void FindTestRulesForPlugins(IReadOnlyList<PluginInfo> Plugins, ModuleRulesContext DefaultContext, Dictionary<FileReference, ModuleRulesContext> ModuleFileToContext, List<FileReference> TargetFiles)
+		{
+			Rules.PrefetchRulesFiles(Plugins.Select(x => DirectoryReference.Combine(x.Directory, "Tests")));
+
+			foreach (PluginInfo Plugin in Plugins)
+			{
+				FindModuleRulesForPluginInFolder(Plugin, "Tests", DefaultContext, ModuleFileToContext);
+				TargetFiles.AddRange(Rules.FindAllRulesFiles(DirectoryReference.Combine(Plugin.Directory, "Tests"), Rules.RulesFileType.Target));
 			}
 		}
 
+		/// <summary>
+		/// Finds all the module rules for a given plugin in a specified folder.
+		/// </summary>
+		/// <param name="Plugin">The plugin to search module rules for</param>
+		/// <param name="Folder">The folder relative to the plugin root to look into, usually "Source" or "Tests"</param>
+		/// <param name="DefaultContext">The default context for any files that are enumerated</param>
+		/// <param name="ModuleFileToContext">Dictionary which is filled with mappings from the module file to its corresponding context</param>
+		private static void FindModuleRulesForPluginInFolder(PluginInfo Plugin, string Folder, ModuleRulesContext DefaultContext, Dictionary<FileReference, ModuleRulesContext> ModuleFileToContext)
+		{
+			List<FileReference> PluginModuleFiles = Rules.FindAllRulesFiles(DirectoryReference.Combine(Plugin.Directory, Folder), Rules.RulesFileType.Module).ToList();
+			foreach (FileReference ChildFile in Plugin.ChildFiles)
+			{
+				PluginModuleFiles.AddRange(Rules.FindAllRulesFiles(DirectoryReference.Combine(ChildFile.Directory, Folder), Rules.RulesFileType.Module));
+			}
+			foreach (FileReference ModuleFile in PluginModuleFiles)
+			{
+				ModuleRulesContext PluginContext = new ModuleRulesContext(DefaultContext);
+				PluginContext.DefaultOutputBaseDir = Plugin.Directory;
+				PluginContext.Plugin = Plugin;
+				ModuleFileToContext[ModuleFile] = PluginContext;
+			}
+		}
+		
 		/// <summary>
 		/// Gets the filename that declares the given type.
 		/// </summary>
