fix UAT getting a wrong exe path from TargetReceipt due to modern Xcode change

diff --git a/Engine/Source/Programs/UnrealBuildTool/ToolChain/AppleToolChain.cs b/Engine/Source/Programs/UnrealBuildTool/ToolChain/AppleToolChain.cs
--- a/Engine/Source/Programs/UnrealBuildTool/ToolChain/AppleToolChain.cs
+++ b/Engine/Source/Programs/UnrealBuildTool/ToolChain/AppleToolChain.cs
@@ -542,9 +542,10 @@ namespace UnrealBuildTool
 			base.ModifyTargetReceipt(Target, Receipt);
 
 			if (Receipt.Platform == UnrealTargetPlatform.Mac && !Target.bIsBuildingConsoleApplication
-				&& Receipt.Launch != null && String.IsNullOrEmpty(Receipt.Launch!.GetExtension()))
+				&& Receipt.Launch != null && String.IsNullOrEmpty(Receipt.Launch!.GetExtension())
+				&& !Receipt.Launch.FullName.Contains(".app/Contents/MacOS/", StringComparison.OrdinalIgnoreCase))
 			{
-				// For Mac binary executables, we are building it outside of app, but expect it to be inside of .app after Xcode does its thing
+				// For Mac binary executables, we may build it outside of app, but expect it to be inside of .app after modern Xcode does its thing
 				Receipt.Launch = Receipt.Launch + ".app/Contents/MacOS/" + Receipt.Launch.GetFileName();
 			}
 		}
