[UBT]
* Added ImmediateActionQueue.GetQueuedActionsCount method that returns number of queued actions left

diff --git a/Engine/Source/Programs/UnrealBuildTool/Executors/ImmediateActionQueue.cs b/Engine/Source/Programs/UnrealBuildTool/Executors/ImmediateActionQueue.cs
--- a/Engine/Source/Programs/UnrealBuildTool/Executors/ImmediateActionQueue.cs
+++ b/Engine/Source/Programs/UnrealBuildTool/Executors/ImmediateActionQueue.cs
@@ -798,6 +798,24 @@ namespace UnrealBuildTool
 			}
 		}
 
+		public int GetQueuedActionsCount()
+		{
+			lock (Actions)
+			{
+				int count = 0;
+
+				for (int actionIndex = _firstPendingAction; actionIndex != Actions.Length; ++actionIndex)
+				{
+
+					if (Actions[actionIndex].Status == ActionStatus.Queued)
+					{
+						++count;
+					}
+				}
+				return count;
+			}
+		}
+
 		/// <summary>
 		/// Purge the pending logging actions
 		/// </summary>
