UnrealBuildTool: Only enable deterministic mode for Shipping by default, as it affects msvc compile time due to disabled codegen multithreading

diff --git a/Engine/Source/Programs/UnrealBuildTool/Configuration/TargetRules.cs b/Engine/Source/Programs/UnrealBuildTool/Configuration/TargetRules.cs
--- a/Engine/Source/Programs/UnrealBuildTool/Configuration/TargetRules.cs
+++ b/Engine/Source/Programs/UnrealBuildTool/Configuration/TargetRules.cs
@@ -1711,12 +1711,18 @@ namespace UnrealBuildTool
 
 		/// <summary>
 		/// Set flags require for deterministic compiling and linking.
+		/// Enabling deterministic mode for msvc disables codegen multithreading so compiling will be slower
 		/// </summary>
 		[CommandLine("-Deterministic", Value = "true")]
 		[CommandLine("-NonDeterministic", Value = "false")]
 		[XmlConfigFile(Category = "BuildConfiguration")]
 		[RequiresUniqueBuildEnvironment]
-		public bool bDeterministic = true;
+		public bool bDeterministic
+		{
+			get => bDeterministicPrivate ?? (Configuration == UnrealTargetConfiguration.Shipping);
+			set => bDeterministicPrivate = value;
+		}
+		private bool? bDeterministicPrivate;
 
 		/// <summary>
 		/// Whether PCHs should be chained when compiling with clang.
diff --git a/Engine/Source/Programs/UnrealBuildTool/Platform/Windows/VCToolChain.cs b/Engine/Source/Programs/UnrealBuildTool/Platform/Windows/VCToolChain.cs
--- a/Engine/Source/Programs/UnrealBuildTool/Platform/Windows/VCToolChain.cs
+++ b/Engine/Source/Programs/UnrealBuildTool/Platform/Windows/VCToolChain.cs
@@ -1409,10 +1409,14 @@ namespace UnrealBuildTool
 		{
 			VCCompileAction BaseCompileAction = CreateBaseCompileAction(CompileEnvironment);
 
-			// MSVC uses multiple processes when compiling CPP files, so the "weight" is more than 1
+			// MSVC uses multiple threads when compiling CPP files, so the "weight" is more than 1
 			if (Target.WindowsPlatform.Compiler.IsMSVC())
 			{
-				BaseCompileAction.Weight = Target.MSVCCompileActionWeight;
+				// If deterministic is enabled, MSVC does not use multiple threads
+				if (!CompileEnvironment.bDeterministic)
+				{
+					BaseCompileAction.Weight = Target.MSVCCompileActionWeight;
+				}
 			}
 			else if (Target.WindowsPlatform.Compiler.IsClang())
 			{
