UnrealBuildTool: Add msvc conformance flags to generated projects for intellisense

#rnx

diff --git a/Engine/Source/Programs/UnrealBuildTool/ProjectFiles/VisualStudio/VCProject.cs b/Engine/Source/Programs/UnrealBuildTool/ProjectFiles/VisualStudio/VCProject.cs
--- a/Engine/Source/Programs/UnrealBuildTool/ProjectFiles/VisualStudio/VCProject.cs
+++ b/Engine/Source/Programs/UnrealBuildTool/ProjectFiles/VisualStudio/VCProject.cs
@@ -657,6 +657,7 @@ namespace UnrealBuildTool
 				{
 					Arguments.Add("/Zc:enumTypes");
 				}
+				return string.Join(' ', Arguments);
 			}
 			return String.Empty;
 		}
@@ -1186,9 +1187,11 @@ namespace UnrealBuildTool
 			VCProjectFileContent.AppendLine("  <PropertyGroup Label=\"UserMacros\" />");
 
 			// Write each project configuration
+			TargetRules? DefaultRules = null;
 			foreach (ProjectConfigAndTargetCombination Combination in ProjectConfigAndTargetCombinations)
 			{
 				WriteConfiguration(ProjectName, Combination, VCProjectFileContent, PlatformProjectGenerators, bGenerateUserFileContent ? VCUserFileContent : null, Logger);
+				DefaultRules = DefaultRules ?? Combination.ProjectTarget?.TargetRules;
 			}
 
 			// Write IntelliSense info
@@ -1204,8 +1207,8 @@ namespace UnrealBuildTool
 				VCProjectFileContent.AppendLine("    <IncludePath>$(IncludePath){0}</IncludePath>", (SharedIncludeSearchPaths.Length > 0 ? (";" + SharedIncludeSearchPaths) : ""));
 				VCProjectFileContent.AppendLine("    <NMakeForcedIncludes>$(NMakeForcedIncludes)</NMakeForcedIncludes>");
 				VCProjectFileContent.AppendLine("    <NMakeAssemblySearchPath>$(NMakeAssemblySearchPath)</NMakeAssemblySearchPath>");
-				VCProjectFileContent.AppendLine("    <AdditionalOptions>{0} {1}</AdditionalOptions>",
-					GetCppStandardCompileArgument(GetIntelliSenseCppVersion()), GetEnableCoroutinesArgument());
+				VCProjectFileContent.AppendLine("    <AdditionalOptions>{0} {1} {2}</AdditionalOptions>",
+					GetCppStandardCompileArgument(GetIntelliSenseCppVersion()), GetEnableCoroutinesArgument(), DefaultRules != null ? GetConformanceCompileArguments(DefaultRules) : string.Empty);
 				VCProjectFileContent.AppendLine("  </PropertyGroup>");
 			}
 
