- For temp single target projects, put the targetname into the workspace name, so that multiple program targets being compiled/finalized at once don't all write to the same workspace file

diff --git a/Engine/Source/Programs/UnrealBuildTool/ProjectFiles/Xcode/XcodeProjectFileGenerator.cs b/Engine/Source/Programs/UnrealBuildTool/ProjectFiles/Xcode/XcodeProjectFileGenerator.cs
--- a/Engine/Source/Programs/UnrealBuildTool/ProjectFiles/Xcode/XcodeProjectFileGenerator.cs
+++ b/Engine/Source/Programs/UnrealBuildTool/ProjectFiles/Xcode/XcodeProjectFileGenerator.cs
@@ -247,7 +247,19 @@ namespace UnrealBuildTool
 
 		private string PrimaryProjectNameForPlatform(UnrealTargetPlatform? Platform)
 		{
-			return Platform == null ? PrimaryProjectName : $"{PrimaryProjectName} ({Platform})";
+			Console.WriteLine("SingleTargetName = '{0}'", SingleTargetName);
+			Console.WriteLine("PrimaryProjectName = '{0}'", PrimaryProjectName);
+			string ProjectName = PrimaryProjectName;
+			if (!string.IsNullOrEmpty(SingleTargetName))
+			{
+				ProjectName = $"{ProjectName}_{SingleTargetName}";
+			}
+			// if there are projectplatforms, then there is a platform name already in the name
+			if (Platform != null && ProjectPlatforms.Count == 0)
+			{
+				ProjectName = $"{ProjectName} ({Platform})";
+			}
+			return ProjectName;
 		}
 
 		private bool WriteXcodeWorkspace(ILogger Logger)
@@ -393,25 +405,28 @@ namespace UnrealBuildTool
 				}
 			}
 
-			// delete outdated workspace files, reduce confusion
-			if (bHasModernProjects)
+			// delete outdated workspace files, reduce confusion (only for real workspaces, not stub ones)
+			if (!bGenerateRunOnlyProject)
 			{
-				DirectoryReference OutdatedWorkspaceDirectory = new DirectoryReference(PrimaryProjectPath + "/" + PrimaryProjectNameForPlatform(null) + ".xcworkspace");
-				if (DirectoryReference.Exists(OutdatedWorkspaceDirectory))
+				if (bHasModernProjects)
 				{
-					DirectoryReference.Delete(OutdatedWorkspaceDirectory, true);
-				}
-			}
-			else
-			{
-				foreach (UnrealTargetPlatform? Platform in WorkspacePlatforms)
-				{
-					DirectoryReference OutdatedWorkspaceDirectory = new DirectoryReference(PrimaryProjectPath + "/" + PrimaryProjectNameForPlatform(Platform) + ".xcworkspace");
+					DirectoryReference OutdatedWorkspaceDirectory = new DirectoryReference(PrimaryProjectPath + "/" + PrimaryProjectNameForPlatform(null) + ".xcworkspace");
 					if (DirectoryReference.Exists(OutdatedWorkspaceDirectory))
 					{
 						DirectoryReference.Delete(OutdatedWorkspaceDirectory, true);
 					}
 				}
+				else
+				{
+					foreach (UnrealTargetPlatform? Platform in WorkspacePlatforms)
+					{
+						DirectoryReference OutdatedWorkspaceDirectory = new DirectoryReference(PrimaryProjectPath + "/" + PrimaryProjectNameForPlatform(Platform) + ".xcworkspace");
+						if (DirectoryReference.Exists(OutdatedWorkspaceDirectory))
+						{
+							DirectoryReference.Delete(OutdatedWorkspaceDirectory, true);
+						}
+					}
+				}
 			}
 
 			return bSuccess;
