Modify BuildProducts for correct in-bundle binary path instead of modifying TargetReceipt and Manifest.xml

diff --git a/Engine/Source/Programs/UnrealBuildTool/Platform/Mac/MacToolChain.cs b/Engine/Source/Programs/UnrealBuildTool/Platform/Mac/MacToolChain.cs
--- a/Engine/Source/Programs/UnrealBuildTool/Platform/Mac/MacToolChain.cs
+++ b/Engine/Source/Programs/UnrealBuildTool/Platform/Mac/MacToolChain.cs
@@ -1120,15 +1120,15 @@ namespace UnrealBuildTool
 
 			if (!BundleContentsDirectories.ContainsKey(Target) && Binary.Type == UEBuildBinaryType.Executable)
 			{
-				if (!Binary.OutputFilePath.Directory.FullName.EndsWith(".app/Contents/MacOS", StringComparison.OrdinalIgnoreCase))
+				// For Mac binary executables, we may build it outside of app, but expect it to be inside of .app after modern Xcode does its thing
+				FileReference FinalBinaryPath = Binary.OutputFilePath;
+				if (!FinalBinaryPath.Directory.FullName.EndsWith(".app/Contents/MacOS", StringComparison.OrdinalIgnoreCase))
 				{
-					// The output binary is outside of the .app bundle (modern Xcode will move it later in the build proccess)
-					BundleContentsDirectories.Add(Target, DirectoryReference.Combine(Binary.OutputFilePath.Directory, Binary.OutputFilePath.GetFileName() + ".app", "Contents"));
-				}
-				else
-				{
-					BundleContentsDirectories.Add(Target, Binary.OutputFilePath.Directory.ParentDirectory!);
+					FinalBinaryPath = FileReference.Combine(FinalBinaryPath.Directory, FinalBinaryPath.GetFileName() + ".app", "Contents", "MacOS", FinalBinaryPath.GetFileName());
+					BuildProducts.Remove(Binary.OutputFilePath);
+					BuildProducts.Add(FinalBinaryPath, BuildProductType.Executable);
 				}
+				BundleContentsDirectories.Add(Target, FinalBinaryPath.Directory.ParentDirectory!);
 			}
 			DirectoryReference? BundleContentsDirectory = BundleContentsDirectories.GetValueOrDefault(Target);
 
diff --git a/Engine/Source/Programs/UnrealBuildTool/ToolChain/AppleToolChain.cs b/Engine/Source/Programs/UnrealBuildTool/ToolChain/AppleToolChain.cs
--- a/Engine/Source/Programs/UnrealBuildTool/ToolChain/AppleToolChain.cs
+++ b/Engine/Source/Programs/UnrealBuildTool/ToolChain/AppleToolChain.cs
@@ -537,19 +537,6 @@ namespace UnrealBuildTool
 			return OutputFiles;
 		}
 
-		public override void ModifyTargetReceipt(ReadOnlyTargetRules Target, TargetReceipt Receipt)
-		{
-			base.ModifyTargetReceipt(Target, Receipt);
-
-			if (Receipt.Platform == UnrealTargetPlatform.Mac && !Target.bIsBuildingConsoleApplication
-				&& Receipt.Launch != null && String.IsNullOrEmpty(Receipt.Launch!.GetExtension())
-				&& !Receipt.Launch.FullName.Contains(".app/Contents/MacOS/", StringComparison.OrdinalIgnoreCase))
-			{
-				// For Mac binary executables, we may build it outside of app, but expect it to be inside of .app after modern Xcode does its thing
-				Receipt.Launch = Receipt.Launch + ".app/Contents/MacOS/" + Receipt.Launch.GetFileName();
-			}
-		}
-
 		#region Stub Xcode Projects
 
 		internal static bool GenerateProjectFiles(FileReference? ProjectFile, string[] Arguments, ILogger Logger, out DirectoryReference? XcodeProjectFile)
