Fixing build warning 320113

diff --git a/Engine/Source/Programs/UnrealBuildTool/System/PluginDescriptor.cs b/Engine/Source/Programs/UnrealBuildTool/System/PluginDescriptor.cs
--- a/Engine/Source/Programs/UnrealBuildTool/System/PluginDescriptor.cs
+++ b/Engine/Source/Programs/UnrealBuildTool/System/PluginDescriptor.cs
@@ -355,6 +355,7 @@ namespace UnrealBuildTool
 			RawObject.TryGetStringArrayField("DisallowedPlugins", out DisallowedPlugins);
 		}
 
+		/// <summary>
 		/// Creates a plugin descriptor from a file on disk
 		/// </summary>
 		/// <param name="FileName">The filename to read</param>
