Support iPad SplitView (dynamic resizing) and StageManager

diff --git a/Engine/Config/BaseEngine.ini b/Engine/Config/BaseEngine.ini
--- a/Engine/Config/BaseEngine.ini
+++ b/Engine/Config/BaseEngine.ini
@@ -2826,6 +2826,7 @@ FrameRateLock=PUFRL_30
 bEnableDynamicMaxFPS=False
 bSupportsIPad=True
 bSupportsIPhone=True
+bEnableSplitView=False
 AdditionalPlistData=
 RemoteServerName=
 RemoteServerOverrideBuildPath=
diff --git a/Engine/Source/Programs/UnrealBuildTool/Platform/IOS/UEDeployIOS.cs b/Engine/Source/Programs/UnrealBuildTool/Platform/IOS/UEDeployIOS.cs
--- a/Engine/Source/Programs/UnrealBuildTool/Platform/IOS/UEDeployIOS.cs
+++ b/Engine/Source/Programs/UnrealBuildTool/Platform/IOS/UEDeployIOS.cs
@@ -341,8 +341,10 @@ namespace UnrealBuildTool
 				Text.AppendLine("\t<key>LSSupportsOpeningDocumentsInPlace</key>");
 				Text.AppendLine("\t<true/>");
 			}
+
 			Text.AppendLine("\t<key>UIRequiresFullScreen</key>");
 			Text.AppendLine("\t<true/>");
+
 			Text.AppendLine("\t<key>UIViewControllerBasedStatusBarAppearance</key>");
 			Text.AppendLine("\t<false/>");
 			if (!String.IsNullOrEmpty(InterfaceOrientation))
@@ -351,7 +353,7 @@ namespace UnrealBuildTool
 			}
 			Text.AppendLine("\t<key>UISupportedInterfaceOrientations</key>");
 			Text.AppendLine("\t<array>");
-			foreach (string Line in SupportedOrientations.Split("\r\n".ToCharArray()))
+			foreach (string Line in SupportedOrientations.Split("\n".ToCharArray()))
 			{
 				if (!String.IsNullOrWhiteSpace(Line))
 				{
@@ -359,6 +361,25 @@ namespace UnrealBuildTool
 				}
 			}
 			Text.AppendLine("\t</array>");
+
+			bool bEnableSplitView = false;
+			Ini.GetBool("/Script/IOSRuntimeSettings.IOSRuntimeSettings", "bEnableSplitView", out bEnableSplitView);
+			if (bEnableSplitView)
+			{
+				// As this is (currently) an iPad only feature, use the iPad descriminator to set it for iPad only
+				// as it also requires supporting all UIOrientations
+				Text.AppendLine("\t<key>UIRequiresFullScreen~ipad</key>");
+				Text.AppendLine("\t<false/>");
+
+				Text.AppendLine("\t<key>UISupportedInterfaceOrientations~ipad</key>");
+				Text.AppendLine("\t<array>");
+					Text.AppendLine($"\t\t<string>UIInterfaceOrientationPortrait</string>");
+					Text.AppendLine($"\t\t<string>UIInterfaceOrientationPortraitUpsideDown</string>");
+					Text.AppendLine($"\t\t<string>UIInterfaceOrientationLandscapeLeft</string>");
+					Text.AppendLine($"\t\t<string>UIInterfaceOrientationLandscapeRight</string>");
+				Text.AppendLine("\t</array>");
+			}
+
 			Text.AppendLine("\t<key>UIRequiredDeviceCapabilities</key>");
 			Text.AppendLine("\t<array>");
 			foreach (string Cap in RequiredCaps)
diff --git a/Engine/Source/Programs/UnrealBuildTool/ProjectFiles/Xcode/XcodeProject.cs b/Engine/Source/Programs/UnrealBuildTool/ProjectFiles/Xcode/XcodeProject.cs
--- a/Engine/Source/Programs/UnrealBuildTool/ProjectFiles/Xcode/XcodeProject.cs
+++ b/Engine/Source/Programs/UnrealBuildTool/ProjectFiles/Xcode/XcodeProject.cs
@@ -1658,6 +1658,15 @@ namespace UnrealBuildTool.XcodeProjectXcconfig
 					// only iphone deals with orientation
 					List<string> SupportedOrientations = XcodeUtils.GetSupportedOrientations(PlatformIni);
 					ExtraConfigLines.Add($"INFOPLIST_KEY_UISupportedInterfaceOrientations = \"{String.Join(" ", SupportedOrientations)}\"");
+
+					// iPhone is always Fullscreen, however, an iPad can support SplitView mode (dynamic view resizeing), so check if that's enabled and then add iPad UISupportedInterfaceOrientations
+					bool bEnableSplitView = false;
+					PlatformIni.GetBool("/Script/IOSRuntimeSettings.IOSRuntimeSettings", "bEnableSplitView", out bEnableSplitView);
+					if (bEnableSplitView)
+					{
+						ExtraConfigLines.Add($"INFOPLIST_KEY_UISupportedInterfaceOrientations_iPad = UIInterfaceOrientationLandscapeLeft UIInterfaceOrientationLandscapeRight UIInterfaceOrientationPortrait UIInterfaceOrientationPortraitUpsideDown");
+						ExtraConfigLines.Add($"INFOPLIST_KEY_UIRequiresFullScreen_iPad = false");
+					}
 				}
 				else // tvos
 				{
diff --git a/Engine/Source/Runtime/ApplicationCore/Private/IOS/IOSApplication.cpp b/Engine/Source/Runtime/ApplicationCore/Private/IOS/IOSApplication.cpp
--- a/Engine/Source/Runtime/ApplicationCore/Private/IOS/IOSApplication.cpp
+++ b/Engine/Source/Runtime/ApplicationCore/Private/IOS/IOSApplication.cpp
@@ -142,7 +142,7 @@ UIEdgeInsets CachedInsets;
 
 void FDisplayMetrics::RebuildDisplayMetrics(FDisplayMetrics& OutDisplayMetrics)
 {
-	const FPlatformRect& Rect = FIOSWindow::GetUIWindowRect();
+	const FPlatformRect& Rect = FIOSWindow::GetScreenRect();
 	const FIOSView *View = [[IOSAppDelegate GetDelegate] IOSView];
 	[View CalculateContentScaleFactor:Rect.Right ScreenHeight:Rect.Bottom];
 	
@@ -155,8 +155,6 @@ void FDisplayMetrics::RebuildDisplayMetrics(FDisplayMetrics& OutDisplayMetrics)
     OutDisplayMetrics.PrimaryDisplayHeight = OutDisplayMetrics.PrimaryDisplayWorkAreaRect.Bottom - OutDisplayMetrics.PrimaryDisplayWorkAreaRect.Top;
     
 #if !PLATFORM_TVOS
-    const double RequestedContentScaleFactor = View.contentScaleFactor;
-    
     //we need to set these according to the orientation
     TAutoConsoleVariable<float>* CVar_Left = nullptr;
     TAutoConsoleVariable<float>* CVar_Top = nullptr;
@@ -192,7 +190,7 @@ void FDisplayMetrics::RebuildDisplayMetrics(FDisplayMetrics& OutDisplayMetrics)
     OutDisplayMetrics.TitleSafePaddingSize.W = Inset_Bottom;
     
     //scale the thing
-    OutDisplayMetrics.TitleSafePaddingSize *= RequestedContentScaleFactor;
+    OutDisplayMetrics.TitleSafePaddingSize *= View.contentScaleFactor;
     
     OutDisplayMetrics.ActionSafePaddingSize = OutDisplayMetrics.TitleSafePaddingSize;
 #endif
diff --git a/Engine/Source/Runtime/ApplicationCore/Private/IOS/IOSView.cpp b/Engine/Source/Runtime/ApplicationCore/Private/IOS/IOSView.cpp
--- a/Engine/Source/Runtime/ApplicationCore/Private/IOS/IOSView.cpp
+++ b/Engine/Source/Runtime/ApplicationCore/Private/IOS/IOSView.cpp
@@ -257,15 +257,12 @@ id<MTLDevice> GMetalDevice = nil;
 	int32 Width = ScreenWidth;
 	int32 Height = ScreenHeight;
 	
+	self.contentScaleFactor = mNativeScale;
+	
 	// 0 means to use native size
-	if (RequestedContentScaleFactor == 0.0f && RequestedResX <= 0 && RequestedResY <= 0)
-	{
-		self.contentScaleFactor = mNativeScale;
-	}
-	else
+	if (RequestedContentScaleFactor != 0.0f || RequestedResX > 0 || RequestedResY > 0)
 	{
 		float AspectRatio = (float)ScreenHeight / (float)ScreenWidth;
-		self.contentScaleFactor = 1.0f;
 		if (RequestedResX > 0)
 		{
 			// set long side for orientation to requested X
@@ -410,8 +407,6 @@ self.accessibilityElements = @[Window.accessibilityContainer];
 
 -(void)HandleTouchAtLoc:(CGPoint)Loc PrevLoc:(CGPoint)PrevLoc TouchIndex:(int)TouchIndex Force:(float)Force Type:(TouchType)Type TouchesArray:(TArray<TouchInput>&)TouchesArray
 {
-	const CGFloat Scale = self.contentScaleFactor;
-
 	// init some things on begin
 	if (Type == TouchBegan)
 	{
@@ -427,8 +422,8 @@ self.accessibilityElements = @[Window.accessibilityContainer];
 	TouchInput TouchMessage;
 	TouchMessage.Handle = TouchIndex;
 	TouchMessage.Type = Type;
-	TouchMessage.Position = FVector2D(FMath::Min<double>(_ViewSize.width - 1, Loc.x), FMath::Min<double>(_ViewSize.height - 1, Loc.y)) * Scale;
-	TouchMessage.LastPosition = FVector2D(FMath::Min<double>(_ViewSize.width - 1, PrevLoc.x), FMath::Min<double>(_ViewSize.height - 1, PrevLoc.y)) * Scale;
+	TouchMessage.Position = FVector2D(FMath::Min<double>(_ViewSize.width - 1, Loc.x), FMath::Min<double>(_ViewSize.height - 1, Loc.y));
+	TouchMessage.LastPosition = FVector2D(FMath::Min<double>(_ViewSize.width - 1, PrevLoc.x), FMath::Min<double>(_ViewSize.height - 1, PrevLoc.y));
 	TouchMessage.Force = Type != TouchEnded ? Force : 0.0f;
 	
 	// skip moves that didn't actually move - this will help input handling to skip over the first
diff --git a/Engine/Source/Runtime/ApplicationCore/Private/IOS/IOSWindow.cpp b/Engine/Source/Runtime/ApplicationCore/Private/IOS/IOSWindow.cpp
--- a/Engine/Source/Runtime/ApplicationCore/Private/IOS/IOSWindow.cpp
+++ b/Engine/Source/Runtime/ApplicationCore/Private/IOS/IOSWindow.cpp
@@ -130,30 +130,14 @@ FPlatformRect FIOSWindow::GetScreenRect()
 		CGFloat Scale = View.contentScaleFactor;
 		
 		ScreenRect.Top = FMath::TruncToInt(Frame.origin.y * Scale);
-		ScreenRect.Bottom = FMath::TruncToInt((Frame.origin.y + View.ViewSize.height) * Scale);
 		ScreenRect.Left = FMath::TruncToInt(Frame.origin.x * Scale);
-		ScreenRect.Right = FMath::TruncToInt((Frame.origin.x + View.ViewSize.width) * Scale);
+		ScreenRect.Bottom = FMath::TruncToInt((Frame.origin.y + Frame.size.height) * Scale);
+		ScreenRect.Right = FMath::TruncToInt((Frame.origin.x + Frame.size.width) * Scale);
 	}
 
 	return ScreenRect;
 }
 
-FPlatformRect FIOSWindow::GetUIWindowRect()
-{
-	// get the main window's bounds
-	IOSAppDelegate* AppDelegate = [IOSAppDelegate GetDelegate];
-	UIWindow* Window = AppDelegate.Window;
-	CGRect Bounds = [Window bounds];
-
-	FPlatformRect WindowRect;
-	WindowRect.Top = FMath::TruncToInt(Bounds.origin.y);
-	WindowRect.Bottom = FMath::TruncToInt(Bounds.origin.y + Bounds.size.height);
-	WindowRect.Left = FMath::TruncToInt(Bounds.origin.x);
-	WindowRect.Right = FMath::TruncToInt(Bounds.origin.x + Bounds.size.width);
-
-	return WindowRect;
-}
-
 bool FIOSWindow::GetFullScreenInfo( int32& X, int32& Y, int32& Width, int32& Height ) const
 {
 	FPlatformRect ScreenRect = GetScreenRect();
diff --git a/Engine/Source/Runtime/ApplicationCore/Private/IOS/IOSWindow.h b/Engine/Source/Runtime/ApplicationCore/Private/IOS/IOSWindow.h
--- a/Engine/Source/Runtime/ApplicationCore/Private/IOS/IOSWindow.h
+++ b/Engine/Source/Runtime/ApplicationCore/Private/IOS/IOSWindow.h
@@ -25,9 +25,6 @@ public:
 	/** Returns a rect describing the main screen */
 	static FPlatformRect GetScreenRect();
 
-	/** Returns a rect describing the main screen */
-	static FPlatformRect GetUIWindowRect();
-
 	/** Returns a void pointer to the hWnd (for other apis)*/
 	virtual void* GetOSWindowHandle() const override { return Window; }
 
diff --git a/Engine/Source/Runtime/IOS/IOSRuntimeSettings/Classes/IOSRuntimeSettings.h b/Engine/Source/Runtime/IOS/IOSRuntimeSettings/Classes/IOSRuntimeSettings.h
--- a/Engine/Source/Runtime/IOS/IOSRuntimeSettings/Classes/IOSRuntimeSettings.h
+++ b/Engine/Source/Runtime/IOS/IOSRuntimeSettings/Classes/IOSRuntimeSettings.h
@@ -262,6 +262,10 @@ public:
     UPROPERTY(GlobalConfig, EditAnywhere, Category = Build, meta = (DisplayName = "Supports iPhone"))
     bool bSupportsIPhone;
     
+    // Whether or not the iPad app supports Split View (also needed for StageManager support)
+    UPROPERTY(GlobalConfig, EditAnywhere, Category = Build, meta = (DisplayName = "Enable iPad Split View"))
+    bool bEnableSplitView;
+    
     /** Set the maximum frame rate to save on power consumption */
     UPROPERTY(GlobalConfig, EditAnywhere, Category = PowerUsage, meta = (ConfigHierarchyEditable))
     EPowerUsageFrameRateLock FrameRateLock;
diff --git a/Engine/Source/Runtime/IOS/IOSRuntimeSettings/Private/IOSRuntimeSettings.cpp b/Engine/Source/Runtime/IOS/IOSRuntimeSettings/Private/IOSRuntimeSettings.cpp
--- a/Engine/Source/Runtime/IOS/IOSRuntimeSettings/Private/IOSRuntimeSettings.cpp
+++ b/Engine/Source/Runtime/IOS/IOSRuntimeSettings/Private/IOSRuntimeSettings.cpp
@@ -31,6 +31,7 @@ UIOSRuntimeSettings::UIOSRuntimeSettings(const FObjectInitializer& ObjectInitial
 	bEnableDynamicMaxFPS = false;
 	bSupportsIPad = true;
 	bSupportsIPhone = true;
+	bEnableSplitView = false;
 	MinimumiOSVersion = EIOSVersion::IOS_Minimum;
     bBuildAsFramework = true;
 	bGeneratedSYMFile = false;
