- Restored framework unzipping into UBT for modern, because if you don't build with Xcode one time, it won't unzip the frameworks in time

diff --git a/Engine/Source/Programs/UnrealBuildTool/ToolChain/AppleToolChain.cs b/Engine/Source/Programs/UnrealBuildTool/ToolChain/AppleToolChain.cs
--- a/Engine/Source/Programs/UnrealBuildTool/ToolChain/AppleToolChain.cs
+++ b/Engine/Source/Programs/UnrealBuildTool/ToolChain/AppleToolChain.cs
@@ -210,11 +210,9 @@ namespace UnrealBuildTool
 			{
 				if (Framework.ZipFile != null)
 				{
-					// modern xcode manages all this in xcode script
-					if (!bUseModernXcode)
-					{
-						ExtractFramework(Framework, Graph, Logger);
-					}
+					// even in modern, we have to try to unzip in UBT, so that commandline builds extract the frameworks _before_ building (the Xcode finalize
+					// happens after building). We still need the unzip in Xcode pre-build phase, because Xcode wants the frameworks around for its build graph
+					ExtractFramework(Framework, Graph, Logger);
 					FrameworkTokenFiles.Add(Framework.ExtractedTokenFile!);
 				}
 			}
