PR #10422: Fixed broken bAddDefaultIncludePaths UBT setting

diff --git a/Engine/Source/Programs/UnrealBuildTool/Configuration/UEBuildModuleCPP.cs b/Engine/Source/Programs/UnrealBuildTool/Configuration/UEBuildModuleCPP.cs
--- a/Engine/Source/Programs/UnrealBuildTool/Configuration/UEBuildModuleCPP.cs
+++ b/Engine/Source/Programs/UnrealBuildTool/Configuration/UEBuildModuleCPP.cs
@@ -216,7 +216,10 @@ namespace UnrealBuildTool
 				}
 			}
 
-			AddDefaultIncludePaths();
+			if (Rules.bAddDefaultIncludePaths)
+			{
+				AddDefaultIncludePaths();
+			}
 		}
 
 		/// <summary>
