[UBT]
* IWYUToolchain - removed code that hardcodes to c++17 now when all c++20 compile errors have been fixed

diff --git a/Engine/Source/Programs/UnrealBuildTool/ToolChain/IWYUToolChain.cs b/Engine/Source/Programs/UnrealBuildTool/ToolChain/IWYUToolChain.cs
--- a/Engine/Source/Programs/UnrealBuildTool/ToolChain/IWYUToolChain.cs
+++ b/Engine/Source/Programs/UnrealBuildTool/ToolChain/IWYUToolChain.cs
@@ -196,12 +196,6 @@ namespace UnrealBuildTool
 			Arguments.Add("-Wno-pragma-once-outside-header");
 		}
 
-		// TODO: REMOVE THIS ONCE c++20 compile fixes have been done
-		protected override void GetCppStandardCompileArgument(CppCompileEnvironment CompileEnvironment, List<string> Arguments)
-		{
-			Arguments.Add("-std=c++17");
-		}
-
 		// Skip ISPC headers
 		public override CPPOutput GenerateISPCHeaders(CppCompileEnvironment CompileEnvironment, List<FileItem> InputFiles, DirectoryReference OutputDir, IActionGraphBuilder Graph)
 		{
