Add binary outside of .app as RequiredResource in BuildProducts, so that Horde will copy it in between agents

diff --git a/Engine/Source/Programs/UnrealBuildTool/Platform/Mac/MacToolChain.cs b/Engine/Source/Programs/UnrealBuildTool/Platform/Mac/MacToolChain.cs
--- a/Engine/Source/Programs/UnrealBuildTool/Platform/Mac/MacToolChain.cs
+++ b/Engine/Source/Programs/UnrealBuildTool/Platform/Mac/MacToolChain.cs
@@ -1121,11 +1121,12 @@ namespace UnrealBuildTool
 			if (!BundleContentsDirectories.ContainsKey(Target) && Binary.Type == UEBuildBinaryType.Executable)
 			{
 				// For Mac binary executables, we may build it outside of app, but expect it to be inside of .app after modern Xcode does its thing
+				// We still keep the executable outside of app as RequiredResource, so that Horde will copy in-between agents
 				FileReference FinalBinaryPath = Binary.OutputFilePath;
 				if (!FinalBinaryPath.Directory.FullName.EndsWith(".app/Contents/MacOS", StringComparison.OrdinalIgnoreCase))
 				{
 					FinalBinaryPath = FileReference.Combine(FinalBinaryPath.Directory, FinalBinaryPath.GetFileName() + ".app", "Contents", "MacOS", FinalBinaryPath.GetFileName());
-					BuildProducts.Remove(Binary.OutputFilePath);
+					BuildProducts[Binary.OutputFilePath] = BuildProductType.RequiredResource;
 					BuildProducts.Add(FinalBinaryPath, BuildProductType.Executable);
 				}
 				BundleContentsDirectories.Add(Target, FinalBinaryPath.Directory.ParentDirectory!);
