UnrealBuildTool: WindowsTargetRules setting to control Inline Function Expansion Level

diff --git a/Engine/Source/Programs/UnrealBuildTool/Platform/Windows/UEBuildWindows.cs b/Engine/Source/Programs/UnrealBuildTool/Platform/Windows/UEBuildWindows.cs
--- a/Engine/Source/Programs/UnrealBuildTool/Platform/Windows/UEBuildWindows.cs
+++ b/Engine/Source/Programs/UnrealBuildTool/Platform/Windows/UEBuildWindows.cs
@@ -197,12 +197,19 @@ namespace UnrealBuildTool
 
 		/// <summary>
 		/// If specified along with -PGOProfile, prevent the usage of extra counters. Please note that by default /FASTGENPROFILE doesnt use extra counters
-		/// https://learn.microsoft.com/en-us/cpp/build/reference/genprofile-fastgenprofile-generate-profiling-instrumented-build?view=msvc-170
 		/// </summary>
+		/// <seealso href="link">https://learn.microsoft.com/en-us/cpp/build/reference/genprofile-fastgenprofile-generate-profiling-instrumented-build</seealso>
 		[XmlConfigFile(Category = "WindowsPlatform")]
 		[CommandLine("-PGONoExtraCounters")]
 		public bool bPGONoExtraCounters = false;
 
+		/// <summary>
+		/// Which level to use for Inline Function Expansion when TargetRules.bUseInlining is enabled
+		/// </summary>
+		/// <seealso href="link">https://learn.microsoft.com/en-us/cpp/build/reference/ob-inline-function-expansion</seealso>
+		[XmlConfigFile(Category = "WindowsPlatform")]
+		public int InlineFunctionExpansionLevel { get; set; } = 2;
+
 		/// <summary>
 		/// Version of the compiler toolchain to use on Windows platform. A value of "default" will be changed to a specific version at UBT start up.
 		/// </summary>
@@ -631,6 +638,8 @@ namespace UnrealBuildTool
 
 		public bool bPGONoExtraCounters => Inner.bPGONoExtraCounters;
 
+		public int InlineFunctionExpansionLevel => Inner.InlineFunctionExpansionLevel;
+
 		public WindowsCompiler Compiler => Inner.Compiler;
 
 		public UnrealArch Architecture => Inner.Architecture;
diff --git a/Engine/Source/Programs/UnrealBuildTool/Platform/Windows/VCToolChain.cs b/Engine/Source/Programs/UnrealBuildTool/Platform/Windows/VCToolChain.cs
--- a/Engine/Source/Programs/UnrealBuildTool/Platform/Windows/VCToolChain.cs
+++ b/Engine/Source/Programs/UnrealBuildTool/Platform/Windows/VCToolChain.cs
@@ -535,7 +535,7 @@ namespace UnrealBuildTool
 			// Do not allow inline method expansion if E&C support is enabled or inline expansion has been disabled
 			if (!CompileEnvironment.bSupportEditAndContinue && CompileEnvironment.bUseInlining)
 			{
-				Arguments.Add("/Ob2");
+				Arguments.Add($"/Ob{Math.Clamp(Target.WindowsPlatform.InlineFunctionExpansionLevel, 1, 3)}");
 			}
 			else
 			{
