- Removed some debugging logs

diff --git a/Engine/Source/Programs/UnrealBuildTool/ProjectFiles/Xcode/XcodeProjectFileGenerator.cs b/Engine/Source/Programs/UnrealBuildTool/ProjectFiles/Xcode/XcodeProjectFileGenerator.cs
--- a/Engine/Source/Programs/UnrealBuildTool/ProjectFiles/Xcode/XcodeProjectFileGenerator.cs
+++ b/Engine/Source/Programs/UnrealBuildTool/ProjectFiles/Xcode/XcodeProjectFileGenerator.cs
@@ -247,8 +247,6 @@ namespace UnrealBuildTool
 
 		private string PrimaryProjectNameForPlatform(UnrealTargetPlatform? Platform)
 		{
-			Console.WriteLine("SingleTargetName = '{0}'", SingleTargetName);
-			Console.WriteLine("PrimaryProjectName = '{0}'", PrimaryProjectName);
 			string ProjectName = PrimaryProjectName;
 			if (!string.IsNullOrEmpty(SingleTargetName))
 			{
