Adding and switching to the 5.4 include order version.

diff --git a/Engine/Source/Programs/UnrealBuildTool/Configuration/TargetRules.cs b/Engine/Source/Programs/UnrealBuildTool/Configuration/TargetRules.cs
--- a/Engine/Source/Programs/UnrealBuildTool/Configuration/TargetRules.cs
+++ b/Engine/Source/Programs/UnrealBuildTool/Configuration/TargetRules.cs
@@ -171,12 +171,17 @@ namespace UnrealBuildTool
 		/// </summary>
 		Unreal5_3,
 
+		/// <summary>
+		/// Include order used in Unreal 5.3
+		/// </summary>
+		Unreal5_4,
+
 		// *** When adding new entries here, be sure to update UEBuildModuleCPP.CurrentIncludeOrderDefine to ensure that the correct guard is used. ***
 
 		/// <summary>
 		/// Always use the latest version of include order.
 		/// </summary>
-		Latest = Unreal5_3,
+		Latest = Unreal5_4,
 
 		/// <summary>
 		/// Contains the oldest version of include order that the engine supports.
@@ -307,6 +312,8 @@ namespace UnrealBuildTool
 			switch (InVersion)
 			{
 				default:
+				case EngineIncludeOrderVersion.Unreal5_4:
+					return "UE_ENABLE_INCLUDE_ORDER_DEPRECATED_IN_5_4";
 				case EngineIncludeOrderVersion.Unreal5_3:
 					return "UE_ENABLE_INCLUDE_ORDER_DEPRECATED_IN_5_3";
 				case EngineIncludeOrderVersion.Unreal5_2:
@@ -334,6 +341,7 @@ namespace UnrealBuildTool
 				GetDeprecationDefine(EngineIncludeOrderVersion.Unreal5_1),
 				GetDeprecationDefine(EngineIncludeOrderVersion.Unreal5_2),
 				GetDeprecationDefine(EngineIncludeOrderVersion.Unreal5_3),
+				GetDeprecationDefine(EngineIncludeOrderVersion.Unreal5_4),
 			};
 		}
 
