- Moved the SingleTargetName property up to ProjectFileGenerator from Xcode generator, and use it to remove Targets from the list to be processed. This is useful because when modern Xcode creates a stub xcode project for a Program, it has no uproject to use with "-game -project", so the GPF code would process all .Target.cs files, which can lead to warnings about missing AdditionalPluginDirectories (build machines may not fully sync a stream, and processing every Target can cause unexpected problems)

diff --git a/Engine/Source/Programs/UnrealBuildTool/Modes/GenerateProjectFilesMode.cs b/Engine/Source/Programs/UnrealBuildTool/Modes/GenerateProjectFilesMode.cs
--- a/Engine/Source/Programs/UnrealBuildTool/Modes/GenerateProjectFilesMode.cs
+++ b/Engine/Source/Programs/UnrealBuildTool/Modes/GenerateProjectFilesMode.cs
@@ -172,6 +172,14 @@ namespace UnrealBuildTool
 			Logger.LogDebug("---   SDK INFO END   ---");
 			Logger.LogDebug("");
 
+			// look for a single target name param
+			string? SingleTargetName = null;
+			if (Arguments.HasValue("-SingleTarget="))
+			{
+				SingleTargetName = Arguments.GetString("-SingleTarget=");
+			}
+
+
 			// Create each project generator and run it
 			Dictionary<ProjectFileFormat, ProjectFileGenerator> Generators = new();
 			foreach (ProjectFileFormat ProjectFileFormat in ProjectFileFormats.Distinct())
@@ -229,6 +237,9 @@ namespace UnrealBuildTool
 					default:
 						throw new BuildException("Unhandled project file type '{0}", ProjectFileFormat);
 				}
+				// remember if we only wanted a single target (similar to -game -project, except usable with progarms without uprojects)
+				Generator.SingleTargetName = SingleTargetName;
+				
 				Generators[ProjectFileFormat] = Generator;
 			}
 
diff --git a/Engine/Source/Programs/UnrealBuildTool/Platform/IOS/IOSToolChain.cs b/Engine/Source/Programs/UnrealBuildTool/Platform/IOS/IOSToolChain.cs
--- a/Engine/Source/Programs/UnrealBuildTool/Platform/IOS/IOSToolChain.cs
+++ b/Engine/Source/Programs/UnrealBuildTool/Platform/IOS/IOSToolChain.cs
@@ -1673,11 +1673,11 @@ namespace UnrealBuildTool
 					DirectoryReference? GeneratedProjectFile;
 					if (AppName == "UnrealGame" || AppName == "UnrealClient" || Target.ProjectFile == null || Target.ProjectFile.IsUnderDirectory(Unreal.EngineDirectory))
 					{
-						GenerateProjectFiles(Target.ProjectFile, new string[] { "-platforms=" + (Target.Platform == UnrealTargetPlatform.IOS ? "IOS" : "TVOS"), "-NoIntellIsense", (Target.Platform == UnrealTargetPlatform.IOS ? "-iosdeployonly" : "-tvosdeployonly"), "-ignorejunk", (Target.bForDistribution ? "-distribution" : "-development"), "-bundleID=" + BundleID, "-includetemptargets", "-appname=" + AppName }, Logger, out GeneratedProjectFile);
+						GenerateProjectFiles(Target.ProjectFile, new string[] { "-platforms=" + (Target.Platform == UnrealTargetPlatform.IOS ? "IOS" : "TVOS"), "-NoIntellIsense", (Target.Platform == UnrealTargetPlatform.IOS ? "-iosdeployonly" : "-tvosdeployonly"), "-ignorejunk", (Target.bForDistribution ? "-distribution" : "-development"), "-bundleID=" + BundleID, "-includetemptargets", "-appname=" + AppName }, null, Logger, out GeneratedProjectFile);
 					}
 					else
 					{
-						GenerateProjectFiles(Target.ProjectFile, new string[] { "-platforms=" + (Target.Platform == UnrealTargetPlatform.IOS ? "IOS" : "TVOS"), "-NoIntellIsense", (Target.Platform == UnrealTargetPlatform.IOS ? "-iosdeployonly" : "-tvosdeployonly"), "-ignorejunk", (Target.bForDistribution ? "-distribution" : "-development"), String.Format("-project={0}", Target.ProjectFile), "-game", "-bundleID=" + BundleID, "-includetemptargets" }, Logger, out GeneratedProjectFile);
+						GenerateProjectFiles(Target.ProjectFile, new string[] { "-platforms=" + (Target.Platform == UnrealTargetPlatform.IOS ? "IOS" : "TVOS"), "-NoIntellIsense", (Target.Platform == UnrealTargetPlatform.IOS ? "-iosdeployonly" : "-tvosdeployonly"), "-ignorejunk", (Target.bForDistribution ? "-distribution" : "-development"), String.Format("-project={0}", Target.ProjectFile), "-game", "-bundleID=" + BundleID, "-includetemptargets" }, null, Logger, out GeneratedProjectFile);
 					}
 					Logger.LogInformation("Generated {0}!", GeneratedProjectFile?.FullName ?? "NONE");
 					// Make sure it exists
diff --git a/Engine/Source/Programs/UnrealBuildTool/ProjectFiles/ProjectFileGenerator.cs b/Engine/Source/Programs/UnrealBuildTool/ProjectFiles/ProjectFileGenerator.cs
--- a/Engine/Source/Programs/UnrealBuildTool/ProjectFiles/ProjectFileGenerator.cs
+++ b/Engine/Source/Programs/UnrealBuildTool/ProjectFiles/ProjectFileGenerator.cs
@@ -351,6 +351,11 @@ namespace UnrealBuildTool
 		[XmlConfigFile(Name = "Configurations")]
 		string[]? ConfigurationNames = null;
 
+		/// <summary>
+		/// If set, strip out other found Target.cs files
+		/// </summary>
+		public string? SingleTargetName = null;
+
 		/// <summary>
 		/// If true, this generator wants to have a project for each target type (UnrealGame, UnrealEditor, etc) and has only the straight configs (Debug, Development)
 		/// </summary>
@@ -966,6 +971,13 @@ namespace UnrealBuildTool
 			// belong to platforms we're generating project files for.
 			List<FileReference> AllTargetFiles = DiscoverTargets(AllGameProjects, Logger, OnlyGameProject, SupportedPlatforms, bIncludeEngineSource, bIncludeTempTargets);
 
+			// if we only want one target, remove the others
+			if (SingleTargetName != null)
+			{
+				AllTargetFiles.RemoveAll(x => !x.GetFileNameWithoutAnyExtensions().Equals(SingleTargetName, StringComparison.OrdinalIgnoreCase) && 
+				!x.GetFileNameWithoutAnyExtensions().Equals(EngineProjectFileNameBase, StringComparison.OrdinalIgnoreCase));
+			}
+
 			// Sort the targets by name. When we have multiple targets of a given type for a project, we'll use the order to determine which goes in the primary project file (so that client names with a suffix will go into their own project).
 			AllTargetFiles = AllTargetFiles.OrderBy(x => x.FullName, StringComparer.OrdinalIgnoreCase).ToList();
 
diff --git a/Engine/Source/Programs/UnrealBuildTool/ProjectFiles/Xcode/XcodeProjectFileGenerator.cs b/Engine/Source/Programs/UnrealBuildTool/ProjectFiles/Xcode/XcodeProjectFileGenerator.cs
--- a/Engine/Source/Programs/UnrealBuildTool/ProjectFiles/Xcode/XcodeProjectFileGenerator.cs
+++ b/Engine/Source/Programs/UnrealBuildTool/ProjectFiles/Xcode/XcodeProjectFileGenerator.cs
@@ -57,11 +57,6 @@ namespace UnrealBuildTool
 		/// </summary>
 		public static FileReference? SingleGameProject = null;
 
-		/// <summary>
-		/// If set, only write out one target/project 
-		/// </summary>
-		string? SingleTargetName = null;
-
 		/// <summary>
 		/// Shared file that the project agnostic projects can point to to get content only projects working nicely in Xcode
 		/// </summary>
@@ -102,11 +97,6 @@ namespace UnrealBuildTool
 				AppName = CommandLine.GetString("-appname=");
 			}
 
-			if (CommandLine.HasValue("-SingleTarget="))
-			{
-				SingleTargetName = CommandLine.GetString("-SingleTarget=");
-			}
-
 			// make sure only one Target writes the file
 			lock (ContentOnlySettingsFile)
 			{
diff --git a/Engine/Source/Programs/UnrealBuildTool/ToolChain/AppleToolChain.cs b/Engine/Source/Programs/UnrealBuildTool/ToolChain/AppleToolChain.cs
--- a/Engine/Source/Programs/UnrealBuildTool/ToolChain/AppleToolChain.cs
+++ b/Engine/Source/Programs/UnrealBuildTool/ToolChain/AppleToolChain.cs
@@ -537,7 +537,7 @@ namespace UnrealBuildTool
 
 		#region Stub Xcode Projects
 
-		internal static bool GenerateProjectFiles(FileReference? ProjectFile, string[] Arguments, ILogger Logger, out DirectoryReference? XcodeProjectFile)
+		internal static bool GenerateProjectFiles(FileReference? ProjectFile, string[] Arguments, string? SingleTargetName, ILogger Logger, out DirectoryReference? XcodeProjectFile)
 		{
 			ProjectFileGenerator.bGenerateProjectFiles = true;
 			try
@@ -550,6 +550,8 @@ namespace UnrealBuildTool
 				PlatformProjectGenerators.RegisterPlatformProjectGenerator(UnrealTargetPlatform.TVOS, new TVOSProjectGenerator(CmdLine, Logger), Logger);
 
 				XcodeProjectFileGenerator Generator = new XcodeProjectFileGenerator(ProjectFile, CmdLine);
+				// this could be improved if ProjectFileGenerator constructor took a Arguments param, and it could parse it there instead of the GenerateProjectFilesMode
+				Generator.SingleTargetName = SingleTargetName;
 				// don't need the editor data since these are stub projects
 				bool bSucces = Generator.GenerateProjectFiles(PlatformProjectGenerators, Arguments, bCacheDataForEditor: false, Logger);
 				XcodeProjectFile = Generator.XCWorkspace;
@@ -612,7 +614,7 @@ namespace UnrealBuildTool
 			// we need to be in Engine/Source for some build.cs files
 			string CurrentCWD = Environment.CurrentDirectory;
 			Environment.CurrentDirectory = DirectoryReference.Combine(Unreal.EngineDirectory, "Source").FullName;
-			GenerateProjectFiles(UProjectFile, Options.ToArray(), Logger, out GeneratedProjectFile);
+			GenerateProjectFiles(UProjectFile, Options.ToArray(), TargetName, Logger, out GeneratedProjectFile);
 			Environment.CurrentDirectory = CurrentCWD;
 		}
 
