- Handle manual codesigning with Modern Xcode (using all new settings, does not use the "green check" provisions in IOS project settings)
- Added support for AppStoreConnect API for automatic code signing (with command line builds, packaging from editor, etc) that doesn't need everyone to be signed in to Xcode (still need to sign in if using Xcode IDE to build/run)
- Fixed Modern to always create the IOS plist in UBT ApplePostBuild step
- Updated a bunch of ini setting names to remove Modern, given Modern will be standard in the future
- Added some settings XcodeProject settings in the editor, and updated some tooltips/category names
- Moved some functions around to clean up
- Added some debug code to catch errors

diff --git a/Engine/Config/BaseEngine.ini b/Engine/Config/BaseEngine.ini
--- a/Engine/Config/BaseEngine.ini
+++ b/Engine/Config/BaseEngine.ini
@@ -3012,15 +3012,21 @@ bEnableRayTracing=false
 
 [/Script/MacTargetPlatform.XcodeProjectSettings]
 bUseModernXcode=false
-bUseModernCodeSigning=true
-ModernSigningTeam=
-ModernSigningPrefix=com.YourCompany
-ModernBundleIdentifier=$(UE_SIGNING_PREFIX).$(UE_PRODUCT_NAME_STRIPPED)
+bUseAutomaticCodeSigning=true
+bMacSignToRunLocally=true
+CodeSigningTeam=
+CodeSigningPrefix=com.YourCompany
+BundleIdentifier=$(UE_SIGNING_PREFIX).$(UE_PRODUCT_NAME_STRIPPED)
 AppCategory=public.app-category.games
 TemplateMacPlist=(FilePath="/Game/Build/Mac/Resources/Info.Template.plist")
 TemplateIOSPlist=(FilePath="/Game/Build/IOS/UBTGenerated/Info.Template.plist")
 PremadeMacEntitlements=(FilePath="/Game/Build/Mac/Resources/Sandbox.Server.entitlements")
 ShippingSpecificMacEntitlements=(FilePath="/Game/Build/Mac/Resources/Sandbox.NoNet.entitlements")
+; The following are used only with Manual codesigning
+MacSigningIdentity=Apple Development
+IOSSigningIdentity=Apple Development
+IOSProvisioningProfile=
+TVOSProvisioningProfile=
 
 [HMDPluginPriority]
 ; Since SteamVR also works with the Oculus Rift, give priority to the native Oculus plugin before trying SteamVR.
diff --git a/Engine/Source/Developer/Mac/MacPlatformEditor/Private/XcodeProjectSettingsDetailsCustomization.cpp b/Engine/Source/Developer/Mac/MacPlatformEditor/Private/XcodeProjectSettingsDetailsCustomization.cpp
--- a/Engine/Source/Developer/Mac/MacPlatformEditor/Private/XcodeProjectSettingsDetailsCustomization.cpp
+++ b/Engine/Source/Developer/Mac/MacPlatformEditor/Private/XcodeProjectSettingsDetailsCustomization.cpp
@@ -29,7 +29,7 @@ void FXcodeProjectSettingsDetailsCustomization::CustomizeDetails(IDetailLayoutBu
     PremadeMacEntitlements = DetailLayout.GetProperty(GET_MEMBER_NAME_CHECKED(UXcodeProjectSettings, PremadeMacEntitlements));
     ShippingEntitlements = DetailLayout.GetProperty(GET_MEMBER_NAME_CHECKED(UXcodeProjectSettings, ShippingSpecificMacEntitlements));
     
-    IDetailCategoryBuilder& PlistCategory = DetailLayout.EditCategory(TEXT("Info.plist"));
+    IDetailCategoryBuilder& PlistCategory = DetailLayout.EditCategory(TEXT("Plist Files"));
     PlistCategory.AddCustomRow(LOCTEXT("InfoPlist", "Info.plist"), false)
     .WholeRowWidget
     [
@@ -45,7 +45,7 @@ void FXcodeProjectSettingsDetailsCustomization::CustomizeDetails(IDetailLayoutBu
          ]
     ];
     
-    IDetailCategoryBuilder& ShipEntitlementCategory = DetailLayout.EditCategory(TEXT("Entitlement"));
+    IDetailCategoryBuilder& ShipEntitlementCategory = DetailLayout.EditCategory(TEXT("Entitlements"));
     ShipEntitlementCategory.AddCustomRow(LOCTEXT("Entitlement", "Entitlement"), false)
     .WholeRowWidget
     [
@@ -72,18 +72,22 @@ void FXcodeProjectSettingsDetailsCustomization::CustomizeDetails(IDetailLayoutBu
             {
                 SortOrder = 0;
             }
-            else if(CategoryName == "Info.plist")
+            else if(CategoryName == "Plist Files")
             {
                 SortOrder = 1;
             }
-            else if(CategoryName == "Entitlement")
-            {
-                SortOrder = 2;
-            }
+			else if(CategoryName == "Entitlements")
+			{
+				SortOrder = 2;
+			}
+			else if(CategoryName == "Code Signing")
+			{
+				SortOrder = 3;
+			}
             else
             {
                 // Unknown category, should explicitly set order
-                ensure(false);
+                ensureMsgf(false, TEXT("Unknown category %s in XcodeProjectSttings"), *CategoryName.ToString());
                 SortOrder = 999;
             }
 
diff --git a/Engine/Source/Developer/Mac/MacTargetPlatform/Classes/XcodeProjectSettings.h b/Engine/Source/Developer/Mac/MacTargetPlatform/Classes/XcodeProjectSettings.h
--- a/Engine/Source/Developer/Mac/MacTargetPlatform/Classes/XcodeProjectSettings.h
+++ b/Engine/Source/Developer/Mac/MacTargetPlatform/Classes/XcodeProjectSettings.h
@@ -18,66 +18,120 @@ class MACTARGETPLATFORM_API UXcodeProjectSettings
     : public UObject
 {
 public:
-    
-    GENERATED_UCLASS_BODY()
-    
-    /**
-     * Enable modernized Xcode, when building from Xcode, use native Xcode for bundle generation and archiving instead of UBT
-     */
-    UPROPERTY(EditAnywhere, config, Category=Xcode, meta = (DisplayName = "Modernized Xcode"))
-    bool bUseModernXcode;
-    
-    /**
-     * Enable native Xcode code signing
-     */
-    UPROPERTY(EditAnywhere, config, Category=Xcode, meta = (EditCondition="bUseModernXcode", DisplayName = "Use Modern Code Signing"))
-    bool bUseModernCodeSigning;
-    
-    /**
-     * Team ID used for native Xcode code signing
-     */
-    UPROPERTY(EditAnywhere, config, Category=Xcode, meta = (EditCondition="bUseModernXcode && bUseModernCodeSigning", DisplayName = "Modern Code Sign Team"))
-    FString ModernSigningTeam;
-    
-    /**
-     * Bundle ID prefix used for native Xcode code signing
-     */
-    UPROPERTY(EditAnywhere, config, Category=Xcode, meta = (EditCondition="bUseModernXcode && bUseModernCodeSigning", DisplayName = "Bundle ID Prefix"))
-    FString ModernSigningPrefix;
-    
-    /**
-     * Bundle ID used for nativr Xcode code signing
-     */
-    UPROPERTY(EditAnywhere, config, Category=Xcode, meta = (EditCondition="bUseModernXcode && bUseModernCodeSigning", DisplayName = "Bundle ID"))
-    FString ModernBundleIdentifier;
-    
-    /**
-     * The App Category that will be used for Mac App Store submission
-     */
-    UPROPERTY(EditAnywhere, config, Category=Xcode, meta = (EditCondition="bUseModernXcode", DisplayName = "App Category"))
-    FString AppCategory;
-    
-    /**
-     * The template info.plist used for Mac game targets
-     */
-    UPROPERTY(EditAnywhere, config, Category="Info.plist", meta = (EditCondition="bUseModernXcode", DisplayName = "Mac Target Info.plist"))
-    FFilePath TemplateMacPlist;
-    
-    /**
-     * The template info.plist used for iOS game targets
-     */
-    UPROPERTY(EditAnywhere, config, Category="Info.plist", meta = (EditCondition="bUseModernXcode", DisplayName = "iOS Target Info.plist"))
-    FFilePath TemplateIOSPlist;
-    
-    /**
-     * The premade entitlement file used for development Mac builds
-     */
-    UPROPERTY(EditAnywhere, config, Category="Entitlement", meta = (EditCondition="bUseModernXcode", DisplayName = "Development Mac Entitlement"))
-    FFilePath PremadeMacEntitlements;
-    
-    /**
-     * The premade entitlement file used for shipping Mac builds
-     */
-    UPROPERTY(EditAnywhere, config, Category="Entitlement", meta = (EditCondition="bUseModernXcode", DisplayName = "Shipping Mac Entitlement"))
-    FFilePath ShippingSpecificMacEntitlements;
+	
+	GENERATED_UCLASS_BODY()
+	
+	/**
+	 * Enable modernized Xcode, when building from Xcode, use native Xcode for bundle generation and archiving instead of UBT
+	 */
+	UPROPERTY(EditAnywhere, config, Category=Xcode, meta = (DisplayName = "Modernized Xcode"))
+	bool bUseModernXcode;
+	
+	/**
+	 * Team ID used for native Xcode code signing. This must be the 10 letters/numbers ID found in Membership Details tab found in https://developer.apple.com/account
+	 */
+	UPROPERTY(EditAnywhere, config, Category=Xcode, meta = (EditCondition="bUseModernXcode", DisplayName = "Apple Dev Account Team ID"))
+	FString CodeSigningTeam;
+	
+	/**
+	 * Bundle ID used for nativr Xcode code signing
+	 */
+	UPROPERTY(EditAnywhere, config, Category=Xcode, meta = (EditCondition="bUseModernXcode", DisplayName = "Bundle ID"))
+	FString BundleIdentifier;
+	
+	/**
+	 * Bundle ID prefix used for native Xcode code signing. This is only needed if you use the default, pieced-together Bundle ID above. If you specify a full Bundle ID, you can ignore this field.
+	 */
+	UPROPERTY(EditAnywhere, config, Category=Xcode, meta = (EditCondition="bUseModernXcode", DisplayName = "Bundle ID Prefix"))
+	FString CodeSigningPrefix;
+	
+	/**
+	 * The App Category that will be used for Apple App Store submissions
+	 */
+	UPROPERTY(EditAnywhere, config, Category=Xcode, meta = (EditCondition="bUseModernXcode", DisplayName = "App Category"))
+	FString AppCategory;
+	
+	/**
+	 * The template info.plist used for Mac game targets
+	 */
+	UPROPERTY(EditAnywhere, config, Category="Plist Files", meta = (EditCondition="bUseModernXcode", DisplayName = "Mac: Info.plist Template"))
+	FFilePath TemplateMacPlist;
+	
+	/**
+	 * The template info.plist used for iOS game targets
+	 */
+	UPROPERTY(EditAnywhere, config, Category="Plist Files", meta = (EditCondition="bUseModernXcode", DisplayName = "IOS / TVOS: Info.plist Template"))
+	FFilePath TemplateIOSPlist;
+	
+	/**
+	 * The premade entitlement file used for development Mac builds
+	 */
+	UPROPERTY(EditAnywhere, config, Category="Entitlements", meta = (EditCondition="bUseModernXcode", DisplayName = "Nac: Development Entitlements"))
+	FFilePath PremadeMacEntitlements;
+	
+	/**
+	 * The premade entitlement file used for shipping Mac builds
+	 */
+	UPROPERTY(EditAnywhere, config, Category="Entitlements", meta = (EditCondition="bUseModernXcode", DisplayName = "Mac: Shipping Entitlements"))
+	FFilePath ShippingSpecificMacEntitlements;
+	
+	/**
+	 * Enable native Xcode code signing
+	 */
+	UPROPERTY(EditAnywhere, config, Category="Code Signing", meta = (EditCondition="bUseModernXcode"))
+	bool bUseAutomaticCodeSigning;
+	
+	/**
+	 * If true, Mac will sign to run locally. Running on another Mac may bring up a dialog preventing running the app. If this and Use Automatic Code Signing are both false, you will need a certificate installed
+	 */
+	UPROPERTY(EditAnywhere, config, Category="Code Signing", meta = (EditCondition="bUseModernXcode"), DisplayName="Mac: Sign To Run Locally")
+	bool bMacSignToRunLocally;
+	
+	/**
+	 * The name (prefix or full) of the certificate to use for Mac code signing.
+	 */
+	UPROPERTY(EditAnywhere, config, Category="Code Signing", meta = (EditCondition="bUseModernXcode && !bUseAutomaticCodeSigning && !bMacSignToRunLocally"), DisplayName = "Mac: Signing Identity")
+	FString MacSigningIdentity;
+	
+	/**
+	 * The name (prefix or full) of the certificate to use for IOS and TVOS code signing
+	 */
+	UPROPERTY(EditAnywhere, config, Category="Code Signing", meta = (EditCondition="bUseModernXcode && !bUseAutomaticCodeSigning", DisplayName = "IOS / TVOS: Signing Identity"))
+	FString IOSSigningIdentity;
+	
+	/**
+	 * The path to a .mobileprovision file to use for signing for IOS. Alternatively, if it's a single name or UUID (no .mobileprovision extension), it will use this as the name/UUID of an already installed provision to sign with
+	 */
+	UPROPERTY(EditAnywhere, config, Category="Code Signing", meta = (EditCondition="bUseModernXcode && !bUseAutomaticCodeSigning", DisplayName = "IOS: Provisioning Profile"))
+	FFilePath IOSProvisioningProfile;
+	
+	/**
+	 * The path to a .mobileprovision file to use for signing for TVOS. Alternatively, if it's a single name or UUID (no .mobileprovision extension), it will use this as the name/UUID of an already installed provision to sign with
+	 */
+	UPROPERTY(EditAnywhere, config, Category="Code Signing", meta = (EditCondition="bUseModernXcode && !bUseAutomaticCodeSigning", DisplayName = "TVOS: Provisioning Profile"))
+	FFilePath TVOSProvisioningProfile;
+	
+	/**
+	 * If true, use AppStore Connect authentication for commandline builds. This allows for automatic codesigning functionality without needing to be signed in to Xcode. See the App Store Connect API section of the Keys tab in your Users and Access page in Apple dev center.
+	 */
+	UPROPERTY(EditAnywhere, config, Category="Code Signing", meta = (EditCondition="bUseModernXcode && bUseAutomaticCodeSigning"))
+	bool bUseAppStoreConnect;
+	
+	/**
+	 * The Issuer ID for your App Store Connect API
+	 */
+	UPROPERTY(EditAnywhere, config, Category="Code Signing", meta = (EditCondition="bUseModernXcode && bUseAutomaticCodeSigning && bUseAppStoreConnect", DisplayName = "App Store Connect: Issuer ID"))
+	FString AppStoreConnectIssuerID;
+	
+	/**
+	 * The Key ID for your App Store Connect generated API key, a 32 hex-character string, including dashes
+	 */
+	UPROPERTY(EditAnywhere, config, Category="Code Signing", meta = (EditCondition="bUseModernXcode && bUseAutomaticCodeSigning && bUseAppStoreConnect", DisplayName = "App Store Connect: Issuer ID"))
+	FString AppStoreConnectKeyID;
+	
+	/**
+	 * The path to the downloaded .p8 file shared with your team
+	 */
+	UPROPERTY(EditAnywhere, config, Category="Code Signing", meta = (EditCondition="bUseModernXcode && bUseAutomaticCodeSigning && bUseAppStoreConnect", DisplayName = "App Store Connect: Key File"))
+	FString AppStoreConnectKeyPath;
 };
diff --git a/Engine/Source/Programs/AutomationTool/IOS/IOSPlatform.Automation.cs b/Engine/Source/Programs/AutomationTool/IOS/IOSPlatform.Automation.cs
--- a/Engine/Source/Programs/AutomationTool/IOS/IOSPlatform.Automation.cs
+++ b/Engine/Source/Programs/AutomationTool/IOS/IOSPlatform.Automation.cs
@@ -804,24 +804,7 @@ public class IOSPlatform : ApplePlatform
 		string SigningCertificate;
 		string TeamUUID;
 		bool bAutomaticSigning;
-		if (AppleExports.UseModernXcode(Params.RawProjectPath))
-		{
-			ConfigHierarchy Ini = ConfigCache.ReadHierarchy(ConfigHierarchyType.Engine, Params.RawProjectPath.Directory!, UnrealTargetPlatform.IOS);
-
-			Ini.TryGetValue("/Script/MacTargetPlatform.XcodeProjectSettings", "bUseModernCodeSigning", out bAutomaticSigning);
-			Ini.TryGetValue("/Script/MacTargetPlatform.XcodeProjectSettings", "ModernSigningTeam", out TeamUUID);
-			//			Ini.TryGetValue("/Script/MacTargetPlatform.XcodeProjectSettings", "ModernSigningPrefix", out SigningPrefix);
-			if (!bAutomaticSigning)
-			{
-				throw new AutomationException("Currently only Modern (automatic) codesigning is supported with Modern xcode");
-			}
-			SigningCertificate = null;
-			MobileProvision = null;
-		}
-		else
-		{
-			GetProvisioningData(Params.RawProjectPath, Params.Distribution, out MobileProvision, out SigningCertificate, out TeamUUID, out bAutomaticSigning);
-		}
+		GetProvisioningData(Params.RawProjectPath, Params.Distribution, out MobileProvision, out SigningCertificate, out TeamUUID, out bAutomaticSigning);
 
 		//@TODO: We should be able to use this code on both platforms, when the following issues are sorted:
 		//   - Raw executable is unsigned & unstripped (need to investigate adding stripping to IPP)
diff --git a/Engine/Source/Programs/UnrealBuildTool/Platform/Apple/AppleExports.cs b/Engine/Source/Programs/UnrealBuildTool/Platform/Apple/AppleExports.cs
--- a/Engine/Source/Programs/UnrealBuildTool/Platform/Apple/AppleExports.cs
+++ b/Engine/Source/Programs/UnrealBuildTool/Platform/Apple/AppleExports.cs
@@ -3,6 +3,7 @@
 using System;
 using EpicGames.Core;
 using Microsoft.Extensions.Logging;
+using UnrealBuildBase;
 
 namespace UnrealBuildTool
 {
@@ -44,6 +45,32 @@ namespace UnrealBuildTool
 			return bUseModernXcode;
 		}
 
+		/// <summary>
+		///  This is a FilePath from UE settings, like /Game/Foo/Bar.txt
+		/// </summary>
+		/// <param name="ProductDirectory">Directory to use for /Game paths</param>
+		/// <param name="FilePath"></param>
+		/// <returns></returns>
+		public static FileReference ConvertFilePath(DirectoryReference? ProductDirectory, string FilePath)
+		{
+			// for FilePath params, pull the path out of the struct
+			if (FilePath.StartsWith("(FilePath=", StringComparison.OrdinalIgnoreCase))
+			{
+				FilePath = ConfigHierarchy.GetStructEntry(FilePath, "FilePath", false)!;
+			}
+
+			if (FilePath.StartsWith("/Engine/", StringComparison.OrdinalIgnoreCase))
+			{
+				return FileReference.Combine(Unreal.EngineDirectory, FilePath.Substring(8));
+			}
+			else if (ProductDirectory != null && FilePath.StartsWith("/Game/", StringComparison.OrdinalIgnoreCase))
+			{
+				return FileReference.Combine(ProductDirectory, FilePath.Substring(6));
+			}
+
+			return new FileReference(FilePath);
+		}
+
 		/// <summary>
 		/// Different ways that xcodebuild is run, so the scripts can behave appropriately
 		/// </summary>
@@ -94,6 +121,32 @@ namespace UnrealBuildTool
 				return 1;
 			}
 
+			// look for the special app store connect key information
+			ConfigHierarchy SharedPlatformIni = ConfigCache.ReadHierarchy(ConfigHierarchyType.Engine, ProjectFile?.Directory, UnrealTargetPlatform.Mac);
+			bool bUseAutomaticCodeSigning;
+			SharedPlatformIni.TryGetValue("/Script/MacTargetPlatform.XcodeProjectSettings", "bUseAutomaticCodeSigning", out bUseAutomaticCodeSigning);
+			if (bUseAutomaticCodeSigning)
+			{
+				ExtraOptions += " -allowProvisioningUpdates";
+
+				// handle AppStore Connect settings
+				bool bUseAppStoreConnect;
+				SharedPlatformIni.TryGetValue("/Script/MacTargetPlatform.XcodeProjectSettings", "bUseAppStoreConnect", out bUseAppStoreConnect);
+				if (bUseAppStoreConnect)
+				{
+					string? IssuerID, KeyID, KeyPath;
+					if (SharedPlatformIni.TryGetValue("/Script/MacTargetPlatform.XcodeProjectSettings", "AppStoreConnectIssuerID", out IssuerID) &&
+						SharedPlatformIni.TryGetValue("/Script/MacTargetPlatform.XcodeProjectSettings", "AppStoreConnectKeyID", out KeyID) &&
+						SharedPlatformIni.TryGetValue("/Script/MacTargetPlatform.XcodeProjectSettings", "AppStoreConnectKeyPath", out KeyPath))
+					{
+						FileReference KeyFile = ConvertFilePath(ProjectFile?.Directory, KeyPath);
+						ExtraOptions += $" -authenticationKeyIssuerID {IssuerID}";
+						ExtraOptions += $" -authenticationKeyID {KeyID}";
+						ExtraOptions += $" -authenticationKeyPath \"{KeyFile}\"";
+					}
+				}
+			}
+
 			// run xcodebuild on the generated project to make the .app
 			string XcodeBuildAction = BuildMode == XcodeBuildMode.Distribute ? "archive" : "build";
 			return AppleExports.FinalizeAppWithXcode(GeneratedProjectFile!, Platform, TargetName, Configuration.ToString(), XcodeBuildAction, ExtraOptions + $" UE_XCODE_BUILD_MODE={BuildMode}", Logger);
diff --git a/Engine/Source/Programs/UnrealBuildTool/Platform/IOS/IOSToolChain.cs b/Engine/Source/Programs/UnrealBuildTool/Platform/IOS/IOSToolChain.cs
--- a/Engine/Source/Programs/UnrealBuildTool/Platform/IOS/IOSToolChain.cs
+++ b/Engine/Source/Programs/UnrealBuildTool/Platform/IOS/IOSToolChain.cs
@@ -1206,17 +1206,8 @@ namespace UnrealBuildTool
 				PostBuildSyncAction.WorkingDirectory = Unreal.EngineSourceDirectory;
 				PostBuildSyncAction.PrerequisiteItems.Add(Executable);
 				PostBuildSyncAction.PrerequisiteItems.UnionWith(OutputFiles);
-				if (bUseModernXcode)
-				{
-					// @todo: do we need one per Target for IOS? Client? I dont think so for Modern
-					FileReference PlistFile = FileReference.Combine(GetActualProjectDirectory(ProjectFile), "Build/IOS/UBTGenerated/Info.Template.plist");
-					PostBuildSyncAction.ProducedItems.Add(FileItem.GetItemByFileReference(PlistFile));
-				}
-				else
-				{
-					PostBuildSyncAction.ProducedItems.Add(FileItem.GetItemByFileReference(GetStagedExecutablePath(Executable.Location, Target.Name)));
-					PostBuildSyncAction.DeleteItems.UnionWith(PostBuildSyncAction.ProducedItems);
-				}
+				PostBuildSyncAction.ProducedItems.Add(FileItem.GetItemByFileReference(GetStagedExecutablePath(Executable.Location, Target.Name)));
+				PostBuildSyncAction.DeleteItems.UnionWith(PostBuildSyncAction.ProducedItems);
 				PostBuildSyncAction.StatusDescription = "Executing PostBuildSync";
 				PostBuildSyncAction.bCanExecuteRemotely = false;
 
@@ -1373,18 +1364,6 @@ namespace UnrealBuildTool
 			FileReference.WriteAllText(MissingPlistPath, Plist);
 		}
 
-		/// <summary>
-		/// If the project is a UnrealGame project, Target.ProjectDirectory refers to the engine dir, not the actual dir of the project. So this method gets the 
-		/// actual directory of the project whether it is a UnrealGame project or not.
-		/// </summary>
-		/// <returns>The actual project directory.</returns>
-		/// <param name="ProjectFile">The path to the project file</param>
-		private static DirectoryReference GetActualProjectDirectory(FileReference? ProjectFile)
-		{
-			DirectoryReference ProjectDirectory = (ProjectFile == null ? DirectoryReference.FromString(UnrealBuildTool.GetRemoteIniPath())! : DirectoryReference.FromFile(ProjectFile)!);
-			return ProjectDirectory;
-		}
-
 		private static void GenerateFrameworkWrapperIfNonexistent(IOSPostBuildSyncTarget Target, ILogger Logger)
 		{
 			DirectoryReference ProjectDirectory = GetActualProjectDirectory(Target.ProjectFile);
diff --git a/Engine/Source/Programs/UnrealBuildTool/Platform/IOS/UEDeployIOS.cs b/Engine/Source/Programs/UnrealBuildTool/Platform/IOS/UEDeployIOS.cs
--- a/Engine/Source/Programs/UnrealBuildTool/Platform/IOS/UEDeployIOS.cs
+++ b/Engine/Source/Programs/UnrealBuildTool/Platform/IOS/UEDeployIOS.cs
@@ -498,7 +498,7 @@ namespace UnrealBuildTool
 
 			// generate the Info.plist for future use
 			string BuildDirectory = ProjectDirectory + "/Build/IOS";
-			string IntermediateDirectory = (bIsUnrealGame ? InEngineDir : ProjectDirectory) + "/Intermediate/IOS";
+			string IntermediateDirectory = ProjectDirectory + "/Intermediate/IOS";
 			string PListFile = IntermediateDirectory + "/" + GameName + "-Info.plist"; ;
 			ProjectName = !String.IsNullOrEmpty(ProjectName) ? ProjectName : GameName;
 			VersionUtilities.BuildDirectory = BuildDirectory;
diff --git a/Engine/Source/Programs/UnrealBuildTool/ProjectFiles/Xcode/XcodeProject.File.cs b/Engine/Source/Programs/UnrealBuildTool/ProjectFiles/Xcode/XcodeProject.File.cs
--- a/Engine/Source/Programs/UnrealBuildTool/ProjectFiles/Xcode/XcodeProject.File.cs
+++ b/Engine/Source/Programs/UnrealBuildTool/ProjectFiles/Xcode/XcodeProject.File.cs
@@ -461,9 +461,14 @@ namespace UnrealBuildTool.XcodeProjectXcconfig
 			}
 		}
 
-		public static string GetRootGroupGuid(Dictionary<string, XcodeFileGroup> GroupsDict)
+		public static string GetRootGroupGuid(Dictionary<string, XcodeFileGroup> GroupsDict, FileSystemReference ProjectFile)
 		{
-			return FindRootFileGroup(GroupsDict)!.GroupGuid;
+			XcodeFileGroup? RootGroup = FindRootFileGroup(GroupsDict);
+			if (RootGroup == null)
+			{
+				throw new BuildException($"The project '{ProjectFile}' had no root group, no files have been added");
+			}
+			return RootGroup.GroupGuid;
 		}
 
 		private static XcodeFileGroup? FindRootFileGroup(Dictionary<string, XcodeFileGroup> GroupsDict)
diff --git a/Engine/Source/Programs/UnrealBuildTool/ProjectFiles/Xcode/XcodeProject.cs b/Engine/Source/Programs/UnrealBuildTool/ProjectFiles/Xcode/XcodeProject.cs
--- a/Engine/Source/Programs/UnrealBuildTool/ProjectFiles/Xcode/XcodeProject.cs
+++ b/Engine/Source/Programs/UnrealBuildTool/ProjectFiles/Xcode/XcodeProject.cs
@@ -162,31 +162,12 @@ namespace UnrealBuildTool.XcodeProjectXcconfig
 		public MetadataItem(DirectoryReference ProductDirectory, DirectoryReference XcodeProject, ConfigHierarchy Ini, string Location, MetadataMode InMode, DirectoryReference CopyFromFolderIfNotFound)
 		{
 			// no extension means it's a .ini entry
-			if (String.IsNullOrEmpty(Path.GetExtension(Location)))
+			if (Path.GetExtension(Location).Length == 0)
 			{
 				string? FileLocation;
-				if (Ini.TryGetValue("/Script/MacTargetPlatform.XcodeProjectSettings", Location, out FileLocation) && !String.IsNullOrEmpty(FileLocation))
+				if (Ini.TryGetValue("/Script/MacTargetPlatform.XcodeProjectSettings", Location, out FileLocation) && FileLocation.Length > 0)
 				{
-					if (FileLocation.StartsWith("(FilePath=", StringComparison.OrdinalIgnoreCase))
-					{
-						FileLocation = ConfigHierarchy.GetStructEntry(FileLocation, "FilePath", false)!;
-						if (FileLocation.StartsWith("/Engine/"))
-						{
-							File = FileReference.Combine(Unreal.EngineDirectory, FileLocation.Substring(8));
-						}
-						else if (FileLocation.StartsWith("/Game/"))
-						{
-							File = FileReference.Combine(ProductDirectory, FileLocation.Substring(6));
-						}
-						else
-						{
-							File = new FileReference(FileLocation);
-						}
-					}
-					else
-					{
-						File = new FileReference(FileLocation);
-					}
+					File = AppleExports.ConvertFilePath(ProductDirectory, FileLocation);
 				}
 			}
 			else
@@ -355,10 +336,7 @@ namespace UnrealBuildTool.XcodeProjectXcconfig
 
 		// settings read from project configs
 		public IOSProjectSettings? IOSProjectSettings;
-		public IOSProvisioningData? IOSProvisioningData;
-
 		public TVOSProjectSettings? TVOSProjectSettings;
-		public TVOSProvisioningData? TVOSProvisioningData;
 
 		// Name of the product (usually the project name, but UE5.xcodeproj is actually UnrealGame product)
 		public string ProductName;
@@ -476,7 +454,7 @@ namespace UnrealBuildTool.XcodeProjectXcconfig
 			if (String.IsNullOrEmpty(BundleIdentifier))
 			{
 				ConfigHierarchy Ini = ConfigCache.ReadHierarchy(ConfigHierarchyType.Engine, UProjectFileLocation?.Directory, UnrealTargetPlatform.Mac);
-				Ini.GetString($"/Script/MacTargetPlatform.XcodeProjectSettings", "ModernBundleIdentifier", out BundleIdentifier);
+				Ini.GetString($"/Script/MacTargetPlatform.XcodeProjectSettings", "BundleIdentifier", out BundleIdentifier);
 			}
 			if (String.IsNullOrEmpty(BundleIdentifier))
 			{
@@ -533,20 +511,12 @@ namespace UnrealBuildTool.XcodeProjectXcconfig
 			{
 				IOSPlatform IOSPlatform = ((IOSPlatform)UEBuildPlatform.GetBuildPlatform(UnrealTargetPlatform.IOS));
 				IOSProjectSettings = IOSPlatform.ReadProjectSettings(UProjectFileLocation);
-				if (!bUseAutomaticSigning)
-				{
-					IOSProvisioningData = IOSPlatform.ReadProvisioningData(IOSProjectSettings, bForDistribution);
-				}
 			}
 
 			if (AllConfigs.Any(x => x.bSupportsTVOS))
 			{
 				TVOSPlatform TVOSPlatform = ((TVOSPlatform)UEBuildPlatform.GetBuildPlatform(UnrealTargetPlatform.TVOS));
 				TVOSProjectSettings = TVOSPlatform.ReadProjectSettings(UProjectFileLocation);
-				if (!bUseAutomaticSigning)
-				{
-					TVOSProvisioningData = TVOSPlatform.ReadProvisioningData(TVOSProjectSettings, bForDistribution);
-				}
 			}
 
 			return true;
@@ -557,7 +527,7 @@ namespace UnrealBuildTool.XcodeProjectXcconfig
 
 			// read setings from the configs, now that we have a project
 			ConfigHierarchy SharedPlatformIni = ConfigCache.ReadHierarchy(ConfigHierarchyType.Engine, UProjectFileLocation?.Directory, UnrealTargetPlatform.Mac);
-			SharedPlatformIni.TryGetValue("/Script/MacTargetPlatform.XcodeProjectSettings", "bUseModernCodeSigning", out bUseAutomaticSigning);
+			SharedPlatformIni.TryGetValue("/Script/MacTargetPlatform.XcodeProjectSettings", "bUseAutomaticCodeSigning", out bUseAutomaticSigning);
 
 			Metadata = new Metadata(ProductDirectory, XcodeProjectFileLocation, SharedPlatformIni, bSupportsMac, bSupportsIOS || bSupportsTVOS, Logger);
 		}
@@ -778,7 +748,7 @@ namespace UnrealBuildTool.XcodeProjectXcconfig
 			Project.FileCollection.AddFileReference(Xcconfig.Guid, Xcconfig.FileRef.MakeRelativeTo(XcodeProjectDirectory), "explicitFileType", "test.xcconfig", "\"<group>\"", "Xcconfigs");
 		}
 
-		public virtual void WriteXcconfigFile()
+		public virtual void WriteXcconfigFile(ILogger Logger)
 		{
 
 		}
@@ -788,8 +758,9 @@ namespace UnrealBuildTool.XcodeProjectXcconfig
 		/// </summary>
 		/// <param name="Content"></param>
 		/// <param name="Node"></param>
+		/// <param name="Logger"></param>
 		/// <param name="WrittenNodes"></param>
-		public static void WriteNodeAndReferences(StringBuilder Content, XcodeProjectNode Node, HashSet<XcodeProjectNode>? WrittenNodes = null)
+		public static void WriteNodeAndReferences(StringBuilder Content, XcodeProjectNode Node, ILogger Logger, HashSet<XcodeProjectNode>? WrittenNodes = null)
 		{
 			if (WrittenNodes == null)
 			{
@@ -798,14 +769,14 @@ namespace UnrealBuildTool.XcodeProjectXcconfig
 
 			// write the node into the xcode project file
 			Node.Write(Content);
-			Node.WriteXcconfigFile();
+			Node.WriteXcconfigFile(Logger);
 
 			foreach (XcodeProjectNode Reference in Node.References)
 			{
 				if (!WrittenNodes.Contains(Reference))
 				{
 					WrittenNodes.Add(Reference);
-					WriteNodeAndReferences(Content, Reference, WrittenNodes);
+					WriteNodeAndReferences(Content, Reference, Logger, WrittenNodes);
 				}
 			}
 		}
@@ -1564,7 +1535,51 @@ namespace UnrealBuildTool.XcodeProjectXcconfig
 			Content.WriteLine($"\t\t\tproductName = \"{UnrealData.ProductName}\";");
 		}
 
-		public override void WriteXcconfigFile()
+		private static Dictionary<string, string> PlistFileMap = new(StringComparer.OrdinalIgnoreCase);
+		private string GetPlistSigningName(string ProvisionSetting, ILogger Logger)
+		{
+			string? SigningName;
+			lock(PlistFileMap)
+			{
+				if (!PlistFileMap.TryGetValue(ProvisionSetting, out SigningName))
+				{
+					FileReference ProfileFile = AppleExports.ConvertFilePath(UnrealData.UProjectFileLocation?.Directory, ProvisionSetting);
+
+					// get mobile provision UUID, either directly or from the probision
+					if (ProfileFile.GetExtension() == ".mobileprovision")
+					{
+						// security will read the provision and dump out the text (plist) bits of the profile, and plutil will extract the UUID to stdout
+						string UUID = Utils.RunLocalProcessAndReturnStdOut("/bin/sh", $"-c 'security cms -D -i {ProfileFile.FullName.Replace(" ", "\\ ")} | plutil -extract UUID raw -'");
+
+						// make sure it's living in user's library
+						// (note: i couldn't find a way to install it without opening Xcode, so copying it works well enough)
+						DirectoryReference UserDir = new DirectoryReference(Environment.GetFolderPath(Environment.SpecialFolder.UserProfile));
+						FileReference InstalledProfileFile = FileReference.Combine(UserDir, "Library", "MobileDevice", "Provisioning Profiles", $"{UUID}.mobileprovision");
+
+						// copy if needed
+						if (!FileReference.Exists(InstalledProfileFile))
+						{
+							DirectoryReference.CreateDirectory(InstalledProfileFile.Directory);
+							FileReference.Copy(ProfileFile, InstalledProfileFile);
+							Logger.LogInformation("Copying project's provision '{SourceProvision}' to your libary: '{TargetProvision}'", ProfileFile, InstalledProfileFile);
+						}
+
+
+						SigningName = UUID;
+					}
+					else
+					{
+						SigningName = ProvisionSetting;
+					}
+
+					PlistFileMap[ProvisionSetting] = SigningName;
+				}
+			}
+
+			return SigningName!;
+		}
+
+		public override void WriteXcconfigFile(ILogger Logger)
 		{
 			// gather general, all-platform, data we are doing to put into the configs
 			UnrealBuildConfig BuildConfig = BuildConfigList!.BuildConfigs[0].Info;
@@ -1592,7 +1607,8 @@ namespace UnrealBuildTool.XcodeProjectXcconfig
 			ConfigHierarchy PlatformIni = ConfigCache.ReadHierarchy(ConfigHierarchyType.Engine, UnrealData.UProjectFileLocation?.Directory, Platform);
 
 			// settings for all platforms
-			bool bAutomaticSigning = true;
+			bool bAutomaticSigning;
+			bool bMacSignToRunLocally;
 			string? SigningTeam;
 			string? SigningPrefix;
 			string? AppCategory;
@@ -1600,19 +1616,23 @@ namespace UnrealBuildTool.XcodeProjectXcconfig
 			string SDKRoot;
 			string DeploymentTarget;
 			string DeploymentTargetKey;
-			string? SigningCert = null;
+			string? SigningIdentity = null;
 			string? ProvisioningProfile = null;
 			string? SupportedDevices = null;
 			string? MarketingVersion = null;
 			string? BundleIdentifier;
 			List<string> ExtraConfigLines = new();
 
-			// get codesigning info (all platforms)
-			PlatformIni.TryGetValue("/Script/MacTargetPlatform.XcodeProjectSettings", "bUseModernCodeSigning", out bAutomaticSigning);
-			PlatformIni.TryGetValue("/Script/MacTargetPlatform.XcodeProjectSettings", "ModernSigningTeam", out SigningTeam);
-			PlatformIni.TryGetValue("/Script/MacTargetPlatform.XcodeProjectSettings", "ModernSigningPrefix", out SigningPrefix);
+			// get signing settings
+			PlatformIni.TryGetValue("/Script/MacTargetPlatform.XcodeProjectSettings", "bUseAutomaticCodeSigning", out bAutomaticSigning);
+			PlatformIni.TryGetValue("/Script/MacTargetPlatform.XcodeProjectSettings", "bMacSignToRunLocally", out bMacSignToRunLocally);
+			PlatformIni.TryGetValue("/Script/MacTargetPlatform.XcodeProjectSettings", "CodeSigningTeam", out SigningTeam);
+			PlatformIni.TryGetValue("/Script/MacTargetPlatform.XcodeProjectSettings", "CodeSigningPrefix", out SigningPrefix);
+			PlatformIni.TryGetValue("/Script/MacTargetPlatform.XcodeProjectSettings", $"{Platform}ProvisioningProfile", out ProvisioningProfile);
+			PlatformIni.TryGetValue("/Script/MacTargetPlatform.XcodeProjectSettings", $"{Platform}SigningIdentity", out SigningIdentity);
 			PlatformIni.TryGetValue("/Script/MacTargetPlatform.XcodeProjectSettings", "AppCategory", out AppCategory);
 
+
 			if (Platform == UnrealTargetPlatform.Mac)
 			{
 				// editor vs game metadata
@@ -1649,12 +1669,6 @@ namespace UnrealBuildTool.XcodeProjectXcconfig
 					SupportedDevices = UnrealData.IOSProjectSettings!.RuntimeDevices;
 					DeploymentTarget = UnrealData.IOSProjectSettings.RuntimeVersion;
 
-					if (UnrealData.IOSProvisioningData != null && UnrealData.bWriteCodeSigningSettings)
-					{
-						ProvisioningProfile = UnrealData.IOSProvisioningData.MobileProvisionUUID;
-						SigningCert = UnrealData.IOSProvisioningData.SigningCertificate ?? (UnrealData.bForDistribution ? "Apple Distribution" : "Apple Developer");
-					}
-
 					// only iphone deals with orientation
 					List<string> SupportedOrientations = XcodeUtils.GetSupportedOrientations(PlatformIni);
 					ExtraConfigLines.Add($"INFOPLIST_KEY_UISupportedInterfaceOrientations = \"{String.Join(" ", SupportedOrientations)}\"");
@@ -1675,12 +1689,6 @@ namespace UnrealBuildTool.XcodeProjectXcconfig
 					DeploymentTargetKey = "TVOS_DEPLOYMENT_TARGET";
 					SupportedDevices = UnrealData.TVOSProjectSettings!.RuntimeDevices;
 					DeploymentTarget = UnrealData.TVOSProjectSettings.RuntimeVersion;
-
-					if (UnrealData.TVOSProvisioningData != null && UnrealData.bWriteCodeSigningSettings)
-					{
-						ProvisioningProfile = UnrealData.TVOSProvisioningData.MobileProvisionUUID;
-						SigningCert = UnrealData.TVOSProvisioningData.SigningCertificate ?? (UnrealData.bForDistribution ? "Apple Distribution" : "Apple Developer");
-					}
 				}
 			}
 
@@ -1767,22 +1775,43 @@ namespace UnrealBuildTool.XcodeProjectXcconfig
 				Xcconfig.AppendLine($"INFOPLIST_KEY_LSApplicationCategoryType = {AppCategory}");
 			}
 
-			if (UnrealData.bWriteCodeSigningSettings)
+			// always use defualt codesigning for Editor, which is for running locally
+			if (UnrealData.bWriteCodeSigningSettings && UnrealData.TargetRules.Type != TargetType.Editor)
 			{
 				Xcconfig.AppendLine("");
 				Xcconfig.AppendLine("// Code-signing settings");
 				Xcconfig.AppendLine("CODE_SIGN_STYLE = " + (bAutomaticSigning ? "Automatic" : "Manual"));
-				if (!String.IsNullOrEmpty(SigningTeam))
-				{
-					Xcconfig.AppendLine($"DEVELOPMENT_TEAM = {SigningTeam}");
-				}
-				if (!String.IsNullOrEmpty(ProvisioningProfile))
+				Xcconfig.AppendLine($"DEVELOPMENT_TEAM = {SigningTeam}");
+
+				// Mac has Sign to Run Locally to deal with
+				if (Platform == UnrealTargetPlatform.Mac)
 				{
-					Xcconfig.AppendLine($"PROVISIONING_PROFILE_SPECIFIER = {ProvisioningProfile}");
+					// when set, always use the - identity, that's how Xcode selects "Sign to Run Locally"
+					// if it is not set, and Automatic is on, then it will use it's default of "Apple Development", but we don't specify it in case default changes
+					if (bMacSignToRunLocally)
+					{
+						Xcconfig.AppendLine($"CODE_SIGN_IDENTITY = -");
+					}
+					else if (!bAutomaticSigning && SigningIdentity != null)
+					{
+						Xcconfig.AppendLine($"CODE_SIGN_IDENTITY = \"{SigningIdentity}\"");
+					}
 				}
-				if (SigningCert != null)
+				else
 				{
-					Xcconfig.AppendLine($"CODE_SIGN_IDENTITY = {SigningCert}");
+					if (!bAutomaticSigning)
+					{
+						// only use profile on IOS
+						if (!String.IsNullOrEmpty(ProvisioningProfile))
+						{
+							string SigningName = GetPlistSigningName(ProvisioningProfile, Logger);
+							Xcconfig.AppendLine($"PROVISIONING_PROFILE_SPECIFIER = {SigningName}");
+						}
+						if (SigningIdentity != null)
+						{
+							Xcconfig.AppendLine($"CODE_SIGN_IDENTITY = \"{SigningIdentity}\"");
+						}
+					}
 				}
 			}
 
@@ -1885,7 +1914,7 @@ namespace UnrealBuildTool.XcodeProjectXcconfig
 			}
 		}
 
-		public override void WriteXcconfigFile()
+		public override void WriteXcconfigFile(ILogger Logger)
 		{
 			// write out settings that for compiling natively
 			Xcconfig!.AppendLine("CLANG_CXX_LANGUAGE_STANDARD = c++17");
@@ -2113,7 +2142,7 @@ namespace UnrealBuildTool.XcodeProjectXcconfig
 			}
 
 			// cache the main group
-			SharedFileCollection.MainGroupGuid = XcodeFileCollection.GetRootGroupGuid(SharedFileCollection.Groups);
+			SharedFileCollection.MainGroupGuid = XcodeFileCollection.GetRootGroupGuid(SharedFileCollection.Groups, UnrealData.XcodeProjectFileLocation);
 
 			// filter each file into the appropriate batch
 			foreach (XcodeSourceFile File in SharedFileCollection.BuildableFilesToResponseFile.Keys)
@@ -2205,7 +2234,7 @@ namespace UnrealBuildTool.XcodeProjectXcconfig
 				{
 					// @todo hahahaah
 					continue;
-					//bSuccess = MergeIntoTemplateProject(PBXFilePath, RootProject, TemplateProject);
+					//bSuccess = MergeIntoTemplateProject(PBXFilePath, RootProject, TemplateProject, Logger);
 				}
 				else
 				{
@@ -2226,7 +2255,7 @@ namespace UnrealBuildTool.XcodeProjectXcconfig
 					FileCollection.Write(Content);
 
 					// now write out the project node and its recursive dependent nodes
-					XcodeProjectNode.WriteNodeAndReferences(Content, RootProject);
+					XcodeProjectNode.WriteNodeAndReferences(Content, RootProject, Logger);
 
 					Content.WriteLine(1, "};");
 					Content.WriteLine(1, $"rootObject = {RootProject.Guid} /* Project object */;");
@@ -2282,7 +2311,7 @@ namespace UnrealBuildTool.XcodeProjectXcconfig
 			return XcodeUtils.Plist(Command);
 		}
 
-		bool MergeIntoTemplateProject(FileReference PBXProjFilePath, XcodeProject RootProject, FileReference TemplateProject)
+		bool MergeIntoTemplateProject(FileReference PBXProjFilePath, XcodeProject RootProject, FileReference TemplateProject, ILogger Logger)
 		{
 			// activate a file for plist reading/writing here
 			XcodeUtils.SetActivePlistFile(PBXFilePathForPlatform(CurrentPlistPlatform).FullName);
@@ -2304,14 +2333,14 @@ namespace UnrealBuildTool.XcodeProjectXcconfig
 			// the runtarget and project need to write out so all of their xcconfigs get written as well,
 			// so write everything to a temp string that is tossed, but all xcconfigs will be done at least
 			StringBuilder Temp = new StringBuilder();
-			XcodeProjectNode.WriteNodeAndReferences(Temp, RootProject!);
+			XcodeProjectNode.WriteNodeAndReferences(Temp, RootProject!, Logger);
 
 			StringBuilder Content = new StringBuilder();
 			Content.WriteLine(0, "{");
 			FileCollection.Write(Content);
-			XcodeProjectNode.WriteNodeAndReferences(Content, BuildTarget);
-			XcodeProjectNode.WriteNodeAndReferences(Content, IndexTarget);
-			XcodeProjectNode.WriteNodeAndReferences(Content, BuildDependency);
+			XcodeProjectNode.WriteNodeAndReferences(Content, BuildTarget, Logger);
+			XcodeProjectNode.WriteNodeAndReferences(Content, IndexTarget, Logger);
+			XcodeProjectNode.WriteNodeAndReferences(Content, BuildDependency, Logger);
 			Content.WriteLine(0, "}");
 
 			// write to disk
diff --git a/Engine/Source/Programs/UnrealBuildTool/ToolChain/AppleToolChain.cs b/Engine/Source/Programs/UnrealBuildTool/ToolChain/AppleToolChain.cs
--- a/Engine/Source/Programs/UnrealBuildTool/ToolChain/AppleToolChain.cs
+++ b/Engine/Source/Programs/UnrealBuildTool/ToolChain/AppleToolChain.cs
@@ -313,6 +313,19 @@ namespace UnrealBuildTool
 			return Framework.ExtractedTokenFile!;
 		}
 
+		/// <summary>
+		/// If the project is a UnrealGame project, Target.ProjectDirectory refers to the engine dir, not the actual dir of the project. So this method gets the 
+		/// actual directory of the project whether it is a UnrealGame project or not.
+		/// </summary>
+		/// <returns>The actual project directory.</returns>
+		/// <param name="ProjectFile">The path to the project file</param>
+		internal static DirectoryReference GetActualProjectDirectory(FileReference? ProjectFile)
+		{
+			DirectoryReference ProjectDirectory = (ProjectFile == null ? DirectoryReference.FromString(UnrealBuildTool.GetRemoteIniPath())! : DirectoryReference.FromFile(ProjectFile)!);
+			return ProjectDirectory;
+		}
+
+
 		/// <inheritdoc/>
 		protected override string EscapePreprocessorDefinition(string Definition)
 		{
@@ -629,7 +642,6 @@ namespace UnrealBuildTool
 				$"-scheme \"{SchemeName}\"",
 				$"-configuration \"{Configuration}\"",
 				$"-destination generic/platform=" + (Platform == UnrealTargetPlatform.TVOS ? "tvOS" : Platform == UnrealTargetPlatform.Mac ? "macOS" : "iOS"),
-				"-allowProvisioningUpdates",
 				"-hideShellScriptEnvironment",
 				// xcode gets confused it we _just_ wrote out entitlements while generating the temp project, and it thinks it was modified _during_ building
 				// but it wasn't, it was written before the build started
@@ -689,6 +701,15 @@ namespace UnrealBuildTool
 
 		private int PostBuildSync(ApplePostBuildSyncTarget Target, ILogger Logger)
 		{
+			// generate the IOS plist file every time
+			if (Target.Platform == UnrealTargetPlatform.IOS || Target.Platform == UnrealTargetPlatform.TVOS)
+			{
+				string GameName = Target.ProjectFile == null ? "UnrealGame" : Target.ProjectFile.GetFileNameWithoutAnyExtensions();
+				// most of these params are uused in modern
+				UEDeployIOS.GenerateIOSPList(Target.ProjectFile, Target.Configuration, Target.ProjectFile!.Directory.FullName, Target.ProjectFile == null, GameName, bIsClient: false,
+					GameName, Unreal.EngineDirectory.FullName, Target.ProjectFile!.Directory.FullName, null, null, false, Logger);
+			}
+
 			// if xcode is building this, it will also do the Run stuff anyway, so no need to do it here as well
 			if (Environment.GetEnvironmentVariable("UE_BUILD_FROM_XCODE") == "1")
 			{
@@ -733,6 +754,14 @@ namespace UnrealBuildTool
 			PostBuildSyncAction.StatusDescription = $"Executing PostBuildSync [{Executable.Location}]";
 			PostBuildSyncAction.bCanExecuteRemotely = false;
 
+
+			if (Target.Platform == UnrealTargetPlatform.IOS || Target.Platform == UnrealTargetPlatform.TVOS)
+			{
+				// @todo: do we need one per Target for IOS? Client? I dont think so for Modern
+				FileReference PlistFile = FileReference.Combine(AppleToolChain.GetActualProjectDirectory(Target.ProjectFile), "Build/IOS/UBTGenerated/Info.Template.plist");
+				PostBuildSyncAction.ProducedItems.Add(FileItem.GetItemByFileReference(PlistFile));
+			}
+
 			return PostBuildSyncAction;
 		}
 	}
