UnrealBuildTool: Build error if VS2019 is used in an installed engine

#rnx

diff --git a/Engine/Source/Programs/UnrealBuildTool/Platform/Windows/UEBuildWindows.cs b/Engine/Source/Programs/UnrealBuildTool/Platform/Windows/UEBuildWindows.cs
--- a/Engine/Source/Programs/UnrealBuildTool/Platform/Windows/UEBuildWindows.cs
+++ b/Engine/Source/Programs/UnrealBuildTool/Platform/Windows/UEBuildWindows.cs
@@ -979,21 +979,27 @@ namespace UnrealBuildTool
 			Target.WindowsPlatform.WindowsSdkVersion = Target.WindowsPlatform.Environment.WindowsSdkVersion.ToString();
 
 			// If we're enabling support for C++ modules, make sure the compiler supports it. VS 16.8 changed which command line arguments are used to enable modules support.
-			if (Target.bEnableCppModules && !ProjectFileGenerator.bGenerateProjectFiles && Target.WindowsPlatform.Environment.CompilerVersion < new VersionNumber(14, 28, 29304))
+			if (Target.bEnableCppModules && !ProjectFileGenerator.bGenerateProjectFiles && Target.WindowsPlatform.Environment.ToolChainVersion < new VersionNumber(14, 28, 29304))
 			{
-				throw new BuildException("Support for C++20 modules requires Visual Studio 2019 16.8 preview 3 or later. Compiler Version Targeted: {0}", Target.WindowsPlatform.Environment.CompilerVersion);
+				throw new BuildException("Support for C++20 modules requires Visual Studio 2019 16.8 Preview 3 (MSVC 17.28.29304) or later. The current compiler version was detected as: {0}", Target.WindowsPlatform.Environment.ToolChainVersion);
 			}
 
 			// Ensure we're using recent enough version of Visual Studio to support ASan builds.
-			if (Target.WindowsPlatform.bEnableAddressSanitizer && Target.WindowsPlatform.Environment.CompilerVersion < new VersionNumber(14, 27, 0))
+			if (Target.WindowsPlatform.bEnableAddressSanitizer && Target.WindowsPlatform.Environment.ToolChainVersion < new VersionNumber(14, 27, 0))
 			{
-				throw new BuildException("Address sanitizer requires Visual Studio 2019 16.7 or later.");
+				throw new BuildException("Address sanitizer requires Visual Studio 2019 16.7 (MSVC 17.27.x) or later. The current compiler version was detected as: {0}", Target.WindowsPlatform.Environment.ToolChainVersion);
 			}
 
 			// Ensure we're using recent enough version of Visual Studio to support LibFuzzer.
-			if (Target.WindowsPlatform.bEnableLibFuzzer && Target.WindowsPlatform.Environment.CompilerVersion < new VersionNumber(14, 30, 0))
+			if (Target.WindowsPlatform.bEnableLibFuzzer && Target.WindowsPlatform.Environment.ToolChainVersion < new VersionNumber(14, 30, 0))
 			{
-				throw new BuildException("LibFuzzer MSVC support requires Visual Studio 2022 17.0 or later (14.30.0). The current compiler version was detected as: {0}", Target.WindowsPlatform.Environment.CompilerVersion.ToString());
+				throw new BuildException("LibFuzzer MSVC support requires Visual Studio 2022 17.0 (MSVC 14.30.x) or later. The current compiler version was detected as: {0}", Target.WindowsPlatform.Environment.ToolChainVersion);
+			}
+
+			// Ensure we're using VS2022 when compiling for the installed engine.
+			if (Unreal.IsEngineInstalled() && (Target.WindowsPlatform.ToolChain == WindowsCompiler.VisualStudio2019 || Target.WindowsPlatform.Environment.ToolChainVersion < new VersionNumber(14, 34, 0)))
+			{
+				throw new BuildException("Microsoft platform targets must be compiled with Visual Studio 2022 17.4 (MSVC 14.34.x) or later for the installed engine. Please update Visual Studio 2022 and ensure no configuration is forcing WindowsTargetRules.Compiler to VisualStudio2019. The current compiler version was detected as: {0}", Target.WindowsPlatform.Environment.ToolChainVersion);
 			}
 
 			//			@Todo: Still getting reports of frequent OOM issues with this enabled as of 15.7.
