Fix project dir for UnrealGame

diff --git a/Engine/Source/Programs/UnrealBuildTool/ToolChain/AppleToolChain.cs b/Engine/Source/Programs/UnrealBuildTool/ToolChain/AppleToolChain.cs
--- a/Engine/Source/Programs/UnrealBuildTool/ToolChain/AppleToolChain.cs
+++ b/Engine/Source/Programs/UnrealBuildTool/ToolChain/AppleToolChain.cs
@@ -321,7 +321,7 @@ namespace UnrealBuildTool
 		/// <param name="ProjectFile">The path to the project file</param>
 		internal static DirectoryReference GetActualProjectDirectory(FileReference? ProjectFile)
 		{
-			DirectoryReference ProjectDirectory = (ProjectFile == null ? DirectoryReference.FromString(UnrealBuildTool.GetRemoteIniPath())! : DirectoryReference.FromFile(ProjectFile)!);
+			DirectoryReference ProjectDirectory = (ProjectFile == null ? Unreal.EngineDirectory : DirectoryReference.FromFile(ProjectFile)!);
 			return ProjectDirectory;
 		}
 
@@ -706,8 +706,8 @@ namespace UnrealBuildTool
 			{
 				string GameName = Target.ProjectFile == null ? "UnrealGame" : Target.ProjectFile.GetFileNameWithoutAnyExtensions();
 				// most of these params are uused in modern
-				UEDeployIOS.GenerateIOSPList(Target.ProjectFile, Target.Configuration, Target.ProjectFile!.Directory.FullName, Target.ProjectFile == null, GameName, bIsClient: false,
-					GameName, Unreal.EngineDirectory.FullName, Target.ProjectFile!.Directory.FullName, null, null, false, Logger);
+				UEDeployIOS.GenerateIOSPList(Target.ProjectFile, Target.Configuration, AppleToolChain.GetActualProjectDirectory(Target.ProjectFile).FullName, Target.ProjectFile == null, GameName, bIsClient: false,
+					GameName, Unreal.EngineDirectory.FullName, "", null, null, false, Logger);
 			}
 
 			// if xcode is building this, it will also do the Run stuff anyway, so no need to do it here as well
