- Fix crash with UnrealEditor builds, caused by earlier fix for AdditionalPluginDirectories warning

diff --git a/Engine/Source/Programs/UnrealBuildTool/ProjectFiles/ProjectFileGenerator.cs b/Engine/Source/Programs/UnrealBuildTool/ProjectFiles/ProjectFileGenerator.cs
--- a/Engine/Source/Programs/UnrealBuildTool/ProjectFiles/ProjectFileGenerator.cs
+++ b/Engine/Source/Programs/UnrealBuildTool/ProjectFiles/ProjectFileGenerator.cs
@@ -971,11 +971,13 @@ namespace UnrealBuildTool
 			// belong to platforms we're generating project files for.
 			List<FileReference> AllTargetFiles = DiscoverTargets(AllGameProjects, Logger, OnlyGameProject, SupportedPlatforms, bIncludeEngineSource, bIncludeTempTargets);
 
-			// if we only want one target, remove the others
+			// if we only want one target, remove the others (keeping Unreal* targets, as some code expects them to exist0
 			if (SingleTargetName != null)
 			{
-				AllTargetFiles.RemoveAll(x => !x.GetFileNameWithoutAnyExtensions().Equals(SingleTargetName, StringComparison.OrdinalIgnoreCase) && 
-				!x.GetFileNameWithoutAnyExtensions().Equals(EngineProjectFileNameBase, StringComparison.OrdinalIgnoreCase));
+				AllTargetFiles.RemoveAll(x => !x.GetFileNameWithoutAnyExtensions().Equals(SingleTargetName, StringComparison.OrdinalIgnoreCase) &&
+					!x.GetFileNameWithoutAnyExtensions().Equals(EngineProjectFileNameBase, StringComparison.OrdinalIgnoreCase) &&
+					!x.GetFileNameWithoutAnyExtensions().Equals(EngineEditorProjectFileNameBase, StringComparison.OrdinalIgnoreCase));
+
 			}
 
 			// Sort the targets by name. When we have multiple targets of a given type for a project, we'll use the order to determine which goes in the primary project file (so that client names with a suffix will go into their own project).
