UBT: When calling with Mode=QueryTargets and the -IncludeAllTargets flag, include a flag in the result to indicate which target is the default target for a given target type.

diff --git a/Engine/Source/Programs/UnrealBuildTool/Modes/QueryTargetsMode.cs b/Engine/Source/Programs/UnrealBuildTool/Modes/QueryTargetsMode.cs
--- a/Engine/Source/Programs/UnrealBuildTool/Modes/QueryTargetsMode.cs
+++ b/Engine/Source/Programs/UnrealBuildTool/Modes/QueryTargetsMode.cs
@@ -133,16 +133,25 @@ namespace UnrealBuildTool
 						continue;
 					}
 
-					// Skip non-default targets if one is specified.
-					if (ProjectFile != null && !bIncludeAllTargets)
+					// Is this a default target?
+					bool? bIsDefaultTarget = null;
+					if (ProjectFile != null)
 					{
 						string? DefaultTargetName = ProjectFileGenerator.GetProjectDefaultTargetNameForType(ProjectFile.Directory, TargetRules.Type);
-						if (DefaultTargetName != null && DefaultTargetName != TargetName)
+						
+						// GetProjectDefaultTargetNameForType returns
+						if (DefaultTargetName != null)
 						{
-							continue;
+							bIsDefaultTarget = DefaultTargetName == TargetName;   
 						}
 					}
 
+					// If we don't want all targets, skip over non-defaults.
+					if (!bIncludeAllTargets && bIsDefaultTarget.HasValue && !bIsDefaultTarget.Value)
+					{
+						continue;
+					}
+
 					// Get the path to the target
 					FileReference? path = Assembly.GetTargetFileName(TargetName);
 
@@ -154,6 +163,11 @@ namespace UnrealBuildTool
 						Writer.WriteValue("Path", path.MakeRelativeTo(OutputFile.Directory));
 					}
 					Writer.WriteValue("Type", TargetRules.Type.ToString());
+
+					if (bIncludeAllTargets && bIsDefaultTarget.HasValue)
+					{
+						Writer.WriteValue("DefaultTarget", bIsDefaultTarget.Value);
+					}
 					Writer.WriteObjectEnd();
 				}
 				Writer.WriteArrayEnd();
