Optimized the searching of included files.

Solution generation:
Before: 1m55s
After: 50s

diff --git a/Engine/Source/Programs/UnrealBuildTool/Configuration/UEBuildModuleCPP.cs b/Engine/Source/Programs/UnrealBuildTool/Configuration/UEBuildModuleCPP.cs
--- a/Engine/Source/Programs/UnrealBuildTool/Configuration/UEBuildModuleCPP.cs
+++ b/Engine/Source/Programs/UnrealBuildTool/Configuration/UEBuildModuleCPP.cs
@@ -1,6 +1,7 @@
 // Copyright Epic Games, Inc. All Rights Reserved.
 
 using System;
+using System.Collections.Concurrent;
 using System.Collections.Generic;
 using System.IO;
 using System.Linq;
@@ -823,11 +824,25 @@ namespace UnrealBuildTool
 			return HeaderFileItems;
 		}
 
-		private void FindIncludedHeaders(HashSet<string> VisitedIncludes, HashSet<FileItem> FoundIncludeFileItems, HashSet<string> IncludesNotFound, CppCompileEnvironment CompileEnvironment, FileItem FileToSearch, ILogger Logger)
+		// Cache of files found while searching the includes
+		static ConcurrentDictionary<DirectoryReference, HashSet<FileReference>> KnownIncludeFilesDict = new();
+		static HashSet<FileReference> GetIncludeFiles(DirectoryReference Directory)
 		{
-			HashSet<DirectoryReference> SearchPaths = new(CompileEnvironment.UserIncludePaths);
-			SearchPaths.Add(FileToSearch.Directory.Location);
+			return KnownIncludeFilesDict.GetOrAdd(Directory, (_) =>
+			{
+				if (DirectoryReference.Exists(Directory))
+				{
+					return DirectoryReference.EnumerateFiles(Directory, "*").ToHashSet();
+				}
+				else
+				{
+					return new HashSet<FileReference>();
+				}
+			});
+		}
 
+		private void FindIncludedHeaders(HashSet<string> VisitedIncludes, HashSet<FileItem> FoundIncludeFileItems, HashSet<string> IncludesNotFound, CppCompileEnvironment CompileEnvironment, FileItem FileToSearch, ILogger Logger)
+		{
 			foreach (string HeaderInclude in CompileEnvironment.MetadataCache.GetHeaderIncludes(FileToSearch))
 			{
 				string TranformedHeaderInclude = HeaderInclude;
@@ -840,12 +855,36 @@ namespace UnrealBuildTool
 
 				if (VisitedIncludes.Add(TranformedHeaderInclude))
 				{
-					var FoundDir = SearchPaths.FirstOrDefault(dir => DirectoryLookupCache.FileExists(FileReference.Combine(dir, TranformedHeaderInclude)));
-					if (FoundDir != null)
+					var SearchForFileItem = (DirectoryReference dir) =>
+					{
+						FileReference FileRef = FileReference.Combine(dir, TranformedHeaderInclude);
+						if (GetIncludeFiles(FileRef.Directory).Contains(FileRef))
+						{
+							return FileItem.GetItemByFileReference(FileRef);
+						}
+						return null;
+					};
+
+					// check to see if the file is a relative path first
+					FileItem? IncludeFileItem = SearchForFileItem(FileToSearch.Directory.Location);
+
+					// search through the include paths if the file isn't relative
+					if (IncludeFileItem == null)
+					{
+						foreach (var IncludePath in CompileEnvironment.UserIncludePaths)
+						{
+							IncludeFileItem = SearchForFileItem(IncludePath);
+							if (IncludeFileItem != null)
+							{
+								break;
+							}
+						}
+					}
+
+					if (IncludeFileItem != null)
 					{
-						FileItem IncludeFileItem = FileItem.GetItemByFileReference(FileReference.Combine(FoundDir, TranformedHeaderInclude));
 						FoundIncludeFileItems.Add(IncludeFileItem);
-						Logger.LogDebug("{0} SharedPCH - '{1}' included '{2}'.", Name, FileToSearch.Location, IncludeFileItem.Location);
+						// Logger.LogDebug("{0} SharedPCH - '{1}' included '{2}'.", Name, FileToSearch.Location, IncludeFileItem.Location);
 
 						FindIncludedHeaders(VisitedIncludes, FoundIncludeFileItems, IncludesNotFound, CompileEnvironment, IncludeFileItem, Logger);
 					}
@@ -887,6 +926,13 @@ namespace UnrealBuildTool
 			HashSet<string> VisitedIncludes = new();
 			HashSet<FileItem> FoundIncludeFileItems = new();
 			HashSet<string> IncludesNotFound = new();
+
+			// pre-populate all the known file locations for this compile environment
+			foreach (DirectoryReference dir in CompileEnvironment.UserIncludePaths)
+			{
+				GetIncludeFiles(dir);
+			}
+
 			FindIncludedHeaders(VisitedIncludes, FoundIncludeFileItems, IncludesNotFound, CompileEnvironment, PCHHeaderFile, Logger);
 			if (IncludesNotFound.Any())
 			{
