- Fixed errors with plugin modules failing to find a project to put into (better fix then yesterday's for the AdditionalPluginDirectories)

diff --git a/Engine/Source/Programs/UnrealBuildTool/ProjectFiles/ProjectFileGenerator.cs b/Engine/Source/Programs/UnrealBuildTool/ProjectFiles/ProjectFileGenerator.cs
--- a/Engine/Source/Programs/UnrealBuildTool/ProjectFiles/ProjectFileGenerator.cs
+++ b/Engine/Source/Programs/UnrealBuildTool/ProjectFiles/ProjectFileGenerator.cs
@@ -974,10 +974,7 @@ namespace UnrealBuildTool
 			// if we only want one target, remove the others (keeping Unreal* targets, as some code expects them to exist0
 			if (SingleTargetName != null)
 			{
-				AllTargetFiles.RemoveAll(x => !x.GetFileNameWithoutAnyExtensions().Equals(SingleTargetName, StringComparison.OrdinalIgnoreCase) &&
-					!x.GetFileNameWithoutAnyExtensions().Equals(EngineProjectFileNameBase, StringComparison.OrdinalIgnoreCase) &&
-					!x.GetFileNameWithoutAnyExtensions().Equals(EngineEditorProjectFileNameBase, StringComparison.OrdinalIgnoreCase));
-
+				AllTargetFiles.RemoveAll(x => !x.GetFileNameWithoutAnyExtensions().Equals(SingleTargetName, StringComparison.OrdinalIgnoreCase));
 			}
 
 			// Sort the targets by name. When we have multiple targets of a given type for a project, we'll use the order to determine which goes in the primary project file (so that client names with a suffix will go into their own project).
@@ -2538,8 +2535,14 @@ namespace UnrealBuildTool
 						FileReference ProjectInfo = Game;
 						BaseFolder = ProjectInfo.Directory;
 
+						// if we are in SingleTargetMode, then skip this and just put it into the project's projet
 						if (bMakeProjectPerTarget && !bIsTemporaryModule)
 						{
+							if (SingleTargetName != null)
+							{
+								return FindOrAddProjectHelper(SingleTargetName, BaseFolder);
+							}
+
 							// find the project that the module is under, and has a TargetType target (useful with bMakeProjectPerTarget)
 							foreach (KeyValuePair<FileReference, ProjectFile> Pair in ProjectFileMap)
 							{
