EngineIncludeOrderVersion cleanup
- Deprecating EngineIncludeOrderVersion.Unreal5_0 and EngineIncludeOrderVersion.Unreal5_1.
- Adding way more comments on EngineIncludeOrderVersion including warnings about using EngineIncludeOrderVersion.Latest.

diff --git a/Engine/Source/Programs/UnrealBuildTool/Configuration/TargetRules.cs b/Engine/Source/Programs/UnrealBuildTool/Configuration/TargetRules.cs
--- a/Engine/Source/Programs/UnrealBuildTool/Configuration/TargetRules.cs
+++ b/Engine/Source/Programs/UnrealBuildTool/Configuration/TargetRules.cs
@@ -147,18 +147,23 @@ namespace UnrealBuildTool
 	}
 
 	/// <summary>
-	/// What version of include order to use when compiling.
+	/// What version of include order to use when compiling. This controls which UE_ENABLE_INCLUDE_ORDER_DEPRECATED_IN_5_X defines are enabled when compiling the target.
+	/// The UE_ENABLE_INCLUDE_ORDER_DEPRECATED_IN_5_X defines are used when removing implicit includes from Unreal Public headers.
+	/// Specific versions can be used by licensees to avoid compile errors when integrating new versions of Unreal.
+	/// Using Latest comes with a high risk of introducing compile errors in your code on newer Unreal versions.
 	/// </summary>
 	public enum EngineIncludeOrderVersion
 	{
 		/// <summary>
 		/// Include order used in Unreal 5.0
 		/// </summary>
+		[Obsolete("The Unreal 5.0 include order is deprecated and will be removed in 5.4.")]
 		Unreal5_0,
 
 		/// <summary>
 		/// Include order used in Unreal 5.1
 		/// </summary>
+		[Obsolete("The Unreal 5.1 include order is deprecated and will be removed in 5.4.")]
 		Unreal5_1,
 
 		/// <summary>
@@ -179,14 +184,16 @@ namespace UnrealBuildTool
 		// *** When adding new entries here, be sure to update UEBuildModuleCPP.CurrentIncludeOrderDefine to ensure that the correct guard is used. ***
 
 		/// <summary>
-		/// Always use the latest version of include order.
+		/// Always use the latest version of include order. This value is updated every Unreal release, use with caution if you intend to integrate newer Unreal releases.
 		/// </summary>
 		Latest = Unreal5_4,
 
 		/// <summary>
 		/// Contains the oldest version of include order that the engine supports.
 		/// </summary>
+#pragma warning disable CS0618 // Type or member is obsolete
 		Oldest = Unreal5_0,
+#pragma warning restore CS0618 // Type or member is obsolete
 	}
 
 	/// <summary>
@@ -318,10 +325,12 @@ namespace UnrealBuildTool
 					return "UE_ENABLE_INCLUDE_ORDER_DEPRECATED_IN_5_3";
 				case EngineIncludeOrderVersion.Unreal5_2:
 					return "UE_ENABLE_INCLUDE_ORDER_DEPRECATED_IN_5_2";
+#pragma warning disable CS0618 // Type or member is obsolete
 				case EngineIncludeOrderVersion.Unreal5_1:
 					return "UE_ENABLE_INCLUDE_ORDER_DEPRECATED_IN_5_1";
 				case EngineIncludeOrderVersion.Unreal5_0:
 					return "UE_ENABLE_INCLUDE_ORDER_DEPRECATED_IN_5_0";
+#pragma warning restore CS0618 // Type or member is obsolete
 			}
 		}
 
@@ -338,7 +347,9 @@ namespace UnrealBuildTool
 		{
 			return new List<string>()
 			{
+#pragma warning disable CS0618 // Type or member is obsolete
 				GetDeprecationDefine(EngineIncludeOrderVersion.Unreal5_1),
+#pragma warning restore CS0618 // Type or member is obsolete
 				GetDeprecationDefine(EngineIncludeOrderVersion.Unreal5_2),
 				GetDeprecationDefine(EngineIncludeOrderVersion.Unreal5_3),
 				GetDeprecationDefine(EngineIncludeOrderVersion.Unreal5_4),
@@ -354,7 +365,9 @@ namespace UnrealBuildTool
 		{
 			return new List<string>()
 			{
+#pragma warning disable CS0618 // Type or member is obsolete
 				GetDeprecationDefine(EngineIncludeOrderVersion.Unreal5_1, InVersion),
+#pragma warning restore CS0618 // Type or member is obsolete
 				GetDeprecationDefine(EngineIncludeOrderVersion.Unreal5_2, InVersion),
 				GetDeprecationDefine(EngineIncludeOrderVersion.Unreal5_3, InVersion),
 			};
