Moved the default VS include paths to the common file.

diff --git a/Engine/Source/Programs/UnrealBuildTool/ProjectFiles/VisualStudio/VCProject.cs b/Engine/Source/Programs/UnrealBuildTool/ProjectFiles/VisualStudio/VCProject.cs
--- a/Engine/Source/Programs/UnrealBuildTool/ProjectFiles/VisualStudio/VCProject.cs
+++ b/Engine/Source/Programs/UnrealBuildTool/ProjectFiles/VisualStudio/VCProject.cs
@@ -559,19 +559,6 @@ namespace UnrealBuildTool
 			}
 		}
 
-		WindowsCompiler GetCompilerForIntellisense()
-		{
-			switch (ProjectFileFormat)
-			{
-				case VCProjectFileFormat.VisualStudio2022:
-					return WindowsCompiler.VisualStudio2022;
-				case VCProjectFileFormat.VisualStudio2019:
-					return WindowsCompiler.VisualStudio2019;
-				default:
-					return WindowsCompiler.VisualStudio2019;
-			}
-		}
-
 		/// <summary>
 		/// Gets highest C++ version which is used in this project
 		/// </summary>
@@ -669,7 +656,7 @@ namespace UnrealBuildTool
 		{
 			if (IntelliSenseEnableCoroutines)
 			{
-				if (GetCompilerForIntellisense().IsMSVC())
+				if (VCProjectFileGenerator.GetCompilerForIntellisense(ProjectFileFormat).IsMSVC())
 				{
 					return "/await:strict";
 				}
@@ -934,18 +921,7 @@ namespace UnrealBuildTool
 					}
 				}
 
-				// Add all the default system include paths
-				if (OperatingSystem.IsWindows())
-				{
-					if (InPlatforms.Contains(UnrealTargetPlatform.Win64))
-					{
-						SharedIncludeSearchPaths.Append(VCToolChain.GetVCIncludePaths(UnrealTargetPlatform.Win64, GetCompilerForIntellisense(), null, Logger) + ";");
-					}
-				}
-				else
-				{
-					Logger.LogInformation("Unable to compute VC include paths on non-Windows host");
-				}
+				SharedIncludeSearchPaths.AppendFormat("$(DefaultSystemIncludePaths);");
 			}
 
 			StringBuilder VCPreprocessorDefinitions = new StringBuilder();
diff --git a/Engine/Source/Programs/UnrealBuildTool/ProjectFiles/VisualStudio/VCProjectFileGenerator.cs b/Engine/Source/Programs/UnrealBuildTool/ProjectFiles/VisualStudio/VCProjectFileGenerator.cs
--- a/Engine/Source/Programs/UnrealBuildTool/ProjectFiles/VisualStudio/VCProjectFileGenerator.cs
+++ b/Engine/Source/Programs/UnrealBuildTool/ProjectFiles/VisualStudio/VCProjectFileGenerator.cs
@@ -242,6 +242,19 @@ namespace UnrealBuildTool
 			return String.Empty;
 		}
 
+		public static WindowsCompiler GetCompilerForIntellisense(VCProjectFileFormat ProjectFileFormat)
+		{
+			switch (ProjectFileFormat)
+			{
+				case VCProjectFileFormat.VisualStudio2022:
+					return WindowsCompiler.VisualStudio2022;
+				case VCProjectFileFormat.VisualStudio2019:
+					return WindowsCompiler.VisualStudio2019;
+				default:
+					return WindowsCompiler.VisualStudio2019;
+			}
+		}
+
 		public static void AppendPlatformToolsetProperty(StringBuilder VCProjectFileContent, VCProjectFileFormat ProjectFileFormat)
 		{
 			string ToolVersionString = GetProjectFileToolVersionString(ProjectFileFormat);
@@ -532,7 +545,23 @@ namespace UnrealBuildTool
 					VCCommonTargetFileContent.AppendLine("    <LibraryWPath />");
 					VCCommonTargetFileContent.AppendLine("    <SourcePath />");
 					VCCommonTargetFileContent.AppendLine("    <ExcludePath />");
+
+					// Add all the default system include paths
+					if (OperatingSystem.IsWindows())
+					{
+						if (SupportedPlatforms.Contains(UnrealTargetPlatform.Win64))
+						{
+							VCCommonTargetFileContent.AppendLine("    <DefaultSystemIncludePaths>{0}</DefaultSystemIncludePaths>", VCToolChain.GetVCIncludePaths(UnrealTargetPlatform.Win64, GetCompilerForIntellisense(Settings.ProjectFileFormat), null, Logger));
+						}
+					}
+					else
+					{
+						Logger.LogInformation("Unable to compute VC include paths on non-Windows host");
+						VCCommonTargetFileContent.AppendLine("    <DefaultSystemIncludePaths />");
+					}
+
 					VCCommonTargetFileContent.AppendLine("  </PropertyGroup>");
+
 				}
 
 				// Write default import group
