Silence a provision profile error, this should not affect code signing

diff --git a/Engine/Source/Programs/UnrealBuildTool/ToolChain/AppleToolChain.cs b/Engine/Source/Programs/UnrealBuildTool/ToolChain/AppleToolChain.cs
--- a/Engine/Source/Programs/UnrealBuildTool/ToolChain/AppleToolChain.cs
+++ b/Engine/Source/Programs/UnrealBuildTool/ToolChain/AppleToolChain.cs
@@ -650,7 +650,38 @@ namespace UnrealBuildTool
 				//$"-sdk {SDKName}",
 			};
 
-			return Utils.RunLocalProcessAndLogOutput("/usr/bin/env", String.Join(" ", Arguments), Logger);
+			Process LocalProcess = new Process();
+			LocalProcess.StartInfo = new ProcessStartInfo("/usr/bin/env", String.Join(" ", Arguments));
+			LocalProcess.OutputDataReceived += (Sender, Args) => { LocalProcessOutput(Args, false, Logger); };
+			LocalProcess.ErrorDataReceived += (Sender, Args) =>
+			{
+				if (Args != null && Args.Data != null
+				&& Args.Data.Contains("Failed to load profile") && Args.Data.Contains("<stdin>"))
+				{
+					Logger.LogInformation("Silencing the following provision profile error, it is not affecting code signing:");
+					LocalProcessOutput(Args, false, Logger);
+				}
+				else
+				{
+					LocalProcessOutput(Args, true, Logger);
+				}
+			};
+			return Utils.RunLocalProcess(LocalProcess);
+		}
+
+		static void LocalProcessOutput(DataReceivedEventArgs? Args, bool bIsError, ILogger Logger)
+		{
+			if (Args != null && Args.Data != null)
+			{
+				if (bIsError)
+				{
+					Logger.LogError("{Message}", Args.Data.TrimEnd());
+				}
+				else
+				{
+					Logger.LogInformation("{Message}", Args.Data.TrimEnd());
+				}
+			}
 		}
 
 		#endregion
@@ -719,7 +750,7 @@ namespace UnrealBuildTool
 			int ExitCode = AppleExports.BuildWithStubXcodeProject(Target.ProjectFile, Target.Platform, Target.Configuration, Target.TargetName, AppleExports.XcodeBuildMode.PostBuildSync, Logger);
 			if (ExitCode != 0)
 			{
-				Logger.LogError("ERROR: Failed to finalize the .app with Xcode. Check the log for ,more information");
+				Logger.LogError("ERROR: Failed to finalize the .app with Xcode. Check the log for more information");
 			}
 
 			return ExitCode;
