- Fixed an crash when no VC toolchains are installed at all

diff --git a/Engine/Source/Programs/UnrealBuildTool/Platform/Windows/MicrosoftPlatformSDK.cs b/Engine/Source/Programs/UnrealBuildTool/Platform/Windows/MicrosoftPlatformSDK.cs
--- a/Engine/Source/Programs/UnrealBuildTool/Platform/Windows/MicrosoftPlatformSDK.cs
+++ b/Engine/Source/Programs/UnrealBuildTool/Platform/Windows/MicrosoftPlatformSDK.cs
@@ -204,9 +204,13 @@ namespace UnrealBuildTool
 			else if (MinVersion == null && MaxVersion == null)
 			{
 				// convert the desired version into a VersionNumber
-				WindowsSdkVersion = VersionNumber.Parse(UEBuildPlatformSDK.GetSDKForPlatform("Win64")!.GetMainVersion());
-				// if it's not a known version, use the highest known
-				if (!CachedWindowsSdkDirs!.ContainsKey(WindowsSdkVersion))
+				VersionNumber MainVersion = VersionNumber.Parse(UEBuildPlatformSDK.GetSDKForPlatform("Win64")!.GetMainVersion());
+				if (CachedWindowsSdkDirs!.ContainsKey(MainVersion))
+				{
+					WindowsSdkVersion = MainVersion;
+				}
+				// if it's not an installed version, use the highest installed version
+				else if (CachedWindowsSdkDirs.Any())
 				{
 					WindowsSdkVersion = CachedWindowsSdkDirs.OrderBy(x => x.Key).Last().Key;
 				}
