[UBT]
* Added so when counting queued actions it is possible to provide a function that can check if the action should be counted or not

diff --git a/Engine/Source/Programs/UnrealBuildTool/Executors/ImmediateActionQueue.cs b/Engine/Source/Programs/UnrealBuildTool/Executors/ImmediateActionQueue.cs
--- a/Engine/Source/Programs/UnrealBuildTool/Executors/ImmediateActionQueue.cs
+++ b/Engine/Source/Programs/UnrealBuildTool/Executors/ImmediateActionQueue.cs
@@ -798,7 +798,7 @@ namespace UnrealBuildTool
 			}
 		}
 
-		public int GetQueuedActionsCount()
+		public int GetQueuedActionsCount(Func<LinkedAction, bool>? CanRunFunc = null)
 		{
 			lock (Actions)
 			{
@@ -807,10 +807,17 @@ namespace UnrealBuildTool
 				for (int actionIndex = _firstPendingAction; actionIndex != Actions.Length; ++actionIndex)
 				{
 
-					if (Actions[actionIndex].Status == ActionStatus.Queued)
+					if (Actions[actionIndex].Status != ActionStatus.Queued)
 					{
-						++count;
+						continue;
 					}
+					
+					if (CanRunFunc != null && !CanRunFunc(Actions[actionIndex].Action))
+					{
+						continue;
+					}
+
+					++count;
 				}
 				return count;
 			}
