Replace encoding args with /utf-8

#rnx

diff --git a/Engine/Source/Programs/UnrealBuildTool/Platform/Windows/VCToolChain.cs b/Engine/Source/Programs/UnrealBuildTool/Platform/Windows/VCToolChain.cs
--- a/Engine/Source/Programs/UnrealBuildTool/Platform/Windows/VCToolChain.cs
+++ b/Engine/Source/Programs/UnrealBuildTool/Platform/Windows/VCToolChain.cs
@@ -457,6 +457,9 @@ namespace UnrealBuildTool
 				Arguments.Add($"/Zm{Target.WindowsPlatform.PCHMemoryAllocationFactor}");
 			}
 
+			// Fix Incredibuild errors with helpers using heterogeneous character sets
+			Arguments.Add("/utf-8");
+
 			// Disable "The file contains a character that cannot be represented in the current code page" warning for non-US windows.
 			Arguments.Add("/wd4819");
 
@@ -525,13 +528,6 @@ namespace UnrealBuildTool
 			// Maintain the old std::aligned_storage behavior from VS from v15.8 onwards, in case of prebuilt third party libraries are reliant on it
 			AddDefinition(Arguments, "_DISABLE_EXTENDED_ALIGNED_STORAGE");
 
-			// Fix Incredibuild errors with helpers using heterogeneous character sets
-			if (Target.WindowsPlatform.Compiler.IsMSVC())
-			{
-				Arguments.Add("/source-charset:utf-8");
-				Arguments.Add("/execution-charset:utf-8");
-			}
-
 			// Do not allow inline method expansion if E&C support is enabled or inline expansion has been disabled
 			if (!CompileEnvironment.bSupportEditAndContinue && CompileEnvironment.bUseInlining)
 			{
