[UBT]
* Removed lock in GetQueuedActionsCount. The return count does not need to be exact so this is fine.

diff --git a/Engine/Source/Programs/UnrealBuildTool/Executors/ImmediateActionQueue.cs b/Engine/Source/Programs/UnrealBuildTool/Executors/ImmediateActionQueue.cs
--- a/Engine/Source/Programs/UnrealBuildTool/Executors/ImmediateActionQueue.cs
+++ b/Engine/Source/Programs/UnrealBuildTool/Executors/ImmediateActionQueue.cs
@@ -800,27 +800,24 @@ namespace UnrealBuildTool
 
 		public int GetQueuedActionsCount(Func<LinkedAction, bool>? CanRunFunc = null)
 		{
-			lock (Actions)
+			int count = 0;
+
+			for (int actionIndex = _firstPendingAction; actionIndex != Actions.Length; ++actionIndex)
 			{
-				int count = 0;
 
-				for (int actionIndex = _firstPendingAction; actionIndex != Actions.Length; ++actionIndex)
+				if (Actions[actionIndex].Status != ActionStatus.Queued)
 				{
-
-					if (Actions[actionIndex].Status != ActionStatus.Queued)
-					{
-						continue;
-					}
+					continue;
+				}
 					
-					if (CanRunFunc != null && !CanRunFunc(Actions[actionIndex].Action))
-					{
-						continue;
-					}
-
-					++count;
+				if (CanRunFunc != null && !CanRunFunc(Actions[actionIndex].Action))
+				{
+					continue;
 				}
-				return count;
+
+				++count;
 			}
+			return count;
 		}
 
 		/// <summary>
