Modified UBT to generate more detailed information when the compiler crashes.

diff --git a/Engine/Source/Programs/UnrealBuildTool/Executors/ImmediateActionQueue.cs b/Engine/Source/Programs/UnrealBuildTool/Executors/ImmediateActionQueue.cs
--- a/Engine/Source/Programs/UnrealBuildTool/Executors/ImmediateActionQueue.cs
+++ b/Engine/Source/Programs/UnrealBuildTool/Executors/ImmediateActionQueue.cs
@@ -968,14 +968,15 @@ namespace UnrealBuildTool
 
 				if (exitCode != 0)
 				{
-					// BEGIN TEMPORARY TO CATCH PVS-STUDIO ISSUES
-					if (logLines == null || logLines.Count == 0)
+
+					// If we have an error code but no output, chances are the tool crashed.  Generate more detailed information to let the
+					// user know something went wrong.
+					if (logLines == null || logLines.Count <= (action.bShouldOutputStatusDescription ? 0 : 1))
 					{
-						Logger.LogError("{TargetDetails} {Description}: Exited with error code {ExitCode}", targetDetails, description, exitCode);
+						Logger.LogError("{TargetDetails} {Description}: Exited with error code {ExitCode}. The build will fail.", targetDetails, description, exitCode);
 						Logger.LogInformation("{TargetDetails} {Description}: WorkingDirectory {WorkingDirectory}", targetDetails, description, action.WorkingDirectory);
 						Logger.LogInformation("{TargetDetails} {Description}: {CommandPath} {CommandArguments}", targetDetails, description, action.CommandPath, action.CommandArguments);
 					}
-					// END TEMPORARY
 
 					// prevent overwriting of error text
 					s_previousLineLength = -1;
