Added the ability to force update the include paths.

diff --git a/Engine/Source/Programs/UnrealBuildTool/Modes/FixIncludePathsMode.cs b/Engine/Source/Programs/UnrealBuildTool/Modes/FixIncludePathsMode.cs
--- a/Engine/Source/Programs/UnrealBuildTool/Modes/FixIncludePathsMode.cs
+++ b/Engine/Source/Programs/UnrealBuildTool/Modes/FixIncludePathsMode.cs
@@ -42,6 +42,9 @@ namespace UnrealBuildTool
 		[CommandLine("-NoIncludeSorting", Description = "Flags that includes should not be sorted.")]
 		public bool bNoIncludeSorting = false;
 
+		[CommandLine("-ForceUpdate", Description = "Forces updating and sorting the includes.")]
+		public bool bForceUpdate = false;
+
 		/// <summary>
 		/// Execute the command
 		/// </summary>
@@ -287,10 +290,11 @@ namespace UnrealBuildTool
 											if (PreferredInclude == null)
 											{
 												Logger.LogInformation("{FileName}({LineNumber}): Could not find path to '{IncludePath}'", InputFile.FullName, LineNumber, Include);
+												continue;
 											}
 										}
 
-										if (PreferredInclude != null && Include != PreferredInclude)
+										if (bForceUpdate || (PreferredInclude != null && Include != PreferredInclude))
 										{
 											Logger.LogInformation("{FileName}({LineNumber}): Updated '{OldInclude}' -> '{NewInclude}'", InputFile.FullName, LineNumber, Include, PreferredInclude);
 											Text[i] = Line.Replace(Include, PreferredInclude);
