- Removed some hard references to VisionOS from XcodeProject.cs and UEBuildIOS.cs, since it's a Plaform Extension it may not exist
- Refactored some general IOS code to be more programmatic

diff --git a/Engine/Platforms/VisionOS/Source/Programs/UnrealBuildTool/UEBuildVisionOS.cs b/Engine/Platforms/VisionOS/Source/Programs/UnrealBuildTool/UEBuildVisionOS.cs
--- a/Engine/Platforms/VisionOS/Source/Programs/UnrealBuildTool/UEBuildVisionOS.cs
+++ b/Engine/Platforms/VisionOS/Source/Programs/UnrealBuildTool/UEBuildVisionOS.cs
@@ -59,14 +59,16 @@ namespace UnrealBuildTool
 				return;
 			}
 
-			if ((Target.Platform == UnrealTargetPlatform.Win64) || (Target.Platform == UnrealTargetPlatform.Mac))
+			if (!Target.bBuildRequiresCookedData)
 			{
-				// allow standalone tools to use targetplatform modules, without needing Engine
-				if (Target.bForceBuildTargetPlatforms)
-				{
-					// @todo visionos: Make the module
-					// InModule.AddPlatformSpecificDynamicallyLoadedModule("VisionOSTargetPlatform");
-				}
+			    if (Target.Platform == UnrealTargetPlatform.Mac || Target.Platform == UnrealTargetPlatform.Win64)
+			    {
+				    if ((ModuleName == "Engine" && Target.bBuildDeveloperTools) ||
+					    (ModuleName == "TargetPlatform" && Target.bForceBuildTargetPlatforms))
+				    {
+					    Rules.DynamicallyLoadedModuleNames.Add("VisionOSTargetPlatform");
+				    }
+			    }
 			}
 		}
 
diff --git a/Engine/Source/Programs/UnrealBuildTool/Platform/IOS/UEBuildIOS.cs b/Engine/Source/Programs/UnrealBuildTool/Platform/IOS/UEBuildIOS.cs
--- a/Engine/Source/Programs/UnrealBuildTool/Platform/IOS/UEBuildIOS.cs
+++ b/Engine/Source/Programs/UnrealBuildTool/Platform/IOS/UEBuildIOS.cs
@@ -1030,7 +1030,6 @@ namespace UnrealBuildTool
 						{
 							Rules.DynamicallyLoadedModuleNames.Add("IOSTargetPlatform");
 							Rules.DynamicallyLoadedModuleNames.Add("TVOSTargetPlatform");
-							Rules.DynamicallyLoadedModuleNames.Add("VisionOSTargetPlatform");
 						}
 					}
 					else if (ModuleName == "TargetPlatform")
@@ -1052,7 +1051,6 @@ namespace UnrealBuildTool
 					{
 						Rules.DynamicallyLoadedModuleNames.Add("IOSTargetPlatform");
 						Rules.DynamicallyLoadedModuleNames.Add("TVOSTargetPlatform");
-						Rules.DynamicallyLoadedModuleNames.Add("VisionOSTargetPlatform");
 					}
 
 					if (bBuildShaderFormats)
diff --git a/Engine/Source/Programs/UnrealBuildTool/ProjectFiles/Xcode/XcodeProject.cs b/Engine/Source/Programs/UnrealBuildTool/ProjectFiles/Xcode/XcodeProject.cs
--- a/Engine/Source/Programs/UnrealBuildTool/ProjectFiles/Xcode/XcodeProject.cs
+++ b/Engine/Source/Programs/UnrealBuildTool/ProjectFiles/Xcode/XcodeProject.cs
@@ -337,8 +337,8 @@ namespace UnrealBuildTool.XcodeProjectXcconfig
 
 		// settings read from project configs
 		public IOSProjectSettings? IOSProjectSettings;
-		public TVOSProjectSettings? TVOSProjectSettings;
-		public VisionOSProjectSettings? VisionOSProjectSettings;
+		public IOSProjectSettings? TVOSProjectSettings;
+		public IOSProjectSettings? VisionOSProjectSettings;
 
 		// Name of the product (usually the project name, but UE5.xcodeproj is actually UnrealGame product)
 		public string ProductName;
@@ -524,8 +524,13 @@ namespace UnrealBuildTool.XcodeProjectXcconfig
 
 			if (AllConfigs.Any(x => x.bSupportsVisionOS))
 			{
-				VisionOSPlatform VisionOSPlatform = ((VisionOSPlatform)UEBuildPlatform.GetBuildPlatform(UnrealTargetPlatform.VisionOS));
-				VisionOSProjectSettings = VisionOSPlatform.ReadProjectSettings(UProjectFileLocation);
+				// this may not exist since it's a PlatformExtension and the VisionOS files may not be preset
+				UEBuildPlatform? BuildPlatform;
+				if (UEBuildPlatform.TryGetBuildPlatform(UnrealTargetPlatform.VisionOS, out BuildPlatform))
+				{
+					IOSPlatform VisionOSPlatform = (IOSPlatform)BuildPlatform;
+					VisionOSProjectSettings = VisionOSPlatform.ReadProjectSettings(UProjectFileLocation);
+				}
 			}
 
 			return true;
diff --git a/Engine/Source/Programs/UnrealBuildTool/ProjectFiles/Xcode/XcodeProjectFileGenerator.cs b/Engine/Source/Programs/UnrealBuildTool/ProjectFiles/Xcode/XcodeProjectFileGenerator.cs
--- a/Engine/Source/Programs/UnrealBuildTool/ProjectFiles/Xcode/XcodeProjectFileGenerator.cs
+++ b/Engine/Source/Programs/UnrealBuildTool/ProjectFiles/Xcode/XcodeProjectFileGenerator.cs
@@ -446,21 +446,12 @@ namespace UnrealBuildTool
 			else
 			{
 				// add platforms that have synced platform support
-				if (InstalledPlatformInfo.IsValidPlatform(UnrealTargetPlatform.Mac, EProjectType.Code))
+				foreach (UnrealTargetPlatform ApplePlatform in UEBuildPlatform.GetPlatformsInGroup(UnrealPlatformGroup.Apple))
 				{
-					XcodePlatforms.Add(UnrealTargetPlatform.Mac);
-				}
-				if (InstalledPlatformInfo.IsValidPlatform(UnrealTargetPlatform.IOS, EProjectType.Code))
-				{
-					XcodePlatforms.Add(UnrealTargetPlatform.IOS);
-				}
-				if (InstalledPlatformInfo.IsValidPlatform(UnrealTargetPlatform.TVOS, EProjectType.Code))
-				{
-					XcodePlatforms.Add(UnrealTargetPlatform.TVOS);
-				}
-				if (InstalledPlatformInfo.IsValidPlatform(UnrealTargetPlatform.VisionOS, EProjectType.Code))
-				{
-					XcodePlatforms.Add(UnrealTargetPlatform.VisionOS);
+					if (InstalledPlatformInfo.IsValidPlatform(ApplePlatform, EProjectType.Code) && UEBuildPlatform.TryGetBuildPlatform(ApplePlatform, out _))
+					{
+						XcodePlatforms.Add(ApplePlatform);
+					}
 				}
 			}
 
