Fixing null dereferences when IWYU hasn't compiled everything.

diff --git a/Engine/Source/Programs/UnrealBuildTool/Modes/IWYUMode.cs b/Engine/Source/Programs/UnrealBuildTool/Modes/IWYUMode.cs
--- a/Engine/Source/Programs/UnrealBuildTool/Modes/IWYUMode.cs
+++ b/Engine/Source/Programs/UnrealBuildTool/Modes/IWYUMode.cs
@@ -812,9 +812,9 @@ namespace UnrealBuildTool
 			Logger.LogInformation($"Generating infos for .generated.h files...");
 			if (GeneratedHeaderInfos.Count > 0)
 			{
-				IWYUIncludeEntry ObjectMacrosInclude = SpecialIncludes["UObject/ObjectMacros.h"]!;
-				IWYUIncludeEntry ScriptMacrosInclude = SpecialIncludes["UObject/ScriptMacros.h"]!;
-				IWYUIncludeEntry VerseInteropUtilsInclude = SpecialIncludes["VerseInteropUtils.h"]!;
+				IWYUIncludeEntry? ObjectMacrosInclude = SpecialIncludes["UObject/ObjectMacros.h"];
+				IWYUIncludeEntry? ScriptMacrosInclude = SpecialIncludes["UObject/ScriptMacros.h"];
+				IWYUIncludeEntry? VerseInteropUtilsInclude = SpecialIncludes["VerseInteropUtils.h"];
 
 				foreach (string Gen in GeneratedHeaderInfos)
 				{
@@ -822,15 +822,24 @@ namespace UnrealBuildTool
 					GenInfo.File = Gen;
 					if (Gen.EndsWith(".generated.h"))
 					{
-						GenInfo.IncludesSeenInFile.Add(ObjectMacrosInclude);
-						GenInfo.IncludesSeenInFile.Add(ScriptMacrosInclude);
-						GenInfo.Includes.Add(ObjectMacrosInclude);
-						GenInfo.Includes.Add(ScriptMacrosInclude);
+						if (ObjectMacrosInclude != null)
+						{
+							GenInfo.IncludesSeenInFile.Add(ObjectMacrosInclude);
+							GenInfo.Includes.Add(ObjectMacrosInclude);
+						}
+						if (ScriptMacrosInclude != null)
+						{
+							GenInfo.IncludesSeenInFile.Add(ScriptMacrosInclude);
+							GenInfo.Includes.Add(ScriptMacrosInclude);
+						}
 					}
 					else
 					{
-						GenInfo.IncludesSeenInFile.Add(VerseInteropUtilsInclude);
-						GenInfo.Includes.Add(VerseInteropUtilsInclude);
+						if (VerseInteropUtilsInclude != null)
+						{
+							GenInfo.IncludesSeenInFile.Add(VerseInteropUtilsInclude);
+							GenInfo.Includes.Add(VerseInteropUtilsInclude);
+						}
 					}
 					Infos.Add(Gen, GenInfo);
 				}
