[UBT]
* Fixed bugs related to header unit tags in sourcefile metadata cache. (tags are also used for iwyu and compile header paths)
* Added support for UE_ENABLE_INCLUDE_ORDER_DEPRECATED_ tags

diff --git a/Engine/Source/Programs/UnrealBuildTool/System/SourceFileMetadataCache.cs b/Engine/Source/Programs/UnrealBuildTool/System/SourceFileMetadataCache.cs
--- a/Engine/Source/Programs/UnrealBuildTool/System/SourceFileMetadataCache.cs
+++ b/Engine/Source/Programs/UnrealBuildTool/System/SourceFileMetadataCache.cs
@@ -89,7 +89,7 @@ namespace UnrealBuildTool
 		/// <summary>
 		/// The current file version
 		/// </summary>
-		public const int CurrentVersion = 7;
+		public const int CurrentVersion = 8;
 
 		/// <summary>
 		/// Location of this dependency cache
@@ -301,62 +301,6 @@ namespace UnrealBuildTool
 			}
 		}
 
-		/// <summary>
-		/// Attempts to find the include in the passed in line
-		/// </summary>
-		/// <param name="Line">Line from a file</param>
-		/// <returns>Include found or null if not found</returns>
-		static private string? FindInclude(string Line)
-		{
-			ReadOnlySpan<char> IncludeSpan = Line.AsSpan().TrimStart();
-			if (IncludeSpan.StartsWith("#include"))
-			{
-				IncludeSpan = IncludeSpan.Slice("#include".Length).TrimStart();
-				if (IncludeSpan.IsEmpty)
-				{
-					return null;
-				}
-
-				char EndChar;
-				bool TrimQuotation = true;
-				if (IncludeSpan[0] == '"')
-				{
-					EndChar = '"';
-				}
-				else if (IncludeSpan[0] == '<')
-				{
-					EndChar = '>';
-				}
-				else
-				{
-					EndChar = ')';
-					TrimQuotation = false;
-				}
-
-				if (TrimQuotation)
-				{
-					IncludeSpan = IncludeSpan.Slice(1);
-				}
-
-				int EndIndex = IncludeSpan.IndexOf(EndChar);
-				if (EndIndex == -1)
-				{
-					return null;
-				}
-
-				// This will include the ')' at the end
-				if (!TrimQuotation)
-				{
-					EndIndex++;
-				}
-
-				IncludeSpan = IncludeSpan.Slice(0, EndIndex);
-
-				return IncludeSpan.ToString();
-			}
-			return null;
-		}
-
 		/// <summary>
 		/// Read entire header file to find markup and includes
 		/// </summary>
@@ -365,41 +309,103 @@ namespace UnrealBuildTool
 		{
 			bool bContainsMarkup = false;
 			bool bUsesAPIDefine = false;
+			int InsideDeprecationScope = 0;
+
 			SortedSet<string> Includes = new();
 			foreach (string Line in FileText)
 			{
-				if (!bContainsMarkup)
-				{
-					bContainsMarkup = ReflectionMarkupRegex.IsMatch(Line);
-				}
-
 				if (!bUsesAPIDefine)
 				{
 					bUsesAPIDefine = Line.Contains("_API", StringComparison.Ordinal);
 				}
 
-				string? Include = FindInclude(Line);
-				if (Include == null)
+				if (!bContainsMarkup)
 				{
-					continue;
+					bContainsMarkup = ReflectionMarkupRegex.IsMatch(Line);
 				}
-				Includes.Add(Include);
 
-				int HeaderUnitIndex = Line.IndexOf("HEADER_UNIT_");
-				if (HeaderUnitIndex != -1)
+				ReadOnlySpan<char> LineSpan = Line.AsSpan().TrimStart();
+				if (LineSpan.StartsWith("#include") && InsideDeprecationScope == 0)
 				{
-					ReadOnlySpan<char> Span = Line.AsSpan(HeaderUnitIndex + "HEADER_UNIT_".Length);
-					if (Span.StartsWith("UNSUPPORTED"))
+					ReadOnlySpan<char> IncludeSpan = LineSpan.Slice("#include".Length).TrimStart();
+					if (IncludeSpan.IsEmpty)
+					{
+						continue;
+					}
+					char EndChar;
+					bool TrimQuotation = true;
+					if (IncludeSpan[0] == '"')
+					{
+						EndChar = '"';
+					}
+					else if (IncludeSpan[0] == '<')
+					{
+						EndChar = '>';
+					}
+					else
 					{
-						HeaderFileInfo.UnitType = HeaderUnitType.Unsupported;
+						EndChar = ')';
+						TrimQuotation = false;
 					}
-					else if (Span.StartsWith("SKIP"))
+
+					if (TrimQuotation)
+					{
+						IncludeSpan = IncludeSpan.Slice(1);
+					}
+
+					if (IncludeSpan.Contains("HEADER_UNIT_IGNORE", StringComparison.OrdinalIgnoreCase))
+					{
+						continue;
+					}
+
+					int EndIndex = IncludeSpan.IndexOf(EndChar);
+
+					if (EndIndex == -1)
 					{
-						HeaderFileInfo.UnitType = HeaderUnitType.Skip;
+						continue;
 					}
-					else if (Span.StartsWith("IGNORE"))
+
+					// This will include the ')' at the end
+					if (!TrimQuotation)
+					{
+						EndIndex++;
+					}
+
+					IncludeSpan = IncludeSpan.Slice(0, EndIndex);
+					Includes.Add(IncludeSpan.ToString());
+				}
+				else if (InsideDeprecationScope != 0)
+				{
+					if (LineSpan.StartsWith("#endif"))
+					{
+						--InsideDeprecationScope;
+					}
+					else if (LineSpan.StartsWith("#if"))
+					{
+						++InsideDeprecationScope;
+					}
+				}
+				else if (LineSpan.StartsWith("#if"))
+				{
+					if (Line.IndexOf("UE_ENABLE_INCLUDE_ORDER_DEPRECATED_") != -1)
+					{
+						++InsideDeprecationScope;
+					}
+				}
+				else
+				{
+					int HeaderUnitIndex = Line.IndexOf("HEADER_UNIT_");
+					if (HeaderUnitIndex != -1)
 					{
-						HeaderFileInfo.UnitType = HeaderUnitType.Ignore;
+						ReadOnlySpan<char> Span = Line.AsSpan(HeaderUnitIndex + "HEADER_UNIT_".Length);
+						if (Span.StartsWith("UNSUPPORTED"))
+						{
+							HeaderFileInfo.UnitType = HeaderUnitType.Unsupported;
+						}
+						else if (Span.StartsWith("SKIP"))
+						{
+							HeaderFileInfo.UnitType = HeaderUnitType.Skip;
+						}
 					}
 				}
 			}
