UnrealBuildTool: Fix fastbuild not handling nested response files

#rnx

diff --git a/Engine/Source/Programs/UnrealBuildTool/Executors/Experimental/FASTBuild.cs b/Engine/Source/Programs/UnrealBuildTool/Executors/Experimental/FASTBuild.cs
--- a/Engine/Source/Programs/UnrealBuildTool/Executors/Experimental/FASTBuild.cs
+++ b/Engine/Source/Programs/UnrealBuildTool/Executors/Experimental/FASTBuild.cs
@@ -620,6 +620,16 @@ namespace UnrealBuildTool
 					}
 				}
 
+				// Handle nested response files
+				if (Token.StartsWith('@'))
+				{
+					foreach (KeyValuePair<string, string> Pair in ParseCommandLineOptions(LocalToolName, Token, SpecialOptions, Logger))
+					{
+						ParsedCompilerOptions.Add(Pair.Key, Pair.Value);
+					}
+					continue;
+				}
+
 				// Defines can have escaped quotes and other strings inside them
 				// so we consume tokens until we've closed any open unescaped parentheses.
 				if ((Token.StartsWith("/D") || Token.StartsWith("-D")) && !QuotesOpened)
@@ -762,6 +772,11 @@ namespace UnrealBuildTool
 				}
 			}
 
+			if (ParsedCompilerOptions.ContainsKey("OtherOptions"))
+			{
+				ProcessedTokens.Insert(0, ParsedCompilerOptions["OtherOptions"]);
+			}
+
 			ParsedCompilerOptions["OtherOptions"] = String.Join(" ", ProcessedTokens) + " ";
 
 			if (SaveResponseFile && !String.IsNullOrEmpty(ResponseFilePath))
