[UBT]
* Fixed typo in code looking for experimental versions of vs2022 (it was using preview path instead of experimental path)

diff --git a/Engine/Source/Programs/UnrealBuildTool/Platform/Windows/MicrosoftPlatformSDK.cs b/Engine/Source/Programs/UnrealBuildTool/Platform/Windows/MicrosoftPlatformSDK.cs
--- a/Engine/Source/Programs/UnrealBuildTool/Platform/Windows/MicrosoftPlatformSDK.cs
+++ b/Engine/Source/Programs/UnrealBuildTool/Platform/Windows/MicrosoftPlatformSDK.cs
@@ -805,7 +805,7 @@ namespace UnrealBuildTool
 								FindVisualStudioToolChains(PreviewBaseDir, null, WindowsCompilerChannel.Preview, ToolChains, IsAutoSdk: true, Logger);
 
 								DirectoryReference ExperimentalBaseDir = DirectoryReference.Combine(PlatformDir, "Win64", $"{VSDir}-Experimental");
-								FindVisualStudioToolChains(PreviewBaseDir, null, WindowsCompilerChannel.Experimental, ToolChains, IsAutoSdk: true, Logger);
+								FindVisualStudioToolChains(ExperimentalBaseDir, null, WindowsCompilerChannel.Experimental, ToolChains, IsAutoSdk: true, Logger);
 							}
 						}
 					}
