UnrealBuildTool: Rename .suppressed.lib -> .sup.lib to reduce filename length

#rnx

diff --git a/Engine/Source/Programs/UnrealBuildTool/Platform/Windows/VCToolChain.cs b/Engine/Source/Programs/UnrealBuildTool/Platform/Windows/VCToolChain.cs
--- a/Engine/Source/Programs/UnrealBuildTool/Platform/Windows/VCToolChain.cs
+++ b/Engine/Source/Programs/UnrealBuildTool/Platform/Windows/VCToolChain.cs
@@ -2316,7 +2316,7 @@ namespace UnrealBuildTool
 			FileReference ImportLibraryFilePath;
 			if (LinkEnvironment.bIsCrossReferenced && !bBuildImportLibraryOnly)
 			{
-				ImportLibraryFilePath = FileReference.Combine(LinkEnvironment.IntermediateDirectory!, LinkEnvironment.OutputFilePath.GetFileNameWithoutExtension() + ".suppressed.lib");
+				ImportLibraryFilePath = FileReference.Combine(LinkEnvironment.IntermediateDirectory!, LinkEnvironment.OutputFilePath.GetFileNameWithoutExtension() + ".sup.lib");
 			}
 			else if (Target.bShouldCompileAsDLL)
 			{
@@ -2484,8 +2484,8 @@ namespace UnrealBuildTool
 				LinkAction.DeleteItems.UnionWith(LinkAction.ProducedItems.Where(x => x.Location.HasExtension(".pdb")));
 			}
 
-			// Delete any .suppressed.lib files before building, even if they're not tracked
-			if (ImportLibraryFilePath.GetFileName().EndsWith(".suppressed.lib"))
+			// Delete any .sup.lib files before building, even if they're not tracked
+			if (ImportLibraryFilePath.GetFileName().EndsWith(".sup.lib", StringComparison.OrdinalIgnoreCase))
 			{
 				FileReference ExportFilePath = ImportLibraryFilePath.ChangeExtension(".exp");
 				LinkAction.DeleteItems.Add(FileItem.GetItemByFileReference(ImportLibraryFilePath));
