Fix various issues with Installed build and modern, removed references of Payload while maintaining create stub

diff --git a/Engine/Source/Programs/UnrealBuildTool/Platform/IOS/IOSToolChain.cs b/Engine/Source/Programs/UnrealBuildTool/Platform/IOS/IOSToolChain.cs
--- a/Engine/Source/Programs/UnrealBuildTool/Platform/IOS/IOSToolChain.cs
+++ b/Engine/Source/Programs/UnrealBuildTool/Platform/IOS/IOSToolChain.cs
@@ -147,23 +147,27 @@ namespace UnrealBuildTool
 				FileReference StubFile = FileReference.Combine(Binary.OutputFilePath.Directory, Binary.OutputFilePath.GetFileNameWithoutExtension() + ".stub");
 				BuildProducts.Add(StubFile, BuildProductType.Package);
 
-				if (Target.Platform == UnrealTargetPlatform.TVOS)
+				// These Payload assets and .app bundle are no longer needed for Modern Xcode
+				if (!bUseModernXcode)
 				{
-					FileReference AssetFile = FileReference.Combine(Binary.OutputFilePath.Directory, "AssetCatalog", "Assets.car");
-					BuildProducts.Add(AssetFile, BuildProductType.RequiredResource);
-				}
-				else if (Target.Platform == UnrealTargetPlatform.IOS)
-				{
-					int Index = Binary.OutputFilePath.GetFileNameWithoutExtension().IndexOf("-");
-					string OutputFile = Binary.OutputFilePath.GetFileNameWithoutExtension().Substring(0, Index > 0 ? Index : Binary.OutputFilePath.GetFileNameWithoutExtension().Length);
-					FileReference AssetFile = FileReference.Combine(Binary.OutputFilePath.Directory, "Payload", OutputFile + ".app", "Assets.car");
-					BuildProducts.Add(AssetFile, BuildProductType.RequiredResource);
-
-					// Default AppIcons need to be copied as is in the Payload's root, even when using an asset catalog.
-					AssetFile = FileReference.Combine(Binary.OutputFilePath.Directory, "Payload", OutputFile + ".app", "[REDACTED_EMAIL]");
-					BuildProducts.Add(AssetFile, BuildProductType.RequiredResource);
-					AssetFile = FileReference.Combine(Binary.OutputFilePath.Directory, "Payload", OutputFile + ".app", "AppIcon76x76@2x~ipad.png");
-					BuildProducts.Add(AssetFile, BuildProductType.RequiredResource);
+					if (Target.Platform == UnrealTargetPlatform.TVOS)
+					{
+						FileReference AssetFile = FileReference.Combine(Binary.OutputFilePath.Directory, "AssetCatalog", "Assets.car");
+						BuildProducts.Add(AssetFile, BuildProductType.RequiredResource);
+					}
+					else if (Target.Platform == UnrealTargetPlatform.IOS)
+					{
+						int Index = Binary.OutputFilePath.GetFileNameWithoutExtension().IndexOf("-");
+						string OutputFile = Binary.OutputFilePath.GetFileNameWithoutExtension().Substring(0, Index > 0 ? Index : Binary.OutputFilePath.GetFileNameWithoutExtension().Length);
+						FileReference AssetFile = FileReference.Combine(Binary.OutputFilePath.Directory, "Payload", OutputFile + ".app", "Assets.car");
+						BuildProducts.Add(AssetFile, BuildProductType.RequiredResource);
+
+						// Default AppIcons need to be copied as is in the Payload's root, even when using an asset catalog.
+						AssetFile = FileReference.Combine(Binary.OutputFilePath.Directory, "Payload", OutputFile + ".app", "[REDACTED_EMAIL]");
+						BuildProducts.Add(AssetFile, BuildProductType.RequiredResource);
+						AssetFile = FileReference.Combine(Binary.OutputFilePath.Directory, "Payload", OutputFile + ".app", "AppIcon76x76@2x~ipad.png");
+						BuildProducts.Add(AssetFile, BuildProductType.RequiredResource);
+					}
 				}
 			}
 			if (Target.IOSPlatform.bGeneratedSYM && ProjectSettings.bGenerateCrashReportSymbols && Binary.Type == UEBuildBinaryType.StaticLibrary)
@@ -770,7 +774,7 @@ namespace UnrealBuildTool
 			return OutputFile;
 		}
 
-		private static void PackageStub(string BinaryPath, string GameName, string ExeName)
+		public static void PackageStub(string BinaryPath, string GameName, string ExeName, bool bUseModernXcode)
 		{
 			// create the ipa
 			string IPAName = BinaryPath + "/" + ExeName + ".stub";
@@ -788,8 +792,8 @@ namespace UnrealBuildTool
 			}
 
 			// set up the directories
-			string ZipWorkingDir = String.Format("Payload/{0}.app/", GameName);
-			string ZipSourceDir = String.Format("{0}/Payload/{1}.app", BinaryPath, GameName);
+			string ZipWorkingDir = String.Format(bUseModernXcode ? "{0}.app/" : "Payload/{0}.app/", GameName);
+			string ZipSourceDir = String.Format(bUseModernXcode ? "{0}/{1}.app" : "{0}/Payload/{1}.app", BinaryPath, GameName);
 
 			// create the file
 			using (ZipFile Zip = new ZipFile())
@@ -1742,11 +1746,11 @@ namespace UnrealBuildTool
 				// Package the stub
 				if (Target.bBuildAsFramework)
 				{
-					PackageStub(FrameworkPayloadDirectory.ParentDirectory!.ToString(), AppName, Target.OutputPath.GetFileNameWithoutExtension());
+					PackageStub(FrameworkPayloadDirectory.ParentDirectory!.ToString(), AppName, Target.OutputPath.GetFileNameWithoutExtension(), AppleExports.UseModernXcode(Target.ProjectFile));
 				}
 				else
 				{
-					PackageStub(RemoteShadowDirectoryMac, AppName, Target.OutputPath.GetFileNameWithoutExtension());
+					PackageStub(RemoteShadowDirectoryMac, AppName, Target.OutputPath.GetFileNameWithoutExtension(), AppleExports.UseModernXcode(Target.ProjectFile));
 				}
 			}
 			else
diff --git a/Engine/Source/Programs/UnrealBuildTool/Platform/IOS/UEDeployIOS.cs b/Engine/Source/Programs/UnrealBuildTool/Platform/IOS/UEDeployIOS.cs
--- a/Engine/Source/Programs/UnrealBuildTool/Platform/IOS/UEDeployIOS.cs
+++ b/Engine/Source/Programs/UnrealBuildTool/Platform/IOS/UEDeployIOS.cs
@@ -1128,6 +1128,7 @@ namespace UnrealBuildTool
 			if (AppleExports.UseModernXcode(ProjectFile))
 			{
 				Logger.LogInformation("Generating plist (only step needed when deploying with Modern Xcode)");
+				AppDirectory = BinaryPath + "/" + GameExeName + ".app";
 				GeneratePList(ProjectFile, Config, InProjectDirectory, bIsUnrealGame, GameExeName, false, InProjectName, InEngineDir, AppDirectory, UPLScripts, BundleID, bBuildAsFramework);
 
 				//// for now, copy the executable into the .app
diff --git a/Engine/Source/Programs/UnrealBuildTool/ToolChain/AppleToolChain.cs b/Engine/Source/Programs/UnrealBuildTool/ToolChain/AppleToolChain.cs
--- a/Engine/Source/Programs/UnrealBuildTool/ToolChain/AppleToolChain.cs
+++ b/Engine/Source/Programs/UnrealBuildTool/ToolChain/AppleToolChain.cs
@@ -698,12 +698,19 @@ namespace UnrealBuildTool
 		public UnrealTargetConfiguration Configuration;
 		public string TargetName;
 
-		public ApplePostBuildSyncTarget(ReadOnlyTargetRules Target)
+		// For iOS/TVOS
+		public bool bCreateStubIPA;
+		public FileReference StubOutputPath;
+
+		public ApplePostBuildSyncTarget(ReadOnlyTargetRules Target, FileItem Executable)
 		{
 			Platform = Target.Platform;
 			Configuration = Target.Configuration;
 			ProjectFile = Target.ProjectFile;
 			TargetName = Target.Name;
+
+			bCreateStubIPA = Target.IOSPlatform.bCreateStubIPA;
+			StubOutputPath = Executable.Location;
 		}
 	}
 
@@ -756,6 +763,12 @@ namespace UnrealBuildTool
 				Logger.LogError("ERROR: Failed to finalize the .app with Xcode. Check the log for more information");
 			}
 
+			if ((Target.Platform == UnrealTargetPlatform.IOS || Target.Platform == UnrealTargetPlatform.TVOS)
+				&& Target.bCreateStubIPA)
+			{
+				IOSToolChain.PackageStub(Target.StubOutputPath.Directory.FullName, Target.TargetName, Target.StubOutputPath.GetFileNameWithoutExtension(), true);
+			}
+
 			return ExitCode;
 		}
 
@@ -775,7 +788,7 @@ namespace UnrealBuildTool
 
 		public static Action CreatePostBuildSyncAction(ReadOnlyTargetRules Target, FileItem Executable, DirectoryReference IntermediateDir, IActionGraphBuilder Graph)
 		{
-			ApplePostBuildSyncTarget PostBuildSync = new(Target);
+			ApplePostBuildSyncTarget PostBuildSync = new(Target, Executable);
 			FileReference PostBuildSyncFile = FileReference.Combine(IntermediateDir!, "PostBuildSync.dat");
 			BinaryFormatterUtils.Save(PostBuildSyncFile, PostBuildSync);
 
@@ -788,12 +801,17 @@ namespace UnrealBuildTool
 			PostBuildSyncAction.StatusDescription = $"Executing PostBuildSync [{Executable.Location}]";
 			PostBuildSyncAction.bCanExecuteRemotely = false;
 
-
 			if (Target.Platform == UnrealTargetPlatform.IOS || Target.Platform == UnrealTargetPlatform.TVOS)
 			{
 				// @todo: do we need one per Target for IOS? Client? I dont think so for Modern
 				FileReference PlistFile = FileReference.Combine(AppleToolChain.GetActualProjectDirectory(Target.ProjectFile), "Build/IOS/UBTGenerated/Info.Template.plist");
 				PostBuildSyncAction.ProducedItems.Add(FileItem.GetItemByFileReference(PlistFile));
+
+				if (PostBuildSync.bCreateStubIPA)
+				{
+					FileReference StubFile = FileReference.Combine(Executable.Directory.Location, Executable.Location.GetFileNameWithoutExtension() + ".stub");
+					PostBuildSyncAction.ProducedItems.Add(FileItem.GetItemByFileReference(StubFile));
+				}
 			}
 
 			return PostBuildSyncAction;
