Adding a log for IWYUMode

diff --git a/Engine/Source/Programs/UnrealBuildTool/Modes/IWYUMode.cs b/Engine/Source/Programs/UnrealBuildTool/Modes/IWYUMode.cs
--- a/Engine/Source/Programs/UnrealBuildTool/Modes/IWYUMode.cs
+++ b/Engine/Source/Programs/UnrealBuildTool/Modes/IWYUMode.cs
@@ -3,6 +3,7 @@
 using System;
 using System.Collections.Concurrent;
 using System.Collections.Generic;
+using System.Diagnostics;
 using System.Diagnostics.CodeAnalysis;
 using System.IO;
 using System.Linq;
@@ -202,6 +203,12 @@ namespace UnrealBuildTool
 	[ToolMode("IWYU", ToolModeOptions.XmlConfig | ToolModeOptions.BuildPlatforms | ToolModeOptions.SingleInstance | ToolModeOptions.StartPrefetchingEngine | ToolModeOptions.ShowExecutionTime)]
 	class IWYUMode : ToolMode
 	{
+		/// <summary>
+		/// Specifies the file to use for logging.
+		/// </summary>
+		[XmlConfigFile(Category = "BuildConfiguration")]
+		public string? IWYUBaseLogFileName;
+
 		/// <summary>
 		/// Will check out files from p4 and write to disk
 		/// </summary>
@@ -349,6 +356,26 @@ namespace UnrealBuildTool
 
 			Arguments.ApplyTo(this);
 
+			// Fixup the log path if it wasn't overridden by a config file
+			if (IWYUBaseLogFileName == null)
+			{
+				IWYUBaseLogFileName = FileReference.Combine(Unreal.EngineProgramSavedDirectory, "UnrealBuildTool", "IWYULog.txt").FullName;
+			}
+
+			// Create the log file, and flush the startup listener to it
+			if (!Arguments.HasOption("-NoLog") && !Log.HasFileWriter())
+			{
+				Log.AddFileWriter("DefaultLogTraceListener", FileReference.FromString(IWYUBaseLogFileName));
+			}
+			else
+			{
+				IEnumerable<StartupTraceListener> StartupListeners = Trace.Listeners.OfType<StartupTraceListener>();
+				if (StartupListeners.Any())
+				{
+					Trace.Listeners.Remove(StartupListeners.First());
+				}
+			}
+
 			// Create the build configuration object, and read the settings
 			CommandLineArguments BuildArguments = Arguments.Append(new[] { "-IWYU" }); // Add in case it is not added (it is needed for iwyu toolchain)
 			BuildConfiguration BuildConfiguration = new BuildConfiguration();
