UnrealBuildTool: Ignore ReleaseChannel for clang toolchains

#rnx

diff --git a/Engine/Source/Programs/UnrealBuildTool/Platform/Windows/MicrosoftPlatformSDK.cs b/Engine/Source/Programs/UnrealBuildTool/Platform/Windows/MicrosoftPlatformSDK.cs
--- a/Engine/Source/Programs/UnrealBuildTool/Platform/Windows/MicrosoftPlatformSDK.cs
+++ b/Engine/Source/Programs/UnrealBuildTool/Platform/Windows/MicrosoftPlatformSDK.cs
@@ -563,9 +563,9 @@ namespace UnrealBuildTool
 
 			foreach (WindowsCompilerChannel Channel in Enum.GetValues(typeof(WindowsCompilerChannel)))
 			{
-				if (String.Compare(CompilerVersion, Channel.ToString(), StringComparison.InvariantCultureIgnoreCase) == 0)
+				if (String.Equals(CompilerVersion, Channel.ToString(), StringComparison.OrdinalIgnoreCase))
 				{
-					ToolChain = SelectToolChain(ToolChains.Where(x => x.ReleaseChannel == Channel), x => x.ThenByDescending(x => x.Version), Architecture);
+					ToolChain = SelectToolChain(ToolChains.Where(x => x.ReleaseChannel.HasFlag(Channel)), x => x.ThenByDescending(x => x.Version), Architecture);
 					if (ToolChain == null)
 					{
 						DumpToolChains(ToolChains, x => x.ThenBy(x => x.ReleaseChannel).ThenByDescending(x => x.Version), Architecture, Logger);
@@ -1055,7 +1055,7 @@ namespace UnrealBuildTool
 				}
 
 				Logger.LogDebug("Found Clang toolchain: {ToolChainDir} (Version={Version}, Is64Bit={Is64Bit}, Rank={Rank}, Error={Error})", ToolChainDir, Version, Is64Bit, Rank, Error != null);
-				ToolChains.Add(new ToolChainInstallation(Family, Rank, Version, Is64Bit, WindowsCompilerChannel.Latest, UnrealArch.X64, Error, ToolChainDir, null, IsAutoSdk));
+				ToolChains.Add(new ToolChainInstallation(Family, Rank, Version, Is64Bit, WindowsCompilerChannel.Any, UnrealArch.X64, Error, ToolChainDir, null, IsAutoSdk));
 			}
 		}
 
@@ -1091,7 +1091,7 @@ namespace UnrealBuildTool
 				}
 
 				Logger.LogDebug("Found Intel OneAPI toolchain: {ToolChainDir} (Version={Version}, Is64Bit={Is64Bit}, Rank={Rank}, Error={Error})", ToolChainDir, Version, Is64Bit, Rank, Error != null);
-				ToolChains.Add(new ToolChainInstallation(Family, Rank, Version, Is64Bit, WindowsCompilerChannel.Latest, UnrealArch.X64, Error, ToolChainDir, null, IsAutoSdk));
+				ToolChains.Add(new ToolChainInstallation(Family, Rank, Version, Is64Bit, WindowsCompilerChannel.Any, UnrealArch.X64, Error, ToolChainDir, null, IsAutoSdk));
 			}
 		}
 
diff --git a/Engine/Source/Programs/UnrealBuildTool/Platform/Windows/UEBuildWindows.cs b/Engine/Source/Programs/UnrealBuildTool/Platform/Windows/UEBuildWindows.cs
--- a/Engine/Source/Programs/UnrealBuildTool/Platform/Windows/UEBuildWindows.cs
+++ b/Engine/Source/Programs/UnrealBuildTool/Platform/Windows/UEBuildWindows.cs
@@ -150,6 +150,7 @@ namespace UnrealBuildTool
 		Latest,
 		Preview,
 		Experimental,
+		Any = Latest | Preview | Experimental
 	}
 
 	/// <summary>
