Update Linux to use v22 toolchain (clang 16.0.6)

diff --git a/Engine/Config/Linux/Linux_SDK.json b/Engine/Config/Linux/Linux_SDK.json
--- a/Engine/Config/Linux/Linux_SDK.json
+++ b/Engine/Config/Linux/Linux_SDK.json
@@ -1,7 +1,7 @@
 {
-"MainVersion" : "v21_clang-15.0.1-centos7",
-"MinVersion" : "v21_clang-15.0.1-centos7",
-"MaxVersion" : "v21_clang-15.0.1-centos7",
+"MainVersion" : "v22_clang-16.0.6-centos7",
+"MinVersion" : "v22_clang-16.0.6-centos7",
+"MaxVersion" : "v22_clang-16.0.6-centos7",
 
 "AutoSDKPlatform" : "Linux_x64",
 }
diff --git a/Engine/Source/Programs/UnrealBuildTool/Platform/Linux/LinuxToolChain.cs b/Engine/Source/Programs/UnrealBuildTool/Platform/Linux/LinuxToolChain.cs
--- a/Engine/Source/Programs/UnrealBuildTool/Platform/Linux/LinuxToolChain.cs
+++ b/Engine/Source/Programs/UnrealBuildTool/Platform/Linux/LinuxToolChain.cs
@@ -55,10 +55,10 @@ namespace UnrealBuildTool
 			: this(UnrealTargetPlatform.Linux, InArchitecture, InSDK, InOptions, InLogger)
 		{
 			// prevent unknown clangs since the build is likely to fail on too old or too new compilers
-			if (CompilerVersionLessThan(15, 0, 0) || CompilerVersionGreaterOrEqual(16, 0, 0))
+			if (CompilerVersionLessThan(16, 0, 0) || CompilerVersionGreaterOrEqual(17, 0, 0))
 			{
 				throw new BuildException(
-					String.Format("This version of the Unreal Engine can only be compiled with clang 15.0. clang {0} may not build it - please use a different version.",
+					String.Format("This version of the Unreal Engine can only be compiled with clang 16.0. clang {0} may not build it - please use a different version.",
 						Info.ClangVersion)
 					);
 			}
@@ -310,7 +310,7 @@ namespace UnrealBuildTool
 				{
 					if (bForLinker)
 					{
-						Arguments.Add(String.Format("-flto=thin -Wl,--thinlto-jobs={0}", Utils.GetPhysicalProcessorCount()));
+						Arguments.Add(String.Format("-flto=thin -Wl,--thinlto-jobs={0}, -Wl,-mllvm,-disable-auto-upgrade-debug-info", Utils.GetPhysicalProcessorCount()));
 					}
 					else
 					{
