Allow building Linux with clang 15.x to unblock some jobs.

diff --git a/Engine/Source/Programs/UnrealBuildTool/Platform/Linux/LinuxToolChain.cs b/Engine/Source/Programs/UnrealBuildTool/Platform/Linux/LinuxToolChain.cs
--- a/Engine/Source/Programs/UnrealBuildTool/Platform/Linux/LinuxToolChain.cs
+++ b/Engine/Source/Programs/UnrealBuildTool/Platform/Linux/LinuxToolChain.cs
@@ -55,10 +55,10 @@ namespace UnrealBuildTool
 			: this(UnrealTargetPlatform.Linux, InArchitecture, InSDK, InOptions, InLogger)
 		{
 			// prevent unknown clangs since the build is likely to fail on too old or too new compilers
-			if (CompilerVersionLessThan(16, 0, 0) || CompilerVersionGreaterOrEqual(17, 0, 0))
+			if (CompilerVersionLessThan(15, 0, 0) || CompilerVersionGreaterOrEqual(17, 0, 0))
 			{
 				throw new BuildException(
-					String.Format("This version of the Unreal Engine can only be compiled with clang 16.0. clang {0} may not build it - please use a different version.",
+					String.Format("This version of the Unreal Engine can only be compiled with clang 15.x or 16.x. clang {0} may not build it - please use a different version.",
 						Info.ClangVersion)
 					);
 			}
