Minor: TestRunner.cpp --no-wait option fix, generate metadata config must be in target file not build file.
Enable in BaseEngie.ini
[LowLevelTestsSettings]
bUpdateBuildGraphPropertiesFile=false

diff --git a/Engine/Config/BaseEngine.ini b/Engine/Config/BaseEngine.ini
--- a/Engine/Config/BaseEngine.ini
+++ b/Engine/Config/BaseEngine.ini
@@ -3304,3 +3304,6 @@ EnvironmentCollisionMaxShapes = 32
 [BuildPatchServices]
 ; Enable/Disable mechanism that provides a disk overflow store for patch chunks.
 bEnableDiskOverflowStore=true
+
+[LowLevelTestsSettings]
+bUpdateBuildGraphPropertiesFile=false
\ No newline at end of file
diff --git a/Engine/Source/Developer/LowLevelTestsRunner/Private/TestRunner.cpp b/Engine/Source/Developer/LowLevelTestsRunner/Private/TestRunner.cpp
--- a/Engine/Source/Developer/LowLevelTestsRunner/Private/TestRunner.cpp
+++ b/Engine/Source/Developer/LowLevelTestsRunner/Private/TestRunner.cpp
@@ -210,7 +210,7 @@ void FTestRunner::ParseCommandLine(TConstArrayView<const ANSICHAR*> Args)
 		}
 		else if (Arg == ANSITEXTVIEW("--no-wait"))
 		{
-			bWaitForInputToTerminate = true;
+			bWaitForInputToTerminate = false;
 		}
 		else if (Arg == ANSITEXTVIEW("--attach-to-debugger"))
 		{
diff --git a/Engine/Source/Programs/UnrealBuildTool/Configuration/TestModuleRules.cs b/Engine/Source/Programs/UnrealBuildTool/Configuration/TestModuleRules.cs
--- a/Engine/Source/Programs/UnrealBuildTool/Configuration/TestModuleRules.cs
+++ b/Engine/Source/Programs/UnrealBuildTool/Configuration/TestModuleRules.cs
@@ -15,9 +15,6 @@ namespace UnrealBuildTool
 	/// </summary>
 	public class TestModuleRules : ModuleRules
 	{
-		[ConfigFile(ConfigHierarchyType.Engine, "LowLevelTestsSettings", "UpdateBuildGraphPropertiesFile")]
-		bool bUpdateBuildGraphPropertiesFile = false;
-
 		private readonly XNamespace BuildGraphNamespace = XNamespace.Get("http://www.epicgames.com/BuildGraph");
 		private readonly XNamespace SchemaInstance = XNamespace.Get("http://www.w3.org/2001/XMLSchema-instance");
 		private readonly XNamespace SchemaLocation = XNamespace.Get("http://www.epicgames.com/BuildGraph ../../Build/Graph/Schema.xsd");
@@ -166,7 +163,16 @@ namespace UnrealBuildTool
 		/// </summary>
 		protected void UpdateBuildGraphPropertiesFile(Metadata TestMetadata)
 		{
-			if (Environment.GetEnvironmentVariable("IsBuildMachine") != null || !bUpdateBuildGraphPropertiesFile || TestMetadata == null)
+			bool bUpdateBuildGraphPropertiesFile = false;
+			TestTargetRules? TestTargetRules = Target.InnerTestTargetRules;
+			if (TestTargetRules != null)
+			{
+				bUpdateBuildGraphPropertiesFile = TestTargetRules.bUpdateBuildGraphPropertiesFile;
+			}
+
+			string? bIsBuildMachineEnv = Environment.GetEnvironmentVariable("IsBuildMachine");
+			bool bIsBuildMachine = bIsBuildMachineEnv != null && bIsBuildMachineEnv.Trim().Equals(1.ToString());
+			if (bIsBuildMachine || !bUpdateBuildGraphPropertiesFile || TestMetadata == null)
 			{
 				return;
 			}
diff --git a/Engine/Source/Programs/UnrealBuildTool/Configuration/TestTargetRules.cs b/Engine/Source/Programs/UnrealBuildTool/Configuration/TestTargetRules.cs
--- a/Engine/Source/Programs/UnrealBuildTool/Configuration/TestTargetRules.cs
+++ b/Engine/Source/Programs/UnrealBuildTool/Configuration/TestTargetRules.cs
@@ -14,6 +14,12 @@ namespace UnrealBuildTool
 	/// </summary>
 	public class TestTargetRules : TargetRules
 	{
+		/// <summary>
+		/// Configuration mapping to control build graph test metadata generation during project files generation.
+		/// </summary>
+		[ConfigFile(ConfigHierarchyType.Engine, "LowLevelTestsSettings")]
+		public bool bUpdateBuildGraphPropertiesFile = false;
+
 		/// <summary>
 		/// Keeps track if low level tests executable must build with the Editor.
 		/// </summary>
@@ -57,7 +63,7 @@ namespace UnrealBuildTool
 		/// <summary>
 		/// If set to true, it will not compile against ApplicationCore even if "ApplicationCore" is in the dependency graph.
 		/// </summary>
-		public static bool bNeverCompileAgainstApplicationCore = false;
+		public bool bNeverCompileAgainstApplicationCore = false;
 
 		/// <summary>
 		/// Test target override for bCompileAgainstCoreUObject.
@@ -72,7 +78,7 @@ namespace UnrealBuildTool
 		/// <summary>
 		/// If set to true, it will not compile against CoreUObject even if "CoreUObject" is in the dependency graph.
 		/// </summary>
-		public static bool bNeverCompileAgainstCoreUObject = false;
+		public bool bNeverCompileAgainstCoreUObject = false;
 
 		/// <summary>
 		/// Test target override for bCompileAgainstEngine.
@@ -87,7 +93,7 @@ namespace UnrealBuildTool
 		/// <summary>
 		/// If set to true, it will not compile against engine even if "Engine" is in the dependency graph.
 		/// </summary>
-		public static bool bNeverCompileAgainstEngine = false;
+		public bool bNeverCompileAgainstEngine = false;
 
 		/// <summary>
 		/// Test target override for bCompileAgainstEditor.
@@ -102,7 +108,7 @@ namespace UnrealBuildTool
 		/// <summary>
 		/// If set to true, it will not compile against editor even if "UnrealEd" is in the dependency graph.
 		/// </summary>
-		public static bool bNeverCompileAgainstEditor = false;
+		public bool bNeverCompileAgainstEditor = false;
 
 		/// <summary>
 		/// Whether to stub the platform file.
