UnrealBuildTool: Use shortname for module shared definition & response file

#rnx

diff --git a/Engine/Source/Programs/UnrealBuildTool/Configuration/UEBuildModuleCPP.cs b/Engine/Source/Programs/UnrealBuildTool/Configuration/UEBuildModuleCPP.cs
--- a/Engine/Source/Programs/UnrealBuildTool/Configuration/UEBuildModuleCPP.cs
+++ b/Engine/Source/Programs/UnrealBuildTool/Configuration/UEBuildModuleCPP.cs
@@ -470,7 +470,7 @@ namespace UnrealBuildTool
 			CreateHeaderForDefinitions(CompileEnvironment, IntermediateDirectory, null, Graph);
 
 			// Create shared rsp for the normal cpp files
-			FileReference SharedResponseFile = FileReference.Combine(IntermediateDirectory, $"{Name}.Shared{UEToolChain.ResponseExt}");
+			FileReference SharedResponseFile = FileReference.Combine(IntermediateDirectory, $"{Rules.ShortName ?? Name}.Shared{UEToolChain.ResponseExt}");
 			CompileEnvironment = ToolChain.CreateSharedResponseFile(CompileEnvironment, SharedResponseFile, Graph);
 
 			// Mapping of source file to unity file. We output this to intermediate directories for other tools (eg. live coding) to use.
@@ -1537,7 +1537,7 @@ namespace UnrealBuildTool
 					{
 						PrecompiledHeaderInstance Instance = FindOrCreateSharedPCH(ToolChain, Template, CompileEnvironment, Graph);
 
-						FileReference PrivateDefinitionsFile = FileReference.Combine(IntermediateDirectory, String.Format("Definitions.{0}.h", Name));
+						FileReference PrivateDefinitionsFile = FileReference.Combine(IntermediateDirectory, $"Definitions.{Rules.ShortName ?? Name}.h");
 
 						FileItem PrivateDefinitionsFileItem;
 						{
