UnrealBuildTool: Add additional info for how plugin modules are found when -AllModules is used

#rnx

diff --git a/Engine/Source/Programs/UnrealBuildTool/Configuration/UEBuildTarget.cs b/Engine/Source/Programs/UnrealBuildTool/Configuration/UEBuildTarget.cs
--- a/Engine/Source/Programs/UnrealBuildTool/Configuration/UEBuildTarget.cs
+++ b/Engine/Source/Programs/UnrealBuildTool/Configuration/UEBuildTarget.cs
@@ -3708,7 +3708,8 @@ namespace UnrealBuildTool
 			}
 
 			// Add all the plugin modules that need to be compiled
-			List<PluginInfo> Plugins = RulesAssembly.EnumeratePlugins().Where(x => x.ChoiceVersion != null).Select(x => x.ChoiceVersion!).ToList();
+			IEnumerable<PluginInfo> Plugins = RulesAssembly.EnumeratePlugins().Select(x => x.ChoiceVersion).OfType<PluginInfo>();
+			Dictionary<string, PluginInfo> ModulePluginSource = new();
 			foreach (PluginInfo Plugin in Plugins)
 			{
 				// Ignore plugins which are specifically disabled by this target
@@ -3762,7 +3763,10 @@ namespace UnrealBuildTool
 							Logger.LogDebug("Excluding plugin {Plugin} module {Module}: Excluded folder name", Plugin.Name, ModuleDescriptor.Name);
 							continue;
 						}
-						FilteredModuleNames.Add(ModuleDescriptor.Name);
+						if (FilteredModuleNames.Add(ModuleDescriptor.Name))
+						{
+							ModulePluginSource.Add(ModuleDescriptor.Name, Plugin);
+						}
 					}
 				}
 			}
@@ -3774,7 +3778,12 @@ namespace UnrealBuildTool
 				ModuleRules? ModuleRules;
 				try
 				{
-					ModuleRules = RulesAssembly.CreateModuleRules(FilteredModuleName, Rules, "all modules option", Logger);
+					string PrecompileReferenceChain = "allmodules option";
+					if (ModulePluginSource.TryGetValue(FilteredModuleName, out PluginInfo? value))
+					{
+						PrecompileReferenceChain = $"{PrecompileReferenceChain} -> {value.File.GetFileName()}";
+					}
+					ModuleRules = RulesAssembly.CreateModuleRules(FilteredModuleName, Rules, PrecompileReferenceChain, Logger);
 				}
 				catch (BuildException)
 				{
@@ -3800,9 +3809,13 @@ namespace UnrealBuildTool
 			// Gather the set of all modules traversed while processing ValidModuleNames, to ensure we have any extra modules found during RecursivelyCreateModules()
 			HashSet<UEBuildModule> AllModules = new HashSet<UEBuildModule>();
 
-			const string PrecompileReferenceChain = "allmodules option";
 			foreach (string ModuleName in ValidModuleNames)
 			{
+				string PrecompileReferenceChain = "allmodules option";
+				if (ModulePluginSource.TryGetValue(ModuleName, out PluginInfo? value))
+				{
+					PrecompileReferenceChain = $"{PrecompileReferenceChain} -> {value.File.GetFileName()}";
+				}
 				UEBuildModule Module = FindOrCreateModuleByName(ModuleName, PrecompileReferenceChain, Logger);
 				AllModules.Add(Module);
 				Module.RecursivelyCreateModules(
