UnrealBuildTool: Fix SupportedPlatformGroups check when unsupported platforms aren't registered into a group

#rnx

diff --git a/Engine/Source/Programs/UnrealBuildTool/Configuration/ModuleRules.cs b/Engine/Source/Programs/UnrealBuildTool/Configuration/ModuleRules.cs
--- a/Engine/Source/Programs/UnrealBuildTool/Configuration/ModuleRules.cs
+++ b/Engine/Source/Programs/UnrealBuildTool/Configuration/ModuleRules.cs
@@ -1638,8 +1638,9 @@ namespace UnrealBuildTool
 			// The platform overrides get applied at a later stage when we actually come to build the module.
 			if (!UEBuildPlatform.GetPlatformFolderNames().Any(Name => ModuleType.Name.EndsWith("_" + Name)))
 			{
-				IEnumerable<UnrealTargetPlatform> SupportedPlatforms = ModuleType.GetCustomAttributes<SupportedPlatformsAttribute>().SelectMany(x => x.Platforms).Distinct();
-				if (SupportedPlatforms.Any() && !SupportedPlatforms.Contains(TargetRules.Platform))
+				IEnumerable<SupportedPlatformsAttribute> PlatformAttributes = ModuleType.GetCustomAttributes<SupportedPlatformsAttribute>();
+				IEnumerable<UnrealTargetPlatform> SupportedPlatforms = PlatformAttributes.SelectMany(x => x.Platforms).Distinct();
+				if (PlatformAttributes.Any() && !SupportedPlatforms.Contains(TargetRules.Platform))
 				{
 					InvalidReason = $"Platform '{TargetRules.Platform}'";
 					return false;
