Fix IWYU tool after Linux toolchain upgrade to clang 16.0.6.

diff --git a/Engine/Source/Programs/IncludeTool/IncludeTool/Support/Utility.cs b/Engine/Source/Programs/IncludeTool/IncludeTool/Support/Utility.cs
--- a/Engine/Source/Programs/IncludeTool/IncludeTool/Support/Utility.cs
+++ b/Engine/Source/Programs/IncludeTool/IncludeTool/Support/Utility.cs
@@ -335,6 +335,12 @@ namespace IncludeTool.Support
 			string CompilerVersionString = DetermineClangVersionDir(CompilerFullPath);
 			if (!String.IsNullOrEmpty(CompilerVersionString))
 			{
+				// starting with 16.x clang includes are in a directory named only after its major version
+				int Major = System.Version.Parse(CompilerVersionString).Major;
+				if (Major >= 16)
+				{
+					CompilerVersionString = System.Version.Parse(CompilerVersionString).Major.ToString();
+				}
 				ExtraSystemIncludePaths.Add(DirectoryReference.Combine(SysRootDir, "lib", "clang", CompilerVersionString, "include"));
 			}
 			ExtraSystemIncludePaths.Add(DirectoryReference.Combine(InputDir, "Engine", "Source", "ThirdParty", "Linux", "LibCxx", "include", "c++", "v1"));
diff --git a/Engine/Source/Programs/UnrealBuildTool/Platform/Linux/LinuxToolChain.cs b/Engine/Source/Programs/UnrealBuildTool/Platform/Linux/LinuxToolChain.cs
--- a/Engine/Source/Programs/UnrealBuildTool/Platform/Linux/LinuxToolChain.cs
+++ b/Engine/Source/Programs/UnrealBuildTool/Platform/Linux/LinuxToolChain.cs
@@ -1640,7 +1640,9 @@ namespace UnrealBuildTool
 			string? InternalSdkPath = UEBuildPlatform.GetSDK(UnrealTargetPlatform.Linux)!.GetInternalSDKPath();
 			if (InternalSdkPath != null)
 			{
-				ExtraArguments.Add(String.Format("-isystem\"{0}\"", System.IO.Path.Combine(InternalSdkPath, "lib/clang/" + Info.ClangVersion + "/include/").Replace("\\", "/")));
+				// starting with clang 16.x the directory naming changed to include major version only
+				string ClangVersionString = (Info.ClangVersion.Major >= 16) ? Info.ClangVersion.Major.ToString() : Info.ClangVersion.ToString();
+				ExtraArguments.Add(String.Format("-isystem\"{0}\"", System.IO.Path.Combine(InternalSdkPath, "lib/clang/" + ClangVersionString + "/include/").Replace("\\", "/")));
 			}
 		}
 	}
diff --git a/Engine/Source/Programs/UnrealBuildTool/ToolChain/IWYUToolChain.cs b/Engine/Source/Programs/UnrealBuildTool/ToolChain/IWYUToolChain.cs
--- a/Engine/Source/Programs/UnrealBuildTool/ToolChain/IWYUToolChain.cs
+++ b/Engine/Source/Programs/UnrealBuildTool/ToolChain/IWYUToolChain.cs
@@ -48,7 +48,9 @@ namespace UnrealBuildTool
 			FileReference ClangPath = FileReference.Combine(BaseLinuxPath!, "bin", $"clang++{BuildHostPlatform.Current.BinarySuffix}");
 			ClangToolChainInfo CompilerToolChainInfo = new ClangToolChainInfo(ClangPath, null!, Logger);
 
-			DirectoryReference SystemPath = DirectoryReference.Combine(BaseLinuxPath!, "lib", "clang", CompilerToolChainInfo.ClangVersion.ToString(), "include");
+			// starting with clang 16.x the directory naming changed to include major version only
+			string ClangVersionString = (CompilerToolChainInfo.ClangVersion.Major >= 16) ? CompilerToolChainInfo.ClangVersion.Major.ToString() : CompilerToolChainInfo.ClangVersion.ToString();
+			DirectoryReference SystemPath = DirectoryReference.Combine(BaseLinuxPath!, "lib", "clang", ClangVersionString, "include");
 			CrossCompilingArguments.Add(GetSystemIncludePathArgument(SystemPath));
 
 			string DevPath = ""; //@"include-what-you-use-0.19\vs_projects\bin\RelWithDebInfo";
