Fix an issue with Linux symbol stripping

#rnx

diff --git a/Engine/Source/Programs/UnrealBuildTool/Platform/Linux/LinuxExports.cs b/Engine/Source/Programs/UnrealBuildTool/Platform/Linux/LinuxExports.cs
--- a/Engine/Source/Programs/UnrealBuildTool/Platform/Linux/LinuxExports.cs
+++ b/Engine/Source/Programs/UnrealBuildTool/Platform/Linux/LinuxExports.cs
@@ -18,7 +18,7 @@ namespace UnrealBuildTool
 		/// <param name="Logger">Logger for output</param>
 		public static void StripSymbols(FileReference SourceFile, FileReference TargetFile, ILogger Logger)
 		{
-			LinuxToolChain ToolChain = new LinuxToolChain(LinuxPlatform.DefaultHostArchitecture, new LinuxPlatformSDK(Logger), ClangToolChainOptions.None, Logger);
+			LinuxToolChain ToolChain = new LinuxToolChain(LinuxPlatform.DefaultHostArchitecture, (LinuxPlatformSDK)UEBuildPlatform.GetSDK(UnrealTargetPlatform.Linux)!, ClangToolChainOptions.None, Logger);
 			ToolChain.StripSymbols(SourceFile, TargetFile, Logger);
 		}
 	}
