- Added basic support for compiling Swift code (not general purpose yet, only the UESwift.swift file in this CL will work for now)
- Added an example SwiftUI view shown on VisionOS, using UIHostingViewController (must be created in Swift, not Obj-C)
- Refactored AppleToolChainSettings and children so AppleToolChain can have a single Swift compile function, and still target the proper toolchain (more cleaning up and pulling into AppleToolChain should now be possible)

diff --git a/Engine/Platforms/VisionOS/Source/Programs/UnrealBuildTool/VisionOSToolChain.cs b/Engine/Platforms/VisionOS/Source/Programs/UnrealBuildTool/VisionOSToolChain.cs
--- a/Engine/Platforms/VisionOS/Source/Programs/UnrealBuildTool/VisionOSToolChain.cs
+++ b/Engine/Platforms/VisionOS/Source/Programs/UnrealBuildTool/VisionOSToolChain.cs
@@ -6,7 +6,7 @@ namespace UnrealBuildTool
 {
 	class VisionOSToolChainSettings : IOSToolChainSettings
 	{
-		public VisionOSToolChainSettings(ILogger Logger) : base("XROS", "XRSimulator", Logger)
+		public VisionOSToolChainSettings(ILogger Logger) : base("XROS", "XRSimulator", "xros", Logger)
 		{
 		}
 	}
diff --git a/Engine/Source/Programs/UnrealBuildTool/Configuration/UEBuildModuleCPP.cs b/Engine/Source/Programs/UnrealBuildTool/Configuration/UEBuildModuleCPP.cs
--- a/Engine/Source/Programs/UnrealBuildTool/Configuration/UEBuildModuleCPP.cs
+++ b/Engine/Source/Programs/UnrealBuildTool/Configuration/UEBuildModuleCPP.cs
@@ -2054,7 +2054,7 @@ namespace UnrealBuildTool
 		{
 			".h", ".ipsh",
 			".cpp", ".ixx", ".c", ".cc",
-			".m", ".mm",
+			".m", ".mm", ".swift",
 			".rc", ".ispc",
 		};
 
@@ -2113,7 +2113,7 @@ namespace UnrealBuildTool
 					SourceFiles.Add(InputFile);
 					InputFiles.CCFiles.Add(InputFile);
 				}
-				else if (InputFile.HasExtension(".m") || InputFile.HasExtension(".mm"))
+				else if (InputFile.HasExtension(".m") || InputFile.HasExtension(".mm") || InputFile.HasExtension(".swift"))
 				{
 					SourceFiles.Add(InputFile);
 					InputFiles.MMFiles.Add(InputFile);
diff --git a/Engine/Source/Programs/UnrealBuildTool/Executors/Experimental/FASTBuild.cs b/Engine/Source/Programs/UnrealBuildTool/Executors/Experimental/FASTBuild.cs
--- a/Engine/Source/Programs/UnrealBuildTool/Executors/Experimental/FASTBuild.cs
+++ b/Engine/Source/Programs/UnrealBuildTool/Executors/Experimental/FASTBuild.cs
@@ -1045,7 +1045,7 @@ namespace UnrealBuildTool
 
 			if (BuildHostPlatform.Current.Platform == UnrealTargetPlatform.Mac)
 			{
-				AddText($".MacBaseSDKDir = '{MacToolChain.Settings.BaseSDKDir}'\n");
+				AddText($".MacBaseSDKDir = '{MacToolChain.Settings.GetSDKPath()}'\n");
 				AddText($".MacToolchainDir = '{MacToolChain.Settings.ToolchainDir}'\n");
 				AddText($"Compiler('UEAppleCompiler') \n{{\n");
 				AddText($"\t.Executable = '$MacToolchainDir$/clang++'\n");
diff --git a/Engine/Source/Programs/UnrealBuildTool/Platform/IOS/IOSToolChain.cs b/Engine/Source/Programs/UnrealBuildTool/Platform/IOS/IOSToolChain.cs
--- a/Engine/Source/Programs/UnrealBuildTool/Platform/IOS/IOSToolChain.cs
+++ b/Engine/Source/Programs/UnrealBuildTool/Platform/IOS/IOSToolChain.cs
@@ -17,64 +17,14 @@ namespace UnrealBuildTool
 {
 	class IOSToolChainSettings : AppleToolChainSettings
 	{
-		/// <summary>
-		/// The version of the iOS SDK to target at build time.
-		/// </summary>
-		[XmlConfigFile(Category = "IOSToolChain")]
-		public string IOSSDKVersion = "latest";
-		public readonly float IOSSDKVersionFloat = 0.0f;
-
-		/// <summary>
-		/// The version of the iOS to allow at build time.
-		/// </summary>
-		[XmlConfigFile(Category = "IOSToolChain")]
-		public string BuildIOSVersion = "15.0";
-
-		/// <summary>
-		/// Directory for the developer binaries
-		/// </summary>
-		public string ToolchainDir = "";
-
-		/// <summary>
-		/// Location of the SDKs
-		/// </summary>
-		public readonly string BaseSDKDir;
-		public readonly string BaseSDKDirSim;
-
-		public readonly string DevicePlatformName;
-		public readonly string SimulatorPlatformName;
-
-		public IOSToolChainSettings(ILogger Logger) : this("iPhoneOS", "iPhoneSimulator", Logger)
-		{
-		}
-
-		protected IOSToolChainSettings(string DevicePlatformName, string SimulatorPlatformName, ILogger Logger) : base(true, Logger)
+		public IOSToolChainSettings(ILogger Logger) 
+			: this("iPhoneOS", "iPhoneSimulator", "ios", Logger)
 		{
-			XmlConfig.ApplyTo(this);
-
-			this.DevicePlatformName = DevicePlatformName;
-			this.SimulatorPlatformName = SimulatorPlatformName;
-
-			// update cached paths
-			BaseSDKDir = XcodeDeveloperDir + "Platforms/" + DevicePlatformName + ".platform/Developer/SDKs";
-			BaseSDKDirSim = XcodeDeveloperDir + "Platforms/" + SimulatorPlatformName + ".platform/Developer/SDKs";
-			ToolchainDir = XcodeDeveloperDir + "Toolchains/XcodeDefault.xctoolchain/usr/bin/";
-
-			// make sure SDK is selected
-			SelectSDK(BaseSDKDir, DevicePlatformName, ref IOSSDKVersion, true, Logger);
-
-			// convert to float for easy comparison
-			IOSSDKVersionFloat = Single.Parse(IOSSDKVersion, System.Globalization.CultureInfo.InvariantCulture);
 		}
 
-		public string GetSDKPath(UnrealArch Architecture)
+		protected IOSToolChainSettings(string DevicePlatformName, string SimulatorPlatformName, string TargetOSName, ILogger Logger) 
+			: base(DevicePlatformName, SimulatorPlatformName, TargetOSName, true, Logger)
 		{
-			if (Architecture == UnrealArch.IOSSimulator || Architecture == UnrealArch.TVOSSimulator)
-			{
-				return BaseSDKDirSim + "/" + SimulatorPlatformName + IOSSDKVersion + ".sdk";
-			}
-
-			return BaseSDKDir + "/" + DevicePlatformName + IOSSDKVersion + ".sdk";
 		}
 	}
 
@@ -90,12 +40,11 @@ namespace UnrealBuildTool
 		{
 		}
 
-		protected IOSToolChain(ReadOnlyTargetRules? Target, IOSProjectSettings InProjectSettings, Func<IOSToolChainSettings> InCreateSettings, ClangToolChainOptions ToolchainOptions, ILogger InLogger)
-			: base((Target == null) ? null : Target.ProjectFile, ToolchainOptions, InLogger)
+		protected IOSToolChain(ReadOnlyTargetRules? Target, IOSProjectSettings InProjectSettings, Func<AppleToolChainSettings> InCreateSettings, ClangToolChainOptions ToolchainOptions, ILogger InLogger)
+			: base((Target == null) ? null : Target.ProjectFile, InCreateSettings, ToolchainOptions, InLogger)
 		{
 			this.Target = Target;
 			ProjectSettings = InProjectSettings;
-			Settings = new Lazy<IOSToolChainSettings>(InCreateSettings);
 		}
 
 		// ***********************************************************************
@@ -113,11 +62,6 @@ namespace UnrealBuildTool
 		[XmlConfigFile]
 		public static bool bUseDangerouslyFastMode = false;
 
-		/// <summary>
-		/// The lazily constructed settings for the toolchain
-		/// </summary>
-		private Lazy<IOSToolChainSettings> Settings;
-
 		/// <summary>
 		/// Which compiler\linker frontend to use
 		/// </summary>
@@ -128,16 +72,18 @@ namespace UnrealBuildTool
 		/// </summary>
 		private const string IOSArchiver = "libtool";
 
+		private IOSToolChainSettings Settings => (IOSToolChainSettings)ToolChainSettings.Value;
+
 		protected override ClangToolChainInfo GetToolChainInfo()
 		{
-			FileReference CompilerPath = new FileReference(Settings.Value.ToolchainDir + IOSCompiler);
-			FileReference ArchiverPath = new FileReference(Settings.Value.ToolchainDir + IOSArchiver);
+			FileReference CompilerPath = FileReference.Combine(Settings.ToolchainDir, IOSCompiler);
+			FileReference ArchiverPath = FileReference.Combine(Settings.ToolchainDir, IOSArchiver);
 			return new AppleToolChainInfo(CompilerPath, ArchiverPath, Logger);
 		}
 
 		public override string GetSDKVersion()
 		{
-			return Settings.Value.IOSSDKVersionFloat.ToString();
+			return Settings.SDKVersionFloat.ToString();
 		}
 
 		public override void ModifyBuildProducts(ReadOnlyTargetRules Target, UEBuildBinary Binary, List<string> Libraries, List<UEBuildBundleResource> BundleResources, Dictionary<FileReference, BuildProductType> BuildProducts)
@@ -204,7 +150,7 @@ namespace UnrealBuildTool
 			base.GetCompileArguments_WarningsAndErrors(CompileEnvironment, Arguments);
 
 			// fix for Xcode 8.3 enabling nonportable include checks, but p4 has some invalid cases in it
-			if (Settings.Value.IOSSDKVersionFloat >= 10.3)
+			if (Settings.SDKVersionFloat >= 10.3)
 			{
 				Arguments.Add("-Wno-nonportable-include-path");
 			}
@@ -270,7 +216,7 @@ namespace UnrealBuildTool
 			// What architecture(s) to build for
 			Arguments.Add(FormatArchitectureArg(CompileEnvironment.Architectures));
 
-			Arguments.Add($"-isysroot \"{Settings.Value.GetSDKPath(CompileEnvironment.Architecture)}\"");
+			Arguments.Add($"-isysroot \"{Settings.GetSDKPath(CompileEnvironment.Architecture)}\"");
 
 			if (GetXcodeMinVersionParam() != "")
 			{
@@ -302,7 +248,7 @@ namespace UnrealBuildTool
 
 			string Extension = Path.GetExtension(SourceFile.AbsolutePath).ToUpperInvariant();
 
-			if (!Extension.Equals(".C"))
+			if (!Extension.Equals(".C") && !Extension.Equals(".SWIFT"))
 			{
 				Arguments.Add(GetObjCExceptionsFlag(CompileEnvironment));
 			}
@@ -363,9 +309,8 @@ namespace UnrealBuildTool
 		{
 			Arguments.Add(FormatArchitectureArg(LinkEnvironment.Architectures));
 
-			bool bIsDevice = LinkEnvironment.Architecture != UnrealArch.IOSSimulator && LinkEnvironment.Architecture != UnrealArch.TVOSSimulator;
-			Arguments.Add(String.Format(" -isysroot \\\"{0}Platforms/{1}.platform/Developer/SDKs/{1}{2}.sdk\\\"",
-				Settings.Value.XcodeDeveloperDir, bIsDevice ? Settings.Value.DevicePlatformName : Settings.Value.SimulatorPlatformName, Settings.Value.IOSSDKVersion));
+			DirectoryReference SDKPath = Settings.GetSDKPath(LinkEnvironment.Architecture);
+			Arguments.Add($" -isysroot \"{SDKPath}\"");
 
 			Arguments.Add("-dead_strip");
 			if (GetXcodeMinVersionParam() != "")
@@ -411,7 +356,7 @@ namespace UnrealBuildTool
 			// link in the frameworks
 			foreach (string Framework in LinkEnvironment.Frameworks)
 			{
-				if (Framework != "ARKit" || Settings.Value.IOSSDKVersionFloat >= 11.0f)
+				if (Framework != "ARKit" || Settings.SDKVersionFloat >= 11.0f)
 				{
 					Arguments.Add("-framework " + Framework);
 				}
@@ -504,6 +449,10 @@ namespace UnrealBuildTool
 					// and add to the commandline
 					LinkCommandArguments += String.Format(" \\\"{0}\\\"", Library.FullName);
 				}
+
+				// allow for auto-linked swift libs to be found
+				// @todo Mac will need this too - a lot of this code should probably be moved up to AppleToolChain and shared
+				LinkCommandArguments += $" -L{ToolChainSettings.Value.GetSDKPath(LinkEnvironment.Architecture)}/usr/lib/swift";
 			}
 
 			// Handle additional framework assets that might need to be shadowed
@@ -1122,7 +1071,7 @@ namespace UnrealBuildTool
 				Action StripAction = Graph.CreateAction(ActionType.CreateAppBundle);
 				StripAction.WorkingDirectory = GetMacDevSrcRoot();
 				StripAction.CommandPath = BuildHostPlatform.Current.Shell;
-				StripAction.CommandArguments = String.Format("-c \"\\\"{0}strip\\\" {1} \\\"{2}\\\" && touch \\\"{3}\\\"\"", Settings.Value.ToolchainDir, StripArguments, Executable.Location, StripCompleteFile);
+				StripAction.CommandArguments = String.Format("-c \"\\\"{0}strip\\\" {1} \\\"{2}\\\" && touch \\\"{3}\\\"\"", Settings.ToolchainDir, StripArguments, Executable.Location, StripCompleteFile);
 				StripAction.PrerequisiteItems.Add(Executable);
 				StripAction.PrerequisiteItems.UnionWith(OutputFiles);
 				StripAction.ProducedItems.Add(StripCompleteFile);
@@ -1162,7 +1111,7 @@ namespace UnrealBuildTool
 				Logger.LogInformation("Adding PostBuildSync action");
 
 				List<string> UPLScripts = UEDeployIOS.CollectPluginDataPaths(BinaryLinkEnvironment.AdditionalProperties, Logger);
-				VersionNumber SdkVersion = VersionNumber.Parse(Settings.Value.IOSSDKVersion);
+				VersionNumber SdkVersion = VersionNumber.Parse(Settings.SDKVersion);
 
 				Dictionary<string, DirectoryReference> FrameworkNameToSourceDir = new Dictionary<string, DirectoryReference>();
 
@@ -1621,7 +1570,7 @@ namespace UnrealBuildTool
 						Writer.WriteLine(String.Format("rm -rf \"{0}\"", FrameworkPayloadDirectory));
 
 						// Build the framework wrapper
-						CmdLine = new IOSToolChainSettings(Logger).XcodeDeveloperDir + "usr/bin/xcodebuild" +
+						CmdLine = FileReference.Combine(AppleToolChainSettings.XcodeDeveloperDir, "usr/bin/xcodebuild").FullName +
 							" -project \"" + WrapperProject + "\"" +
 								" -configuration \"" + ConfigName + "\"" +
 							" -scheme '" + SchemeName + "'" +
@@ -1634,7 +1583,7 @@ namespace UnrealBuildTool
 					else
 					{
 						// code sign the project
-						CmdLine = new IOSToolChainSettings(Logger).XcodeDeveloperDir + "usr/bin/xcodebuild" +
+						CmdLine = FileReference.Combine(AppleToolChainSettings.XcodeDeveloperDir, "usr/bin/xcodebuild").FullName +
 							" -workspace \"" + XcodeWorkspaceDir + "\"" +
 								" -configuration \"" + ConfigName + "\"" +
 							" -scheme '" + SchemeName + "'" +
@@ -1803,7 +1752,7 @@ namespace UnrealBuildTool
 
 		public void StripSymbols(FileReference SourceFile, FileReference TargetFile)
 		{
-			StripSymbolsWithXcode(SourceFile, TargetFile, Settings.Value.ToolchainDir);
+			StripSymbolsWithXcode(SourceFile, TargetFile, Settings.ToolchainDir);
 		}
 
 		FileItem CopyBundleResource(UEBuildBundleResource Resource, FileItem Executable, DirectoryReference BundleDirectory, IActionGraphBuilder Graph)
diff --git a/Engine/Source/Programs/UnrealBuildTool/Platform/Mac/MacToolChain.cs b/Engine/Source/Programs/UnrealBuildTool/Platform/Mac/MacToolChain.cs
--- a/Engine/Source/Programs/UnrealBuildTool/Platform/Mac/MacToolChain.cs
+++ b/Engine/Source/Programs/UnrealBuildTool/Platform/Mac/MacToolChain.cs
@@ -14,12 +14,6 @@ namespace UnrealBuildTool
 {
 	class MacToolChainSettings : AppleToolChainSettings
 	{
-		/// <summary>
-		/// Which version of the Mac OS SDK to target at build time
-		/// </summary>
-		public string MacOSSDKVersion = "latest";
-		public float MacOSSDKVersionFloat = 0.0f;
-
 		/// <summary>
 		/// Which version of the Mac OS X to allow at run time
 		/// </summary>
@@ -30,37 +24,19 @@ namespace UnrealBuildTool
 		/// </summary>
 		public string MinMacOSVersion = "10.15.7";
 
-		/// <summary>
-		/// Directory for the developer binaries
-		/// </summary>
-		public string ToolchainDir = "";
-
-		/// <summary>
-		/// Location of the SDKs
-		/// </summary>
-		public string BaseSDKDir;
-
 		/// <summary>
 		/// Constructor
 		/// </summary>
 		/// <param name="bVerbose">Whether to output verbose logging</param>
 		/// <param name="Logger">Logger for output</param>
-		public MacToolChainSettings(bool bVerbose, ILogger Logger) : base(bVerbose, Logger)
+		public MacToolChainSettings(bool bVerbose, ILogger Logger) 
+			: base("MacOSX", null, "macos", bVerbose, Logger)
 		{
-			BaseSDKDir = XcodeDeveloperDir + "Platforms/MacOSX.platform/Developer/SDKs";
-			ToolchainDir = XcodeDeveloperDir + "Toolchains/XcodeDefault.xctoolchain/usr/bin/";
-
-			SelectSDK(BaseSDKDir, "MacOSX", ref MacOSSDKVersion, bVerbose, Logger);
+		}
 
-			// convert to float for easy comparison
-			if (String.IsNullOrWhiteSpace(MacOSSDKVersion))
-			{
-				throw new BuildException("Unable to find installed MacOS SDK on remote agent.");
-			}
-			else if (!Single.TryParse(MacOSSDKVersion, NumberStyles.AllowDecimalPoint, CultureInfo.InvariantCulture.NumberFormat, out MacOSSDKVersionFloat))
-			{
-				throw new BuildException("Unable to parse installed MacOS version (\"{0}\")", MacOSSDKVersion);
-			}
+		public DirectoryReference GetSDKPath()
+		{
+			return GetSDKPath(UnrealArch.Host.Value);
 		}
 	}
 
@@ -70,15 +46,11 @@ namespace UnrealBuildTool
 	class MacToolChain : AppleToolChain
 	{
 		public MacToolChain(FileReference? InProjectFile, ClangToolChainOptions InOptions, ILogger InLogger)
-			: base(InProjectFile, InOptions, InLogger)
+			: base(InProjectFile, () => new MacToolChainSettings(false, InLogger), InOptions, InLogger)
 		{
 		}
 
-		public static Lazy<MacToolChainSettings> SettingsPrivate = new Lazy<MacToolChainSettings>(() => new MacToolChainSettings(false, Log.Logger));
-
-		public static MacToolChainSettings Settings => SettingsPrivate.Value;
-
-		public static string SDKPath => Settings.BaseSDKDir + "/MacOSX.sdk";
+		public static MacToolChainSettings Settings => new MacToolChainSettings(false, Log.Logger);
 
 		/// <summary>
 		/// Which compiler\linker frontend to use
@@ -92,15 +64,11 @@ namespace UnrealBuildTool
 
 		protected override ClangToolChainInfo GetToolChainInfo()
 		{
-			FileReference CompilerPath = new FileReference(Settings.ToolchainDir + MacCompiler);
-			FileReference ArchiverPath = new FileReference(Settings.ToolchainDir + MacArchiver);
+			FileReference CompilerPath = FileReference.Combine(Settings.ToolchainDir, MacCompiler);
+			FileReference ArchiverPath = FileReference.Combine(Settings.ToolchainDir, MacArchiver);
 			return new AppleToolChainInfo(CompilerPath, ArchiverPath, Logger);
 		}
 
-		private static void SetupXcodePaths(bool bVerbose)
-		{
-		}
-
 		public static DirectoryReference FindProductDirectory(FileReference? ProjectFile, DirectoryReference BinaryDir, string? NameIfProgram)
 		{
 			// a project file is always used if there is one
@@ -156,13 +124,6 @@ namespace UnrealBuildTool
 			return null;
 		}
 
-		public override void SetUpGlobalEnvironment(ReadOnlyTargetRules Target)
-		{
-			base.SetUpGlobalEnvironment(Target);
-
-			SetupXcodePaths(true);
-		}
-
 		/// <inheritdoc/>
 		protected override void GetCompileArguments_WarningsAndErrors(CppCompileEnvironment CompileEnvironment, List<string> Arguments)
 		{
@@ -237,7 +198,7 @@ namespace UnrealBuildTool
 
 			// Pass through architecture and OS info
 			Arguments.Add("" + FormatArchitectureArg(CompileEnvironment.Architectures));
-			Arguments.Add($"-isysroot \"{SDKPath}\"");
+			Arguments.Add($"-isysroot \"{Settings.GetSDKPath()}\"");
 			Arguments.Add("-mmacosx-version-min=" + (CompileEnvironment.bEnableOSX109Support ? "10.9" : Settings.MacOSVersion));
 
 			List<string> FrameworksSearchPaths = new List<string>();
@@ -268,7 +229,7 @@ namespace UnrealBuildTool
 		{
 			// Pass through architecture and OS info		
 			Arguments.Add(FormatArchitectureArg(LinkEnvironment.Architectures));
-			Arguments.Add(String.Format("-isysroot \"{0}\"", SDKPath));
+			Arguments.Add(String.Format("-isysroot \"{0}\"", ToolChainSettings.Value.GetSDKPath(LinkEnvironment.Architecture)));
 			Arguments.Add("-mmacosx-version-min=" + Settings.MacOSVersion);
 			Arguments.Add("-dead_strip");
 
@@ -1244,8 +1205,6 @@ namespace UnrealBuildTool
 
 		public void StripSymbols(FileReference SourceFile, FileReference TargetFile)
 		{
-			SetupXcodePaths(false);
-
 			StripSymbolsWithXcode(SourceFile, TargetFile, Settings.ToolchainDir);
 		}
 	};
diff --git a/Engine/Source/Programs/UnrealBuildTool/Platform/Mac/UEBuildMac.cs b/Engine/Source/Programs/UnrealBuildTool/Platform/Mac/UEBuildMac.cs
--- a/Engine/Source/Programs/UnrealBuildTool/Platform/Mac/UEBuildMac.cs
+++ b/Engine/Source/Programs/UnrealBuildTool/Platform/Mac/UEBuildMac.cs
@@ -143,8 +143,19 @@ namespace UnrealBuildTool
 			return GetProjectArchitectures(ProjectFile, TargetName, true, false);
 		}
 
+		private static Dictionary<string, UnrealArchitectures> ProjectArchitectureCache = new();
 		private UnrealArchitectures GetProjectArchitectures(FileReference? ProjectFile, string? TargetName, bool bGetAllSupported, bool bIsDistributionMode)
 		{
+			string Key = $"{ProjectFile}{TargetName}{bGetAllSupported}{bIsDistributionMode}";
+			lock (ProjectArchitectureCache)
+			{
+				UnrealArchitectures? CachedArches;
+				if (ProjectArchitectureCache.TryGetValue(Key, out CachedArches))
+				{
+					return CachedArches;
+				}
+			}
+
 			bool bIsEditor = false;
 			bool bIsBuildMachine = Environment.GetEnvironmentVariable("IsBuildMachine") == "1";
 
@@ -240,7 +251,12 @@ namespace UnrealBuildTool
 				throw new BuildException($"Unknown {DefaultKey} value found ('{DefaultArchitecture}') in .ini");
 			}
 
-			return new UnrealArchitectures(Architectures);
+			UnrealArchitectures Result = new UnrealArchitectures(Architectures);
+			lock (ProjectArchitectureCache)
+			{
+				ProjectArchitectureCache.Add(Key, Result);
+			}
+			return Result;
 		}
 	}
 
@@ -299,7 +315,7 @@ namespace UnrealBuildTool
 			}
 
 			// Needs OS X 10.11 for Metal. The remote toolchain has not been initialized yet, so just assume it's a recent SDK.
-			if ((BuildHostPlatform.Current.Platform != UnrealTargetPlatform.Mac || MacToolChain.Settings.MacOSSDKVersionFloat >= 10.11f) && Target.bCompileAgainstEngine)
+			if ((BuildHostPlatform.Current.Platform != UnrealTargetPlatform.Mac || MacToolChain.Settings.SDKVersionFloat >= 10.11f) && Target.bCompileAgainstEngine)
 			{
 				Target.GlobalDefinitions.Add("HAS_METAL=1");
 				Target.ExtraModuleNames.Add("MetalRHI");
diff --git a/Engine/Source/Programs/UnrealBuildTool/Platform/TVOS/TVOSToolChain.cs b/Engine/Source/Programs/UnrealBuildTool/Platform/TVOS/TVOSToolChain.cs
--- a/Engine/Source/Programs/UnrealBuildTool/Platform/TVOS/TVOSToolChain.cs
+++ b/Engine/Source/Programs/UnrealBuildTool/Platform/TVOS/TVOSToolChain.cs
@@ -6,7 +6,7 @@ namespace UnrealBuildTool
 {
 	class TVOSToolChainSettings : IOSToolChainSettings
 	{
-		public TVOSToolChainSettings(ILogger Logger) : base("AppleTVOS", "AppleTVSimulator", Logger)
+		public TVOSToolChainSettings(ILogger Logger) : base("AppleTVOS", "AppleTVSimulator", "tvos", Logger)
 		{
 		}
 	}
diff --git a/Engine/Source/Programs/UnrealBuildTool/ProjectFiles/Rider/RiderProjectFile.cs b/Engine/Source/Programs/UnrealBuildTool/ProjectFiles/Rider/RiderProjectFile.cs
--- a/Engine/Source/Programs/UnrealBuildTool/ProjectFiles/Rider/RiderProjectFile.cs
+++ b/Engine/Source/Programs/UnrealBuildTool/ProjectFiles/Rider/RiderProjectFile.cs
@@ -927,17 +927,17 @@ namespace UnrealBuildTool
 			{
 				if (Platform == UnrealTargetPlatform.Mac)
 				{
-					return MacToolChain.SDKPath;
+					return MacToolChain.Settings.GetSDKPath().FullName;
 				}
 
 				if (Platform == UnrealTargetPlatform.IOS)
 				{
-					return new IOSToolChainSettings(Logger).GetSDKPath(Architecture);
+					return new IOSToolChainSettings(Logger).GetSDKPath(Architecture).FullName;
 				}
 
 				if (Platform == UnrealTargetPlatform.TVOS)
 				{
-					return new TVOSToolChainSettings(Logger).GetSDKPath(Architecture);
+					return new TVOSToolChainSettings(Logger).GetSDKPath(Architecture).FullName;
 				}
 
 				throw new NotImplementedException("Path to SDK has to be specified for each Apple's platform");
diff --git a/Engine/Source/Programs/UnrealBuildTool/ProjectFiles/VisualStudioCode/VSCodeProjectFileGenerator.cs b/Engine/Source/Programs/UnrealBuildTool/ProjectFiles/VisualStudioCode/VSCodeProjectFileGenerator.cs
--- a/Engine/Source/Programs/UnrealBuildTool/ProjectFiles/VisualStudioCode/VSCodeProjectFileGenerator.cs
+++ b/Engine/Source/Programs/UnrealBuildTool/ProjectFiles/VisualStudioCode/VSCodeProjectFileGenerator.cs
@@ -351,8 +351,8 @@ namespace UnrealBuildTool
 			else if (OperatingSystem.IsMacOS())
 			{
 				MacToolChainSettings Settings = new MacToolChainSettings(false, Logger);
-				CompilerPath = FileReference.FromString(Settings.ToolchainDir + "clang++");
-				SysRootPath = DirectoryReference.FromString(Settings.BaseSDKDir + "/MacOSX" + Settings.MacOSSDKVersion + ".sdk");
+				CompilerPath = FileReference.Combine(Settings.ToolchainDir, "clang++");
+				SysRootPath = Settings.GetSDKPath();
 			}
 			else
 			{
diff --git a/Engine/Source/Programs/UnrealBuildTool/ToolChain/AppleToolChain.cs b/Engine/Source/Programs/UnrealBuildTool/ToolChain/AppleToolChain.cs
--- a/Engine/Source/Programs/UnrealBuildTool/ToolChain/AppleToolChain.cs
+++ b/Engine/Source/Programs/UnrealBuildTool/ToolChain/AppleToolChain.cs
@@ -5,6 +5,7 @@ using System.Collections.Generic;
 using System.Diagnostics;
 using System.IO;
 using System.Linq;
+using System.Runtime.InteropServices;
 using System.Text;
 using System.Threading.Tasks;
 using EpicGames.Core;
@@ -13,59 +14,135 @@ using UnrealBuildBase;
 
 namespace UnrealBuildTool
 {
-	abstract class AppleToolChainSettings
+	/// <summary>
+	/// Helper class for managing Xcode paths, versions, etc. Helps to differentiate between Apple xcode platforms
+	/// </summary>
+	public abstract class AppleToolChainSettings
 	{
 		/// <summary>
 		/// Which developer directory to root from? If this is "xcode-select", UBT will query for the currently selected Xcode
 		/// </summary>
-		public string XcodeDeveloperDir = "xcode-select";
+		public static DirectoryReference XcodeDeveloperDir => _XcodeDeveloperDir!;
 
-		public AppleToolChainSettings(bool bVerbose, ILogger Logger)
+		// This is not expected to be used before we've made any Settings options, so XcodeDeveloperDir will assert if it's used too early
+		private static DirectoryReference? _XcodeDeveloperDir = null;
+
+		/// <summary>
+		/// Directory for the developer binaries
+		/// </summary>
+		public DirectoryReference ToolchainDir;
+
+		/// <summary>
+		/// A portion of the target "tuple"
+		/// </summary>
+		private string TargetOSName;
+
+		/// <summary>
+		/// The version of the iOS SDK to target at build time.
+		/// </summary>
+		public string SDKVersion;
+
+		/// <summary>
+		/// The version in a floating point value for easy comparison
+		/// </summary>
+		public readonly float SDKVersionFloat = 0.0f;
+
+		/// <summary>
+		/// Cache SDK dir for device and simulator (for non-Mac)
+		/// </summary>
+		DirectoryReference SDKDir;
+		DirectoryReference? SimulatorSDKDir = null;
+
+		/// <summary>
+		/// Constructor, called by platform sybclasses
+		/// </summary>
+		/// <param name="OSPrefix">SDK name (like "MacOSX")</param>
+		/// <param name="SimulatorOSPrefix">Sinulator SDK name (like "iPhoneOSSimulator")</param>
+		/// <param name="TargetOSName">Platform name used in the -target parameter</param>
+		/// <param name="bVerbose"></param>
+		/// <param name="Logger"></param>
+		protected AppleToolChainSettings(string OSPrefix, string? SimulatorOSPrefix, string TargetOSName, bool bVerbose, ILogger Logger)
 		{
-			SelectXcode(ref XcodeDeveloperDir, bVerbose, Logger);
+			this.TargetOSName = TargetOSName;
+
+			if (_XcodeDeveloperDir == null)
+			{
+				SelectXcode(bVerbose, Logger);
+			}
+
+			// set up directories
+			ToolchainDir = DirectoryReference.Combine(XcodeDeveloperDir, "Toolchains/XcodeDefault.xctoolchain/usr/bin");
+			SDKDir = DirectoryReference.Combine(XcodeDeveloperDir, $"Platforms/{OSPrefix}.platform/Developer/SDKs/{OSPrefix}.sdk");
+			if (SimulatorOSPrefix != null)
+			{   
+				SimulatorSDKDir = DirectoryReference.Combine(XcodeDeveloperDir, $"Platforms/{SimulatorOSPrefix}.platform/Developer/SDKs/{SimulatorOSPrefix}.sdk");
+			}
+
+			// use one up fro the SDK dir to look in for the version'd name directories next to the SDK dir
+			SDKVersion = SelectSDK(SDKDir.ParentDirectory! , OSPrefix, bVerbose, Logger);
+
+			// convert to float for easy comparison
+			SDKVersionFloat = Single.Parse(SDKVersion, System.Globalization.CultureInfo.InvariantCulture);
 		}
 
-		private static void SelectXcode(ref string DeveloperDir, bool bVerbose, ILogger Logger)
+		/// <summary>
+		/// Get the path to the SDK diretory in xcode for the given architecture
+		/// </summary>
+		/// <param name="Architecture"></param>
+		/// <returns></returns>
+		public DirectoryReference GetSDKPath(UnrealArch Architecture)
 		{
-			string Reason = "hardcoded";
-
-			if (DeveloperDir == "xcode-select")
+			// note that VisionOS uses IOSSimulator (as TVOS should eventually do as well)
+			if (Architecture == UnrealArch.IOSSimulator || Architecture == UnrealArch.TVOSSimulator)
 			{
-				Reason = "xcode-select";
+				return SimulatorSDKDir!;
+			}
+			return SDKDir;
+		}
 
-				// on the Mac, run xcode-select directly.
-				int ReturnCode;
-				DeveloperDir = Utils.RunLocalProcessAndReturnStdOut("xcode-select", "--print-path", null, out ReturnCode);
-				if (ReturnCode != 0)
-				{
-					string? MinVersion = UEBuildPlatform.GetSDK(UnrealTargetPlatform.Mac)!.GetSDKInfo("Sdk")!.Min;
-					throw new BuildException($"We were unable to find your build tools (via 'xcode-select --print-path'). Please install Xcode, version {MinVersion} or later");
-				}
+		/// <summary>
+		/// Gets the string used by xcode to taget a platform and version (will return something like "arm64-apple-ios17.0-simulator"
+		/// </summary>
+		/// <param name="Architecture"></param>
+		/// <returns></returns>
+		public string GetTargetTuple(UnrealArch Architecture)
+		{
+			string Prefix = Architecture.AppleName;
+			string Suffix = (Architecture == UnrealArch.IOSSimulator || Architecture == UnrealArch.TVOSSimulator) ? "-simulator" : "";
 
-				// make sure we get a full path
-				if (Directory.Exists(DeveloperDir) == false)
-				{
-					throw new BuildException("Selected Xcode ('{0}') doesn't exist, cannot continue.", DeveloperDir);
-				}
+			return $"{Prefix}-apple-{TargetOSName}{SDKVersion}{Suffix}";
+		}
 
-				if (DeveloperDir.Contains("CommandLineTools", StringComparison.InvariantCultureIgnoreCase))
-				{
-					throw new BuildException($"Your Mac is set to use CommandLineTools for its build tools ({DeveloperDir}). Unreal expects Xcode as the build tools. Please install Xcode if it's not already, then do one of the following:\n" +
-						"  - Run Xcode, go to Settings, and in the Locations tab, choose your Xcode in Command Line Tools dropdown.\n" +
-						"  - In Terminal, run 'sudo xcode-select -s /Applications/Xcode.app' (or an alternate location if you installed Xcode to a non-standard location)\n" +
-						"Either way, you will need to enter your Mac password.");
-				}
+		private static void SelectXcode(bool bVerbose, ILogger Logger)
+		{
+			// on the Mac, run xcode-select directly.
+			int ReturnCode;
+			string XcodeSelectResult = Utils.RunLocalProcessAndReturnStdOut("xcode-select", "--print-path", null, out ReturnCode);
+			if (ReturnCode != 0)
+			{
+				string? MinVersion = UEBuildPlatform.GetSDK(UnrealTargetPlatform.Mac)!.GetSDKInfo("Sdk")!.Min;
+				throw new BuildException($"We were unable to find your build tools (via 'xcode-select --print-path'). Please install Xcode, version {MinVersion} or later");
+			}
 
-				if (DeveloperDir.EndsWith("/") == false)
-				{
-					// we expect this to end with a slash
-					DeveloperDir += "/";
-				}
+			_XcodeDeveloperDir = new DirectoryReference(XcodeSelectResult);
+
+			// make sure we get a full path
+			if (DirectoryReference.Exists(XcodeDeveloperDir) == false)
+			{
+				throw new BuildException("Selected Xcode ('{0}') doesn't exist, cannot continue.", XcodeDeveloperDir);
 			}
 
-			if (bVerbose && !DeveloperDir.StartsWith("/Applications/Xcode.app"))
+			if (XcodeDeveloperDir.ContainsName("CommandLineTools", 0))
 			{
-				Log.TraceInformationOnce("Compiling with non-standard Xcode ({0}): {1}", Reason, DeveloperDir);
+				throw new BuildException($"Your Mac is set to use CommandLineTools for its build tools ({XcodeDeveloperDir}). Unreal expects Xcode as the build tools. Please install Xcode if it's not already, then do one of the following:\n" +
+					"  - Run Xcode, go to Settings, and in the Locations tab, choose your Xcode in Command Line Tools dropdown.\n" +
+					"  - In Terminal, run 'sudo xcode-select -s /Applications/Xcode.app' (or an alternate location if you installed Xcode to a non-standard location)\n" +
+					"Either way, you will need to enter your Mac password.");
+			}
+
+			if (bVerbose && !XcodeDeveloperDir.FullName.StartsWith("/Applications/Xcode.app"))
+			{
+				Log.TraceInformationOnce("Compiling with non-standard Xcode: {0}", XcodeDeveloperDir);
 			}
 
 			// Installed engine requires Xcode 13
@@ -83,80 +160,70 @@ namespace UnrealBuildTool
 			}
 		}
 
-		protected void SelectSDK(string BaseSDKDir, string OSPrefix, ref string PlatformSDKVersion, bool bVerbose, ILogger Logger)
+		private static string SelectSDK(DirectoryReference BaseSDKDir, string OSPrefix, bool bVerbose, ILogger Logger)
 		{
-			if (PlatformSDKVersion == "latest")
+			string PlatformSDKVersion = "";
+			try
 			{
-				PlatformSDKVersion = "";
-				try
+				// loop over the subdirs and parse out the version
+				int MaxSDKVersionMajor = 0;
+				int MaxSDKVersionMinor = 0;
+				string? MaxSDKVersionString = null;
+				foreach (DirectoryReference SubDir in DirectoryReference.EnumerateDirectories(BaseSDKDir))
 				{
-					// on the Mac, we can just get the directory name
-					string[] SubDirs = System.IO.Directory.GetDirectories(BaseSDKDir);
-
-					// loop over the subdirs and parse out the version
-					int MaxSDKVersionMajor = 0;
-					int MaxSDKVersionMinor = 0;
-					string? MaxSDKVersionString = null;
-					foreach (string SubDir in SubDirs)
+					string SubDirName = Path.GetFileNameWithoutExtension(SubDir.GetDirectoryName());
+					if (SubDirName.StartsWith(OSPrefix))
 					{
-						string SubDirName = Path.GetFileNameWithoutExtension(SubDir);
-						if (SubDirName.StartsWith(OSPrefix))
-						{
-							// get the SDK version from the directory name
-							string SDKString = SubDirName.Replace(OSPrefix, "");
-							int Major = 0;
-							int Minor = 0;
+						// get the SDK version from the directory name
+						string SDKString = SubDirName.Replace(OSPrefix, "");
+						int Major = 0;
+						int Minor = 0;
 
-							// parse it into whole and fractional parts (since 10.10 > 10.9 in versions, but not in math)
-							try
-							{
-								string[] Tokens = SDKString.Split(".".ToCharArray());
-								if (Tokens.Length == 2)
-								{
-									Major = Int32.Parse(Tokens[0]);
-									Minor = Int32.Parse(Tokens[1]);
-								}
-							}
-							catch (Exception)
-							{
-								// weirdly formatted SDKs
-								continue;
-							}
-
-							// update largest SDK version number
-							if (Major > MaxSDKVersionMajor || (Major == MaxSDKVersionMajor && Minor > MaxSDKVersionMinor))
+						// parse it into whole and fractional parts (since 10.10 > 10.9 in versions, but not in math)
+						try
+						{
+							string[] Tokens = SDKString.Split(".".ToCharArray());
+							if (Tokens.Length == 2)
 							{
-								MaxSDKVersionString = SDKString;
-								MaxSDKVersionMajor = Major;
-								MaxSDKVersionMinor = Minor;
+								Major = Int32.Parse(Tokens[0]);
+								Minor = Int32.Parse(Tokens[1]);
 							}
 						}
-					}
+						catch (Exception)
+						{
+							// weirdly formatted SDKs
+							continue;
+						}
 
-					// use the largest version
-					if (MaxSDKVersionString != null)
-					{
-						PlatformSDKVersion = MaxSDKVersionString;
+						// update largest SDK version number
+						if (Major > MaxSDKVersionMajor || (Major == MaxSDKVersionMajor && Minor > MaxSDKVersionMinor))
+						{
+							MaxSDKVersionString = SDKString;
+							MaxSDKVersionMajor = Major;
+							MaxSDKVersionMinor = Minor;
+						}
 					}
 				}
-				catch (Exception Ex)
+
+				// use the largest version
+				if (MaxSDKVersionString != null)
 				{
-					// on any exception, just use the backup version
-					Logger.LogInformation("Triggered an exception while looking for SDK directory in Xcode.app");
-					Logger.LogInformation("{Ex}", Ex.ToString());
+					PlatformSDKVersion = MaxSDKVersionString;
 				}
 			}
-
-			// make sure we have a valid SDK directory
-			if (!RuntimePlatform.IsWindows && !Directory.Exists(Path.Combine(BaseSDKDir, OSPrefix + PlatformSDKVersion + ".sdk")))
+			catch (Exception Ex)
 			{
-				throw new BuildException("Invalid SDK {0}{1}.sdk, not found in {2}", OSPrefix, PlatformSDKVersion, BaseSDKDir);
+				// on any exception, just use the backup version
+				Logger.LogInformation("Triggered an exception while looking for SDK directory in Xcode.app");
+				Logger.LogInformation("{Ex}", Ex.ToString());
 			}
 
 			if (bVerbose && !ProjectFileGenerator.bGenerateProjectFiles)
 			{
 				Logger.LogInformation("Compiling with {Os} SDK {Sdk}", OSPrefix, PlatformSDKVersion);
 			}
+
+			return PlatformSDKVersion;
 		}
 	}
 
@@ -174,13 +241,16 @@ namespace UnrealBuildTool
 			protected override string QueryArchiverVersionString() => ClangVersionString;
 		}
 
+		public Lazy<AppleToolChainSettings> ToolChainSettings;
+
 		protected FileReference? ProjectFile;
 
 		protected bool bUseModernXcode => AppleExports.UseModernXcode(ProjectFile);
 
-		public AppleToolChain(FileReference? InProjectFile, ClangToolChainOptions InOptions, ILogger InLogger) : base(InOptions, InLogger)
+		public AppleToolChain(FileReference? InProjectFile, Func<AppleToolChainSettings> InCreateSettings, ClangToolChainOptions InOptions, ILogger InLogger) : base(InOptions, InLogger)
 		{
 			ProjectFile = InProjectFile;
+			ToolChainSettings = new Lazy<AppleToolChainSettings>(InCreateSettings);
 		}
 
 		/// <summary>
@@ -227,7 +297,7 @@ namespace UnrealBuildTool
 			return Result;
 		}
 
-		protected void StripSymbolsWithXcode(FileReference SourceFile, FileReference TargetFile, string ToolchainDir)
+		protected void StripSymbolsWithXcode(FileReference SourceFile, FileReference TargetFile, DirectoryReference ToolchainDir)
 		{
 			if (SourceFile != TargetFile)
 			{
@@ -236,7 +306,7 @@ namespace UnrealBuildTool
 			}
 
 			ProcessStartInfo StartInfo = new ProcessStartInfo();
-			StartInfo.FileName = Path.Combine(ToolchainDir, "strip");
+			StartInfo.FileName = Path.Combine(ToolchainDir.FullName, "strip");
 			StartInfo.Arguments = String.Format("\"{0}\" -S", TargetFile.FullName);
 			StartInfo.UseShellExecute = false;
 			StartInfo.CreateNoWindow = true;
@@ -472,6 +542,124 @@ namespace UnrealBuildTool
 			return CompileEnvironment;
 		}
 
+		protected override FileItem GetCompileArguments_FileType(CppCompileEnvironment CompileEnvironment, FileItem SourceFile, DirectoryReference OutputDir, List<string> Arguments, Action CompileAction, CPPOutput CompileResult)
+		{
+			FileItem Output = base.GetCompileArguments_FileType(CompileEnvironment, SourceFile, OutputDir, Arguments, CompileAction, CompileResult);
+
+			// @todo hack fix this better - it's so that files can see the generated -Swift.h headers - we don't know when we need them at this point
+			if (CompileEnvironment.Platform == UnrealTargetPlatform.VisionOS)
+			{
+				if (!SourceFile.HasExtension(".swift") && SourceFile.Location.ContainsName("Launch", 0))
+				{
+					FileItem OutputInteropHeader = GetBridgingHeader("UESwift", OutputDir);
+					Arguments.Add(GetUserIncludePathArgument(OutputInteropHeader.GetDirectoryItem().Location));
+					CompileAction.PrerequisiteItems.Add(OutputInteropHeader);
+				}
+			}
+			return Output;
+		}
+
+		private FileItem GetBridgingHeader(string SourceFile, DirectoryReference OutputDir)
+		{
+			string Filename = $"{Path.GetFileNameWithoutExtension(SourceFile)}-Swift.h";
+			return FileItem.GetItemByFileReference(FileReference.Combine(OutputDir, "Bridging", Filename));
+		}
+
+		protected override Action CompileCPPFile(CppCompileEnvironment CompileEnvironment, FileItem SourceFile, DirectoryReference OutputDir, string ModuleName, IActionGraphBuilder Graph, IReadOnlyCollection<string> GlobalArguments, CPPOutput Result)
+		{
+			if (SourceFile.HasExtension(".swift"))
+			{
+				return CompileSwiftFile(CompileEnvironment, SourceFile, OutputDir, Graph, Result);
+			}
+
+			return base.CompileCPPFile(CompileEnvironment, SourceFile, OutputDir, ModuleName, Graph, GlobalArguments, Result);
+		}
+
+		private Action CompileSwiftFile(CppCompileEnvironment CompileEnvironment, FileItem SourceFile, DirectoryReference OutputDir, IActionGraphBuilder Graph, CPPOutput CompileResult)
+		{
+			FileItem OutputFile = FileItem.GetItemByFileReference(FileReference.Combine(OutputDir, GetFileNameFromExtension(SourceFile.Name, ".o")));
+
+			List<string> Arguments = new();
+			Arguments.Add($"\"{SourceFile}\"");
+
+			// output file settings
+			Arguments.Add("-emit-object"); // same as -c
+			Arguments.Add($"-parse-as-library"); // don't create a main function
+			Arguments.Add($"-o \"{OutputFile}\"");
+
+			// platform settings
+			Arguments.Add($"-target {ToolChainSettings.Value.GetTargetTuple(CompileEnvironment.Architecture)}");
+			Arguments.Add($"-sdk {ToolChainSettings.Value.GetSDKPath(CompileEnvironment.Architecture)}");
+
+			// misc settings copied from Xcode
+			Arguments.Add("-Xllvm");
+			Arguments.Add("-aarch64-use-tbi");
+			Arguments.Add("-stack-check");
+			Arguments.Add($"-swift-version 5");
+			
+			Arguments.Add($"-module-name Launch");
+
+			Action CompileAction = Graph.CreateAction(ActionType.Compile);
+			CompileAction.Weight = CompileActionWeight;
+			CompileAction.CommandArguments = String.Join(" ", Arguments);
+			CompileAction.CommandPath = FileReference.Combine(ToolChainSettings.Value.ToolchainDir, "swift-frontend");
+			CompileAction.PrerequisiteItems.Add(SourceFile);
+			CompileAction.ProducedItems.Add(OutputFile);
+			CompileResult.ObjectFiles.Add(OutputFile);
+			CompileAction.WorkingDirectory = Unreal.EngineSourceDirectory;
+			CompileAction.CommandDescription = "Compile";
+			UnrealArchitectureConfig ArchConfig = UnrealArchitectureConfig.ForPlatform(CompileEnvironment.Platform);
+			if (ArchConfig.Mode != UnrealArchitectureMode.SingleArchitecture)
+			{
+				string ReadableArch = ArchConfig.ConvertToReadableArchitecture(CompileEnvironment.Architecture);
+				CompileAction.CommandDescription += $" [{ReadableArch}]";
+			}
+			CompileAction.StatusDescription = Path.GetFileName(SourceFile.AbsolutePath);
+			CompileAction.bIsGCCCompiler = true;
+			CompileAction.bCanExecuteRemotely = true;
+
+
+			FileItem OutputInteropHeader = GetBridgingHeader(SourceFile.FullName, OutputDir);
+
+			// obj-c bridging header settings
+			Arguments.Clear();
+			Arguments.Add($"\"{SourceFile}\"");
+			Arguments.Add("-parse"); // this will allow it to generate the header without writing out any .o/executable
+			Arguments.Add("-emit-objc-header");
+			Arguments.Add($"-emit-objc-header-path \"{OutputInteropHeader}\"");
+
+			// platform settings
+			Arguments.Add($"-target {ToolChainSettings.Value.GetTargetTuple(CompileEnvironment.Architecture)}");
+			Arguments.Add($"-sdk {ToolChainSettings.Value.GetSDKPath(CompileEnvironment.Architecture)}");
+
+			Arguments.Add($"-swift-version 5");
+
+			Arguments.Add($"-module-name Launch");
+
+			// now make an action to export the swift code as a header Obj-C bridging
+			Action HeaderAction = Graph.CreateAction(ActionType.CompileModuleInterface);
+			HeaderAction.CommandPath = FileReference.Combine(ToolChainSettings.Value.ToolchainDir, "swiftc");
+			HeaderAction.CommandArguments = String.Join(" ", Arguments);
+			HeaderAction.PrerequisiteItems.Add(SourceFile);
+			HeaderAction.ProducedItems.Add(OutputInteropHeader);
+			HeaderAction.CommandDescription = "Generate Header";
+			HeaderAction.StatusDescription = Path.GetFileName(OutputInteropHeader.AbsolutePath);
+			if (ArchConfig.Mode != UnrealArchitectureMode.SingleArchitecture)
+			{
+				string ReadableArch = ArchConfig.ConvertToReadableArchitecture(CompileEnvironment.Architecture);
+				CompileAction.CommandDescription += $" [{ReadableArch}]";
+			}
+
+			HeaderAction.WorkingDirectory = CompileAction.WorkingDirectory;
+			HeaderAction.bIsGCCCompiler = CompileAction.bIsGCCCompiler;
+			HeaderAction.bCanExecuteRemotely = CompileAction.bCanExecuteRemotely;
+
+			Console.WriteLine($"{CompileAction.CommandPath} {CompileAction.CommandArguments}");
+
+			// this is likely ignored, but the Compile action is the important one
+			return CompileAction;
+		}
+
 		protected string GetDsymutilPath(ILogger Logger, out string ExtraOptions, bool bIsForLTOBuild = false)
 		{
 			FileReference DsymutilLocation = new FileReference("/usr/bin/dsymutil");
@@ -639,7 +827,7 @@ namespace UnrealBuildTool
 			List<string> Arguments = new()
 			{
 				"UBT_NO_POST_DEPLOY=true",
-				new IOSToolChainSettings(Logger).XcodeDeveloperDir + "usr/bin/xcodebuild",
+				FileReference.Combine(AppleToolChainSettings.XcodeDeveloperDir, "usr/bin/xcodebuild").FullName,
 				Action,
 				$"-workspace \"{XcodeProject.FullName}\"",
 				$"-scheme \"{SchemeName}\"",
diff --git a/Engine/Source/Runtime/Launch/Private/IOS/LaunchIOS.cpp b/Engine/Source/Runtime/Launch/Private/IOS/LaunchIOS.cpp
--- a/Engine/Source/Runtime/Launch/Private/IOS/LaunchIOS.cpp
+++ b/Engine/Source/Runtime/Launch/Private/IOS/LaunchIOS.cpp
@@ -289,6 +289,11 @@ void FAppEntry::DecrementAudioSuspendCounters()
     }
 }
 
+#if PLATFORM_VISIONOS
+// SwiftUI test integration
+#include "UESwift-Swift.h"
+#endif
+
 void FAppEntry::PreInit(IOSAppDelegate* AppDelegate, UIApplication* Application)
 {
 	// make a controller object
@@ -312,6 +317,15 @@ void FAppEntry::PreInit(IOSAppDelegate* AppDelegate, UIApplication* Application)
 	// reset badge count on launch
 	Application.applicationIconBadgeNumber = 0;
 #endif
+	
+#if PLATFORM_VISIONOS
+	UIViewController* HostingController = [HostingViewFactory MakeSwiftUIViewOnClick:
+	^{
+		// close on click
+		[[IOSController presentedViewController] dismissViewControllerAnimated:YES completion:nil];
+	}];
+	[IOSController presentViewController:HostingController animated:YES completion:nil];
+#endif
 }
 
 static void MainThreadInit()
diff --git a/Engine/Source/Runtime/Launch/Private/VisionOS/UESwift.swift b/Engine/Source/Runtime/Launch/Private/VisionOS/UESwift.swift
new file mode 100644
--- /dev/null
+++ b/Engine/Source/Runtime/Launch/Private/VisionOS/UESwift.swift
@@ -0,0 +1,39 @@
+//
+//  SwiftUIView.swift
+//  IOSOnVision
+//
+//  Created by Josh Adams on 6/28/23.
+//
+
+import SwiftUI
+
+struct SwiftUIView: View {
+  var onClick: () -> Void = {}
+  
+  var body: some View {
+	VStack(spacing: 8) {
+	  Text("SwiftUI in Unreal!")
+		.font(.title)
+		.bold()
+	  Button(action: {
+		onClick()
+	  }, label: {
+		Text("Test Button")
+	  })
+	}
+  }
+}
+
+
+class HostingViewFactory: NSObject
+{
+  @objc static func MakeSwiftUIView(OnClick: @escaping (() -> Void)) -> UIViewController
+  {
+	return UIHostingController(rootView: SwiftUIView(onClick: OnClick))
+  }
+}
+
+
+//#Preview {
+//    SwiftUIView()
+//}
