- Fixing the name of the .app when staging/deploying a BP IOS project with modern
- Uses a new AppleExports function to call MakeBinaryFilename, so we can use standardized name format for the BP .app

diff --git a/Engine/Source/Programs/AutomationTool/Apple/ApplePlatform.Automation.cs b/Engine/Source/Programs/AutomationTool/Apple/ApplePlatform.Automation.cs
--- a/Engine/Source/Programs/AutomationTool/Apple/ApplePlatform.Automation.cs
+++ b/Engine/Source/Programs/AutomationTool/Apple/ApplePlatform.Automation.cs
@@ -98,7 +98,9 @@ public abstract class ApplePlatform : Platform
 					$"UE_OVERRIDE_STAGE_DIR=\"{SC.StageDirectory}\"";
 				if (!Params.IsCodeBasedProject)
 				{
-					ExtraOptions += $" PRODUCT_NAME={Params.ShortProjectName}";
+					// instead of staging an UnrealGame.app, stage something with the project name, like say MyProjectClient-IOS-Shipping.app
+					string ProductName = AppleExports.MakeBinaryFileName(SC.ShortProjectName, Target.Platform, Target.Configuration, Target.Architectures, UnrealTargetConfiguration.Development, null);
+					ExtraOptions += $" PRODUCT_NAME={ProductName}";
 				}
 
 				AppleExports.BuildWithStubXcodeProject(SC.RawProjectPath, Target.Platform, Target.Configuration, Target.TargetName, AppleExports.XcodeBuildMode.Stage, Logger, ExtraOptions);
diff --git a/Engine/Source/Programs/AutomationTool/IOS/IOSPlatform.Automation.cs b/Engine/Source/Programs/AutomationTool/IOS/IOSPlatform.Automation.cs
--- a/Engine/Source/Programs/AutomationTool/IOS/IOSPlatform.Automation.cs
+++ b/Engine/Source/Programs/AutomationTool/IOS/IOSPlatform.Automation.cs
@@ -1839,7 +1839,15 @@ public class IOSPlatform : ApplePlatform
 		string AppToDeploy;
 		if (AppleExports.UseModernXcode(Params.RawProjectPath))
 		{
-			AppToDeploy = FileReference.Combine(SC.StageDirectory, SC.StageExecutables[0] + ".app").FullName;
+			// modern uses a project-name for the content-only .app, but the StageExecutable is UnrealGame-IOS-Shipping or similar, so
+			// fix up the Unreal part
+			string AppBaseName = SC.StageExecutables[0];
+			if (!SC.IsCodeBasedProject)
+			{
+				TargetReceipt Target = SC.StageTargets[0].Receipt;
+				AppBaseName = AppleExports.MakeBinaryFileName(SC.ShortProjectName, Target.Platform, Target.Configuration, Target.Architectures, UnrealTargetConfiguration.Development, null);
+			}
+			AppToDeploy = FileReference.Combine(SC.StageDirectory, AppBaseName + ".app").FullName;
 
 			// @todo: deal with iterative deploy - we have the uncombined .app in the Staged dir inin Binaries/IOS
 			if (Params.IterativeDeploy)
diff --git a/Engine/Source/Programs/UnrealBuildTool/Platform/Apple/AppleExports.cs b/Engine/Source/Programs/UnrealBuildTool/Platform/Apple/AppleExports.cs
--- a/Engine/Source/Programs/UnrealBuildTool/Platform/Apple/AppleExports.cs
+++ b/Engine/Source/Programs/UnrealBuildTool/Platform/Apple/AppleExports.cs
@@ -208,5 +208,21 @@ namespace UnrealBuildTool
 		{
 			return AppleToolChain.FinalizeAppWithXcode(XcodeProject, Platform, SchemeName, Configuration, Action, ExtraOptions, Logger);
 		}
+
+		/// <summary>
+		/// Pass along the call to UEBuildTarget.MakeBinaryFileName, but taking an extension instead of a binary type, since the type is hidden
+		/// </summary>
+		/// <param name="BinaryName"></param>
+		/// <param name="Platform"></param>
+		/// <param name="Configuration"></param>
+		/// <param name="Architectures"></param>
+		/// <param name="UndecoratedConfiguration"></param>
+		/// <param name="Extension"></param>
+		/// <returns></returns>
+		public static string MakeBinaryFileName(string BinaryName, UnrealTargetPlatform Platform, UnrealTargetConfiguration Configuration, UnrealArchitectures Architectures, UnrealTargetConfiguration UndecoratedConfiguration, string? Extension)
+		{
+			string StandardBinaryName = UEBuildTarget.MakeBinaryFileName(BinaryName, Platform, Configuration, Architectures, UndecoratedConfiguration, UEBuildBinaryType.Executable);
+			return System.IO.Path.ChangeExtension(StandardBinaryName, Extension);
+		}
 	}
 }
