UnrealBuildTool: Fix analyze mode to actually show the shortest chain

#rnx

diff --git a/Engine/Source/Programs/UnrealBuildTool/Modes/AnalyzeMode.cs b/Engine/Source/Programs/UnrealBuildTool/Modes/AnalyzeMode.cs
--- a/Engine/Source/Programs/UnrealBuildTool/Modes/AnalyzeMode.cs
+++ b/Engine/Source/Programs/UnrealBuildTool/Modes/AnalyzeMode.cs
@@ -51,7 +51,7 @@ namespace UnrealBuildTool.Modes
 		class ModuleInfo
 		{
 			public UEBuildModule Module;
-			public string Chain;
+			public IReadOnlyList<string> IncludeChain;
 			public HashSet<UEBuildModule> InwardRefs = new HashSet<UEBuildModule>();
 			public HashSet<UEBuildModule> UniqueInwardRefs = new HashSet<UEBuildModule>();
 			public HashSet<UEBuildModule> OutwardRefs = new HashSet<UEBuildModule>();
@@ -62,36 +62,72 @@ namespace UnrealBuildTool.Modes
 			public List<FileReference> BinaryFiles = new List<FileReference>();
 			public long BinSize = 0;
 
-			public ModuleInfo(UEBuildModule Module, string Chain)
+			public ModuleInfo(UEBuildModule Module, params string[] IncludeChain)
 			{
 				this.Module = Module;
-				this.Chain = Chain;
+				this.IncludeChain = IncludeChain;
 			}
+
+			public string Chain => String.Join(" -> ", IncludeChain.Where(x => !String.IsNullOrEmpty(x)));
 		}
 
-		private void AnalyzeTarget(TargetDescriptor TargetDescriptor, BuildConfiguration BuildConfiguration, ILogger Logger)
+		private static void AnalyzeModuleChains(string ModuleName, List<string> ParentChain, UEBuildTarget Target, Dictionary<UEBuildModule, ModuleInfo> ModuleToInfo, HashSet<UEBuildModule> Visited, ILogger Logger)
+		{
+			// Prevent recursive includes, they'll never be shorter
+			if (ParentChain.Contains(ModuleName))
+			{
+				return;
+			}
+
+			UEBuildModule Module = Target.GetModuleByName(ModuleName);
+
+			List<string> CurrentChain = new(ParentChain);
+			CurrentChain.Add(Module.Name);
+
+			if (!ModuleToInfo.ContainsKey(Module))
+			{
+				ModuleToInfo[Module] = new ModuleInfo(Module, CurrentChain.ToArray());
+			}
+			
+			if (ModuleToInfo[Module].IncludeChain.Count > CurrentChain.Count)
+			{
+				// Now we need to recheck all downstream dependencies again because the chain may be shorter
+				List<UEBuildModule> RecheckModules = new();
+				Module.GetAllDependencyModules(RecheckModules, new(), true, false, true);
+				Visited.ExceptWith(RecheckModules);
+				Logger.LogDebug("Found shorter chain for {Module} {Prev} -> {New}, rechecking {Count} already visited dependencies", Module.Name, ModuleToInfo[Module].IncludeChain.Count, CurrentChain.Count, RecheckModules.Count);
+				ModuleToInfo[Module].IncludeChain = CurrentChain.ToArray();
+			}
+			else if (Visited.Contains(Module))
+			{
+				return;
+			}
+
+			Visited.Add(Module);
+
+			List<UEBuildModule> TargetModules = new();
+			Module.GetAllDependencyModules(TargetModules, new(), true, false, true);
+			TargetModules.ForEach(x => AnalyzeModuleChains(x.Name, CurrentChain, Target, ModuleToInfo, Visited, Logger));
+		}
+
+		private static void AnalyzeTarget(TargetDescriptor TargetDescriptor, BuildConfiguration BuildConfiguration, ILogger Logger)
 		{
 			// Create a makefile for the target
 			UEBuildTarget Target = UEBuildTarget.Create(TargetDescriptor, BuildConfiguration.bSkipRulesCompile, BuildConfiguration.bForceRulesCompile, BuildConfiguration.bUsePrecompiled, Logger);
 			DirectoryReference.CreateDirectory(Target.ReceiptFileName.Directory);
 
 			// Find the shortest path from the target to each module
+			HashSet<UEBuildModule> Visited = new();
 			Dictionary<UEBuildModule, ModuleInfo> ModuleToInfo = new Dictionary<UEBuildModule, ModuleInfo>();
 
-			List<string> RootModuleNames = new List<string>(Target.Rules.ExtraModuleNames);
 			if (Target.Rules.LaunchModuleName != null)
 			{
-				RootModuleNames.Add(Target.Rules.LaunchModuleName);
+				AnalyzeModuleChains(Target.Rules.LaunchModuleName, new List<string>() { "target" }, Target, ModuleToInfo, Visited, Logger);
 			}
 
-			foreach (string RootModuleName in RootModuleNames)
+			foreach (string RootModuleName in Target.Rules.ExtraModuleNames)
 			{
-				UEBuildModule Module = Target.GetModuleByName(RootModuleName);
-				if (Module != null)
-				{
-					string Chain = $"target -> {RootModuleName}";
-					ModuleToInfo[Module] = new ModuleInfo(Module, Chain);
-				}
+				AnalyzeModuleChains(RootModuleName, new List<string>() { "target" }, Target, ModuleToInfo, Visited, Logger);
 			}
 
 			// Also enable all the plugin modules
@@ -99,34 +135,36 @@ namespace UnrealBuildTool.Modes
 			{
 				foreach (UEBuildModule Module in Plugin.Modules)
 				{
-					string Chain = $"{Plugin.ReferenceChain} -> {Module.Name}";
-					ModuleToInfo[Module] = new ModuleInfo(Module, Chain);
+					if (!ModuleToInfo.ContainsKey(Module))
+					{
+						AnalyzeModuleChains(Module.Name, new List<string>() { "target", Plugin.ReferenceChain, Plugin.File.GetFileName() }, Target, ModuleToInfo, Visited, Logger);
+					}
 				}
 			}
 
-			// Set of visited modules
-			HashSet<UEBuildModule> VisitedModules = new HashSet<UEBuildModule>();
-
-			// Recurse out to find new modules and the shortest path to each
-			List<UEBuildModule> SourceModules = new List<UEBuildModule>(ModuleToInfo.Keys);
-			while (SourceModules.Count > 0)
+			if (Target.Rules.bBuildAllModules)
 			{
-				List<UEBuildModule> TargetModules = new List<UEBuildModule>();
-
-				foreach (UEBuildModule SourceModule in SourceModules)
+				foreach (UEBuildBinary Binary in Target.Binaries)
 				{
-					int Idx = TargetModules.Count;
-					SourceModule.GetAllDependencyModules(TargetModules, VisitedModules, true, false, true);
-
-					for (; Idx < TargetModules.Count; Idx++)
+					foreach (UEBuildModule Module in Binary.Modules)
 					{
-						UEBuildModule TargetModule = TargetModules[Idx];
-						string Chain = $"{ModuleToInfo[SourceModule].Chain} -> {TargetModule.Name}";
-						ModuleToInfo[TargetModule] = new ModuleInfo(TargetModule, Chain);
+						if (!ModuleToInfo.ContainsKey(Module))
+						{
+							// quick hack to make allmodules always worse (empty entries are ignored when writing)
+							List<string> IncludeChain = Enumerable.Repeat(String.Empty, 1000).ToList();
+							IncludeChain.Add("allmodules option");
+							if (Module.Rules.Plugin != null)
+							{
+								IncludeChain.Add(Module.Rules.Plugin.File.GetFileName());
+								AnalyzeModuleChains(Module.Name, IncludeChain, Target, ModuleToInfo, Visited, Logger);
+							}
+							else
+							{
+								AnalyzeModuleChains(Module.Name, IncludeChain, Target, ModuleToInfo, Visited, Logger);
+							}
+						}
 					}
 				}
-
-				SourceModules = TargetModules;
 			}
 
 			// Find all the outward dependencies of each module
@@ -240,8 +278,8 @@ namespace UnrealBuildTool.Modes
 				foreach (ModuleInfo ModuleInfo in ModuleToInfo.Values.OrderByDescending(x => x.InwardRefs.Count).ThenBy(x => x.BinSize))
 				{
 					Writer.WriteLine("");
-					Writer.WriteLine("Module:        \"{0}\"", ModuleInfo.Module.Name);
-					Writer.WriteLine("Shortest path: {0}", ModuleInfo.Chain);
+					Writer.WriteLine("Module:                  \"{0}\"", ModuleInfo.Module.Name);
+					Writer.WriteLine("Shortest path:           {0}", ModuleInfo.Chain);
 					WriteDependencyList(Writer, "Unique inward refs:     ", ModuleInfo.UniqueInwardRefs);
 					WriteDependencyList(Writer, "Unique outward refs:    ", ModuleInfo.UniqueOutwardRefs);
 					WriteDependencyList(Writer, "Recursive inward refs:  ", ModuleInfo.InwardRefs);
@@ -304,7 +342,7 @@ namespace UnrealBuildTool.Modes
 			}
 		}
 
-		private void WriteDependencyList(TextWriter Writer, string Prefix, HashSet<UEBuildModule> Modules)
+		private static void WriteDependencyList(TextWriter Writer, string Prefix, HashSet<UEBuildModule> Modules)
 		{
 			if (Modules.Count == 0)
 			{
@@ -316,7 +354,7 @@ namespace UnrealBuildTool.Modes
 			}
 		}
 
-		private void WriteDependencyGraph(UEBuildTarget Target, Dictionary<UEBuildModule, ModuleInfo> ModuleToInfo, FileReference FileName)
+		private static void WriteDependencyGraph(UEBuildTarget Target, Dictionary<UEBuildModule, ModuleInfo> ModuleToInfo, FileReference FileName)
 		{
 			List<GraphNode> Nodes = new List<GraphNode>();
 
@@ -361,7 +399,7 @@ namespace UnrealBuildTool.Modes
 			GraphVisualization.WriteGraphFile(FileName, $"Module dependency graph for {Target.TargetName}", Nodes, Edges);
 		}
 
-		private void WriteShortestPathGraph(UEBuildTarget Target, Dictionary<UEBuildModule, ModuleInfo> ModuleToInfo, FileReference FileName)
+		private static void WriteShortestPathGraph(UEBuildTarget Target, Dictionary<UEBuildModule, ModuleInfo> ModuleToInfo, FileReference FileName)
 		{
 			Dictionary<string, GraphNode> NameToNode = new Dictionary<string, GraphNode>(StringComparer.Ordinal);
 
