UnrealBuildTool: Disable bDetailedUnityFiles to prevent excess stale intermediate files

#rnx

diff --git a/Engine/Source/Programs/UnrealBuildTool/Configuration/TargetRules.cs b/Engine/Source/Programs/UnrealBuildTool/Configuration/TargetRules.cs
--- a/Engine/Source/Programs/UnrealBuildTool/Configuration/TargetRules.cs
+++ b/Engine/Source/Programs/UnrealBuildTool/Configuration/TargetRules.cs
@@ -1676,12 +1676,12 @@ namespace UnrealBuildTool
 		public bool bStressTestUnity = false;
 
 		/// <summary>
-		/// Whether to add additional information to the unity files, such as '_of_X' in the file name.
+		/// Whether to add additional information to the unity files, such as '_of_X' in the file name. Not recommended.
 		/// </summary>
 		[CommandLine("-DetailedUnityFiles", Value = "true")]
 		[CommandLine("-NoDetailedUnityFiles", Value = "false")]
 		[XmlConfigFile(Category = "BuildConfiguration")]
-		public bool bDetailedUnityFiles = true;
+		public bool bDetailedUnityFiles = false;
 
 		/// <summary>
 		/// Whether to force debug info to be generated.
