Updated the VS project generator code to check to see if there are common settings that should apply to all of the source files.

diff --git a/Engine/Source/Programs/UnrealBuildTool/ProjectFiles/VisualStudio/VCProject.cs b/Engine/Source/Programs/UnrealBuildTool/ProjectFiles/VisualStudio/VCProject.cs
--- a/Engine/Source/Programs/UnrealBuildTool/ProjectFiles/VisualStudio/VCProject.cs
+++ b/Engine/Source/Programs/UnrealBuildTool/ProjectFiles/VisualStudio/VCProject.cs
@@ -875,65 +875,6 @@
 
 			bool bSuccess = true;
 
-			// Merge as many include paths as possible into the shared list
-			HashSet<DirectoryReference> SharedIncludeSearchPathsSet = new HashSet<DirectoryReference>();
-
-			// Build up the new include search path string
-			StringBuilder SharedIncludeSearchPaths = new StringBuilder();
-			{
-				// Find out how many source files there are in each directory
-				ConcurrentDictionary<DirectoryReference, int> SourceDirToCount = new();
-				Parallel.ForEach(SourceFiles.Where(sf => sf.Reference.HasExtension(".cpp")), SourceFile =>
-				{
-					SourceDirToCount.AddOrUpdate(SourceFile.Reference.Directory, _ => 1, (k, v) => v + 1);
-				});
-
-				// Figure out the most common include paths
-				ConcurrentDictionary<DirectoryReference, int> IncludePathToCount = new();
-				Parallel.ForEach(SourceDirToCount, Pair =>
-				{
-					if (TryGetBuildEnvironment(Pair.Key, out BuildEnvironment? OutBuildEnvironment))
-					{
-						foreach (DirectoryReference IncludePath in OutBuildEnvironment.UserIncludePaths.AbsolutePaths)
-						{
-							IncludePathToCount.AddOrUpdate(IncludePath, _ => Pair.Value, (k, v) => v + Pair.Value);
-						}
-						return;
-					}
-				});
-
-				// Append the most common include paths to the search list.
-				if (Settings.MaxSharedIncludePaths > 0)
-				{
-					foreach (DirectoryReference IncludePath in IncludePathToCount.OrderByDescending(x => x.Value).ThenBy(x => x.Key).Select(x => x.Key))
-					{
-						string RelativePath = NormalizeProjectPath(IncludePath);
-						if (SharedIncludeSearchPaths.Length + RelativePath.Length >= Settings.MaxSharedIncludePaths)
-						{
-							break;
-						}
-
-						if (!IncludePathIsFilteredOut(IncludePath))
-						{
-							SharedIncludeSearchPathsSet.Add(IncludePath);
-							SharedIncludeSearchPaths.AppendFormat("{0};", RelativePath);
-						}
-					}
-				}
-
-				SharedIncludeSearchPaths.AppendFormat("$(DefaultSystemIncludePaths);");
-			}
-
-			StringBuilder VCPreprocessorDefinitions = new StringBuilder();
-			foreach (string CurDef in IntelliSensePreprocessorDefinitions)
-			{
-				if (VCPreprocessorDefinitions.Length > 0)
-				{
-					VCPreprocessorDefinitions.Append(';');
-				}
-				VCPreprocessorDefinitions.Append(CurDef);
-			}
-
 			// Setup VC project file content
 			StringBuilder VCProjectFileContent = new StringBuilder();
 			StringBuilder VCFiltersFileContent = new StringBuilder();
@@ -1170,28 +1111,61 @@
 				DefaultRules = DefaultRules ?? Combination.ProjectTarget?.TargetRules;
 			}
 
-			// Write IntelliSense info
+			// Merge as many include paths as possible into the shared list
+			HashSet<DirectoryReference> SharedIncludeSearchPathsSet = new HashSet<DirectoryReference>();
+
+			// Build up the new include search path string
+			StringBuilder SharedIncludeSearchPaths = new StringBuilder();
 			{
-				// @todo projectfiles: Currently we are storing defines/include paths for ALL configurations rather than using ConditionString and storing
-				//      this data uniquely for each target configuration.  IntelliSense may behave better if we did that, but it will result in a LOT more
-				//      data being stored into the project file, and might make the IDE perform worse when switching configurations!
-				VCProjectFileContent.AppendLine("  <PropertyGroup>");
-				VCProjectFileContent.AppendLine("    <NMakePreprocessorDefinitions>$(NMakePreprocessorDefinitions){0}</NMakePreprocessorDefinitions>", (VCPreprocessorDefinitions.Length > 0 ? (";" + VCPreprocessorDefinitions) : ""));
-				// NOTE: Setting the IncludePath property rather than NMakeIncludeSearchPath results in significantly less
-				// memory usage, because NMakeIncludeSearchPath metadata is duplicated to each output item. Functionality should be identical for
-				// intellisense results.
-				VCProjectFileContent.AppendLine("    <IncludePath>$(IncludePath){0}</IncludePath>", (SharedIncludeSearchPaths.Length > 0 ? (";" + SharedIncludeSearchPaths) : ""));
-				VCProjectFileContent.AppendLine("    <NMakeForcedIncludes>$(NMakeForcedIncludes)</NMakeForcedIncludes>");
-				VCProjectFileContent.AppendLine("    <NMakeAssemblySearchPath>$(NMakeAssemblySearchPath)</NMakeAssemblySearchPath>");
-				VCProjectFileContent.AppendLine("    <AdditionalOptions>{0} {1} {2}</AdditionalOptions>",
-					GetCppStandardCompileArgument(GetIntelliSenseCppVersion()), GetEnableCoroutinesArgument(), DefaultRules != null ? GetConformanceCompileArguments(DefaultRules) : string.Empty);
-				VCProjectFileContent.AppendLine("  </PropertyGroup>");
+				// Find out how many source files there are in each directory
+				ConcurrentDictionary<DirectoryReference, int> SourceDirToCount = new();
+				Parallel.ForEach(SourceFiles.Where(sf => sf.Reference.HasExtension(".cpp")), SourceFile =>
+				{
+					SourceDirToCount.AddOrUpdate(SourceFile.Reference.Directory, _ => 1, (k, v) => v + 1);
+				});
+
+				// Figure out the most common include paths
+				ConcurrentDictionary<DirectoryReference, int> IncludePathToCount = new();
+				Parallel.ForEach(SourceDirToCount, Pair =>
+				{
+					if (TryGetBuildEnvironment(Pair.Key, out BuildEnvironment? OutBuildEnvironment))
+					{
+						foreach (DirectoryReference IncludePath in OutBuildEnvironment.UserIncludePaths.AbsolutePaths)
+						{
+							IncludePathToCount.AddOrUpdate(IncludePath, _ => Pair.Value, (k, v) => v + Pair.Value);
+						}
+						return;
+					}
+				});
+
+				// Append the most common include paths to the search list.
+				if (Settings.MaxSharedIncludePaths > 0)
+				{
+					foreach (DirectoryReference IncludePath in IncludePathToCount.OrderByDescending(x => x.Value).ThenBy(x => x.Key).Select(x => x.Key))
+					{
+						string RelativePath = NormalizeProjectPath(IncludePath);
+						if (SharedIncludeSearchPaths.Length + RelativePath.Length >= Settings.MaxSharedIncludePaths)
+						{
+							break;
+						}
+
+						if (!IncludePathIsFilteredOut(IncludePath))
+						{
+							SharedIncludeSearchPathsSet.Add(IncludePath);
+							SharedIncludeSearchPaths.AppendFormat("{0};", RelativePath);
+						}
+					}
+				}
+
+				SharedIncludeSearchPaths.AppendFormat("$(DefaultSystemIncludePaths);");
 			}
 
-			// Source folders and files
+			// Gather source folder and file info
+			List<AliasedFile> LocalAliasedFiles = new List<AliasedFile>(AliasedFiles);
+			ConcurrentDictionary<DirectoryReference, string?> DirectoryToPchFile = new();
+			ConcurrentDictionary<DirectoryReference, string> DirectoryToForceIncludePaths = new();
+			ConcurrentDictionary<DirectoryReference, string> DirectoryToIncludeSearchPaths = new();
 			{
-				List<AliasedFile> LocalAliasedFiles = new List<AliasedFile>(AliasedFiles);
-
 				foreach (SourceFile CurFile in SourceFiles)
 				{
 					// We want all source file and directory paths in the project files to be relative to the project file's
@@ -1224,10 +1198,6 @@
 
 					LocalAliasedFiles.Add(new AliasedFile(CurFile.Reference, ProjectRelativeSourceFile, FilterRelativeSourceDirectory));
 				}
-
-				ConcurrentDictionary<DirectoryReference, string?> DirectoryToPchFile = new();
-				ConcurrentDictionary<DirectoryReference, string> DirectoryToForceIncludePaths = new();
-				ConcurrentDictionary<DirectoryReference, string> DirectoryToIncludeSearchPaths = new();
 				Parallel.ForEach(LocalAliasedFiles, LocalAliasedFile =>
 				{
 					// get the filetype as represented to Visual Studio
@@ -1304,7 +1274,78 @@
 						});
 					}
 				});
+			}
 
+			// Check to see if all the source settings are the same
+			string CommonForcedIncludes = string.Empty;
+			string CommonAdditionalOptions = string.Empty;
+			{
+				if (DirectoryToIncludeSearchPaths.Any())
+				{
+					string IncludePathToCheck = DirectoryToIncludeSearchPaths.Values.First();
+					if (DirectoryToIncludeSearchPaths.Values.All(x => x == IncludePathToCheck))
+					{
+						SharedIncludeSearchPaths.Append(IncludePathToCheck);
+						DirectoryToIncludeSearchPaths.Clear();
+					}
+				}
+
+				if (DirectoryToForceIncludePaths.Any())
+				{
+					string ForceIncludePathToCheck = DirectoryToForceIncludePaths.Values.First();
+					if (DirectoryToForceIncludePaths.Values.All(x => x == ForceIncludePathToCheck))
+					{
+						CommonForcedIncludes = ForceIncludePathToCheck;
+						DirectoryToForceIncludePaths.Clear();
+					}
+				}
+
+				if (DirectoryToPchFile.Any())
+				{
+					string? PchFileToCheck = DirectoryToPchFile.Values.FirstOrDefault();
+					if (DirectoryToPchFile.Values.All(x => x == PchFileToCheck))
+					{
+						if (!string.IsNullOrEmpty(PchFileToCheck))
+						{
+							CommonAdditionalOptions = $"/Yu\"{PchFileToCheck}\"";
+						}
+						DirectoryToPchFile.Clear();
+					}
+				}
+			}
+
+			StringBuilder VCPreprocessorDefinitions = new StringBuilder();
+			foreach (string CurDef in IntelliSensePreprocessorDefinitions)
+			{
+				if (VCPreprocessorDefinitions.Length > 0)
+				{
+					VCPreprocessorDefinitions.Append(';');
+				}
+				VCPreprocessorDefinitions.Append(CurDef);
+			}
+
+			// Write IntelliSense info
+			{
+				// @todo projectfiles: Currently we are storing defines/include paths for ALL configurations rather than using ConditionString and storing
+				//      this data uniquely for each target configuration.  IntelliSense may behave better if we did that, but it will result in a LOT more
+				//      data being stored into the project file, and might make the IDE perform worse when switching configurations!
+				VCProjectFileContent.AppendLine("  <PropertyGroup>");
+				VCProjectFileContent.AppendLine("    <NMakePreprocessorDefinitions>$(NMakePreprocessorDefinitions){0}</NMakePreprocessorDefinitions>", (VCPreprocessorDefinitions.Length > 0 ? (";" + VCPreprocessorDefinitions) : string.Empty));
+				// NOTE: Setting the IncludePath property rather than NMakeIncludeSearchPath results in significantly less
+				// memory usage, because NMakeIncludeSearchPath metadata is duplicated to each output item. Functionality should be identical for
+				// intellisense results.
+				VCProjectFileContent.AppendLine("    <IncludePath>$(IncludePath){0}</IncludePath>", (SharedIncludeSearchPaths.Length > 0 ? (";" + SharedIncludeSearchPaths) : ""));
+				VCProjectFileContent.AppendLine("    <NMakeForcedIncludes>$(NMakeForcedIncludes){0}</NMakeForcedIncludes>", (CommonForcedIncludes.Length > 0 ? (";" + CommonForcedIncludes) : ""));
+				VCProjectFileContent.AppendLine("    <NMakeAssemblySearchPath>$(NMakeAssemblySearchPath)</NMakeAssemblySearchPath>");
+				VCProjectFileContent.AppendLine("    <AdditionalOptions>{0} {1}{2}{3}</AdditionalOptions>",
+					GetCppStandardCompileArgument(GetIntelliSenseCppVersion()), 
+					GetEnableCoroutinesArgument(), 
+					DefaultRules != null ? (" " + GetConformanceCompileArguments(DefaultRules)) : string.Empty, 
+					CommonAdditionalOptions.Length > 0 ? (" " + CommonAdditionalOptions) : string.Empty);
+				VCProjectFileContent.AppendLine("  </PropertyGroup>");
+			}
+
+			{ 
 				// Collapse common values
 				{
 					StringBuilder CommonProjectFileContent = new StringBuilder();
@@ -1479,18 +1520,36 @@
 						// Find the include search paths
 						if (TryGetBuildEnvironment(Directory, out BuildEnvironment? BuildEnvironment))
 						{
-							VCProjectFileContent.AppendLine("    <{0} Include=\"{1}\">", VCFileType, EscapeFileName(AliasedFile.FileSystemPath));
-
-							VCProjectFileContent.AppendLine($"      <AdditionalIncludeDirectories>$({DirectoryToIncludeSearchPaths[Directory]})</AdditionalIncludeDirectories>");
-							VCProjectFileContent.AppendLine($"      <ForcedIncludeFiles>$({DirectoryToForceIncludePaths[Directory]})</ForcedIncludeFiles>");
-
-							string? PchHeaderFile = DirectoryToPchFile[Directory];
-							if (PchHeaderFile != null && ProjectFileFormat >= VCProjectFileFormat.VisualStudio2022)
+							StringBuilder ClCompileInfo = new();
+							if (DirectoryToIncludeSearchPaths.Any())
 							{
-								VCProjectFileContent.AppendLine($"      <AdditionalOptions>$({DirectoryToPchFile[Directory]})</AdditionalOptions>");
+								ClCompileInfo.AppendLine($"      <AdditionalIncludeDirectories>$({DirectoryToIncludeSearchPaths[Directory]})</AdditionalIncludeDirectories>");
 							}
 
-							VCProjectFileContent.AppendLine("    </{0}>", VCFileType);
+							if (DirectoryToForceIncludePaths.Any())
+							{
+								ClCompileInfo.AppendLine($"      <ForcedIncludeFiles>$({DirectoryToForceIncludePaths[Directory]})</ForcedIncludeFiles>");
+							}
+
+							if (DirectoryToPchFile.Any())
+							{
+								string? PchHeaderFile = DirectoryToPchFile[Directory];
+								if (PchHeaderFile != null && ProjectFileFormat >= VCProjectFileFormat.VisualStudio2022)
+								{
+									ClCompileInfo.AppendLine($"      <AdditionalOptions>$({DirectoryToPchFile[Directory]})</AdditionalOptions>");
+								}
+							}
+
+							if (ClCompileInfo.Length == 0)
+							{
+								VCProjectFileContent.AppendLine("    <{0} Include=\"{1}\"/>", VCFileType, EscapeFileName(AliasedFile.FileSystemPath));
+							}
+							else
+							{
+								VCProjectFileContent.AppendLine("    <{0} Include=\"{1}\">", VCFileType, EscapeFileName(AliasedFile.FileSystemPath));
+								VCProjectFileContent.Append(ClCompileInfo.ToString());
+								VCProjectFileContent.AppendLine("    </{0}>", VCFileType);
+							}
 						}
 						else
 						{
