Removed the redundant writing of the additional options.

diff --git a/Engine/Source/Programs/UnrealBuildTool/ProjectFiles/VisualStudio/VCProject.cs b/Engine/Source/Programs/UnrealBuildTool/ProjectFiles/VisualStudio/VCProject.cs
--- a/Engine/Source/Programs/UnrealBuildTool/ProjectFiles/VisualStudio/VCProject.cs
+++ b/Engine/Source/Programs/UnrealBuildTool/ProjectFiles/VisualStudio/VCProject.cs
@@ -1103,14 +1103,6 @@ namespace UnrealBuildTool
 			VCProjectFileContent.AppendLine("  <ImportGroup Label=\"ExtensionSettings\" />");
 			VCProjectFileContent.AppendLine("  <PropertyGroup Label=\"UserMacros\" />");
 
-			// Write each project configuration
-			TargetRules? DefaultRules = null;
-			foreach (ProjectConfigAndTargetCombination Combination in ProjectConfigAndTargetCombinations)
-			{
-				WriteConfiguration(ProjectName, Combination, VCProjectFileContent, PlatformProjectGenerators, bGenerateUserFileContent ? VCUserFileContent : null, Logger);
-				DefaultRules = DefaultRules ?? Combination.ProjectTarget?.TargetRules;
-			}
-
 			// Merge as many include paths as possible into the shared list
 			HashSet<DirectoryReference> SharedIncludeSearchPathsSet = new HashSet<DirectoryReference>();
 
@@ -1324,7 +1316,20 @@ namespace UnrealBuildTool
 				VCPreprocessorDefinitions.Append(CurDef);
 			}
 
-			// Write IntelliSense info
+			ProjectConfigAndTargetCombination? FoundCombo = ProjectConfigAndTargetCombinations.FirstOrDefault(combo => combo != null && combo.ProjectTarget != null && combo.ProjectTarget.TargetRules != null);
+			TargetRules? DefaultRules = FoundCombo != null ? FoundCombo.ProjectTarget?.TargetRules : null;
+
+			var GetAdditionalOptionsString = (TargetRules? TargetRules) =>
+			{
+				return string.Format("{0} {1}{2}{3}", GetCppStandardCompileArgument(GetIntelliSenseCppVersion()),
+					GetEnableCoroutinesArgument(),
+					DefaultRules != null ? (" " + GetConformanceCompileArguments(DefaultRules)) : string.Empty,
+					CommonAdditionalOptions.Length > 0 ? (" " + CommonAdditionalOptions) : string.Empty);
+			};
+			
+			string DefaultAdditionalOptions = GetAdditionalOptionsString(DefaultRules);
+
+			// Write common IntelliSense info
 			{
 				// @todo projectfiles: Currently we are storing defines/include paths for ALL configurations rather than using ConditionString and storing
 				//      this data uniquely for each target configuration.  IntelliSense may behave better if we did that, but it will result in a LOT more
@@ -1337,14 +1342,38 @@ namespace UnrealBuildTool
 				VCProjectFileContent.AppendLine("    <IncludePath>$(IncludePath){0}</IncludePath>", (SharedIncludeSearchPaths.Length > 0 ? (";" + SharedIncludeSearchPaths) : ""));
 				VCProjectFileContent.AppendLine("    <NMakeForcedIncludes>$(NMakeForcedIncludes){0}</NMakeForcedIncludes>", (CommonForcedIncludes.Length > 0 ? (";" + CommonForcedIncludes) : ""));
 				VCProjectFileContent.AppendLine("    <NMakeAssemblySearchPath>$(NMakeAssemblySearchPath)</NMakeAssemblySearchPath>");
-				VCProjectFileContent.AppendLine("    <AdditionalOptions>{0} {1}{2}{3}</AdditionalOptions>",
-					GetCppStandardCompileArgument(GetIntelliSenseCppVersion()), 
-					GetEnableCoroutinesArgument(), 
-					DefaultRules != null ? (" " + GetConformanceCompileArguments(DefaultRules)) : string.Empty, 
-					CommonAdditionalOptions.Length > 0 ? (" " + CommonAdditionalOptions) : string.Empty);
+				VCProjectFileContent.AppendLine("    <AdditionalOptions>{0}</AdditionalOptions>", DefaultAdditionalOptions);
 				VCProjectFileContent.AppendLine("  </PropertyGroup>");
 			}
 
+			// Write platform properties
+			HashSet<UnrealTargetPlatform?> WrittenPlatforms = new();
+			foreach (ProjectConfigAndTargetCombination Combination in ProjectConfigAndTargetCombinations)
+			{
+				if (WrittenPlatforms.Add(Combination.Platform))
+				{
+					StringBuilder PlatformProperties = new();
+					string PlatformAdditionalOptions = GetAdditionalOptionsString(Combination.ProjectTarget?.TargetRules);
+					if (PlatformAdditionalOptions != DefaultAdditionalOptions)
+					{
+						PlatformProperties.AppendLine("    <AdditionalOptions>{0}</AdditionalOptions>", PlatformAdditionalOptions);
+					}
+
+					if (PlatformProperties.Length != 0)
+					{
+						VCProjectFileContent.AppendLine("  <PropertyGroup Condition=\"'$(Platform)'=='" + Combination.ProjectPlatformName + "'\">");
+						VCProjectFileContent.Append(PlatformAdditionalOptions);
+						VCProjectFileContent.AppendLine("  </PropertyGroup>");
+					}
+				}
+			}
+
+			// Write each configuration
+			foreach (ProjectConfigAndTargetCombination Combination in ProjectConfigAndTargetCombinations)
+			{
+				WriteConfiguration(ProjectName, Combination, VCProjectFileContent, PlatformProjectGenerators, bGenerateUserFileContent ? VCUserFileContent : null);
+			}
+
 			{ 
 				// Collapse common values
 				{
@@ -1998,7 +2027,7 @@ namespace UnrealBuildTool
 		}
 
 		// Anonymous function that writes project configuration data
-		private void WriteConfiguration(string ProjectName, ProjectConfigAndTargetCombination Combination, StringBuilder VCProjectFileContent, PlatformProjectGeneratorCollection PlatformProjectGenerators, StringBuilder? VCUserFileContent, ILogger Logger)
+		private void WriteConfiguration(string ProjectName, ProjectConfigAndTargetCombination Combination, StringBuilder VCProjectFileContent, PlatformProjectGeneratorCollection PlatformProjectGenerators, StringBuilder? VCUserFileContent)
 		{
 			UnrealTargetConfiguration Configuration = Combination.Configuration;
 
@@ -2135,7 +2164,6 @@ namespace UnrealBuildTool
 					VCProjectFileContent.AppendLine("    <NMakeReBuildCommandLine>$(RebuildBatchScript) {0}</NMakeReBuildCommandLine>", BuildArguments);
 					VCProjectFileContent.AppendLine("    <NMakeCleanCommandLine>$(CleanBatchScript) {0}</NMakeCleanCommandLine>", BuildArguments);
 					VCProjectFileContent.AppendLine("    <NMakeOutput>{0}</NMakeOutput>", NormalizeProjectPath(NMakePath.FullName));
-					VCProjectFileContent.AppendLine("    <AdditionalOptions>{0} {1} {2}</AdditionalOptions>", GetCppStandardCompileArgument(TargetRulesObject.CppStandard), GetEnableCoroutinesArgument(), GetConformanceCompileArguments(TargetRulesObject));
 
 					if (TargetRulesObject.Type == TargetType.Game || TargetRulesObject.Type == TargetType.Client || TargetRulesObject.Type == TargetType.Server)
 					{
