[UBT]
* ParallelExecutor - Fixed bug where totalactions was used to set maxweight.. this means that if 35 cl.exe actions were queued (which has weight 1.5) it would only build max 23 actions in parallel even though NumParallelProcesses might be higher (computer might have 128 cores)

diff --git a/Engine/Source/Programs/UnrealBuildTool/Executors/ParallelExecutor.cs b/Engine/Source/Programs/UnrealBuildTool/Executors/ParallelExecutor.cs
--- a/Engine/Source/Programs/UnrealBuildTool/Executors/ParallelExecutor.cs
+++ b/Engine/Source/Programs/UnrealBuildTool/Executors/ParallelExecutor.cs
@@ -184,7 +184,7 @@ namespace UnrealBuildTool
 			{
 				using ImmediateActionQueue queue = CreateActionQueue(ActionsToExecute, actionArtifactCache, Logger);
 				int actionLimit = Math.Min(NumParallelProcesses, queue.TotalActions);
-				queue.CreateAutomaticRunner(action => RunAction(queue, action), bUseActionWeights, actionLimit, actionLimit);
+				queue.CreateAutomaticRunner(action => RunAction(queue, action), bUseActionWeights, actionLimit, NumParallelProcesses);
 				queue.Start();
 				queue.StartManyActions();
 				return await queue.RunTillDone();
