[UBT]
* Added additional command line arguments handling in BoxExecutor when -Target="..." is used
* Fixed so max remote action count is sent to native side

diff --git a/Engine/Source/Programs/UnrealBuildTool/Executors/ImmediateActionQueue.cs b/Engine/Source/Programs/UnrealBuildTool/Executors/ImmediateActionQueue.cs
--- a/Engine/Source/Programs/UnrealBuildTool/Executors/ImmediateActionQueue.cs
+++ b/Engine/Source/Programs/UnrealBuildTool/Executors/ImmediateActionQueue.cs
@@ -798,19 +798,24 @@ namespace UnrealBuildTool
 			}
 		}
 
-		public int GetQueuedActionsCount(Func<LinkedAction, bool>? CanRunFunc = null)
+		/// <summary>
+		/// Returns the number of queued actions left (not including ArtifactCheck actions)
+		/// Note, this method is lockless and will not always return accurate count
+		/// </summary>
+		/// <param name="filterFunc">Optional function to filter out actions. Return false if action should not be included</param>
+		public uint GetQueuedActionsCount(Func<LinkedAction, bool>? filterFunc = null)
 		{
-			int count = 0;
+			uint count = 0;
 
 			for (int actionIndex = _firstPendingAction; actionIndex != Actions.Length; ++actionIndex)
 			{
 
-				if (Actions[actionIndex].Status != ActionStatus.Queued)
+				if (Actions[actionIndex].Status != ActionStatus.Queued || Actions[actionIndex].Phase != ActionPhase.Compile)
 				{
 					continue;
 				}
 					
-				if (CanRunFunc != null && !CanRunFunc(Actions[actionIndex].Action))
+				if (filterFunc != null && !filterFunc(Actions[actionIndex].Action))
 				{
 					continue;
 				}
diff --git a/Engine/Source/Programs/UnrealBuildTool/System/ActionGraph.cs b/Engine/Source/Programs/UnrealBuildTool/System/ActionGraph.cs
--- a/Engine/Source/Programs/UnrealBuildTool/System/ActionGraph.cs
+++ b/Engine/Source/Programs/UnrealBuildTool/System/ActionGraph.cs
@@ -356,7 +356,8 @@ namespace UnrealBuildTool
 			bool bAnySingleFile = TargetDescriptors.Any(Descriptor => Descriptor.SpecificFilesToCompile.Count == 1);
 			if (!bAnySingleFile && BuildConfiguration.bAllowBoxExecutor && BoxExecutor.IsAvailable(Logger))
 			{
-				return new BoxExecutor(BuildConfiguration.MaxParallelActions, BuildConfiguration.bAllCores, BuildConfiguration.bCompactOutput, Logger);
+				CommandLineArguments? AdditionalArguments = TargetDescriptors.FirstOrDefault()?.AdditionalArguments;
+				return new BoxExecutor(BuildConfiguration.MaxParallelActions, BuildConfiguration.bAllCores, BuildConfiguration.bCompactOutput, Logger, AdditionalArguments);
 			}
 #endif // #if __BOXEXECUTOR_AVAILABLE__
 
