UnrealBuildTool: Fix CompilerVersion selection

#rnx

diff --git a/Engine/Source/Programs/UnrealBuildTool/Platform/Windows/MicrosoftPlatformSDK.cs b/Engine/Source/Programs/UnrealBuildTool/Platform/Windows/MicrosoftPlatformSDK.cs
--- a/Engine/Source/Programs/UnrealBuildTool/Platform/Windows/MicrosoftPlatformSDK.cs
+++ b/Engine/Source/Programs/UnrealBuildTool/Platform/Windows/MicrosoftPlatformSDK.cs
@@ -565,7 +565,7 @@ namespace UnrealBuildTool
 			{
 				if (String.Equals(CompilerVersion, Channel.ToString(), StringComparison.OrdinalIgnoreCase))
 				{
-					ToolChain = SelectToolChain(ToolChains.Where(x => x.ReleaseChannel.HasFlag(Channel)), x => x.ThenByDescending(x => x.Version), Architecture);
+					ToolChain = SelectToolChain(ToolChains.Where(x => x.ReleaseChannel.HasFlag(Channel)), x => x.ThenBy(x => x.ReleaseChannel).ThenByDescending(x => x.Version), Architecture);
 					if (ToolChain == null)
 					{
 						DumpToolChains(ToolChains, x => x.ThenBy(x => x.ReleaseChannel).ThenByDescending(x => x.Version), Architecture, Logger);
