Handle clang directory version change in asan RPATH.

diff --git a/Engine/Source/Programs/UnrealBuildTool/Platform/Linux/LinuxToolChain.cs b/Engine/Source/Programs/UnrealBuildTool/Platform/Linux/LinuxToolChain.cs
--- a/Engine/Source/Programs/UnrealBuildTool/Platform/Linux/LinuxToolChain.cs
+++ b/Engine/Source/Programs/UnrealBuildTool/Platform/Linux/LinuxToolChain.cs
@@ -661,16 +661,18 @@ namespace UnrealBuildTool
 
 				if (CrossCompiling())
 				{
+					string ClangVersionFolder = (Info.ClangVersion.Major < 16) ? String.Format("{0}.{1}.{2}", Info.ClangVersion.Major, Info.ClangVersion.Minor, Info.ClangVersion.Build) : String.Format("{0}", Info.ClangVersion.Major);
+
 					// x64 only replaced the linux folder with arch, while on arm64 its still linux
 					if (LinkEnvironment.Architecture == UnrealArch.Arm64)
 					{
-						Arguments.Add(String.Format("-Wl,-rpath=\"{0}/lib/clang/{1}.{2}.{3}/lib/linux\"",
-								LinuxInfo.BaseLinuxPath, Info.ClangVersion.Major, Info.ClangVersion.Minor, Info.ClangVersion.Build));
+						Arguments.Add(String.Format("-Wl,-rpath=\"{0}/lib/clang/{1}/lib/linux\"",
+								LinuxInfo.BaseLinuxPath, ClangVersionFolder));
 					}
 					else
 					{
-						Arguments.Add(String.Format("-Wl,-rpath=\"{0}/lib/clang/{1}.{2}.{3}/lib/x86_64-unknown-linux-gnu\"",
-								LinuxInfo.BaseLinuxPath, Info.ClangVersion.Major, Info.ClangVersion.Minor, Info.ClangVersion.Build));
+						Arguments.Add(String.Format("-Wl,-rpath=\"{0}/lib/clang/{1}/lib/x86_64-unknown-linux-gnu\"",
+								LinuxInfo.BaseLinuxPath, ClangVersionFolder));
 					}
 				}
 			}
