UnrealBuildTool: Allow configuring toolchain and toolchain version to control what msvc version is used when clang is used to compile

diff --git a/Engine/Source/Programs/UnrealBuildTool/Executors/Experimental/FASTBuild.cs b/Engine/Source/Programs/UnrealBuildTool/Executors/Experimental/FASTBuild.cs
--- a/Engine/Source/Programs/UnrealBuildTool/Executors/Experimental/FASTBuild.cs
+++ b/Engine/Source/Programs/UnrealBuildTool/Executors/Experimental/FASTBuild.cs
@@ -851,6 +851,7 @@ namespace UnrealBuildTool
 						Platform: UnrealTargetPlatform.Win64,
 						Architecture: UnrealArch.X64,
 						CompilerVersion: String.IsNullOrEmpty(CompilerVersion) ? null : CompilerVersion,
+						ToolchainVersion: null,
 						WindowsSdkVersion: null,
 						SuppliedSdkDirectoryForVersion: null,
 						bUseCPPWinRT: false,
diff --git a/Engine/Source/Programs/UnrealBuildTool/Modes/QueryMode.cs b/Engine/Source/Programs/UnrealBuildTool/Modes/QueryMode.cs
--- a/Engine/Source/Programs/UnrealBuildTool/Modes/QueryMode.cs
+++ b/Engine/Source/Programs/UnrealBuildTool/Modes/QueryMode.cs
@@ -309,7 +309,7 @@ namespace UnrealBuildTool
 							if (CurrentTarget.Platform == UnrealTargetPlatform.Win64)
 							{
 								// TODO: Correct compiler
-								Settings.IncludePaths.AddRange(VCToolChain.GetVCIncludePaths(UnrealTargetPlatform.Win64, WindowsCompiler.VisualStudio2022, null, Logger).Split(";"));
+								Settings.IncludePaths.AddRange(VCToolChain.GetVCIncludePaths(UnrealTargetPlatform.Win64, WindowsCompiler.VisualStudio2022, null, null, Logger).Split(";"));
 							}
 						}
 						Settings.IncludePaths.AddRange(ModuleCompileEnvironment.SystemIncludePaths.Select(x => x.ToString()));
diff --git a/Engine/Source/Programs/UnrealBuildTool/Platform/Windows/MicrosoftPlatformSDK.Versions.cs b/Engine/Source/Programs/UnrealBuildTool/Platform/Windows/MicrosoftPlatformSDK.Versions.cs
--- a/Engine/Source/Programs/UnrealBuildTool/Platform/Windows/MicrosoftPlatformSDK.Versions.cs
+++ b/Engine/Source/Programs/UnrealBuildTool/Platform/Windows/MicrosoftPlatformSDK.Versions.cs
@@ -1,6 +1,7 @@
 // Copyright Epic Games, Inc. All Rights Reserved.
 
-using System.Linq;
+using System.Collections.Generic;
+using System.Diagnostics;
 using EpicGames.Core;
 
 namespace UnrealBuildTool
@@ -31,6 +32,9 @@ namespace UnrealBuildTool
 			VersionNumberRange.Parse("13.0.0", "13.999"), // VS2019 16.11 runtime requires Clang 13
 		};
 
+		/// <summary>
+		/// The minimum supported Clang compiler
+		/// </summary>
 		static readonly VersionNumber MinimumClangVersion = new VersionNumber(13, 0, 0);
 
 		/// <summary>
@@ -45,6 +49,18 @@ namespace UnrealBuildTool
 			VersionNumberRange.Parse("14.29.30133", "14.29.99999"), // VS2019 16.11.x
 		};
 
+		/// <summary>
+		/// Minimum Clang version required for MSVC toolchain versions
+		/// </summary>
+		static readonly IReadOnlyDictionary<VersionNumber, VersionNumber> MinimumRequiredClangVersion = new Dictionary<VersionNumber, VersionNumber>()
+		{
+			{ new VersionNumber(14, 37), new VersionNumber(16) }, // VS2022 17.7.x
+			{ new VersionNumber(14, 35), new VersionNumber(15) }, // VS2022 17.5.x - 17.6.x
+			{ new VersionNumber(14, 34), new VersionNumber(14) }, // VS2022 17.4.x
+			{ new VersionNumber(14, 29), new VersionNumber(13) }, // VS2019 16.11.x
+
+		};
+
 		/// <summary>
 		/// Tested compiler toolchains that should not be allowed.
 		/// </summary>
@@ -53,6 +69,9 @@ namespace UnrealBuildTool
 			VersionNumberRange.Parse("14.30.0", "14.33.99999"), // VS2022 17.0.x - 17.3.x
 		};
 
+		/// <summary>
+		/// The minimum supported MSVC compiler
+		/// </summary>
 		static readonly VersionNumber MinimumVisualCppVersion = new VersionNumber(14, 29, 30133);
 
 		/// <summary>
@@ -64,6 +83,9 @@ namespace UnrealBuildTool
 			VersionNumberRange.Parse("2023.1.0", "2023.9999"),
 		};
 
+		/// <summary>
+		/// The minimum supported Intel compiler
+		/// </summary>
 		static readonly VersionNumber MinimumIntelOneApiVersion = new VersionNumber(2023, 0, 0);
 
 		/// <inheritdoc/>
@@ -73,6 +95,41 @@ namespace UnrealBuildTool
 			MaxVersion = MaximumWindowsSDKVersion?.ToString();
 		}
 
+		/// <summary>
+		/// The minimum supported Clang version for a given MSVC toolchain
+		/// </summary>
+		/// <param name="VcVersion"></param>
+		/// <returns></returns>
+		public static VersionNumber GetMinimumClangVersionForVcVersion(VersionNumber VcVersion)
+		{
+			foreach (KeyValuePair<VersionNumber, VersionNumber> Item in MinimumRequiredClangVersion)
+			{
+				if (VcVersion >= Item.Key)
+				{
+					return Item.Value;
+				}
+			}
+			return MinimumClangVersion;
+		}
+
+		/// <summary>
+		/// The base Clang version for a given Intel toolchain
+		/// </summary>
+		/// <param name="IntelCompilerPath"></param>
+		/// <returns></returns>
+		public static VersionNumber GetClangVersionForIntelCompiler(FileReference IntelCompilerPath)
+		{
+			FileReference LdLLdPath = FileReference.Combine(IntelCompilerPath.Directory, "..", "bin-llvm", "ld.lld.exe");
+			if (FileReference.Exists(LdLLdPath))
+			{
+				FileVersionInfo VersionInfo = FileVersionInfo.GetVersionInfo(LdLLdPath.FullName);
+				VersionNumber Version = new VersionNumber(VersionInfo.FileMajorPart, VersionInfo.FileMinorPart, VersionInfo.FileBuildPart);
+				return Version;
+			}
+
+			return MinimumClangVersion;
+		}
+
 		/// <summary>
 		/// Whether toolchain errors should be ignored. Enable to ignore banned toolchains when generating projects,
 		/// as components such as the recommended toolchain can be installed by opening the generated solution via the .vsconfig file.
diff --git a/Engine/Source/Programs/UnrealBuildTool/Platform/Windows/UEBuildWindows.cs b/Engine/Source/Programs/UnrealBuildTool/Platform/Windows/UEBuildWindows.cs
--- a/Engine/Source/Programs/UnrealBuildTool/Platform/Windows/UEBuildWindows.cs
+++ b/Engine/Source/Programs/UnrealBuildTool/Platform/Windows/UEBuildWindows.cs
@@ -224,7 +224,13 @@ namespace UnrealBuildTool
 		/// <summary>
 		/// Version of the toolchain to use on Windows platform when a non-msvc Compiler is in use, to locate include paths etc.
 		/// </summary>
-		public WindowsCompiler ToolChain => Compiler.IsMSVC() ? Compiler : WindowsPlatform.GetDefaultCompiler(Target.ProjectFile, Architecture, Target.Logger);
+		[CommandLine("-VCToolchain=")]
+		public WindowsCompiler ToolChain
+		{
+			get => ToolChainPrivate ?? (Compiler.IsMSVC() ? Compiler : WindowsPlatform.GetDefaultCompiler(Target.ProjectFile, Architecture, Target.Logger));
+			set => ToolChainPrivate = value;
+		}
+		private WindowsCompiler? ToolChainPrivate = null;
 
 		/// <summary>
 		/// Architecture of Target.
@@ -237,7 +243,7 @@ namespace UnrealBuildTool
 		= UnrealArch.X64;
 
 		/// <summary>
-		/// The specific toolchain version to use. This may be a specific version number (for example, "14.13.26128"), the string "Latest" to select the newest available version, or
+		/// The specific compiler version to use. This may be a specific version number (for example, "14.13.26128"), the string "Latest" to select the newest available version, or
 		/// the string "Preview" to select the newest available preview version. By default, and if it is available, we use the toolchain version indicated by
 		/// WindowsPlatform.DefaultToolChainVersion (otherwise, we use the latest version).
 		/// </summary>
@@ -246,6 +252,16 @@ namespace UnrealBuildTool
 		[CommandLine("-CompilerVersion")]
 		public string? CompilerVersion = null;
 
+		/// <summary>
+		/// The specific msvc toolchain version to use if the compiler is not msvc. This may be a specific version number (for example, "14.13.26128"), the string "Latest" to select the newest available version, or
+		/// the string "Preview" to select the newest available preview version. By default, and if it is available, we use the toolchain version indicated by
+		/// WindowsPlatform.DefaultToolChainVersion (otherwise, we use the latest version).
+		/// </summary>
+		[ConfigFile(ConfigHierarchyType.Engine, "/Script/WindowsTargetPlatform.WindowsTargetSettings", "ToolchainVersion")]
+		[XmlConfigFile(Category = "WindowsPlatform")]
+		[CommandLine("-VCToolchainVersion")]
+		public string? ToolchainVersion = null;
+
 		/// <summary>
 		/// True if /fastfail should be passed to the msvc compiler and linker
 		/// </summary>
@@ -517,21 +533,6 @@ namespace UnrealBuildTool
 		/// </summary>
 		internal VCEnvironment? Environment;
 
-		/// <summary>
-		/// Directory containing the toolchain
-		/// </summary>
-		public string? ToolChainDir => (Environment == null) ? null : Environment.ToolChainDir.FullName;
-
-		/// <summary>
-		/// The version number of the toolchain
-		/// </summary>
-		public string? ToolChainVersion => (Environment == null) ? null : Environment.ToolChainVersion.ToString();
-
-		/// <summary>
-		/// Root directory containing the Windows Sdk
-		/// </summary>
-		public string? WindowsSdkDir => (Environment == null) ? null : Environment.WindowsSdkDir.FullName;
-
 		/// <summary>
 		/// Directory containing the NETFXSDK
 		/// </summary>
@@ -643,10 +644,14 @@ namespace UnrealBuildTool
 
 		public WindowsCompiler Compiler => Inner.Compiler;
 
+		public WindowsCompiler ToolChain => Inner.ToolChain;
+
 		public UnrealArch Architecture => Inner.Architecture;
 
 		public string? CompilerVersion => Inner.CompilerVersion;
 
+		public string? ToolchainVerison => Inner.ToolchainVersion;
+
 		public string? WindowsSdkVersion => Inner.WindowsSdkVersion;
 
 		public int TargetWindowsVersion => Inner.TargetWindowsVersion;
@@ -735,13 +740,12 @@ namespace UnrealBuildTool
 
 		internal VCEnvironment? Environment => Inner.Environment;
 
-		public WindowsCompiler ToolChain => Inner.ToolChain;
 
-		public string? ToolChainDir => Inner.ToolChainDir;
+		public string? ToolChainDir => Inner.Environment?.ToolChainDir.FullName ?? null;
 
-		public string? ToolChainVersion => Inner.ToolChainVersion;
+		public string? ToolChainVersion => Inner.Environment?.ToolChainVersion.ToString() ?? null;
 
-		public string? WindowsSdkDir => Inner.WindowsSdkDir;
+		public string? WindowsSdkDir => Inner.Environment?.WindowsSdkDir.ToString() ?? null;
 
 		public string? NetFxSdkDir => Inner.NetFxSdkDir;
 
@@ -861,7 +865,7 @@ namespace UnrealBuildTool
 		[SupportedOSPlatform("windows")]
 		protected virtual VCEnvironment CreateVCEnvironment(TargetRules Target)
 		{
-			return VCEnvironment.Create(Target.WindowsPlatform.Compiler, Target.WindowsPlatform.ToolChain, Platform, Target.WindowsPlatform.Architecture, Target.WindowsPlatform.CompilerVersion, Target.WindowsPlatform.WindowsSdkVersion, null, Target.WindowsPlatform.bUseCPPWinRT, Target.WindowsPlatform.bAllowClangLinker, Logger);
+			return VCEnvironment.Create(Target.WindowsPlatform.Compiler, Target.WindowsPlatform.ToolChain, Platform, Target.WindowsPlatform.Architecture, Target.WindowsPlatform.CompilerVersion, Target.WindowsPlatform.ToolchainVersion, Target.WindowsPlatform.WindowsSdkVersion, null, Target.WindowsPlatform.bUseCPPWinRT, Target.WindowsPlatform.bAllowClangLinker, Logger);
 		}
 
 		/// <summary>
@@ -978,6 +982,8 @@ namespace UnrealBuildTool
 			// pull some things from it
 			Target.WindowsPlatform.Compiler = Target.WindowsPlatform.Environment.Compiler;
 			Target.WindowsPlatform.CompilerVersion = Target.WindowsPlatform.Environment.CompilerVersion.ToString();
+			Target.WindowsPlatform.ToolChain = Target.WindowsPlatform.Environment.ToolChain;
+			Target.WindowsPlatform.ToolchainVersion = Target.WindowsPlatform.Environment.ToolChainVersion.ToString();
 			Target.WindowsPlatform.WindowsSdkVersion = Target.WindowsPlatform.Environment.WindowsSdkVersion.ToString();
 
 			// If we're enabling support for C++ modules, make sure the compiler supports it. VS 16.8 changed which command line arguments are used to enable modules support.
@@ -998,6 +1004,17 @@ namespace UnrealBuildTool
 				throw new BuildException("LibFuzzer MSVC support requires Visual Studio 2022 17.0 (MSVC 14.30.x) or later. The current compiler version was detected as: {0}", Target.WindowsPlatform.Environment.ToolChainVersion);
 			}
 
+			// Ensure we're using a recent enough version of Clang given the MSVC version
+			if (Target.WindowsPlatform.Compiler.IsClang())
+			{
+				VersionNumber ClangVersion = Target.WindowsPlatform.Compiler == WindowsCompiler.Intel ? MicrosoftPlatformSDK.GetClangVersionForIntelCompiler(Target.WindowsPlatform.Environment.CompilerPath) : Target.WindowsPlatform.Environment.CompilerVersion;
+				VersionNumber MinimumClang = MicrosoftPlatformSDK.GetMinimumClangVersionForVcVersion(Target.WindowsPlatform.Environment.ToolChainVersion);
+				if (ClangVersion < MinimumClang)
+				{
+					throw new BuildException("MSVC toolchain version {0} requires Clang compiler version {1} or later. The current Clang compiler version was detected as: {2}", Target.WindowsPlatform.Environment.ToolChainVersion, MinimumClang, ClangVersion);
+				}
+			}
+
 			// Ensure we're using VS2022 when compiling for the installed engine.
 			if (Unreal.IsEngineInstalled() && (Target.WindowsPlatform.ToolChain == WindowsCompiler.VisualStudio2019 || Target.WindowsPlatform.Environment.ToolChainVersion < new VersionNumber(14, 34, 0)))
 			{
diff --git a/Engine/Source/Programs/UnrealBuildTool/Platform/Windows/VCEnvironment.cs b/Engine/Source/Programs/UnrealBuildTool/Platform/Windows/VCEnvironment.cs
--- a/Engine/Source/Programs/UnrealBuildTool/Platform/Windows/VCEnvironment.cs
+++ b/Engine/Source/Programs/UnrealBuildTool/Platform/Windows/VCEnvironment.cs
@@ -378,7 +378,8 @@ namespace UnrealBuildTool
 		/// <param name="ToolChain">The toolchain version to use, when a non-msvc compiler is used</param>
 		/// <param name="Platform">The platform to target</param>
 		/// <param name="Architecture">The Architecture to target</param>
-		/// <param name="CompilerVersion">The specific toolchain version to use</param>
+		/// <param name="CompilerVersion">The specific compiler version to use</param>
+		/// <param name="ToolchainVersion">The specific toolchain version to use (if the compiler isn't msvc)</param>
 		/// <param name="WindowsSdkVersion">Version of the Windows SDK to use</param>
 		/// <param name="SuppliedSdkDirectoryForVersion">If specified, this is the SDK directory to use, otherwise, attempt to look up via registry. If specified, the WindowsSdkVersion is used directly</param>
 		/// <param name="bUseCPPWinRT">Include the CPP/WinRT language projection</param>
@@ -386,9 +387,9 @@ namespace UnrealBuildTool
 		/// <param name="Logger">Logger for output</param>
 		/// <returns>New environment object with paths for the given settings</returns>
 		[SupportedOSPlatform("windows")]
-		public static VCEnvironment Create(WindowsCompiler Compiler, WindowsCompiler ToolChain, UnrealTargetPlatform Platform, UnrealArch Architecture, string? CompilerVersion, string? WindowsSdkVersion, string? SuppliedSdkDirectoryForVersion, bool bUseCPPWinRT, bool bAllowClangLinker, ILogger Logger)
+		public static VCEnvironment Create(WindowsCompiler Compiler, WindowsCompiler ToolChain, UnrealTargetPlatform Platform, UnrealArch Architecture, string? CompilerVersion, string? ToolchainVersion, string? WindowsSdkVersion, string? SuppliedSdkDirectoryForVersion, bool bUseCPPWinRT, bool bAllowClangLinker, ILogger Logger)
 		{
-			return Create(new VCEnvironmentParameters(Compiler, ToolChain, Platform, Architecture, CompilerVersion, WindowsSdkVersion, SuppliedSdkDirectoryForVersion, bUseCPPWinRT, bAllowClangLinker, Logger), Logger);
+			return Create(new VCEnvironmentParameters(Compiler, ToolChain, Platform, Architecture, CompilerVersion, ToolchainVersion, WindowsSdkVersion, SuppliedSdkDirectoryForVersion, bUseCPPWinRT, bAllowClangLinker, Logger), Logger);
 		}
 
 		/// <summary>
@@ -452,7 +453,8 @@ namespace UnrealBuildTool
 		/// <param name="ToolChain">The toolchain version to use, when a non-msvc compiler is used</param>
 		/// <param name="Platform">The platform to target</param>
 		/// <param name="Architecture">The Architecture to target</param>
-		/// <param name="CompilerVersion">The specific toolchain version to use</param>
+		/// <param name="CompilerVersion">The specific compiler version to use</param>
+		/// <param name="ToolchainVersion">The specific toolchain version to use (if the compiler isn't msvc)</param>
 		/// <param name="WindowsSdkVersion">Version of the Windows SDK to use</param>
 		/// <param name="SuppliedSdkDirectoryForVersion">If specified, this is the SDK directory to use, otherwise, attempt to look up via registry. If specified, the WindowsSdkVersion is used directly</param>
 		/// <param name="bUseCPPWinRT">Include the CPP/WinRT language projection</param>
@@ -460,7 +462,7 @@ namespace UnrealBuildTool
 		/// <param name="Logger">Logger for output</param>
 		/// <returns>Creation parameters for VC environment</returns>
 		[SupportedOSPlatform("windows")]
-		public VCEnvironmentParameters(WindowsCompiler Compiler, WindowsCompiler ToolChain, UnrealTargetPlatform Platform, UnrealArch Architecture, string? CompilerVersion, string? WindowsSdkVersion, string? SuppliedSdkDirectoryForVersion, bool bUseCPPWinRT, bool bAllowClangLinker, ILogger Logger)
+		public VCEnvironmentParameters(WindowsCompiler Compiler, WindowsCompiler ToolChain, UnrealTargetPlatform Platform, UnrealArch Architecture, string? CompilerVersion, string? ToolchainVersion, string? WindowsSdkVersion, string? SuppliedSdkDirectoryForVersion, bool bUseCPPWinRT, bool bAllowClangLinker, ILogger Logger)
 		{
 			// Get the compiler version info
 			VersionNumber? SelectedCompilerVersion;
@@ -481,7 +483,7 @@ namespace UnrealBuildTool
 					throw new BuildException("{0} is not a valid ToolChain for Compiler {1}", WindowsPlatform.GetCompilerName(ToolChain), WindowsPlatform.GetCompilerName(Compiler));
 				}
 
-				if (!WindowsPlatform.TryGetToolChainDir(ToolChain, null, Architecture, Logger, out SelectedToolChainVersion, out SelectedToolChainDir, out SelectedRedistDir))
+				if (!WindowsPlatform.TryGetToolChainDir(ToolChain, ToolchainVersion, Architecture, Logger, out SelectedToolChainVersion, out SelectedToolChainDir, out SelectedRedistDir))
 				{
 					throw new BuildException("{0} or {1} must be installed in order to build this target.", WindowsPlatform.GetCompilerName(WindowsCompiler.VisualStudio2019), WindowsPlatform.GetCompilerName(WindowsCompiler.VisualStudio2022));
 				}
diff --git a/Engine/Source/Programs/UnrealBuildTool/Platform/Windows/VCToolChain.cs b/Engine/Source/Programs/UnrealBuildTool/Platform/Windows/VCToolChain.cs
--- a/Engine/Source/Programs/UnrealBuildTool/Platform/Windows/VCToolChain.cs
+++ b/Engine/Source/Programs/UnrealBuildTool/Platform/Windows/VCToolChain.cs
@@ -88,6 +88,10 @@ namespace UnrealBuildTool
 			{
 				Lines.Add($"Using {WindowsPlatform.GetCompilerName(EnvVars.Compiler)} {EnvVars.ToolChainVersion} toolchain ({EnvVars.ToolChainDir}) and Windows {EnvVars.WindowsSdkVersion} SDK ({EnvVars.WindowsSdkDir}).");
 			}
+			else if (EnvVars.Compiler == WindowsCompiler.Intel)
+			{
+				Lines.Add($"Using {WindowsPlatform.GetCompilerName(EnvVars.Compiler)} {EnvVars.CompilerVersion} compiler ({EnvVars.CompilerDir}) based on Clang {MicrosoftPlatformSDK.GetClangVersionForIntelCompiler(EnvVars.CompilerPath)} with {WindowsPlatform.GetCompilerName(EnvVars.ToolChain)} {EnvVars.ToolChainVersion} runtime ({EnvVars.ToolChainDir}) and Windows {EnvVars.WindowsSdkVersion} SDK ({EnvVars.WindowsSdkDir}).");
+			}
 			else
 			{
 				Lines.Add($"Using {WindowsPlatform.GetCompilerName(EnvVars.Compiler)} {EnvVars.CompilerVersion} compiler ({EnvVars.CompilerDir}) with {WindowsPlatform.GetCompilerName(EnvVars.ToolChain)} {EnvVars.ToolChainVersion} runtime ({EnvVars.ToolChainDir}) and Windows {EnvVars.WindowsSdkVersion} SDK ({EnvVars.WindowsSdkDir}).");
@@ -1006,9 +1010,10 @@ namespace UnrealBuildTool
 
 				ClangWarnings.GetEnabledWarnings(Arguments);
 
+				VersionNumber ClangVersion = Target.WindowsPlatform.Compiler == WindowsCompiler.Intel ? MicrosoftPlatformSDK.GetClangVersionForIntelCompiler(EnvVars.CompilerPath) : EnvVars.CompilerVersion;
 				ClangWarnings.GetDisabledWarnings(CompileEnvironment,
 					Target.StaticAnalyzer,
-					EnvVars.CompilerVersion,
+					ClangVersion,
 					Arguments);
 
 				// Additional disabled warnings for msvc. Everything below should be checked if it is necessary
@@ -2651,10 +2656,10 @@ namespace UnrealBuildTool
 		/// Gets the default include paths for the given platform.
 		/// </summary>
 		[SupportedOSPlatform("windows")]
-		public static string GetVCIncludePaths(UnrealTargetPlatform Platform, WindowsCompiler Compiler, string? CompilerVersion, ILogger Logger)
+		public static string GetVCIncludePaths(UnrealTargetPlatform Platform, WindowsCompiler Compiler, string? CompilerVersion, string? ToolchainVersion, ILogger Logger)
 		{
 			// Make sure we've got the environment variables set up for this target
-			VCEnvironment EnvVars = VCEnvironment.Create(Compiler, WindowsCompiler.Default, Platform, UnrealArch.X64, CompilerVersion, null, null, false, false, Logger);
+			VCEnvironment EnvVars = VCEnvironment.Create(Compiler, WindowsCompiler.Default, Platform, UnrealArch.X64, CompilerVersion, ToolchainVersion, null, null, false, false, Logger);
 
 			// Also add any include paths from the INCLUDE environment variable.  MSVC is not necessarily running with an environment that
 			// matches what UBT extracted from the vcvars*.bat using SetEnvironmentVariablesFromBatchFile().  We'll use the variables we
diff --git a/Engine/Source/Programs/UnrealBuildTool/ProjectFiles/VisualStudio/VCProjectFileGenerator.cs b/Engine/Source/Programs/UnrealBuildTool/ProjectFiles/VisualStudio/VCProjectFileGenerator.cs
--- a/Engine/Source/Programs/UnrealBuildTool/ProjectFiles/VisualStudio/VCProjectFileGenerator.cs
+++ b/Engine/Source/Programs/UnrealBuildTool/ProjectFiles/VisualStudio/VCProjectFileGenerator.cs
@@ -551,7 +551,7 @@ namespace UnrealBuildTool
 					{
 						if (SupportedPlatforms.Contains(UnrealTargetPlatform.Win64))
 						{
-							VCCommonTargetFileContent.AppendLine("    <DefaultSystemIncludePaths>{0}</DefaultSystemIncludePaths>", VCToolChain.GetVCIncludePaths(UnrealTargetPlatform.Win64, GetCompilerForIntellisense(Settings.ProjectFileFormat), null, Logger));
+							VCCommonTargetFileContent.AppendLine("    <DefaultSystemIncludePaths>{0}</DefaultSystemIncludePaths>", VCToolChain.GetVCIncludePaths(UnrealTargetPlatform.Win64, GetCompilerForIntellisense(Settings.ProjectFileFormat), null, null, Logger));
 						}
 					}
 					else
diff --git a/Engine/Source/Programs/UnrealBuildTool/ProjectFiles/VisualStudioCode/VSCodeProjectFileGenerator.cs b/Engine/Source/Programs/UnrealBuildTool/ProjectFiles/VisualStudioCode/VSCodeProjectFileGenerator.cs
--- a/Engine/Source/Programs/UnrealBuildTool/ProjectFiles/VisualStudioCode/VSCodeProjectFileGenerator.cs
+++ b/Engine/Source/Programs/UnrealBuildTool/ProjectFiles/VisualStudioCode/VSCodeProjectFileGenerator.cs
@@ -335,7 +335,7 @@ namespace UnrealBuildTool
 			DirectoryReference? SysRootPath = null;
 			if (OperatingSystem.IsWindows())
 			{
-				VCEnvironment Environment = VCEnvironment.Create(WindowsPlatform.GetDefaultCompiler(null, Target.Rules.WindowsPlatform.Architecture, Logger), WindowsCompiler.Default, Target.Platform, Target.Rules.WindowsPlatform.Architecture, null, Target.Rules.WindowsPlatform.WindowsSdkVersion, null, Target.Rules.WindowsPlatform.bUseCPPWinRT, Target.Rules.WindowsPlatform.bAllowClangLinker, Logger);
+				VCEnvironment Environment = VCEnvironment.Create(WindowsPlatform.GetDefaultCompiler(null, Target.Rules.WindowsPlatform.Architecture, Logger), WindowsCompiler.Default, Target.Platform, Target.Rules.WindowsPlatform.Architecture, null, null, Target.Rules.WindowsPlatform.WindowsSdkVersion, null, Target.Rules.WindowsPlatform.bUseCPPWinRT, Target.Rules.WindowsPlatform.bAllowClangLinker, Logger);
 				CompilerPath = FileReference.FromString(Environment.CompilerPath.FullName);
 				UsingClang = false;
 			}
