UnrealBuildTool: Set preferred clang to 16

#rnx

diff --git a/Engine/Source/Programs/UnrealBuildTool/Platform/Windows/MicrosoftPlatformSDK.Versions.cs b/Engine/Source/Programs/UnrealBuildTool/Platform/Windows/MicrosoftPlatformSDK.Versions.cs
--- a/Engine/Source/Programs/UnrealBuildTool/Platform/Windows/MicrosoftPlatformSDK.Versions.cs
+++ b/Engine/Source/Programs/UnrealBuildTool/Platform/Windows/MicrosoftPlatformSDK.Versions.cs
@@ -28,6 +28,7 @@ namespace UnrealBuildTool
 		/// </summary>
 		static readonly VersionNumberRange[] PreferredClangVersions =
 		{
+			VersionNumberRange.Parse("16.0.0", "16.999"), // VS2022 17.7.x runtime requires Clang 16
 			VersionNumberRange.Parse("15.0.0", "15.999"), // VS2022 17.5.x runtime requires Clang 15
 			VersionNumberRange.Parse("14.0.0", "14.999"), // VS2022 17.4.x runtime requires Clang 14
 			VersionNumberRange.Parse("13.0.0", "13.999"), // VS2019 16.11 runtime requires Clang 13
