Fix PATH separator on non-Windows platforms

#rnx

diff --git a/Engine/Source/Programs/Shared/EpicGames.Core/DirectoryReference.cs b/Engine/Source/Programs/Shared/EpicGames.Core/DirectoryReference.cs
--- a/Engine/Source/Programs/Shared/EpicGames.Core/DirectoryReference.cs
+++ b/Engine/Source/Programs/Shared/EpicGames.Core/DirectoryReference.cs
@@ -4,6 +4,7 @@ using System;
 using System.Collections.Generic;
 using System.Diagnostics.CodeAnalysis;
 using System.IO;
+using System.Linq;
 
 namespace EpicGames.Core
 {
@@ -116,6 +117,20 @@ namespace EpicGames.Core
 			}
 		}
 
+		/// <summary>
+		/// Add a directory to the PATH environment variable
+		/// </summary>
+		/// <param name="path">The path to add</param>
+		public static void AddDirectoryToPath(DirectoryReference path)
+		{
+			string pathEnvironmentVariable = Environment.GetEnvironmentVariable("PATH") ?? "";
+			if (!pathEnvironmentVariable.Split(Path.PathSeparator).Any(x => String.Equals(x, path.FullName, StringComparison.Ordinal)))
+			{
+				pathEnvironmentVariable = $"{path.FullName}{Path.PathSeparator}{pathEnvironmentVariable}";
+				Environment.SetEnvironmentVariable("PATH", pathEnvironmentVariable);
+			}
+		}
+
 		/// <summary>
 		/// Gets the path for a special folder
 		/// </summary>
diff --git a/Engine/Source/Programs/UnrealBuildTool/Configuration/UEBuildPlatform.cs b/Engine/Source/Programs/UnrealBuildTool/Configuration/UEBuildPlatform.cs
--- a/Engine/Source/Programs/UnrealBuildTool/Configuration/UEBuildPlatform.cs
+++ b/Engine/Source/Programs/UnrealBuildTool/Configuration/UEBuildPlatform.cs
@@ -738,19 +738,10 @@ namespace UnrealBuildTool
 		/// <summary>
 		/// Returns the delimiter used to separate paths in the PATH environment variable for the platform we are executing on.
 		/// </summary>
-		public static String GetPathVarDelimiter()
+		[Obsolete("Replace with System.IO.Path.PathSeparator")]
+		public static string GetPathVarDelimiter()
 		{
-			if (BuildHostPlatform.Current.Platform == UnrealTargetPlatform.Linux || BuildHostPlatform.Current.Platform == UnrealTargetPlatform.LinuxArm64 ||
-				BuildHostPlatform.Current.Platform == UnrealTargetPlatform.Mac)
-			{
-				return ":";
-			}
-			if (BuildHostPlatform.Current.Platform == UnrealTargetPlatform.Win64)
-			{
-				return ";";
-			}
-
-			throw new InvalidOperationException($"PATH variable delimiter unknown for platform {BuildHostPlatform.Current.Platform}");
+			return Path.PathSeparator.ToString();
 		}
 
 		/// <summary>
diff --git a/Engine/Source/Programs/UnrealBuildTool/Platform/Windows/VCEnvironment.cs b/Engine/Source/Programs/UnrealBuildTool/Platform/Windows/VCEnvironment.cs
--- a/Engine/Source/Programs/UnrealBuildTool/Platform/Windows/VCEnvironment.cs
+++ b/Engine/Source/Programs/UnrealBuildTool/Platform/Windows/VCEnvironment.cs
@@ -153,12 +153,12 @@ namespace UnrealBuildTool
 			Environment.SetEnvironmentVariable("VC_COMPILER_PATH", CompilerPath.FullName, EnvironmentVariableTarget.Process);
 			Environment.SetEnvironmentVariable("VC_COMPILER_DIR", CompilerPath.Directory.FullName, EnvironmentVariableTarget.Process);
 
-			AddDirectoryToPath(GetVCToolPath(ToolChainDir, Architecture));
+			DirectoryReference.AddDirectoryToPath(GetVCToolPath(ToolChainDir, Architecture));
 			if (Architecture == UnrealArch.Arm64)
 			{
 				// Add both toolchain paths to the PATH environment variable. There are some support DLLs which are only added to one of the paths, but which the toolchain in the other directory
 				// needs to run (eg. mspdbcore.dll).
-				AddDirectoryToPath(GetVCToolPath(ToolChainDir, UnrealArch.X64));
+				DirectoryReference.AddDirectoryToPath(GetVCToolPath(ToolChainDir, UnrealArch.X64));
 			}
 
 			// Add the Windows SDK directory to the path too, for mt.exe.
@@ -170,25 +170,11 @@ namespace UnrealBuildTool
 					Debug.Assert(BuildHostPlatform.Current.Platform == UnrealTargetPlatform.Win64);
 					BuildHostArch = "x64";
 				}
-				AddDirectoryToPath(DirectoryReference.Combine(WindowsSdkDir, "bin", WindowsSdkVersion.ToString(), BuildHostArch));
+				DirectoryReference.AddDirectoryToPath(DirectoryReference.Combine(WindowsSdkDir, "bin", WindowsSdkVersion.ToString(), BuildHostArch));
 
 			}
 		}
 
-		/// <summary>
-		/// Add a directory to the PATH environment variable
-		/// </summary>
-		/// <param name="ToolPath">The path to add</param>
-		static void AddDirectoryToPath(DirectoryReference ToolPath)
-		{
-			string PathEnvironmentVariable = Environment.GetEnvironmentVariable("PATH") ?? "";
-			if (!PathEnvironmentVariable.Split(';').Any(x => String.Compare(x, ToolPath.FullName, true) == 0))
-			{
-				PathEnvironmentVariable = ToolPath.FullName + ";" + PathEnvironmentVariable;
-				Environment.SetEnvironmentVariable("PATH", PathEnvironmentVariable);
-			}
-		}
-
 		/// <summary>
 		/// Gets the path to the tool binaries.
 		/// </summary>
diff --git a/Engine/Source/Programs/UnrealBuildTool/UnrealBuildTool.cs b/Engine/Source/Programs/UnrealBuildTool/UnrealBuildTool.cs
--- a/Engine/Source/Programs/UnrealBuildTool/UnrealBuildTool.cs
+++ b/Engine/Source/Programs/UnrealBuildTool/UnrealBuildTool.cs
@@ -488,11 +488,7 @@ namespace UnrealBuildTool
 				AssemblyUtils.InstallAssemblyResolver(Path.GetDirectoryName(Assembly.GetEntryAssembly()!.GetOriginalLocation())!);
 
 				// Add the application directory to PATH
-				string path = Environment.GetEnvironmentVariable("PATH") ?? String.Empty;
-				if (!path.Split(';').Any(x => String.Compare(x, Unreal.UnrealBuildToolDllPath.Directory.FullName, true) == 0))
-				{
-					Environment.SetEnvironmentVariable("PATH", $"{path};{Unreal.UnrealBuildToolDllPath.Directory.FullName}");
-				}
+				DirectoryReference.AddDirectoryToPath(Unreal.UnrealBuildToolDllPath.Directory);
 
 				// Change the working directory to be the Engine/Source folder. We are likely running from Engine/Binaries/DotNET
 				// This is critical to be done early so any code that relies on the current directory being Engine/Source will work.
