UnrealBuildTool: Add ini setting for VCToolchain

#rnx

diff --git a/Engine/Source/Programs/UnrealBuildTool/Platform/Windows/UEBuildWindows.cs b/Engine/Source/Programs/UnrealBuildTool/Platform/Windows/UEBuildWindows.cs
--- a/Engine/Source/Programs/UnrealBuildTool/Platform/Windows/UEBuildWindows.cs
+++ b/Engine/Source/Programs/UnrealBuildTool/Platform/Windows/UEBuildWindows.cs
@@ -224,6 +224,8 @@ namespace UnrealBuildTool
 		/// <summary>
 		/// Version of the toolchain to use on Windows platform when a non-msvc Compiler is in use, to locate include paths etc.
 		/// </summary>
+		[ConfigFile(ConfigHierarchyType.Engine, "/Script/WindowsTargetPlatform.WindowsTargetSettings", "Toolchain")]
+		[XmlConfigFile(Category = "WindowsPlatform")]
 		[CommandLine("-VCToolchain=")]
 		public WindowsCompiler ToolChain
 		{
