- Added a few new properties to ModuleRules (.Build.cs files) to make handling PlatformExtensions easier to manage platform extensions:
  * bIsPlatformExtension - true if there's a platform extension .Build.cs with a subclass
  * PlatformDirectory - location of the main Build.cs or the platform extension's subclass Build.cs
  * PlatformModuleDirectory - string version of above
  * PlatformSubdirectoryName - this is the name underneath the ModuleDirectory that platform files are in, or '.' (because platform extensions are already under a platform dir, we don't add another directory - there would be ThirdParty/MyLib/MyPlatform/foo.lib vs Platforms/MyPlatform/ThirdParty/MyLib/foo.lib, in the extension case, there's no MyPlatform directory underneath the module dir)

- With these changes, basic Build.cs files can be written without needing to know if a platform is an extension or not, and the extension subclasses can be much simpler

diff --git a/Engine/Source/Programs/UnrealBuildTool/Configuration/ModuleRules.cs b/Engine/Source/Programs/UnrealBuildTool/Configuration/ModuleRules.cs
--- a/Engine/Source/Programs/UnrealBuildTool/Configuration/ModuleRules.cs
+++ b/Engine/Source/Programs/UnrealBuildTool/Configuration/ModuleRules.cs
@@ -553,11 +553,21 @@ namespace UnrealBuildTool
 		/// </summary>
 		internal DirectoryReference Directory { get; set; }
 
+		/// <summary>
+		/// The Directory that contains either the Build.cs file, or the platform extension's SubClass directory
+		/// </summary>
+		protected DirectoryReference PlatformDirectory => DirectoriesForModuleSubClasses[GetType()]; // this is expected to always exist if we are able to be queried
+
+		/// <summary>
+		/// Returns true if the rules are a platform extension subclass
+		/// </summary>
+		protected bool bIsPlatformExtension => (Directory != PlatformDirectory);
+
 		/// <summary>
 		/// Additional directories that contribute to this module (likely in UnrealBuildTool.EnginePlatformExtensionsDirectory). 
 		/// The dictionary tracks module subclasses 
 		/// </summary>
-		internal Dictionary<Type, DirectoryReference>? DirectoriesForModuleSubClasses;
+		internal Dictionary<Type, DirectoryReference> DirectoriesForModuleSubClasses;
 
 		/// <summary>
 		/// Additional directories that contribute to this module but are not based on a subclass (NotForLicensees, etc)
@@ -1390,6 +1400,16 @@ namespace UnrealBuildTool
 		/// </summary>
 		public string ModuleDirectory => Directory.FullName;
 
+		/// <summary>
+		/// Property for the directory containing this module or the platform extension's subclass. Useful for adding paths to third party dependencies.
+		/// </summary>
+		protected string PlatformModuleDirectory => PlatformDirectory.FullName;
+
+		/// <summary>
+		/// Name of a platform under the PlatformModuleDirectory - for a PlatfomrExtension, we don't use platform subdirectories since it's already in a platform dir, so this returns '.'
+		/// </summary>
+		protected string PlatformSubdirectoryName => bIsPlatformExtension ? "." : Target.Platform.ToString();
+
 		/// <summary>
 		/// Returns module's low level tests directory "Tests".
 		/// </summary>
@@ -1766,11 +1786,6 @@ namespace UnrealBuildTool
 		/// <returns>Directory where the subclass's .Build.cs lives, or null if not found</returns>
 		public DirectoryReference? GetModuleDirectoryForSubClass(Type Type)
 		{
-			if (DirectoriesForModuleSubClasses == null)
-			{
-				return null;
-			}
-
 			DirectoryReference? Directory;
 			if (DirectoriesForModuleSubClasses.TryGetValue(Type, out Directory))
 			{
@@ -1785,14 +1800,8 @@ namespace UnrealBuildTool
 		/// <returns>List of directories, or null if none were added</returns>
 		public DirectoryReference[] GetAllModuleDirectories()
 		{
-			List<DirectoryReference> AllDirectories = new List<DirectoryReference> { Directory };
+			List<DirectoryReference> AllDirectories = new List<DirectoryReference>(DirectoriesForModuleSubClasses.Values);
 			AllDirectories.AddRange(AdditionalModuleDirectories);
-
-			if (DirectoriesForModuleSubClasses != null)
-			{
-				AllDirectories.AddRange(DirectoriesForModuleSubClasses.Values);
-			}
-
 			return AllDirectories.ToArray();
 		}
 
diff --git a/Engine/Source/Programs/UnrealBuildTool/System/RulesAssembly.cs b/Engine/Source/Programs/UnrealBuildTool/System/RulesAssembly.cs
--- a/Engine/Source/Programs/UnrealBuildTool/System/RulesAssembly.cs
+++ b/Engine/Source/Programs/UnrealBuildTool/System/RulesAssembly.cs
@@ -500,15 +500,16 @@ namespace UnrealBuildTool
 				RulesObject.bPrecompile = (RulesObject.bTreatAsEngineModule || ModuleName.Equals("UnrealGame", StringComparison.OrdinalIgnoreCase)) && Target.bPrecompile;
 				RulesObject.bUsePrecompiled = bReadOnly;
 				RulesObject.RulesAssembly = this;
+				RulesObject.DirectoriesForModuleSubClasses = new();
+				RulesObject.DirectoriesForModuleSubClasses.Add(BaseRulesObjectType, RulesObject.Directory);
 
 				// go up the type hierarchy (if there is a hierarchy), looking for any extra directories for the module
 				if (RulesObjectType != BaseRulesObjectType && RulesObjectType != typeof(ModuleRules))
 				{
 					Type SubType = RulesObjectType;
 
-					RulesObject.DirectoriesForModuleSubClasses = new Dictionary<Type, DirectoryReference>();
 					RulesObject.SubclassRules = new List<string>();
-					while (SubType != null && SubType != BaseRulesObjectType)
+					while (SubType != BaseRulesObjectType)
 					{
 						FileReference? SubTypeFileName;
 						if (TryGetFileNameFromType(SubType, out SubTypeFileName))
