Change Linux PGO file location to match other platforms.

- Also change the name to match the target.

diff --git a/Engine/Source/Programs/UnrealBuildTool/Platform/Linux/UEBuildLinux.cs b/Engine/Source/Programs/UnrealBuildTool/Platform/Linux/UEBuildLinux.cs
--- a/Engine/Source/Programs/UnrealBuildTool/Platform/Linux/UEBuildLinux.cs
+++ b/Engine/Source/Programs/UnrealBuildTool/Platform/Linux/UEBuildLinux.cs
@@ -496,8 +496,8 @@ namespace UnrealBuildTool
 				{
 					BaseDir = DirectoryReference.FromFile(Target.ProjectFile);
 				}
-				CompileEnvironment.PGODirectory = Path.Combine(BaseDir.FullName, "Build", Target.Platform.ToString(), "PGO").Replace('\\', '/') + "/";
-				CompileEnvironment.PGOFilenamePrefix = "profile.profdata";
+				CompileEnvironment.PGODirectory = DirectoryReference.Combine(Target.ProjectFile.Directory, "Platforms", "Linux", "Build", "PGO").FullName.Replace('\\', '/') + "/";
+				CompileEnvironment.PGOFilenamePrefix = string.Format("{0}-{1}-{2}.profdata", Target.Name, Target.Platform, Target.Configuration);
 
 				LinkEnvironment.PGODirectory = CompileEnvironment.PGODirectory;
 				LinkEnvironment.PGOFilenamePrefix = CompileEnvironment.PGOFilenamePrefix;
