Fix wrong Linux PGO path for a project-less build.

diff --git a/Engine/Source/Programs/UnrealBuildTool/Platform/Linux/UEBuildLinux.cs b/Engine/Source/Programs/UnrealBuildTool/Platform/Linux/UEBuildLinux.cs
--- a/Engine/Source/Programs/UnrealBuildTool/Platform/Linux/UEBuildLinux.cs
+++ b/Engine/Source/Programs/UnrealBuildTool/Platform/Linux/UEBuildLinux.cs
@@ -495,8 +495,15 @@ namespace UnrealBuildTool
 				if (Target.ProjectFile != null)
 				{
 					BaseDir = DirectoryReference.FromFile(Target.ProjectFile);
+					// projects put PGO data in Platform/Linux/Build/PGO, even if Linux platform isn't a Platform Extension
+					CompileEnvironment.PGODirectory = Path.Combine(BaseDir.FullName, "Platforms", Target.Platform.ToString(), "Build", "PGO");
 				}
-				CompileEnvironment.PGODirectory = DirectoryReference.Combine(Target.ProjectFile.Directory, "Platforms", "Linux", "Build", "PGO").FullName.Replace('\\', '/') + "/";
+				else
+				{
+					// project-less build put PGO data in Engine/Build/Linux/PGO, because Linux platform isn't a Platform Extension
+					CompileEnvironment.PGODirectory = Path.Combine(BaseDir.FullName, "Build", Target.Platform.ToString(), "PGO");
+				}
+				CompileEnvironment.PGODirectory = CompileEnvironment.PGODirectory.Replace('\\', '/') + "/";
 				CompileEnvironment.PGOFilenamePrefix = string.Format("{0}-{1}-{2}.profdata", Target.Name, Target.Platform, Target.Configuration);
 
 				LinkEnvironment.PGODirectory = CompileEnvironment.PGODirectory;
