UnrealBuildTool: Add ini config for InlineFunctionExpansionLevel

#rnx

diff --git a/Engine/Source/Programs/UnrealBuildTool/Platform/Windows/UEBuildWindows.cs b/Engine/Source/Programs/UnrealBuildTool/Platform/Windows/UEBuildWindows.cs
--- a/Engine/Source/Programs/UnrealBuildTool/Platform/Windows/UEBuildWindows.cs
+++ b/Engine/Source/Programs/UnrealBuildTool/Platform/Windows/UEBuildWindows.cs
@@ -208,6 +208,7 @@ namespace UnrealBuildTool
 		/// Which level to use for Inline Function Expansion when TargetRules.bUseInlining is enabled
 		/// </summary>
 		/// <seealso href="link">https://learn.microsoft.com/en-us/cpp/build/reference/ob-inline-function-expansion</seealso>
+		[ConfigFile(ConfigHierarchyType.Engine, "/Script/WindowsTargetPlatform.WindowsTargetSettings", "InlineFunctionExpansionLevel")]
 		[XmlConfigFile(Category = "WindowsPlatform")]
 		public int InlineFunctionExpansionLevel { get; set; } = 2;
 
