Make the AutoRTFM compiler having changed cause the compiler sources to be clean compiled.

diff --git a/Engine/Source/Programs/UnrealBuildTool/Configuration/UEBuildModuleCPP.cs b/Engine/Source/Programs/UnrealBuildTool/Configuration/UEBuildModuleCPP.cs
--- a/Engine/Source/Programs/UnrealBuildTool/Configuration/UEBuildModuleCPP.cs
+++ b/Engine/Source/Programs/UnrealBuildTool/Configuration/UEBuildModuleCPP.cs
@@ -1844,6 +1844,8 @@ namespace UnrealBuildTool
 			Result.IncludeOrderVersion = Rules.IncludeOrderVersion;
 			Result.DeterministicWarningLevel = Rules.DeterministicWarningLevel;
 
+			Result.bUseAutoRTFMCompiler = Target.bUseAutoRTFMCompiler;
+
 			// Only enable the AutoRTFM flag if we are using the AutoRTFM compiler
 			if (Target.bUseAutoRTFMCompiler)
 			{
diff --git a/Engine/Source/Programs/UnrealBuildTool/System/CppCompileEnvironment.cs b/Engine/Source/Programs/UnrealBuildTool/System/CppCompileEnvironment.cs
--- a/Engine/Source/Programs/UnrealBuildTool/System/CppCompileEnvironment.cs
+++ b/Engine/Source/Programs/UnrealBuildTool/System/CppCompileEnvironment.cs
@@ -264,6 +264,11 @@ namespace UnrealBuildTool
 		/// </summary>
 		public bool bEnableBufferSecurityChecks = true;
 
+		/// <summary>
+		/// Whether the AutoRTFM compiler is being used or not.
+		/// </summary>
+		public bool bUseAutoRTFMCompiler = false;
+
 		/// <summary>
 		/// Enables AutoRTFM instrumentation to this cpp file only when AutoRTFMCompiler is enabled
 		/// </summary>
diff --git a/Engine/Source/Programs/UnrealBuildTool/ToolChain/ClangToolChain.cs b/Engine/Source/Programs/UnrealBuildTool/ToolChain/ClangToolChain.cs
--- a/Engine/Source/Programs/UnrealBuildTool/ToolChain/ClangToolChain.cs
+++ b/Engine/Source/Programs/UnrealBuildTool/ToolChain/ClangToolChain.cs
@@ -1020,6 +1020,15 @@ namespace UnrealBuildTool
 		protected virtual Action CompileCPPFile(CppCompileEnvironment CompileEnvironment, FileItem SourceFile, DirectoryReference OutputDir, string ModuleName, IActionGraphBuilder Graph, IReadOnlyCollection<string> GlobalArguments, CPPOutput Result)
 		{
 			Action CompileAction = Graph.CreateAction(ActionType.Compile);
+
+			// If we are using the AutoRTFM compiler, we make the compile action depend on the version of the compiler itself.
+			// This lets us update the compiler (which might not cause a version update of the compiler, which instead tracks
+			// the LLVM versioning scheme that Clang uses), but ensure that we rebuild the source if the compiler has changed.
+			if (CompileEnvironment.bUseAutoRTFMCompiler)
+			{
+				CompileAction.PrerequisiteItems.Add(FileItem.GetItemByFileReference(GetToolChainInfo().Clang));
+			}
+
 			CompileAction.Weight = CompileActionWeight;
 
 			// copy the global arguments into the file arguments, so GetCompileArguments_FileType can remove entries if needed (special case but can be important)
