1) Fixed internal performance issue where attempts to start actions were being spammed.
2) Fixed an extremely rare issue where action execution might be terminated early.
3) Improved logging to provide better visibility when actions fail resulting in dependent actions not getting executed (i.e. not all actions being logged with no error feedback)
4) Fixed issue where not all error messages would be logged prior to the build summary

diff --git a/Engine/Source/Programs/UnrealBuildTool/Executors/ImmediateActionQueue.cs b/Engine/Source/Programs/UnrealBuildTool/Executors/ImmediateActionQueue.cs
--- a/Engine/Source/Programs/UnrealBuildTool/Executors/ImmediateActionQueue.cs
+++ b/Engine/Source/Programs/UnrealBuildTool/Executors/ImmediateActionQueue.cs
@@ -258,6 +258,11 @@ namespace UnrealBuildTool
 		/// </summary>
 		private readonly Action<string> _writeToolOutput;
 
+		/// <summary>
+		/// Flush the tool output after logging has completed
+		/// </summary>
+		private readonly System.Action _flushToolOutput;
+
 		/// <summary>
 		/// Timer used to collect CPU utilization
 		/// </summary>
@@ -271,7 +276,7 @@ namespace UnrealBuildTool
 		/// <summary>
 		/// Collection of all actions remaining to be logged
 		/// </summary>
-		private readonly List<int> _actionsToLog = new();
+		private readonly List<(int, ActionStatus, ExecuteResults?)> _actionsToLog = new();
 
 		/// <summary>
 		/// Task waiting to process logging
@@ -283,6 +288,11 @@ namespace UnrealBuildTool
 		/// </summary>
 		private int _loggedCompletedActions = 0;
 
+		/// <summary>
+		/// Number of logged errors
+		/// </summary>
+		private int _loggedErrors = 0;
+
 		/// <summary>
 		/// Tracks the number of completed actions.
 		/// </summary>
@@ -298,6 +308,11 @@ namespace UnrealBuildTool
 		/// </summary>
 		private IActionArtifactCache? _actionArtifactCache;
 
+		/// <summary>
+		/// Queue of running and pending start many action calls
+		/// </summary>
+		private List<TaskCompletionSource> _startManyActionsQueue = new(2);
+
 		static ExecuteResults s_copiedFromCacheResults = new(new List<string>(), 0, TimeSpan.Zero, TimeSpan.Zero, "copied from cache");
 
 		/// <summary>
@@ -308,8 +323,9 @@ namespace UnrealBuildTool
 		/// <param name="maxActionArtifactCacheTasks">Max number of concurrent artifact cache tasks</param>
 		/// <param name="progressWriterText">Text to be displayed with the progress writer</param>
 		/// <param name="writeToolOutput">Action to invoke when writing tool output</param>
+		/// <param name="flushToolOutput">Action to invoke when flushing tool output</param>
 		/// <param name="logger">Logging interface</param>
-		public ImmediateActionQueue(IEnumerable<LinkedAction> actions, IActionArtifactCache? actionArtifactCache, int maxActionArtifactCacheTasks, string progressWriterText, Action<string> writeToolOutput, ILogger logger)
+		public ImmediateActionQueue(IEnumerable<LinkedAction> actions, IActionArtifactCache? actionArtifactCache, int maxActionArtifactCacheTasks, string progressWriterText, Action<string> writeToolOutput, System.Action flushToolOutput, ILogger logger)
 		{
 			int count = actions.Count();
 			Actions = new ActionState[count];
@@ -318,6 +334,7 @@ namespace UnrealBuildTool
 			ProgressWriter = new(progressWriterText, false, logger);
 			_actionArtifactCache = actionArtifactCache;
 			_writeToolOutput = writeToolOutput;
+			_flushToolOutput = flushToolOutput;
 
 			bool readArtifacts = _actionArtifactCache != null && _actionArtifactCache.EnableReads;
 			ActionPhase initialPhase = readArtifacts ? ActionPhase.ArtifactCheck : ActionPhase.Compile;
@@ -455,13 +472,33 @@ namespace UnrealBuildTool
 				Func<Task>? runAction = null;
 				LinkedAction? action = null;
 				int completedActions = 0;
+				bool prematureDone = false;
 				lock (Actions)
 				{
+
+					// Try to find an action
 					hasCanceled = CancellationTokenSource.IsCancellationRequested;
 					if (!hasCanceled)
 					{
 						(runAction, action, completedActions) = TryStartOneActionInternal(runner);
 					}
+
+					// If we have no completed actions, then we need to check for a stall.  
+					// Check to see if we have no active running tasks and no manual runners.
+					// If we don't, then it is assumed we can't queue any more work.
+					if ((runAction == null || action == null) && completedActions == 0)
+					{
+						// We found nothing, 
+						prematureDone = true;
+						foreach (ImmediateActionQueueRunner tryRunner in _runners)
+						{
+							if ((tryRunner.Type == ImmediateActionQueueRunnerType.Manual && !hasCanceled) || tryRunner.ActiveActions != 0)
+							{
+								prematureDone = false;
+								break;
+							}
+						}
+					}
 				}
 
 				// If we have an action, run it and account for any completed actions
@@ -482,24 +519,10 @@ namespace UnrealBuildTool
 					return true;
 				}
 
-				// If we have no completed actions, then we need to check for a stall
-				if (completedActions == 0)
+				// If we were prematurely done, max out the completed actions to stop the process
+				if (prematureDone)
 				{
-					// We found nothing, check to see if we have no active running tasks and no manual runners.
-					// If we don't, then it is assumed we can't queue any more work.
-					bool prematureDone = true;
-					foreach (ImmediateActionQueueRunner tryRunner in _runners)
-					{
-						if ((tryRunner.Type == ImmediateActionQueueRunnerType.Manual && !hasCanceled) || tryRunner.ActiveActions != 0)
-						{
-							prematureDone = false;
-							break;
-						}
-					}
-					if (prematureDone)
-					{
-						AddCompletedActions(Int32.MaxValue);
-					}
+					AddCompletedActions(Int32.MaxValue);
 					return false;
 				}
 
@@ -639,8 +662,50 @@ namespace UnrealBuildTool
 		/// <param name="runner">If specified, all actions will be limited to the runner</param>
 		public void StartManyActions(ImmediateActionQueueRunner? runner = null)
 		{
+
+			// We will only have one running and one pending call at any given time.  The second call 
+			// will wait on the first to complete.  If there are already two, we will just return and
+			// the pending call will run in lieu of that request.
+			TaskCompletionSource? activeTask = null;
+			TaskCompletionSource? myTask = null;
+			lock (_startManyActionsQueue)
+			{
+				switch (_startManyActionsQueue.Count)
+				{
+					case 0:
+						_startManyActionsQueue.Add(myTask = new());
+						break;
+
+					case 1:
+						activeTask = _startManyActionsQueue[0];
+						_startManyActionsQueue.Add(myTask = new());
+						break;
+
+					case 2:
+						return;
+
+					default:
+						throw new BuildException("Unexpected StartManyActions queue depth");
+				}
+			}
+
+			// Wait on the active task and then start actions
+			activeTask?.Task.Wait();
 			while (TryStartOneAction(runner))
 			{ }
+
+			// Remove my task from the queue.
+			lock (_startManyActionsQueue)
+			{
+				if (_startManyActionsQueue.Count < 1 || _startManyActionsQueue[0] != myTask)
+				{
+					throw new BuildException("Unexpected top of StartManyActions queue");
+				}
+				_startManyActionsQueue.RemoveAt(0);
+			}
+
+			// Set my result so that anybody waiting on my will run now.
+			myTask!.SetResult();
 		}
 
 		/// <summary>
@@ -728,7 +793,7 @@ namespace UnrealBuildTool
 				{
 					lock (_actionsToLog)
 					{
-						_actionsToLog.Add(action.SortIndex);
+						_actionsToLog.Add((action.SortIndex, status, results));
 						if (_actionsToLogTask == null)
 						{
 							_actionsToLogTask = Task.Run(LogActions);
@@ -832,7 +897,7 @@ namespace UnrealBuildTool
 		{
 			for (; ; )
 			{
-				int[]? actionsToLog = null;
+				(int, ActionStatus, ExecuteResults?)[]? actionsToLog = null;
 				lock (_actionsToLog)
 				{
 					if (_actionsToLog.Count == 0)
@@ -851,9 +916,9 @@ namespace UnrealBuildTool
 					return;
 				}
 
-				foreach (int index in actionsToLog)
+				foreach ((int index, ActionStatus status, ExecuteResults? results) in actionsToLog)
 				{
-					LogAction(Actions[index].Action, Actions[index].Results);
+					LogAction(Actions[index].Action, status, results);
 				}
 			}
 		}
@@ -864,8 +929,9 @@ namespace UnrealBuildTool
 		/// Log an action that has completed
 		/// </summary>
 		/// <param name="action">Action that has completed</param>
+		/// <param name="status">Status associated log event</param>
 		/// <param name="executeTaskResult">Results of the action</param>
-		private void LogAction(LinkedAction action, ExecuteResults? executeTaskResult)
+		private void LogAction(LinkedAction action, ActionStatus status, ExecuteResults? executeTaskResult)
 		{
 			List<string>? logLines = null;
 			int exitCode = Int32.MaxValue;
@@ -899,7 +965,7 @@ namespace UnrealBuildTool
 			lock (ProgressWriter)
 			{
 				int totalActions = Actions.Length;
-				int completedActions = Interlocked.Increment(ref _loggedCompletedActions);
+				int completedActions = ++_loggedCompletedActions;
 				ProgressWriter.Write(completedActions, Actions.Length);
 
 				// Canceled
@@ -936,6 +1002,10 @@ namespace UnrealBuildTool
 				}
 
 				string message = ($"[{completedActions}/{totalActions}]{targetDetails}{compilationTimes} {description}");
+				if (status == ActionStatus.Error)
+				{
+					message += " (failed)";
+				}
 
 				if (CompactOutput)
 				{
@@ -971,8 +1041,9 @@ namespace UnrealBuildTool
 					}
 				}
 
-				if (exitCode != 0)
+				if (status == ActionStatus.Error)
 				{
+					_loggedErrors++;
 
 					// If we have an error code but no output, chances are the tool crashed.  Generate more detailed information to let the
 					// user know something went wrong.
@@ -1008,6 +1079,7 @@ namespace UnrealBuildTool
 				loggingTask = _actionsToLogTask;
 			}
 			loggingTask?.Wait();
+			_flushToolOutput();
 
 			if (ShowCPUUtilization)
 			{
@@ -1021,6 +1093,12 @@ namespace UnrealBuildTool
 				}
 			}
 
+			if (_loggedErrors != 0)
+			{
+				Logger.LogInformation("");
+				Logger.LogInformation("Action Failure Count: {FailedActionCount}", _loggedErrors);
+			}
+
 			if (!ShowCompilationTimes)
 			{
 				return;
diff --git a/Engine/Source/Programs/UnrealBuildTool/Executors/ParallelExecutor.cs b/Engine/Source/Programs/UnrealBuildTool/Executors/ParallelExecutor.cs
--- a/Engine/Source/Programs/UnrealBuildTool/Executors/ParallelExecutor.cs
+++ b/Engine/Source/Programs/UnrealBuildTool/Executors/ParallelExecutor.cs
@@ -158,7 +158,7 @@ namespace UnrealBuildTool
 		/// <returns>Action queue</returns>
 		public ImmediateActionQueue CreateActionQueue(IEnumerable<LinkedAction> actionsToExecute, IActionArtifactCache? actionArtifactCache, ILogger logger)
 		{
-			return new(actionsToExecute, actionArtifactCache, NumParallelProcesses, "Compiling C++ source code...", x => WriteToolOutput(x), logger)
+			return new(actionsToExecute, actionArtifactCache, NumParallelProcesses, "Compiling C++ source code...", x => WriteToolOutput(x), () => FlushToolOutput(), logger)
 			{
 				ShowCompilationTimes = bShowCompilationTimes,
 				ShowCPUUtilization = bShowCPUUtilization,
diff --git a/Engine/Source/Programs/UnrealBuildTool/System/ActionExecutor.cs b/Engine/Source/Programs/UnrealBuildTool/System/ActionExecutor.cs
--- a/Engine/Source/Programs/UnrealBuildTool/System/ActionExecutor.cs
+++ b/Engine/Source/Programs/UnrealBuildTool/System/ActionExecutor.cs
@@ -76,5 +76,13 @@ namespace UnrealBuildTool
 				Parser.WriteLine(Line);
 			}
 		}
+
+		protected void FlushToolOutput()
+		{
+			lock (Parser)
+			{
+				Parser.Flush();
+			}
+		}
 	}
 }
