Fix for AppX manifest only including cultures with localized strings/images in the Resources node.

 - fixes a case on some platforms where culture tags used in other parts of the packaging process would cause an "unknown culture" build failure

#rnx

diff --git a/Engine/Source/Programs/UnrealBuildTool/Platform/Windows/AppXManifestGeneratorBase.cs b/Engine/Source/Programs/UnrealBuildTool/Platform/Windows/AppXManifestGeneratorBase.cs
--- a/Engine/Source/Programs/UnrealBuildTool/Platform/Windows/AppXManifestGeneratorBase.cs
+++ b/Engine/Source/Programs/UnrealBuildTool/Platform/Windows/AppXManifestGeneratorBase.cs
@@ -243,6 +243,7 @@ namespace UnrealBuildTool
 			{
 				bHasStageIdOverrides = true;
 				UEStageIdToAppXCultureId = StageIdOverrides;
+				AppXResources!.AddCultures(UEStageIdToAppXCultureId.Values);
 			}
 
 
