-  Restored test for long path names under Intermediate folders. We want to limit the depth from Engine to 200 chars to leave 60 chars for the root.
- Added a test for long path names for Prerequisite items.

#Swarm 26256453

diff --git a/Engine/Source/Programs/UnrealBuildTool/System/ActionGraph.cs b/Engine/Source/Programs/UnrealBuildTool/System/ActionGraph.cs
--- a/Engine/Source/Programs/UnrealBuildTool/System/ActionGraph.cs
+++ b/Engine/Source/Programs/UnrealBuildTool/System/ActionGraph.cs
@@ -299,6 +299,14 @@ namespace UnrealBuildTool
 						{
 							FailPaths.Add(PrerequisiteItem.Location);
 						}
+
+						if (PrerequisiteItem.Location.FullName.Length > Unreal.RootDirectory.FullName.Length + BuildConfiguration.MaxNestedPathLength &&
+							PrerequisiteItem.Location.IsUnderDirectory(Unreal.RootDirectory)
+							)
+						{
+							WarnPaths.Add(PrerequisiteItem.Location);
+						}
+
 					}
 
 					foreach (FileItem ProducedItem in Action.ProducedItems)
@@ -310,8 +318,7 @@ namespace UnrealBuildTool
 
 						// don't look in Intermediate directories - these aren't portable between machines, so don't need to be cbecked for length underneath the root
 						if (ProducedItem.Location.FullName.Length > Unreal.RootDirectory.FullName.Length + BuildConfiguration.MaxNestedPathLength &&
-							ProducedItem.Location.IsUnderDirectory(Unreal.RootDirectory) &&
-							!ProducedItem.Location.ContainsName("Intermediate", Unreal.RootDirectory)
+							ProducedItem.Location.IsUnderDirectory(Unreal.RootDirectory)
 							)
 						{
 							WarnPaths.Add(ProducedItem.Location);
