Adding UnrealIntermediateEnvironment to allow different intermediate paths based on build settings.
This removes the old hack that renames the target based on nonunity being set while also adding a new intermediates folder for IWYU.

diff --git a/Engine/Source/Programs/UnrealBuildTool/Configuration/BuildConfiguration.cs b/Engine/Source/Programs/UnrealBuildTool/Configuration/BuildConfiguration.cs
--- a/Engine/Source/Programs/UnrealBuildTool/Configuration/BuildConfiguration.cs
+++ b/Engine/Source/Programs/UnrealBuildTool/Configuration/BuildConfiguration.cs
@@ -234,5 +234,44 @@ namespace UnrealBuildTool
 		/// and only if output is not being redirected e.g. during a build from within Visual Studio
 		/// </summary>
 		public bool bCompactOutput => bCompactOutputCommandLine && !Console.IsOutputRedirected;
+
+		/// <summary>
+		/// Whether to unify C++ code into larger files for faster compilation.
+		/// </summary>
+		[CommandLine("-DisableUnity", Value = "false")]
+		[XmlConfigFile(Category = "BuildConfiguration")]
+		public bool bUseUnityBuild = true;
+
+		/// <summary>
+		/// Whether to force C++ source files to be combined into larger files for faster compilation.
+		/// </summary>
+		[CommandLine("-ForceUnity")]
+		[XmlConfigFile(Category = "BuildConfiguration")]
+		public bool bForceUnityBuild = false;
+
+		/// <summary>
+		/// Enables "include what you use" mode.
+		/// </summary>
+		[CommandLine("-IWYU")]
+		public bool bIWYU = false;
+
+		/// <summary>
+		/// Intermediate environment. Determines if the intermediates end up in a different folder than normal.
+		/// </summary>
+		public UnrealIntermediateEnvironment IntermediateEnvironment
+		{
+			get
+			{
+				if (bIWYU)
+				{
+					return UnrealIntermediateEnvironment.IWYU;
+				}
+				if (!bUseUnityBuild && !bForceUnityBuild)
+				{
+					return UnrealIntermediateEnvironment.NonUnity;
+				}
+				return UnrealIntermediateEnvironment.Default;
+			}
+		}
 	}
 }
diff --git a/Engine/Source/Programs/UnrealBuildTool/Configuration/TargetDescriptor.cs b/Engine/Source/Programs/UnrealBuildTool/Configuration/TargetDescriptor.cs
--- a/Engine/Source/Programs/UnrealBuildTool/Configuration/TargetDescriptor.cs
+++ b/Engine/Source/Programs/UnrealBuildTool/Configuration/TargetDescriptor.cs
@@ -127,6 +127,11 @@ namespace UnrealBuildTool
 		[CommandLine("-Rebuild")]
 		public bool bRebuild;
 
+		/// <summary>
+		/// The intermediate environment to use. This primarily controls the target's intermediate path.
+		/// </summary>
+		public UnrealIntermediateEnvironment IntermediateEnvironment = UnrealIntermediateEnvironment.Default;
+
 		/// <summary>
 		/// Constructor
 		/// </summary>
@@ -222,6 +227,37 @@ namespace UnrealBuildTool
 			return new TargetDescriptor(Info.ProjectFile, Info.Name, Info.Platform, Info.Configuration, Info.Architectures, Info.Arguments);
 		}
 
+		/// <summary>
+		/// Parse a list of target descriptors from the command line
+		/// </summary>
+		/// <param name="Arguments">Command-line arguments</param>
+		/// <param name="BuildConfiguration">Build configuration to get common flags from</param>
+		/// <param name="Logger">Logger for output</param>
+		/// <returns>List of target descriptors</returns>
+		public static List<TargetDescriptor> ParseCommandLine(CommandLineArguments Arguments, BuildConfiguration BuildConfiguration, ILogger Logger)
+		{
+			List<TargetDescriptor> TargetDescriptors = new List<TargetDescriptor>();
+			ParseCommandLine(Arguments, BuildConfiguration.bUsePrecompiled, BuildConfiguration.bSkipRulesCompile, BuildConfiguration.bForceRulesCompile, TargetDescriptors, Logger);
+
+			// apply the intermediate environment from the build configuration
+			TargetDescriptors.ForEach(x => x.IntermediateEnvironment = BuildConfiguration.IntermediateEnvironment);
+
+			return TargetDescriptors;
+		}
+
+		/// <summary>
+		/// Parse a list of target descriptors from the command line
+		/// </summary>
+		/// <param name="Arguments">Command-line arguments</param>
+		/// <param name="Logger">Logger for output</param>
+		/// <returns>List of target descriptors</returns>
+		public static List<TargetDescriptor> ParseCommandLine(CommandLineArguments Arguments, ILogger Logger)
+		{
+			List<TargetDescriptor> TargetDescriptors = new List<TargetDescriptor>();
+			ParseCommandLine(Arguments, false, false, false, TargetDescriptors, Logger);
+			return TargetDescriptors;
+		}
+
 		/// <summary>
 		/// Parse a list of target descriptors from the command line
 		/// </summary>
diff --git a/Engine/Source/Programs/UnrealBuildTool/Configuration/TargetInfo.cs b/Engine/Source/Programs/UnrealBuildTool/Configuration/TargetInfo.cs
--- a/Engine/Source/Programs/UnrealBuildTool/Configuration/TargetInfo.cs
+++ b/Engine/Source/Programs/UnrealBuildTool/Configuration/TargetInfo.cs
@@ -31,6 +31,11 @@ namespace UnrealBuildTool
 		/// </summary>
 		public readonly UnrealArchitectures Architectures;
 
+		/// <summary>
+		/// Intermediate environment. Determines if the intermediates end up in a different folder than normal.
+		/// </summary>
+		public UnrealIntermediateEnvironment IntermediateEnvironment;
+
 		/// <summary>
 		/// The project containing the target
 		/// </summary>
@@ -55,11 +60,13 @@ namespace UnrealBuildTool
 		/// <param name="Architectures">The architectures being built for</param>
 		/// <param name="ProjectFile">Path to the project file containing the target</param>
 		/// <param name="Arguments">Additional command line arguments for this target</param>
-		public TargetInfo(string Name, UnrealTargetPlatform Platform, UnrealTargetConfiguration Configuration, UnrealArchitectures? Architectures, FileReference? ProjectFile, CommandLineArguments? Arguments)
+		/// <param name="IntermediateEnvironment">Intermediate environment to use</param>
+		public TargetInfo(string Name, UnrealTargetPlatform Platform, UnrealTargetConfiguration Configuration, UnrealArchitectures? Architectures, FileReference? ProjectFile, CommandLineArguments? Arguments, UnrealIntermediateEnvironment IntermediateEnvironment = UnrealIntermediateEnvironment.Default)
 		{
 			this.Name = Name;
 			this.Platform = Platform;
 			this.Configuration = Configuration;
+			this.IntermediateEnvironment = IntermediateEnvironment;
 			this.ProjectFile = ProjectFile;
 			this.Arguments = Arguments;
 
@@ -92,6 +99,12 @@ namespace UnrealBuildTool
 					String.Join(",", Enum.GetValues(typeof(UnrealTargetConfiguration)).Cast<UnrealTargetConfiguration>().Select(x => x.ToString()))));
 			}
 
+			string? IntermediateEnvironmentStr = Reader.ReadString();
+			if (IntermediateEnvironmentStr != null)
+			{
+				UnrealIntermediateEnvironment.TryParse(IntermediateEnvironmentStr, out IntermediateEnvironment);
+			}
+
 			Arguments = ArgumentStrs == null ? null : new CommandLineArguments(ArgumentStrs);
 		}
 
@@ -107,6 +120,7 @@ namespace UnrealBuildTool
 			Writer.WriteArray(Architectures.Architectures.ToArray(), Item => Writer.WriteString(Item.ToString()));
 			Writer.WriteFileReference(ProjectFile);
 			Writer.WriteArray(Arguments?.GetRawArray(), Item => Writer.WriteString(Item));
+			Writer.WriteString(IntermediateEnvironment.ToString());
 		}
 	}
 }
diff --git a/Engine/Source/Programs/UnrealBuildTool/Configuration/TargetRules.cs b/Engine/Source/Programs/UnrealBuildTool/Configuration/TargetRules.cs
--- a/Engine/Source/Programs/UnrealBuildTool/Configuration/TargetRules.cs
+++ b/Engine/Source/Programs/UnrealBuildTool/Configuration/TargetRules.cs
@@ -430,10 +430,6 @@ namespace UnrealBuildTool
 				{
 					return NameOverride;
 				}
-				else if (!bUseUnityBuild)
-				{
-					return DefaultName + "NonUnity";
-				}
 
 				return DefaultName;
 			}
@@ -532,6 +528,11 @@ namespace UnrealBuildTool
 		/// </summary>
 		public UnrealArch Architecture => Architectures.SingleArchitecture;
 
+		/// <summary>
+		/// Intermediate environment. Determines if the intermediates end up in a different folder than normal.
+		/// </summary>
+		public UnrealIntermediateEnvironment IntermediateEnvironment;
+
 		/// <summary>
 		/// Path to the project file for the project containing this target.
 		/// </summary>
@@ -1372,8 +1373,7 @@ namespace UnrealBuildTool
 		/// <summary>
 		/// Enables "include what you use" by default for modules in this target. Changes the default PCH mode for any module in this project to PCHUsageMode.UseExplicitOrSharedPCHs.
 		/// </summary>
-		[CommandLine("-IWYU")]
-		public bool bIWYU = false;
+		public bool bIWYU { get => IntermediateEnvironment == UnrealIntermediateEnvironment.IWYU; }
 
 		/// <summary>
 		/// Adds header files in included modules to the build.
@@ -1449,15 +1449,13 @@ namespace UnrealBuildTool
 		/// </summary>
 		public bool bUseUnityBuild
 		{
-			get => bUseUnityBuildOverride ?? !bEnableCppModules;
+			get => bUseUnityBuildOverride ?? IntermediateEnvironment == UnrealIntermediateEnvironment.Default && !bEnableCppModules;
 			set => bUseUnityBuildOverride = value;
 		}
 
 		/// <summary>
 		/// Whether to unify C++ code into larger files for faster compilation.
 		/// </summary>
-		[CommandLine("-DisableUnity", Value = "false")]
-		[XmlConfigFile(Category = "BuildConfiguration", Name = nameof(bUseUnityBuild))]
 		bool? bUseUnityBuildOverride = null;
 
 		/// <summary>
@@ -2500,6 +2498,7 @@ namespace UnrealBuildTool
 			Platform = Target.Platform;
 			Configuration = Target.Configuration;
 			Architectures = Target.Architectures;
+			IntermediateEnvironment = Target.IntermediateEnvironment;
 			ProjectFile = Target.ProjectFile;
 			Version = Target.Version;
 			WindowsPlatform = new WindowsTargetRules(this);
@@ -2530,6 +2529,19 @@ namespace UnrealBuildTool
 			UEBuildPlatform.GetBuildPlatform(Platform).ResetTarget(this);
 			bDeployAfterCompile = bForceSkipDeploy ? false : bDeployAfterCompile;
 
+			// Determine intermediate environment overrides based on command line flags
+			if (IntermediateEnvironment == UnrealIntermediateEnvironment.Default)
+			{
+				if (bIWYU)
+				{
+					IntermediateEnvironment = UnrealIntermediateEnvironment.IWYU;
+				}
+				else if (!bUseUnityBuild)
+				{
+					IntermediateEnvironment = UnrealIntermediateEnvironment.NonUnity;
+				}
+			}
+
 			// Set the default build version
 			if (String.IsNullOrEmpty(BuildVersion))
 			{
@@ -2894,6 +2906,8 @@ namespace UnrealBuildTool
 
 		public UnrealArchitectures Architectures => Inner.Architectures;
 
+		public UnrealIntermediateEnvironment IntermediateEnvironment => Inner.IntermediateEnvironment;
+
 		public UnrealArch Architecture => Inner.Architecture;
 
 		public FileReference? ProjectFile => Inner.ProjectFile;
diff --git a/Engine/Source/Programs/UnrealBuildTool/Configuration/UEBuildTarget.cs b/Engine/Source/Programs/UnrealBuildTool/Configuration/UEBuildTarget.cs
--- a/Engine/Source/Programs/UnrealBuildTool/Configuration/UEBuildTarget.cs
+++ b/Engine/Source/Programs/UnrealBuildTool/Configuration/UEBuildTarget.cs
@@ -1096,6 +1096,25 @@ namespace UnrealBuildTool
 		Shipping,
 	}
 
+	/// <summary>
+	/// Intermediate environment. Determines if the intermediates end up in a different folder than normal.
+	/// </summary>
+	public enum UnrealIntermediateEnvironment
+	{
+		/// <summary>
+		/// 
+		/// </summary>
+		Default,
+		/// <summary>
+		/// 
+		/// </summary>
+		IWYU,
+		/// <summary>
+		/// 
+		/// </summary>
+		NonUnity,
+	}
+
 	/// <summary>
 	/// A container for a binary files (dll, exe) with its associated debug info.
 	/// </summary>
@@ -1151,6 +1170,43 @@ namespace UnrealBuildTool
 	/// </summary>
 	class UEBuildTarget
 	{
+		/// <summary>
+		/// Creates a target object for the specified target name.
+		/// </summary>
+		/// <param name="Descriptor">Information about the target</param>
+		/// <param name="BuildConfiguration">Build configuration</param>
+		/// <param name="Logger">Logger for output</param>
+		/// <returns></returns>
+		public static UEBuildTarget Create(TargetDescriptor Descriptor, BuildConfiguration BuildConfiguration, ILogger Logger)
+		{
+			return Create(Descriptor, BuildConfiguration.bSkipRulesCompile, BuildConfiguration.bForceRulesCompile, BuildConfiguration.bUsePrecompiled, BuildConfiguration.IntermediateEnvironment, Logger);
+		}
+
+		/// <summary>
+		/// Creates a target object for the specified target name.
+		/// </summary>
+		/// <param name="Descriptor">Information about the target</param>
+		/// <param name="Logger">Logger for output</param>
+		/// <returns></returns>
+		public static UEBuildTarget Create(TargetDescriptor Descriptor, ILogger Logger)
+		{
+			return Create(Descriptor, false, false, false, UnrealIntermediateEnvironment.Default, Logger);
+		}
+
+		/// <summary>
+		/// Creates a target object for the specified target name.
+		/// </summary>
+		/// <param name="Descriptor">Information about the target</param>
+		/// <param name="bSkipRulesCompile"></param>
+		/// <param name="bForceRulesCompile"></param>
+		/// <param name="bUsePrecompiled"></param>
+		/// <param name="Logger">Logger for output</param>
+		/// <returns></returns>
+		public static UEBuildTarget Create(TargetDescriptor Descriptor, bool bSkipRulesCompile, bool bForceRulesCompile, bool bUsePrecompiled, ILogger Logger)
+		{
+			return Create(Descriptor, bSkipRulesCompile, bForceRulesCompile, bUsePrecompiled, UnrealIntermediateEnvironment.Default, Logger);
+		}
+
 		/// <summary>
 		/// Creates a target object for the specified target name.
 		/// </summary>
@@ -1158,9 +1214,10 @@ namespace UnrealBuildTool
 		/// <param name="bSkipRulesCompile">Whether to skip compiling any rules assemblies</param>
 		/// <param name="bForceRulesCompile">Whether to always compile all rules assemblies</param>
 		/// <param name="bUsePrecompiled">Whether to use a precompiled engine build</param>
+		/// <param name="IntermediateEnvironment">Intermediate environment to use</param>
 		/// <param name="Logger">Logger for output</param>
 		/// <returns>The build target object for the specified build rules source file</returns>
-		public static UEBuildTarget Create(TargetDescriptor Descriptor, bool bSkipRulesCompile, bool bForceRulesCompile, bool bUsePrecompiled, ILogger Logger)
+		public static UEBuildTarget Create(TargetDescriptor Descriptor, bool bSkipRulesCompile, bool bForceRulesCompile, bool bUsePrecompiled, UnrealIntermediateEnvironment IntermediateEnvironment, ILogger Logger)
 		{
 			// make sure we are allowed to build this platform
 			if (!UEBuildPlatform.IsPlatformAvailable(Descriptor.Platform))
@@ -1177,7 +1234,7 @@ namespace UnrealBuildTool
 			TargetRules RulesObject;
 			using (GlobalTracer.Instance.BuildSpan("RulesAssembly.CreateTargetRules()").StartActive())
 			{
-				RulesObject = RulesAssembly.CreateTargetRules(Descriptor.Name, Descriptor.Platform, Descriptor.Configuration, Descriptor.Architectures, Descriptor.ProjectFile, Descriptor.AdditionalArguments, Logger, Descriptor.IsTestsTarget);
+				RulesObject = RulesAssembly.CreateTargetRules(Descriptor.Name, Descriptor.Platform, Descriptor.Configuration, Descriptor.Architectures, Descriptor.ProjectFile, Descriptor.AdditionalArguments, Logger, Descriptor.IsTestsTarget, IntermediateEnvironment: IntermediateEnvironment);
 			}
 			if ((ProjectFileGenerator.bGenerateProjectFiles == false) && !RulesObject.GetSupportedPlatforms().Contains(Descriptor.Platform))
 			{
@@ -1318,7 +1375,7 @@ namespace UnrealBuildTool
 			}
 
 			// Create the target rules for it
-			TargetRules BaseRules = RulesAssembly.CreateTargetRules(BaseTargetName, ThisRules.Platform, ThisRules.Configuration, ThisRules.Architectures, null, Arguments, Logger);
+			TargetRules BaseRules = RulesAssembly.CreateTargetRules(BaseTargetName, ThisRules.Platform, ThisRules.Configuration, ThisRules.Architectures, null, Arguments, Logger, IntermediateEnvironment: ThisRules.IntermediateEnvironment);
 
 			// Get all the configurable objects
 			object[] BaseObjects = BaseRules.GetConfigurableObjects().ToArray();
@@ -1450,6 +1507,11 @@ namespace UnrealBuildTool
 		/// </summary>
 		public UnrealArchitectures Architectures;
 
+		/// <summary>
+		/// Intermediate environment. Determines if the intermediates end up in a different folder than normal.
+		/// </summary>
+		public UnrealIntermediateEnvironment IntermediateEnvironment;
+
 		/// <summary>
 		/// Relative path for platform-specific intermediates (eg. Intermediate/Build/Win64/x64)
 		/// </summary>
@@ -1583,6 +1645,7 @@ namespace UnrealBuildTool
 			Platform = InDescriptor.Platform;
 			Configuration = InDescriptor.Configuration;
 			Architectures = InDescriptor.Architectures;
+			IntermediateEnvironment = InDescriptor.IntermediateEnvironment;
 			Rules = InRules;
 			RulesAssembly = InRulesAssembly;
 			TargetType = Rules.Type;
@@ -1632,8 +1695,8 @@ namespace UnrealBuildTool
 			}
 
 			// Build the project intermediate directory
-			ProjectIntermediateDirectory = DirectoryReference.Combine(OutputRootDirectory, PlatformIntermediateFolder, TargetName, Configuration.ToString());
-			ProjectIntermediateDirectoryNoArch = DirectoryReference.Combine(OutputRootDirectory, PlatformIntermediateFolderNoArch, TargetName, Configuration.ToString());
+			ProjectIntermediateDirectory = DirectoryReference.Combine(OutputRootDirectory, PlatformIntermediateFolder, GetTargetIntermediateFolderName(TargetName, IntermediateEnvironment), Configuration.ToString());
+			ProjectIntermediateDirectoryNoArch = DirectoryReference.Combine(OutputRootDirectory, PlatformIntermediateFolderNoArch, GetTargetIntermediateFolderName(TargetName, IntermediateEnvironment), Configuration.ToString());
 
 			// Build the engine intermediate directory. If we're building agnostic engine binaries, we can use the engine intermediates folder. Otherwise we need to use the project intermediates directory.
 			if (!bUseSharedBuildEnvironment)
@@ -1677,6 +1740,27 @@ namespace UnrealBuildTool
 			return FolderPath;
 		}
 
+		/// <summary>
+		/// Adjusts a target name to be a usable intermediate folder name.
+		/// </summary>
+		/// <param name="TargetName">Base target name</param>
+		/// <param name="IntermediateEnvironment">Intermediate environment to use</param>
+		/// <returns></returns>
+		public static string GetTargetIntermediateFolderName(string TargetName, UnrealIntermediateEnvironment IntermediateEnvironment)
+		{
+			string TargetFolderName = TargetName;
+			switch (IntermediateEnvironment)
+			{
+				case UnrealIntermediateEnvironment.IWYU:
+					TargetFolderName += "IWYU";
+					break;
+				case UnrealIntermediateEnvironment.NonUnity:
+					TargetFolderName += "NonUnity";
+					break;
+			}
+			return TargetFolderName;
+		}
+
 		/// <summary>
 		/// Gets the app name for a given target type
 		/// </summary>
@@ -3919,7 +4003,7 @@ namespace UnrealBuildTool
 			bool bUseExternalFolder = ModuleRules.Plugin != null && (ModuleRules.Plugin.Type == PluginType.External && !ModuleRules.Plugin.bExplicitPluginTarget);
 
 			// Get the output and intermediate directories for this module
-			DirectoryReference IntermediateDirectory = DirectoryReference.Combine(BaseOutputDirectory, GetPlatformIntermediateFolder(Platform, Architectures, bUseExternalFolder), AppName, ModuleConfiguration.ToString());
+			DirectoryReference IntermediateDirectory = DirectoryReference.Combine(BaseOutputDirectory, GetPlatformIntermediateFolder(Platform, Architectures, bUseExternalFolder), GetTargetIntermediateFolderName(AppName, IntermediateEnvironment), ModuleConfiguration.ToString());
 
 			// Append a subdirectory if the module rules specifies one
 			if (!String.IsNullOrEmpty(ModuleRules.BinariesSubFolder))
diff --git a/Engine/Source/Programs/UnrealBuildTool/Modes/AnalyzeMode.cs b/Engine/Source/Programs/UnrealBuildTool/Modes/AnalyzeMode.cs
--- a/Engine/Source/Programs/UnrealBuildTool/Modes/AnalyzeMode.cs
+++ b/Engine/Source/Programs/UnrealBuildTool/Modes/AnalyzeMode.cs
@@ -32,7 +32,7 @@ namespace UnrealBuildTool.Modes
 			Arguments.ApplyTo(BuildConfiguration);
 
 			// Parse all the target descriptors
-			List<TargetDescriptor> TargetDescriptors = TargetDescriptor.ParseCommandLine(Arguments, BuildConfiguration.bUsePrecompiled, BuildConfiguration.bSkipRulesCompile, BuildConfiguration.bForceRulesCompile, Logger);
+			List<TargetDescriptor> TargetDescriptors = TargetDescriptor.ParseCommandLine(Arguments, BuildConfiguration, Logger);
 
 			// Generate the compile DB for each target
 			using (ISourceFileWorkingSet WorkingSet = new EmptySourceFileWorkingSet())
@@ -113,7 +113,7 @@ namespace UnrealBuildTool.Modes
 		private static void AnalyzeTarget(TargetDescriptor TargetDescriptor, BuildConfiguration BuildConfiguration, ILogger Logger)
 		{
 			// Create a makefile for the target
-			UEBuildTarget Target = UEBuildTarget.Create(TargetDescriptor, BuildConfiguration.bSkipRulesCompile, BuildConfiguration.bForceRulesCompile, BuildConfiguration.bUsePrecompiled, Logger);
+			UEBuildTarget Target = UEBuildTarget.Create(TargetDescriptor, BuildConfiguration, Logger);
 			DirectoryReference.CreateDirectory(Target.ReceiptFileName.Directory);
 
 			// Find the shortest path from the target to each module
diff --git a/Engine/Source/Programs/UnrealBuildTool/Modes/BuildMode.cs b/Engine/Source/Programs/UnrealBuildTool/Modes/BuildMode.cs
--- a/Engine/Source/Programs/UnrealBuildTool/Modes/BuildMode.cs
+++ b/Engine/Source/Programs/UnrealBuildTool/Modes/BuildMode.cs
@@ -174,7 +174,7 @@ namespace UnrealBuildTool
 				// Parse all the target descriptors
 				using (GlobalTracer.Instance.BuildSpan("TargetDescriptor.ParseCommandLine()").StartActive())
 				{
-					TargetDescriptors = TargetDescriptor.ParseCommandLine(Arguments, BuildConfiguration.bUsePrecompiled, BuildConfiguration.bSkipRulesCompile, BuildConfiguration.bForceRulesCompile, Logger);
+					TargetDescriptors = TargetDescriptor.ParseCommandLine(Arguments, BuildConfiguration, Logger);
 
 					// Handle BuildConfiguration arguments nested inside -Target= args
 					TargetDescriptors.ForEach(x => x.AdditionalArguments.ApplyTo(BuildConfiguration));
@@ -528,7 +528,7 @@ namespace UnrealBuildTool
 					using (GlobalTracer.Instance.BuildSpan("Reading dependency cache").StartActive())
 					{
 						TargetDescriptor TargetDescriptor = TargetDescriptors[TargetIdx];
-						CppDependencies.Mount(TargetDescriptor.ProjectFile, TargetDescriptor.Name, TargetDescriptor.Platform, TargetDescriptor.Configuration, Makefiles[TargetIdx].TargetType, TargetDescriptor.Architectures, Logger);
+						CppDependencies.Mount(TargetDescriptor, Makefiles[TargetIdx].TargetType, Logger);
 					}
 				}
 
@@ -1067,7 +1067,7 @@ namespace UnrealBuildTool
 			FileReference? MakefileLocation = null;
 			if (BuildConfiguration.bUseUBTMakefiles)
 			{
-				MakefileLocation = TargetMakefile.GetLocation(TargetDescriptor.ProjectFile, TargetDescriptor.Name, TargetDescriptor.Platform, TargetDescriptor.Architectures, TargetDescriptor.Configuration);
+				MakefileLocation = TargetMakefile.GetLocation(TargetDescriptor.ProjectFile, TargetDescriptor.Name, TargetDescriptor.Platform, TargetDescriptor.Architectures, TargetDescriptor.Configuration, BuildConfiguration.IntermediateEnvironment);
 			}
 
 			// Try to load an existing makefile
@@ -1111,7 +1111,7 @@ namespace UnrealBuildTool
 				UEBuildTarget Target;
 				using (GlobalTracer.Instance.BuildSpan("UEBuildTarget.Create()").StartActive())
 				{
-					Target = UEBuildTarget.Create(TargetDescriptor, BuildConfiguration.bSkipRulesCompile, BuildConfiguration.bForceRulesCompile, BuildConfiguration.bUsePrecompiled, Logger);
+					Target = UEBuildTarget.Create(TargetDescriptor, BuildConfiguration, Logger);
 				}
 
 				// Create the pre-build scripts
diff --git a/Engine/Source/Programs/UnrealBuildTool/Modes/CleanMode.cs b/Engine/Source/Programs/UnrealBuildTool/Modes/CleanMode.cs
--- a/Engine/Source/Programs/UnrealBuildTool/Modes/CleanMode.cs
+++ b/Engine/Source/Programs/UnrealBuildTool/Modes/CleanMode.cs
@@ -48,7 +48,7 @@ namespace UnrealBuildTool
 			Arguments.ApplyTo(BuildConfiguration);
 
 			// Parse all the targets being built
-			List<TargetDescriptor> TargetDescriptors = TargetDescriptor.ParseCommandLine(Arguments, BuildConfiguration.bUsePrecompiled, bSkipRulesCompile, BuildConfiguration.bForceRulesCompile, Logger);
+			List<TargetDescriptor> TargetDescriptors = TargetDescriptor.ParseCommandLine(Arguments, BuildConfiguration, Logger);
 
 			Clean(TargetDescriptors, BuildConfiguration, Logger);
 
@@ -141,7 +141,7 @@ namespace UnrealBuildTool
 					}
 
 					// Add all the makefiles and caches to be deleted
-					FilesToDelete.Add(TargetMakefile.GetLocation(Target.ProjectFile, Target.Name, Target.Platform, Target.Architectures, Target.Configuration));
+					FilesToDelete.Add(TargetMakefile.GetLocation(Target.ProjectFile, Target.Name, Target.Platform, Target.Architectures, Target.Configuration, TargetDescriptor.IntermediateEnvironment));
 					FilesToDelete.UnionWith(SourceFileMetadataCache.GetFilesToClean(Target.ProjectFile));
 
 					// Add all the intermediate folders to be deleted
diff --git a/Engine/Source/Programs/UnrealBuildTool/Modes/IWYUMode.cs b/Engine/Source/Programs/UnrealBuildTool/Modes/IWYUMode.cs
--- a/Engine/Source/Programs/UnrealBuildTool/Modes/IWYUMode.cs
+++ b/Engine/Source/Programs/UnrealBuildTool/Modes/IWYUMode.cs
@@ -384,7 +384,7 @@ namespace UnrealBuildTool
 			BuildConfiguration.MaxNestedPathLength = 220; // For now since the path is slightly longer
 
 			// Parse all the target descriptors
-			List<TargetDescriptor> TargetDescriptors = TargetDescriptor.ParseCommandLine(BuildArguments, BuildConfiguration.bUsePrecompiled, BuildConfiguration.bSkipRulesCompile, BuildConfiguration.bForceRulesCompile, Logger);
+			List<TargetDescriptor> TargetDescriptors = TargetDescriptor.ParseCommandLine(BuildArguments, BuildConfiguration, Logger);
 
 			if (TargetDescriptors.Count != 1)
 			{
@@ -392,6 +392,8 @@ namespace UnrealBuildTool
 				return 0;
 			}
 
+			TargetDescriptors[0].IntermediateEnvironment = UnrealIntermediateEnvironment.IWYU;
+
 			string? ModuleToUpdateName = GetModuleToUpdateName(TargetDescriptors[0]);
 
 			if (!String.IsNullOrEmpty(ModuleToUpdateName))
@@ -433,7 +435,7 @@ namespace UnrealBuildTool
 			{
 				// Create target to be able to traverse all existing modules
 				Logger.LogInformation($"Creating BuildTarget for {TargetName}...");
-				UEBuildTarget Target = UEBuildTarget.Create(TargetDescriptors[0], BuildConfiguration.bSkipRulesCompile, BuildConfiguration.bForceRulesCompile, BuildConfiguration.bUsePrecompiled, Logger);
+				UEBuildTarget Target = UEBuildTarget.Create(TargetDescriptors[0], BuildConfiguration, Logger);
 
 				Logger.LogInformation($"Calculating file filter for IWYU...");
 
diff --git a/Engine/Source/Programs/UnrealBuildTool/ProjectFiles/ProjectFileGenerator.cs b/Engine/Source/Programs/UnrealBuildTool/ProjectFiles/ProjectFileGenerator.cs
--- a/Engine/Source/Programs/UnrealBuildTool/ProjectFiles/ProjectFileGenerator.cs
+++ b/Engine/Source/Programs/UnrealBuildTool/ProjectFiles/ProjectFileGenerator.cs
@@ -2186,7 +2186,7 @@ namespace UnrealBuildTool
 								// If we're generating project files, then go ahead and wipe out the existing UBTMakefile for every target, to make sure that
 								// it gets a full dependency scan next time.
 								// NOTE: This is just a safeguard and doesn't have to be perfect.  We also check for newer project file timestamps in LoadUBTMakefile()
-								FileReference MakefileLocation = TargetMakefile.GetLocation(TargetDesc.ProjectFile, TargetDesc.Name, TargetDesc.Platform, TargetDesc.Architectures, TargetDesc.Configuration);
+								FileReference MakefileLocation = TargetMakefile.GetLocation(TargetDesc.ProjectFile, TargetDesc.Name, TargetDesc.Platform, TargetDesc.Architectures, TargetDesc.Configuration, TargetDesc.IntermediateEnvironment);
 								if (FileReference.Exists(MakefileLocation))
 								{
 									FileReference.Delete(MakefileLocation);
diff --git a/Engine/Source/Programs/UnrealBuildTool/System/CppDependencyCache.cs b/Engine/Source/Programs/UnrealBuildTool/System/CppDependencyCache.cs
--- a/Engine/Source/Programs/UnrealBuildTool/System/CppDependencyCache.cs
+++ b/Engine/Source/Programs/UnrealBuildTool/System/CppDependencyCache.cs
@@ -332,6 +332,18 @@ namespace UnrealBuildTool
 			return true;
 		}
 
+		/// <summary>
+		/// Creates a cache hierarchy for a particular target
+		/// </summary>
+		/// <param name="TargetDescriptor">Target descriptor being built</param>
+		/// <param name="TargetType">The target type</param>
+		/// <param name="Logger">Logger for output</param>
+		/// <returns>Dependency cache hierarchy for the given project</returns>
+		public void Mount(TargetDescriptor TargetDescriptor, TargetType TargetType, ILogger Logger)
+		{
+			Mount(TargetDescriptor.ProjectFile, TargetDescriptor.Name, TargetDescriptor.Platform, TargetDescriptor.Configuration, TargetType, TargetDescriptor.Architectures, TargetDescriptor.IntermediateEnvironment, Logger);
+		}
+
 		/// <summary>
 		/// Creates a cache hierarchy for a particular target
 		/// </summary>
@@ -341,9 +353,10 @@ namespace UnrealBuildTool
 		/// <param name="Configuration">Configuration being built</param>
 		/// <param name="TargetType">The target type</param>
 		/// <param name="Architectures">The target architectures</param>
+		/// <param name="IntermediateEnvironment">Intermediate environment to use</param>
 		/// <param name="Logger">Logger for output</param>
 		/// <returns>Dependency cache hierarchy for the given project</returns>
-		public void Mount(FileReference? ProjectFile, string TargetName, UnrealTargetPlatform Platform, UnrealTargetConfiguration Configuration, TargetType TargetType, UnrealArchitectures Architectures, ILogger Logger)
+		public void Mount(FileReference? ProjectFile, string TargetName, UnrealTargetPlatform Platform, UnrealTargetConfiguration Configuration, TargetType TargetType, UnrealArchitectures Architectures, UnrealIntermediateEnvironment IntermediateEnvironment, ILogger Logger)
 		{
 			string AppName = TargetName;
 			if (ProjectFile == null || !Unreal.IsEngineInstalled())
@@ -358,12 +371,12 @@ namespace UnrealBuildTool
 				}
 			}
 
-			FileReference EngineCacheLocation = FileReference.Combine(Unreal.WritableEngineDirectory, UEBuildTarget.GetPlatformIntermediateFolder(Platform, Architectures, false), AppName, Configuration.ToString(), "DependencyCache.bin");
+			FileReference EngineCacheLocation = FileReference.Combine(Unreal.WritableEngineDirectory, UEBuildTarget.GetPlatformIntermediateFolder(Platform, Architectures, false), UEBuildTarget.GetTargetIntermediateFolderName(AppName, IntermediateEnvironment), Configuration.ToString(), "DependencyCache.bin");
 			FindOrAddPartition(EngineCacheLocation, Unreal.EngineDirectory, Logger);
 
 			if (ProjectFile != null)
 			{
-				FileReference ProjectCacheLocation = FileReference.Combine(ProjectFile.Directory, UEBuildTarget.GetPlatformIntermediateFolder(Platform, Architectures, false), TargetName, Configuration.ToString(), "DependencyCache.bin");
+				FileReference ProjectCacheLocation = FileReference.Combine(ProjectFile.Directory, UEBuildTarget.GetPlatformIntermediateFolder(Platform, Architectures, false), UEBuildTarget.GetTargetIntermediateFolderName(TargetName, IntermediateEnvironment), Configuration.ToString(), "DependencyCache.bin");
 				FindOrAddPartition(ProjectCacheLocation, ProjectFile.Directory, Logger);
 			}
 		}
diff --git a/Engine/Source/Programs/UnrealBuildTool/System/HotReload.cs b/Engine/Source/Programs/UnrealBuildTool/System/HotReload.cs
--- a/Engine/Source/Programs/UnrealBuildTool/System/HotReload.cs
+++ b/Engine/Source/Programs/UnrealBuildTool/System/HotReload.cs
@@ -584,7 +584,7 @@ namespace UnrealBuildTool
 		{
 			// Get the dependency history
 			CppDependencyCache CppDependencies = new CppDependencyCache();
-			CppDependencies.Mount(TargetDescriptor.ProjectFile, TargetDescriptor.Name, TargetDescriptor.Platform, TargetDescriptor.Configuration, Makefile.TargetType, TargetDescriptor.Architectures, Logger);
+			CppDependencies.Mount(TargetDescriptor, Makefile.TargetType, Logger);
 
 			ActionHistory History = new ActionHistory();
 			if (TargetDescriptor.ProjectFile != null)
diff --git a/Engine/Source/Programs/UnrealBuildTool/System/RulesAssembly.cs b/Engine/Source/Programs/UnrealBuildTool/System/RulesAssembly.cs
--- a/Engine/Source/Programs/UnrealBuildTool/System/RulesAssembly.cs
+++ b/Engine/Source/Programs/UnrealBuildTool/System/RulesAssembly.cs
@@ -766,8 +766,9 @@ namespace UnrealBuildTool
 		/// <param name="Logger"></param>
 		/// <param name="IsTestTarget">If building a low level test target</param>
 		/// <param name="bSkipValidation">If validation should be skipped (QueryTargetMode)</param>
+		/// <param name="IntermediateEnvironment">Intermediate environment to use</param>
 		/// <returns>The build target rules for the specified target</returns>
-		public TargetRules CreateTargetRules(string TargetName, UnrealTargetPlatform Platform, UnrealTargetConfiguration Configuration, UnrealArchitectures? Architectures, FileReference? ProjectFile, CommandLineArguments? Arguments, ILogger Logger, bool IsTestTarget = false, bool bSkipValidation = false)
+		public TargetRules CreateTargetRules(string TargetName, UnrealTargetPlatform Platform, UnrealTargetConfiguration Configuration, UnrealArchitectures? Architectures, FileReference? ProjectFile, CommandLineArguments? Arguments, ILogger Logger, bool IsTestTarget = false, bool bSkipValidation = false, UnrealIntermediateEnvironment IntermediateEnvironment = UnrealIntermediateEnvironment.Default)
 		{
 			if (IsTestTarget)
 			{
@@ -803,7 +804,7 @@ namespace UnrealBuildTool
 				}
 				else
 				{
-					return Parent.CreateTargetRules(TargetName, Platform, Configuration, Architectures, ProjectFile, Arguments, Logger, IsTestTarget, bSkipValidation);
+					return Parent.CreateTargetRules(TargetName, Platform, Configuration, Architectures, ProjectFile, Arguments, Logger, IsTestTarget, bSkipValidation, IntermediateEnvironment);
 				}
 			}
 
@@ -811,7 +812,7 @@ namespace UnrealBuildTool
 			string TargetTypeName = TargetName + "Target";
 
 			// The build module must define a type named '<TargetName>Target' that derives from our 'TargetRules' type.  
-			TargetRules? TargetRules = CreateTargetRulesInstance(TargetTypeName, new TargetInfo(TargetName, Platform, Configuration, Architectures, ProjectFile, Arguments), Logger, IsTestTarget, bSkipValidation);
+			TargetRules? TargetRules = CreateTargetRulesInstance(TargetTypeName, new TargetInfo(TargetName, Platform, Configuration, Architectures, ProjectFile, Arguments, IntermediateEnvironment), Logger, IsTestTarget, bSkipValidation);
 
 			if (TargetRules == null)
 			{
diff --git a/Engine/Source/Programs/UnrealBuildTool/System/TargetMakefile.cs b/Engine/Source/Programs/UnrealBuildTool/System/TargetMakefile.cs
--- a/Engine/Source/Programs/UnrealBuildTool/System/TargetMakefile.cs
+++ b/Engine/Source/Programs/UnrealBuildTool/System/TargetMakefile.cs
@@ -989,8 +989,9 @@ namespace UnrealBuildTool
 		/// <param name="Platform">The platform that the target is being built for</param>
 		/// <param name="Architectures">The architectures the target is being built for (can be blank to signify a default)</param>
 		/// <param name="Configuration">The configuration being built</param>
+		/// <param name="IntermediateEnvironment">Intermediate environment to use</param>
 		/// <returns>Path to the makefile</returns>
-		public static FileReference GetLocation(FileReference? ProjectFile, string TargetName, UnrealTargetPlatform Platform, UnrealArchitectures Architectures, UnrealTargetConfiguration Configuration)
+		public static FileReference GetLocation(FileReference? ProjectFile, string TargetName, UnrealTargetPlatform Platform, UnrealArchitectures Architectures, UnrealTargetConfiguration Configuration, UnrealIntermediateEnvironment IntermediateEnvironment)
 		{
 			DirectoryReference BaseDirectory = Unreal.EngineDirectory;
 			// Programs with .uprojects still want the Intermediate under Engine (see UEBuildTarget constructor for similar code) 
@@ -998,7 +999,9 @@ namespace UnrealBuildTool
 			{
 				BaseDirectory = ProjectFile.Directory;
 			}
-			return FileReference.Combine(BaseDirectory, UEBuildTarget.GetPlatformIntermediateFolder(Platform, Architectures, false), TargetName, Configuration.ToString(), "Makefile.bin");
+			string IntermediateFolder = UEBuildTarget.GetPlatformIntermediateFolder(Platform, Architectures, false);
+			string TargetFolderName = UEBuildTarget.GetTargetIntermediateFolderName(TargetName, IntermediateEnvironment);
+			return FileReference.Combine(BaseDirectory, IntermediateFolder, TargetFolderName, Configuration.ToString(), "Makefile.bin");
 		}
 	}
 }
