UnrealBuildTool: Shorten NonUnity folder suffix

#rnx

diff --git a/Engine/Source/Programs/UnrealBuildTool/Configuration/UEBuildTarget.cs b/Engine/Source/Programs/UnrealBuildTool/Configuration/UEBuildTarget.cs
--- a/Engine/Source/Programs/UnrealBuildTool/Configuration/UEBuildTarget.cs
+++ b/Engine/Source/Programs/UnrealBuildTool/Configuration/UEBuildTarget.cs
@@ -1755,7 +1755,7 @@ namespace UnrealBuildTool
 					TargetFolderName += "IWYU";
 					break;
 				case UnrealIntermediateEnvironment.NonUnity:
-					TargetFolderName += "NonUnity";
+					TargetFolderName += "NU";
 					break;
 			}
 			return TargetFolderName;
