Ignore Restricted and NotForLicensees code when warning for long paths.

diff --git a/Engine/Source/Programs/UnrealBuildTool/System/ActionGraph.cs b/Engine/Source/Programs/UnrealBuildTool/System/ActionGraph.cs
--- a/Engine/Source/Programs/UnrealBuildTool/System/ActionGraph.cs
+++ b/Engine/Source/Programs/UnrealBuildTool/System/ActionGraph.cs
@@ -301,7 +301,9 @@ namespace UnrealBuildTool
 						}
 
 						if (PrerequisiteItem.Location.FullName.Length > Unreal.RootDirectory.FullName.Length + BuildConfiguration.MaxNestedPathLength &&
-							PrerequisiteItem.Location.IsUnderDirectory(Unreal.RootDirectory)
+							PrerequisiteItem.Location.IsUnderDirectory(Unreal.RootDirectory) &&
+							(PrerequisiteItem.Location.ContainsName("Restricted", 0) == false) && //Be more relaxed for internal only code
+							(PrerequisiteItem.Location.ContainsName("NotForLicensees", 0) == false)
 							)
 						{
 							WarnPaths.Add(PrerequisiteItem.Location);
@@ -316,9 +318,10 @@ namespace UnrealBuildTool
 							FailPaths.Add(ProducedItem.Location);
 						}
 
-						// don't look in Intermediate directories - these aren't portable between machines, so don't need to be cbecked for length underneath the root
 						if (ProducedItem.Location.FullName.Length > Unreal.RootDirectory.FullName.Length + BuildConfiguration.MaxNestedPathLength &&
-							ProducedItem.Location.IsUnderDirectory(Unreal.RootDirectory)
+							ProducedItem.Location.IsUnderDirectory(Unreal.RootDirectory) &&
+							(ProducedItem.Location.ContainsName("Restricted", 0) == false) && //Be more relaxed for internal only code
+							(ProducedItem.Location.ContainsName("NotForLicensees", 0) == false)
 							)
 						{
 							WarnPaths.Add(ProducedItem.Location);
@@ -337,7 +340,6 @@ namespace UnrealBuildTool
 
 					throw new BuildException(Message.ToString());
 				}
-
 				if (WarnPaths.Count > 0)
 				{
 					StringBuilder Message = new StringBuilder();
