- Wrapped xcodebuild in a Mutex so only one runs at once. Xcode was having a rare issue of concurrent builds blocking access to a build.db file

diff --git a/Engine/Source/Programs/UnrealBuildTool/ToolChain/AppleToolChain.cs b/Engine/Source/Programs/UnrealBuildTool/ToolChain/AppleToolChain.cs
--- a/Engine/Source/Programs/UnrealBuildTool/ToolChain/AppleToolChain.cs
+++ b/Engine/Source/Programs/UnrealBuildTool/ToolChain/AppleToolChain.cs
@@ -824,40 +824,46 @@ namespace UnrealBuildTool
 
 		internal static int FinalizeAppWithXcode(DirectoryReference XcodeProject, UnrealTargetPlatform Platform, string SchemeName, string Configuration, string Action, string ExtraOptions, ILogger Logger)
 		{
-			List<string> Arguments = new()
-			{
-				"UBT_NO_POST_DEPLOY=true",
-				FileReference.Combine(AppleToolChainSettings.XcodeDeveloperDir, "usr/bin/xcodebuild").FullName,
-				Action,
-				$"-workspace \"{XcodeProject.FullName}\"",
-				$"-scheme \"{SchemeName}\"",
-				$"-configuration \"{Configuration}\"",
-				$"-destination generic/platform={AppleExports.GetDestinationPlatform(Platform)}",
-				"-hideShellScriptEnvironment",
-				// xcode gets confused it we _just_ wrote out entitlements while generating the temp project, and it thinks it was modified _during_ building
-				// but it wasn't, it was written before the build started
-				"CODE_SIGN_ALLOW_ENTITLEMENTS_MODIFICATION=YES",
-				ExtraOptions,
-				//$"-sdk {SDKName}",
-			};
-
-			Process LocalProcess = new Process();
-			LocalProcess.StartInfo = new ProcessStartInfo("/usr/bin/env", String.Join(" ", Arguments));
-			LocalProcess.OutputDataReceived += (Sender, Args) => { LocalProcessOutput(Args, false, Logger); };
-			LocalProcess.ErrorDataReceived += (Sender, Args) =>
+			// Acquire a different mutex to the regular UBT instance, since this mode will be called as part of a build. We need the mutex to ensure that building two modular configurations 
+			// in parallel don't clash over writing shared *.modules files (eg. DebugGame and Development editors).
+			string MutexName = SingleInstanceMutex.GetUniqueMutexForPath("UnrealBuildTool_XcodeBuild", Unreal.RootDirectory.FullName);
+			using (new SingleInstanceMutex(MutexName, true))
 			{
-				if (Args != null && Args.Data != null
-				&& Args.Data.Contains("Failed to load profile") && Args.Data.Contains("<stdin>"))
+				List<string> Arguments = new()
 				{
-					Logger.LogInformation("Silencing the following provision profile error, it is not affecting code signing:");
-					LocalProcessOutput(Args, false, Logger);
-				}
-				else
+					"UBT_NO_POST_DEPLOY=true",
+					FileReference.Combine(AppleToolChainSettings.XcodeDeveloperDir, "usr/bin/xcodebuild").FullName,
+					Action,
+					$"-workspace \"{XcodeProject.FullName}\"",
+					$"-scheme \"{SchemeName}\"",
+					$"-configuration \"{Configuration}\"",
+					$"-destination generic/platform={AppleExports.GetDestinationPlatform(Platform)}",
+					"-hideShellScriptEnvironment",
+					// xcode gets confused it we _just_ wrote out entitlements while generating the temp project, and it thinks it was modified _during_ building
+					// but it wasn't, it was written before the build started
+					"CODE_SIGN_ALLOW_ENTITLEMENTS_MODIFICATION=YES",
+					ExtraOptions,
+					//$"-sdk {SDKName}",
+				};
+
+				Process LocalProcess = new Process();
+				LocalProcess.StartInfo = new ProcessStartInfo("/usr/bin/env", String.Join(" ", Arguments));
+				LocalProcess.OutputDataReceived += (Sender, Args) => { LocalProcessOutput(Args, false, Logger); };
+				LocalProcess.ErrorDataReceived += (Sender, Args) =>
 				{
-					LocalProcessOutput(Args, true, Logger);
-				}
-			};
-			return Utils.RunLocalProcess(LocalProcess);
+					if (Args != null && Args.Data != null
+					&& Args.Data.Contains("Failed to load profile") && Args.Data.Contains("<stdin>"))
+					{
+						Logger.LogInformation("Silencing the following provision profile error, it is not affecting code signing:");
+						LocalProcessOutput(Args, false, Logger);
+					}
+					else
+					{
+						LocalProcessOutput(Args, true, Logger);
+					}
+				};
+				return Utils.RunLocalProcess(LocalProcess);
+			}
 		}
 
 		static void LocalProcessOutput(DataReceivedEventArgs? Args, bool bIsError, ILogger Logger)
