- Wrapped xcodebuild in a Mutex so only one runs at once. Xcode was having a rare issue of concurrent builds blocking access to a build.db file

diff --git a/Engine/Source/Programs/UnrealBuildTool/ToolChain/AppleToolChain.cs b/Engine/Source/Programs/UnrealBuildTool/ToolChain/AppleToolChain.cs
--- a/Engine/Source/Programs/UnrealBuildTool/ToolChain/AppleToolChain.cs
+++ b/Engine/Source/Programs/UnrealBuildTool/ToolChain/AppleToolChain.cs
@@ -845,7 +845,22 @@ namespace UnrealBuildTool
 					ExtraOptions,
 					//$"-sdk {SDKName}",
 				};
-
+				{
+					"UBT_NO_POST_DEPLOY=true",
+					new IOSToolChainSettings(Logger).XcodeDeveloperDir + "usr/bin/xcodebuild",
+					Action,
+					$"-workspace \"{XcodeProject.FullName}\"",
+					$"-scheme \"{SchemeName}\"",
+					$"-configuration \"{Configuration}\"",
+					$"-destination generic/platform=" + (Platform == UnrealTargetPlatform.TVOS ? "tvOS" : Platform == UnrealTargetPlatform.Mac ? "macOS" : "iOS"),
+					"-hideShellScriptEnvironment",
+					// xcode gets confused it we _just_ wrote out entitlements while generating the temp project, and it thinks it was modified _during_ building
+					// but it wasn't, it was written before the build started
+					"CODE_SIGN_ALLOW_ENTITLEMENTS_MODIFICATION=YES",
+					ExtraOptions,
+					//$"-sdk {SDKName}",
+				};
+	
 				Process LocalProcess = new Process();
 				LocalProcess.StartInfo = new ProcessStartInfo("/usr/bin/env", String.Join(" ", Arguments));
 				LocalProcess.OutputDataReceived += (Sender, Args) => { LocalProcessOutput(Args, false, Logger); };
