Making more TargetDescriptor uses correctly handle different intermediate paths.
Making the global link environment use the modified intermediate path.

diff --git a/Engine/Source/Programs/UnrealBuildTool/Configuration/UEBuildTarget.cs b/Engine/Source/Programs/UnrealBuildTool/Configuration/UEBuildTarget.cs
--- a/Engine/Source/Programs/UnrealBuildTool/Configuration/UEBuildTarget.cs
+++ b/Engine/Source/Programs/UnrealBuildTool/Configuration/UEBuildTarget.cs
@@ -4752,8 +4752,10 @@ namespace UnrealBuildTool
 
 			[REDACTED_EMAIL]: Do any platform specific tool chain initialization here if required
 
+			string IntermediateAppName = GetTargetIntermediateFolderName(AppName, IntermediateEnvironment);
+
 			UnrealTargetConfiguration EngineTargetConfiguration = Configuration == UnrealTargetConfiguration.DebugGame ? UnrealTargetConfiguration.Development : Configuration;
-			DirectoryReference LinkIntermediateDirectory = DirectoryReference.Combine(Unreal.EngineDirectory, PlatformIntermediateFolder, AppName, EngineTargetConfiguration.ToString());
+			DirectoryReference LinkIntermediateDirectory = DirectoryReference.Combine(Unreal.EngineDirectory, PlatformIntermediateFolder, IntermediateAppName, EngineTargetConfiguration.ToString());
 
 			// Installed Engine intermediates go to the project's intermediate folder. Installed Engine never writes to the engine intermediate folder. (Those files are immutable)
 			// Also, when compiling in monolithic, all intermediates go to the project's folder.  This is because a project can change definitions that affects all engine translation
@@ -4762,11 +4764,11 @@ namespace UnrealBuildTool
 			{
 				if (ProjectFile != null)
 				{
-					LinkIntermediateDirectory = DirectoryReference.Combine(ProjectFile.Directory, PlatformIntermediateFolder, AppName, Configuration.ToString());
+					LinkIntermediateDirectory = DirectoryReference.Combine(ProjectFile.Directory, PlatformIntermediateFolder, IntermediateAppName, Configuration.ToString());
 				}
 				else if (ForeignPlugin != null)
 				{
-					LinkIntermediateDirectory = DirectoryReference.Combine(ForeignPlugin.Directory, PlatformIntermediateFolder, AppName, Configuration.ToString());
+					LinkIntermediateDirectory = DirectoryReference.Combine(ForeignPlugin.Directory, PlatformIntermediateFolder, IntermediateAppName, Configuration.ToString());
 				}
 			}
 
diff --git a/Engine/Source/Programs/UnrealBuildTool/Modes/ExecuteMode.cs b/Engine/Source/Programs/UnrealBuildTool/Modes/ExecuteMode.cs
--- a/Engine/Source/Programs/UnrealBuildTool/Modes/ExecuteMode.cs
+++ b/Engine/Source/Programs/UnrealBuildTool/Modes/ExecuteMode.cs
@@ -54,7 +54,7 @@ namespace UnrealBuildTool
 			// Execute the actions
 			using (GlobalTracer.Instance.BuildSpan("ActionGraph.ExecuteActions()").StartActive())
 			{
-				List<TargetDescriptor> TargetDescriptors = TargetDescriptor.ParseCommandLine(Arguments, false, false, false, Logger);
+				List<TargetDescriptor> TargetDescriptors = TargetDescriptor.ParseCommandLine(Arguments, BuildConfiguration, Logger);
 				await ActionGraph.ExecuteActionsAsync(BuildConfiguration, Actions, TargetDescriptors, Logger);
 			}
 
diff --git a/Engine/Source/Programs/UnrealBuildTool/Modes/FixIncludePathsMode.cs b/Engine/Source/Programs/UnrealBuildTool/Modes/FixIncludePathsMode.cs
--- a/Engine/Source/Programs/UnrealBuildTool/Modes/FixIncludePathsMode.cs
+++ b/Engine/Source/Programs/UnrealBuildTool/Modes/FixIncludePathsMode.cs
@@ -86,7 +86,7 @@ namespace UnrealBuildTool
 			UEBuildModuleCPP.bForceAddGeneratedCodeIncludePath = true;
 
 			// Parse all the target descriptors
-			List<TargetDescriptor> TargetDescriptors = TargetDescriptor.ParseCommandLine(Arguments, BuildConfiguration.bUsePrecompiled, BuildConfiguration.bSkipRulesCompile, BuildConfiguration.bForceRulesCompile, Logger);
+			List<TargetDescriptor> TargetDescriptors = TargetDescriptor.ParseCommandLine(Arguments, BuildConfiguration, Logger);
 
 			// Generate the compile DB for each target
 			using (ISourceFileWorkingSet WorkingSet = new EmptySourceFileWorkingSet())
@@ -98,7 +98,7 @@ namespace UnrealBuildTool
 				foreach (TargetDescriptor TargetDescriptor in TargetDescriptors)
 				{
 					// Create a makefile for the target
-					UEBuildTarget Target = UEBuildTarget.Create(TargetDescriptor, BuildConfiguration.bSkipRulesCompile, BuildConfiguration.bForceRulesCompile, BuildConfiguration.bUsePrecompiled, Logger);
+					UEBuildTarget Target = UEBuildTarget.Create(TargetDescriptor, BuildConfiguration, Logger);
 
 					// Get InputLists and build Include Filter if necessary
 					HashSet<FileReference> IncludeFileList = new HashSet<FileReference>();
diff --git a/Engine/Source/Programs/UnrealBuildTool/Modes/GenerateClangDatabase.cs b/Engine/Source/Programs/UnrealBuildTool/Modes/GenerateClangDatabase.cs
--- a/Engine/Source/Programs/UnrealBuildTool/Modes/GenerateClangDatabase.cs
+++ b/Engine/Source/Programs/UnrealBuildTool/Modes/GenerateClangDatabase.cs
@@ -135,7 +135,7 @@ namespace UnrealBuildTool
 			UEBuildModuleCPP.bForceAddGeneratedCodeIncludePath = true;
 
 			// Parse all the target descriptors
-			List<TargetDescriptor> TargetDescriptors = TargetDescriptor.ParseCommandLine(Arguments, BuildConfiguration.bUsePrecompiled, BuildConfiguration.bSkipRulesCompile, BuildConfiguration.bForceRulesCompile, Logger);
+			List<TargetDescriptor> TargetDescriptors = TargetDescriptor.ParseCommandLine(Arguments, BuildConfiguration, Logger);
 
 			// Generate the compile DB for each target
 			using (ISourceFileWorkingSet WorkingSet = new EmptySourceFileWorkingSet())
@@ -148,7 +148,7 @@ namespace UnrealBuildTool
 					TargetDescriptor.AdditionalArguments = TargetDescriptor.AdditionalArguments.Append(new string[] { "-NoPCH", "-DisableUnity" });
 
 					// Create a makefile for the target
-					UEBuildTarget Target = UEBuildTarget.Create(TargetDescriptor, BuildConfiguration.bSkipRulesCompile, BuildConfiguration.bForceRulesCompile, BuildConfiguration.bUsePrecompiled, Logger);
+					UEBuildTarget Target = UEBuildTarget.Create(TargetDescriptor, BuildConfiguration, Logger);
 					UEToolChain TargetToolChain = Target.CreateToolchain(Target.Platform);
 
 					// Execute code generation actions
diff --git a/Engine/Source/Programs/UnrealBuildTool/Modes/InlineGeneratedCpps.cs b/Engine/Source/Programs/UnrealBuildTool/Modes/InlineGeneratedCpps.cs
--- a/Engine/Source/Programs/UnrealBuildTool/Modes/InlineGeneratedCpps.cs
+++ b/Engine/Source/Programs/UnrealBuildTool/Modes/InlineGeneratedCpps.cs
@@ -78,7 +78,7 @@ namespace UnrealBuildTool
 			UEBuildModuleCPP.bForceAddGeneratedCodeIncludePath = true;
 
 			// Parse all the target descriptors
-			List<TargetDescriptor> TargetDescriptors = TargetDescriptor.ParseCommandLine(Arguments, BuildConfiguration.bUsePrecompiled, BuildConfiguration.bSkipRulesCompile, BuildConfiguration.bForceRulesCompile, Logger);
+			List<TargetDescriptor> TargetDescriptors = TargetDescriptor.ParseCommandLine(Arguments, BuildConfiguration, Logger);
 
 			// Generate the compile DB for each target
 			using (ISourceFileWorkingSet WorkingSet = new EmptySourceFileWorkingSet())
@@ -88,7 +88,7 @@ namespace UnrealBuildTool
 				foreach (TargetDescriptor TargetDescriptor in TargetDescriptors)
 				{
 					// Create a makefile for the target
-					UEBuildTarget Target = UEBuildTarget.Create(TargetDescriptor, BuildConfiguration.bSkipRulesCompile, BuildConfiguration.bForceRulesCompile, BuildConfiguration.bUsePrecompiled, Logger);
+					UEBuildTarget Target = UEBuildTarget.Create(TargetDescriptor, BuildConfiguration, Logger);
 
 					// Create all the binaries and modules
 					CppCompileEnvironment GlobalCompileEnvironment = Target.CreateCompileEnvironmentForProjectFiles(Logger);
diff --git a/Engine/Source/Programs/UnrealBuildTool/Modes/JsonExportMode.cs b/Engine/Source/Programs/UnrealBuildTool/Modes/JsonExportMode.cs
--- a/Engine/Source/Programs/UnrealBuildTool/Modes/JsonExportMode.cs
+++ b/Engine/Source/Programs/UnrealBuildTool/Modes/JsonExportMode.cs
@@ -31,11 +31,16 @@ namespace UnrealBuildTool
 		{
 			Arguments.ApplyTo(this);
 
-			List<TargetDescriptor> TargetDescriptors = TargetDescriptor.ParseCommandLine(Arguments, false, false, false, Logger);
+			// Create the build configuration object, and read the settings
+			BuildConfiguration BuildConfiguration = new BuildConfiguration();
+			XmlConfig.ApplyTo(BuildConfiguration);
+			Arguments.ApplyTo(BuildConfiguration);
+
+			List<TargetDescriptor> TargetDescriptors = TargetDescriptor.ParseCommandLine(Arguments, BuildConfiguration, Logger);
 			foreach (TargetDescriptor TargetDescriptor in TargetDescriptors)
 			{
 				// Create the target
-				UEBuildTarget Target = UEBuildTarget.Create(TargetDescriptor, false, false, false, Logger);
+				UEBuildTarget Target = UEBuildTarget.Create(TargetDescriptor, BuildConfiguration, Logger);
 
 				// Get the output file
 				FileReference? OutputFile = TargetDescriptor.AdditionalArguments.GetFileReferenceOrDefault("-OutputFile=", null);
@@ -49,11 +54,6 @@ namespace UnrealBuildTool
 				{
 					using (ISourceFileWorkingSet WorkingSet = new EmptySourceFileWorkingSet())
 					{
-						// Create the build configuration object, and read the settings
-						BuildConfiguration BuildConfiguration = new BuildConfiguration();
-						XmlConfig.ApplyTo(BuildConfiguration);
-						Arguments.ApplyTo(BuildConfiguration);
-
 						// Create the makefile
 						TargetMakefile Makefile = await Target.BuildAsync(BuildConfiguration, WorkingSet, TargetDescriptor, Logger);
 						List<LinkedAction> Actions = Makefile.Actions.ConvertAll(x => new LinkedAction(x, TargetDescriptor));
diff --git a/Engine/Source/Programs/UnrealBuildTool/Modes/PrintBuildGraphInfo.cs b/Engine/Source/Programs/UnrealBuildTool/Modes/PrintBuildGraphInfo.cs
--- a/Engine/Source/Programs/UnrealBuildTool/Modes/PrintBuildGraphInfo.cs
+++ b/Engine/Source/Programs/UnrealBuildTool/Modes/PrintBuildGraphInfo.cs
@@ -32,7 +32,7 @@ namespace UnrealBuildTool
 			Arguments.ApplyTo(BuildConfiguration);
 
 			// Parse all the target descriptors
-			List<TargetDescriptor> TargetDescriptors = TargetDescriptor.ParseCommandLine(Arguments, BuildConfiguration.bUsePrecompiled, BuildConfiguration.bSkipRulesCompile, BuildConfiguration.bForceRulesCompile, Logger);
+			List<TargetDescriptor> TargetDescriptors = TargetDescriptor.ParseCommandLine(Arguments, BuildConfiguration, Logger);
 
 			using (ISourceFileWorkingSet WorkingSet = new EmptySourceFileWorkingSet())
 			{
@@ -45,7 +45,7 @@ namespace UnrealBuildTool
 					Logger.LogInformation("------------------------");
 
 					// Create a makefile for the target
-					UEBuildTarget Target = UEBuildTarget.Create(TargetDescriptor, BuildConfiguration.bSkipRulesCompile, BuildConfiguration.bForceRulesCompile, BuildConfiguration.bUsePrecompiled, Logger);
+					UEBuildTarget Target = UEBuildTarget.Create(TargetDescriptor, BuildConfiguration, Logger);
 
 					List<UEBuildModule> NoUnityModules = new();
 					List<UEBuildModule> NotOptimizedModules = new();
diff --git a/Engine/Source/Programs/UnrealBuildTool/Modes/ProfileUnitySizesMode.cs b/Engine/Source/Programs/UnrealBuildTool/Modes/ProfileUnitySizesMode.cs
--- a/Engine/Source/Programs/UnrealBuildTool/Modes/ProfileUnitySizesMode.cs
+++ b/Engine/Source/Programs/UnrealBuildTool/Modes/ProfileUnitySizesMode.cs
@@ -135,7 +135,7 @@ namespace UnrealBuildTool
 			UEBuildModuleCPP.bForceAddGeneratedCodeIncludePath = true;
 
 			// Parse all the target descriptors
-			List<TargetDescriptor> TargetDescriptors = TargetDescriptor.ParseCommandLine(Arguments, BuildConfiguration.bUsePrecompiled, BuildConfiguration.bSkipRulesCompile, BuildConfiguration.bForceRulesCompile, Logger);
+			List<TargetDescriptor> TargetDescriptors = TargetDescriptor.ParseCommandLine(Arguments, BuildConfiguration, Logger);
 
 			foreach (TargetDescriptor TargetDescriptor in TargetDescriptors)
 			{
@@ -145,7 +145,7 @@ namespace UnrealBuildTool
 
 				// Create a makefile for the target
 				TimingLogger TimingLogger = new(Logger);
-				UEBuildTarget Target = UEBuildTarget.Create(TargetDescriptor, BuildConfiguration.bSkipRulesCompile, BuildConfiguration.bForceRulesCompile, BuildConfiguration.bUsePrecompiled, TimingLogger);
+				UEBuildTarget Target = UEBuildTarget.Create(TargetDescriptor, BuildConfiguration, TimingLogger);
 				UEToolChain TargetToolChain = Target.CreateToolchain(Target.Platform);
 
 				CppCompileEnvironment GlobalCompileEnvironment = Target.CreateCompileEnvironmentForProjectFiles(TimingLogger);
diff --git a/Engine/Source/Programs/UnrealBuildTool/Modes/TestMode.cs b/Engine/Source/Programs/UnrealBuildTool/Modes/TestMode.cs
--- a/Engine/Source/Programs/UnrealBuildTool/Modes/TestMode.cs
+++ b/Engine/Source/Programs/UnrealBuildTool/Modes/TestMode.cs
@@ -29,7 +29,7 @@ namespace UnrealBuildTool
 			XmlConfig.ApplyTo(BuildConfiguration);
 
 			// Parse all the targets being built
-			List<TargetDescriptor> TargetDescriptors = TargetDescriptor.ParseCommandLine(Arguments, false, false, false, Logger);
+			List<TargetDescriptor> TargetDescriptors = TargetDescriptor.ParseCommandLine(Arguments, BuildConfiguration, Logger);
 
 			await BuildTestsAsync(TargetDescriptors, BuildConfiguration, Logger);
 
diff --git a/Engine/Source/Programs/UnrealBuildTool/Modes/UnrealHeaderToolMode.cs b/Engine/Source/Programs/UnrealBuildTool/Modes/UnrealHeaderToolMode.cs
--- a/Engine/Source/Programs/UnrealBuildTool/Modes/UnrealHeaderToolMode.cs
+++ b/Engine/Source/Programs/UnrealBuildTool/Modes/UnrealHeaderToolMode.cs
@@ -1184,8 +1184,13 @@ namespace UnrealBuildTool.Modes
 
 				if (TargetArgumentIndex >= 0)
 				{
+					// Create the build configuration object, and read the settings
+					BuildConfiguration BuildConfiguration = new();
+					XmlConfig.ApplyTo(BuildConfiguration);
+					Arguments.ApplyTo(BuildConfiguration);
+
 					CommandLineArguments LocalArguments = new(new string[] { Arguments[TargetArgumentIndex] });
-					List<TargetDescriptor> TargetDescriptors = TargetDescriptor.ParseCommandLine(LocalArguments, false, false, false, Logger);
+					List<TargetDescriptor> TargetDescriptors = TargetDescriptor.ParseCommandLine(LocalArguments, BuildConfiguration, Logger);
 					if (TargetDescriptors.Count == 0)
 					{
 						Logger.LogError("No target descriptors found.");
@@ -1195,16 +1200,11 @@ namespace UnrealBuildTool.Modes
 					TargetDescriptor TargetDesc = TargetDescriptors[0];
 
 					// Create the target
-					UEBuildTarget Target = UEBuildTarget.Create(TargetDesc, false, false, false, Logger);
+					UEBuildTarget Target = UEBuildTarget.Create(TargetDesc, BuildConfiguration, Logger);
 
 					// Create the makefile for the target and export the module information
 					using ISourceFileWorkingSet WorkingSet = new EmptySourceFileWorkingSet();
 
-					// Create the build configuration object, and read the settings
-					BuildConfiguration BuildConfiguration = new();
-					XmlConfig.ApplyTo(BuildConfiguration);
-					Arguments.ApplyTo(BuildConfiguration);
-
 					// Create the makefile
 					TargetMakefile Makefile = await Target.BuildAsync(BuildConfiguration, WorkingSet, TargetDesc, Logger, true);
 
