- Programs now no longer use Game/Client entitlements, but can set bUseEntitlementsForPrograms to true in the .ini if really needed
- Added an Unrestricted.entitlements, an empty entitlements file
- Changed Programs to find the .ini for codesigning, without needing a .uproject

diff --git a/Engine/Config/BaseEngine.ini b/Engine/Config/BaseEngine.ini
--- a/Engine/Config/BaseEngine.ini
+++ b/Engine/Config/BaseEngine.ini
@@ -3024,6 +3024,8 @@ TemplateMacPlist=(FilePath="/Game/Build/Mac/Resources/Info.Template.plist")
 TemplateIOSPlist=(FilePath="/Game/Build/IOS/UBTGenerated/Info.Template.plist")
 PremadeMacEntitlements=(FilePath="/Game/Build/Mac/Resources/Sandbox.Server.entitlements")
 ShippingSpecificMacEntitlements=(FilePath="/Game/Build/Mac/Resources/Sandbox.NoNet.entitlements")
+; Programs generally don't need entitlements. Override in a particular program's .ini to true to use the above Premade entitlements files
+bUseEntitlementsForPrograms=false
 ; The following are used only with Manual codesigning
 MacSigningIdentity=Apple Development
 IOSSigningIdentity=Apple Development
@@ -3306,4 +3308,4 @@ EnvironmentCollisionMaxShapes = 32
 bEnableDiskOverflowStore=true
 
 [LowLevelTestsSettings]
-bUpdateBuildGraphPropertiesFile=false
\ No newline at end of file
+bUpdateBuildGraphPropertiesFile=false
diff --git a/Engine/Source/Programs/UnrealBuildTool/ProjectFiles/Xcode/XcodeProject.cs b/Engine/Source/Programs/UnrealBuildTool/ProjectFiles/Xcode/XcodeProject.cs
--- a/Engine/Source/Programs/UnrealBuildTool/ProjectFiles/Xcode/XcodeProject.cs
+++ b/Engine/Source/Programs/UnrealBuildTool/ProjectFiles/Xcode/XcodeProject.cs
@@ -334,6 +334,7 @@ namespace UnrealBuildTool.XcodeProjectXcconfig
 		public FileReference? UProjectFileLocation = null;
 		public DirectoryReference XcodeProjectFileLocation;
 		public DirectoryReference ProductDirectory;
+		public DirectoryReference? ConfigDirectory = null;
 
 		// settings read from project configs
 		public IOSProjectSettings? IOSProjectSettings;
@@ -453,21 +454,11 @@ namespace UnrealBuildTool.XcodeProjectXcconfig
 
 			InitializeUProjectFileLocation(ProjectFile);
 
-			// setup BundleIdentifier from ini file (if there's a specified plist file with one, that will override this)
-			if (String.IsNullOrEmpty(BundleIdentifier))
-			{
-				ConfigHierarchy Ini = ConfigCache.ReadHierarchy(ConfigHierarchyType.Engine, UProjectFileLocation?.Directory, UnrealTargetPlatform.Mac);
-				Ini.GetString($"/Script/MacTargetPlatform.XcodeProjectSettings", "BundleIdentifier", out BundleIdentifier);
-			}
-			if (String.IsNullOrEmpty(BundleIdentifier))
-			{
-				BundleIdentifier = "$(UE_SIGNING_PREFIX).$(UE_PRODUCT_NAME_STRIPPED)";
-			}
-
 			// make sure ProjectDir is something good
 			if (UProjectFileLocation != null)
 			{
 				ProductDirectory = UProjectFileLocation.Directory;
+				ConfigDirectory = ProductDirectory;
 			}
 			else if (ProjectFile.ProjectTargets[0].TargetRules!.Type == TargetType.Program)
 			{
@@ -484,10 +475,22 @@ namespace UnrealBuildTool.XcodeProjectXcconfig
 					if (DirectoryReference.Exists(ProgramFinder))
 					{
 						ProductDirectory = ProgramFinder;
+						ConfigDirectory = ProductDirectory;
 					}
 				}
 			}
 
+			// setup BundleIdentifier from ini file (if there's a specified plist file with one, that will override this)
+			if (String.IsNullOrEmpty(BundleIdentifier))
+			{
+				ConfigHierarchy Ini = ConfigCache.ReadHierarchy(ConfigHierarchyType.Engine, ConfigDirectory, UnrealTargetPlatform.Mac);
+				Ini.GetString($"/Script/MacTargetPlatform.XcodeProjectSettings", "BundleIdentifier", out BundleIdentifier);
+			}
+			if (String.IsNullOrEmpty(BundleIdentifier))
+			{
+				BundleIdentifier = "$(UE_SIGNING_PREFIX).$(UE_PRODUCT_NAME_STRIPPED)";
+			}
+
 			InitializeMetadata(Logger);
 
 			// Figure out all the desired configurations on the unreal side
@@ -540,7 +543,7 @@ namespace UnrealBuildTool.XcodeProjectXcconfig
 		{
 
 			// read setings from the configs, now that we have a project
-			ConfigHierarchy SharedPlatformIni = ConfigCache.ReadHierarchy(ConfigHierarchyType.Engine, UProjectFileLocation?.Directory, UnrealTargetPlatform.Mac);
+			ConfigHierarchy SharedPlatformIni = ConfigCache.ReadHierarchy(ConfigHierarchyType.Engine, ConfigDirectory, UnrealTargetPlatform.Mac);
 			SharedPlatformIni.TryGetValue("/Script/MacTargetPlatform.XcodeProjectSettings", "bUseAutomaticCodeSigning", out bUseAutomaticSigning);
 
 			Metadata = new Metadata(ProductDirectory, XcodeProjectFileLocation, SharedPlatformIni, bSupportsMac, bSupportsIOS || bSupportsTVOS || bSupportsVisionOS, Logger);
@@ -1619,11 +1622,12 @@ namespace UnrealBuildTool.XcodeProjectXcconfig
 			}
 
 			// get ini file for the platform
-			ConfigHierarchy PlatformIni = ConfigCache.ReadHierarchy(ConfigHierarchyType.Engine, UnrealData.UProjectFileLocation?.Directory, Platform);
+			ConfigHierarchy PlatformIni = ConfigCache.ReadHierarchy(ConfigHierarchyType.Engine, UnrealData.ConfigDirectory, Platform);
 
 			// settings for all platforms
 			bool bAutomaticSigning;
 			bool bMacSignToRunLocally;
+			bool bUseEntitlementsForPrograms;
 			string? SigningTeam;
 			string? SigningPrefix;
 			string? AppCategory;
@@ -1640,6 +1644,7 @@ namespace UnrealBuildTool.XcodeProjectXcconfig
 
 			// get signing settings
 			PlatformIni.TryGetValue("/Script/MacTargetPlatform.XcodeProjectSettings", "bUseAutomaticCodeSigning", out bAutomaticSigning);
+			PlatformIni.TryGetValue("/Script/MacTargetPlatform.XcodeProjectSettings", "bUseEntitlementsForPrograms", out bUseEntitlementsForPrograms);
 			PlatformIni.TryGetValue("/Script/MacTargetPlatform.XcodeProjectSettings", "bMacSignToRunLocally", out bMacSignToRunLocally);
 			PlatformIni.TryGetValue("/Script/MacTargetPlatform.XcodeProjectSettings", "CodeSigningTeam", out SigningTeam);
 			PlatformIni.TryGetValue("/Script/MacTargetPlatform.XcodeProjectSettings", "CodeSigningPrefix", out SigningPrefix);
@@ -1866,7 +1871,7 @@ namespace UnrealBuildTool.XcodeProjectXcconfig
 				string ProductName = ExecutableName;
 				string ExecutableKey = $"UE_{Platform.ToString().ToUpper()}_EXECUTABLE_NAME";
 
-				MetadataItem EntitlementsMetadata = UnrealData.Metadata!.EntitlementsFiles[MetadataPlatform];
+				MetadataItem? EntitlementsMetadata = UnrealData.Metadata!.EntitlementsFiles[MetadataPlatform];
 
 				// if we have shipping specified, use it instead
 				if (Config.BuildConfig == UnrealTargetConfiguration.Shipping)
@@ -1878,6 +1883,12 @@ namespace UnrealBuildTool.XcodeProjectXcconfig
 					}
 				}
 
+				// programs almost never want entitlements, so handle them specially
+				if (UnrealData.TargetRules.Type == TargetType.Program && !bUseEntitlementsForPrograms)
+				{
+					EntitlementsMetadata = null;
+				}
+
 				ConfigXcconfig.AppendLine("// pull in the shared settings for all configs for this target");
 				ConfigXcconfig.AppendLine($"#include \"{Xcconfig.Name}.xcconfig\"");
 				ConfigXcconfig.AppendLine("");
@@ -1886,7 +1897,7 @@ namespace UnrealBuildTool.XcodeProjectXcconfig
 				ConfigXcconfig.AppendLine($"UE_UBT_BINARY_SUBPATH = {ExetuableSubPath}");
 				ConfigXcconfig.AppendLine($"{ExecutableKey} = {ExecutableName}");
 				ConfigXcconfig.AppendLine($"PRODUCT_NAME = {ProductName}");
-				if (EntitlementsMetadata.Mode == MetadataMode.UsePremade)
+				if (EntitlementsMetadata != null && EntitlementsMetadata.Mode == MetadataMode.UsePremade)
 				{
 					ConfigXcconfig.AppendLine($"CODE_SIGN_ENTITLEMENTS = {EntitlementsMetadata.XcodeProjectRelative}");
 				}
