UnrealBuildTool: Always pass /ALTERNATENAME for any msvc platform for both vs2019 and vs2022, instead of only for vs2022 so pgo linker args will match

#rnx
#presubmit 64b9cdc0fb725fded24c9781

diff --git a/Engine/Source/Programs/UnrealBuildTool/Platform/Windows/VCToolChain.cs b/Engine/Source/Programs/UnrealBuildTool/Platform/Windows/VCToolChain.cs
--- a/Engine/Source/Programs/UnrealBuildTool/Platform/Windows/VCToolChain.cs
+++ b/Engine/Source/Programs/UnrealBuildTool/Platform/Windows/VCToolChain.cs
@@ -1286,11 +1286,8 @@ namespace UnrealBuildTool
 
 			// Workaround for linker errors when linking against static libraries that were compiled with an older msvc
 			// [REDACTED_PROFILE_URL]/STL/issues/2655
-			if (Target.WindowsPlatform.Architecture == UnrealArch.Arm64ec)
-			{
-				Arguments.Add("/ALTERNATENAME:__imp___std_init_once_begin_initialize=__imp_InitOnceBeginInitialize");
-				Arguments.Add("/ALTERNATENAME:__imp___std_init_once_complete=__imp_InitOnceComplete");
-			}
+			Arguments.Add("/ALTERNATENAME:__imp___std_init_once_begin_initialize=__imp_InitOnceBeginInitialize");
+			Arguments.Add("/ALTERNATENAME:__imp___std_init_once_complete=__imp_InitOnceComplete");
 		}
 
 		protected virtual void AppendLibArguments(LinkEnvironment LinkEnvironment, List<string> Arguments)
