- Fixed path to strip to add a missing slash
- Pass -legacyxcode along to the legacy PostBuildSync step to make sure it uses legacy when overriding the .ini setting

diff --git a/Engine/Source/Programs/UnrealBuildTool/Platform/IOS/IOSToolChain.cs b/Engine/Source/Programs/UnrealBuildTool/Platform/IOS/IOSToolChain.cs
--- a/Engine/Source/Programs/UnrealBuildTool/Platform/IOS/IOSToolChain.cs
+++ b/Engine/Source/Programs/UnrealBuildTool/Platform/IOS/IOSToolChain.cs
@@ -1071,7 +1071,7 @@ namespace UnrealBuildTool
 				Action StripAction = Graph.CreateAction(ActionType.CreateAppBundle);
 				StripAction.WorkingDirectory = GetMacDevSrcRoot();
 				StripAction.CommandPath = BuildHostPlatform.Current.Shell;
-				StripAction.CommandArguments = String.Format("-c \"\\\"{0}strip\\\" {1} \\\"{2}\\\" && touch \\\"{3}\\\"\"", Settings.ToolchainDir, StripArguments, Executable.Location, StripCompleteFile);
+				StripAction.CommandArguments = String.Format("-c \"\\\"{0}/strip\\\" {1} \\\"{2}\\\" && touch \\\"{3}\\\"\"", Settings.ToolchainDir, StripArguments, Executable.Location, StripCompleteFile);
 				StripAction.PrerequisiteItems.Add(Executable);
 				StripAction.PrerequisiteItems.UnionWith(OutputFiles);
 				StripAction.ProducedItems.Add(StripCompleteFile);
@@ -1152,7 +1152,7 @@ namespace UnrealBuildTool
 				FileReference PostBuildSyncFile = FileReference.Combine(BinaryLinkEnvironment.IntermediateDirectory!, "PostBuildSync.dat");
 				BinaryFormatterUtils.Save(PostBuildSyncFile, PostBuildSyncTarget);
 
-				string PostBuildSyncArguments = String.Format("-Input=\"{0}\" -XmlConfigCache=\"{1}\" -remoteini=\"{2}\"", PostBuildSyncFile, XmlConfig.CacheFile, UnrealBuildTool.GetRemoteIniPath());
+				string PostBuildSyncArguments = String.Format("-Input=\"{0}\" -XmlConfigCache=\"{1}\" -remoteini=\"{2}\" -legacyxcode", PostBuildSyncFile, XmlConfig.CacheFile, UnrealBuildTool.GetRemoteIniPath());
 
 				if (Log.OutputFile != null)
 				{
