- Added -legacyxcode commandline handler for legacy PostbuildSync
#trivial

diff --git a/Engine/Source/Programs/UnrealBuildTool/Platform/IOS/IOSPostBuildSyncMode.cs b/Engine/Source/Programs/UnrealBuildTool/Platform/IOS/IOSPostBuildSyncMode.cs
--- a/Engine/Source/Programs/UnrealBuildTool/Platform/IOS/IOSPostBuildSyncMode.cs
+++ b/Engine/Source/Programs/UnrealBuildTool/Platform/IOS/IOSPostBuildSyncMode.cs
@@ -63,6 +63,11 @@ namespace UnrealBuildTool
 		[CommandLine("-XmlConfigCache=")]
 		public FileReference? XmlConfigCache = null;
 
+		// this isn't actually used, but is helpful to pass -legacyxcode along in CreatePostBuildSyncAction, and UBT won't
+		// complain that nothing is using it, because where we _do_ use it is outside the normal cmdline parsing functionality
+		[CommandLine("-LegacyXcode")]
+		public bool bLegacyXcode;
+
 		public override Task<int> ExecuteAsync(CommandLineArguments Arguments, ILogger Logger)
 		{
 			Arguments.ApplyTo(this);
