Improved performance of StartManyActions and TryStartOneAction to avoid excessive processing.

diff --git a/Engine/Source/Programs/UnrealBuildTool/Executors/ImmediateActionQueue.cs b/Engine/Source/Programs/UnrealBuildTool/Executors/ImmediateActionQueue.cs
--- a/Engine/Source/Programs/UnrealBuildTool/Executors/ImmediateActionQueue.cs
+++ b/Engine/Source/Programs/UnrealBuildTool/Executors/ImmediateActionQueue.cs
@@ -78,6 +78,12 @@ namespace UnrealBuildTool
 		/// </summary>
 		public double ActiveActionWeight = 0;
 
+		/// <summary>
+		/// Used to track if the runner has already scanned the action list
+		/// for a given change count but found nothing to run.
+		/// </summary>
+		public int LastActionChange = 0;
+
 		/// <summary>
 		/// True if the current limits have not been reached.
 		/// </summary>
@@ -288,6 +294,16 @@ namespace UnrealBuildTool
 		/// </summary>
 		private int _completedActions = 0;
 
+		/// <summary>
+		/// Flags used to track how StartManyActions should run
+		/// </summary>
+		private int _startManyFlags = 0;
+
+		/// <summary>
+		/// Number of changes to the action list
+		/// </summary>
+		private int _lastActionChange = 1;
+
 		/// <summary>
 		/// Used to terminate the run with status
 		/// </summary>
@@ -458,9 +474,26 @@ namespace UnrealBuildTool
 				lock (Actions)
 				{
 					hasCanceled = CancellationTokenSource.IsCancellationRequested;
-					if (!hasCanceled)
+
+					// Don't bother if we have been canceled or the specified running has already failed to find
+					// anything at the given change number.
+					if (!hasCanceled && (runner == null || runner.LastActionChange != _lastActionChange))
 					{
+
+						// Try to get an action
 						(runAction, action, completedActions) = TryStartOneActionInternal(runner);
+
+						// If we have completed actions (i.e. error propagations), then increment the change counter
+						if (completedActions != 0)
+						{
+							_lastActionChange++;
+						}
+						
+						// Otherwise if nothing was found, remember that we have already scanned at this change.
+						else if ((runAction == null || action == null) && runner != null)
+						{
+							runner.LastActionChange = _lastActionChange;
+						}
 					}
 				}
 
@@ -639,8 +672,31 @@ namespace UnrealBuildTool
 		/// <param name="runner">If specified, all actions will be limited to the runner</param>
 		public void StartManyActions(ImmediateActionQueueRunner? runner = null)
 		{
-			while (TryStartOneAction(runner))
-			{ }
+			const int Running = 1 << 1;
+			const int ScanRequested = 1 << 0;
+
+			// If both flags were clear, I need to start running actions
+			int old = Interlocked.Or(ref _startManyFlags, Running | ScanRequested);
+			if (old == 0)
+			{
+				for(; ; )
+				{
+
+					// Clear the changed flag since we are about to scan
+					Interlocked.And(ref _startManyFlags, Running);
+
+					// If nothing started
+					if (!TryStartOneAction(runner))
+					{
+
+						// If we only have the running flag (nothing new changed), then exit
+						if (Interlocked.CompareExchange(ref _startManyFlags, 0, Running) == Running)
+						{
+							return;
+						}
+					}
+				}
+			}
 		}
 
 		/// <summary>
@@ -704,6 +760,9 @@ namespace UnrealBuildTool
 				runner.ActiveActions--;
 				runner.ActiveActionWeight -= Actions[actionIndex].Action.Weight;
 
+				// Use to track that action states have changed
+				_lastActionChange++;
+
 				// If we are doing an artifact check, then move to compile phase
 				bool wasArtifactCheck = false;
 				if (Actions[actionIndex].Phase == ActionPhase.ArtifactCheck)
