Set the correct sdk to use for iOS device and iOS Simulator builds with updated Main logic

diff --git a/Engine/Platforms/VisionOS/Source/Programs/UnrealBuildTool/VisionOSToolChain.cs b/Engine/Platforms/VisionOS/Source/Programs/UnrealBuildTool/VisionOSToolChain.cs
--- a/Engine/Platforms/VisionOS/Source/Programs/UnrealBuildTool/VisionOSToolChain.cs
+++ b/Engine/Platforms/VisionOS/Source/Programs/UnrealBuildTool/VisionOSToolChain.cs
@@ -18,7 +18,7 @@ namespace UnrealBuildTool
 		{
 		}
 
-		public override string GetXcodeMinVersionParam()
+		public override string GetXcodeMinVersionParam(UnrealArch Architecture)
 		{
 			return "";
 		}
diff --git a/Engine/Source/Programs/UnrealBuildTool/Platform/IOS/IOSToolChain.cs b/Engine/Source/Programs/UnrealBuildTool/Platform/IOS/IOSToolChain.cs
--- a/Engine/Source/Programs/UnrealBuildTool/Platform/IOS/IOSToolChain.cs
+++ b/Engine/Source/Programs/UnrealBuildTool/Platform/IOS/IOSToolChain.cs
@@ -218,9 +218,9 @@ namespace UnrealBuildTool
 
 			Arguments.Add($"-isysroot \"{Settings.GetSDKPath(CompileEnvironment.Architecture)}\"");
 
-			if (GetXcodeMinVersionParam() != "")
+			if (GetXcodeMinVersionParam(CompileEnvironment.Architecture) != "")
 			{
-				Arguments.Add("-m" + GetXcodeMinVersionParam() + "=" + ProjectSettings.RuntimeVersion);
+				Arguments.Add("-m" + GetXcodeMinVersionParam(CompileEnvironment.Architecture) + "=" + ProjectSettings.RuntimeVersion);
 			}
 
 			// Add additional frameworks so that their headers can be found
@@ -288,9 +288,9 @@ namespace UnrealBuildTool
 			RunExecutableAndWait("mkdir", String.Format("-p \"{0}\"", Path), out ResultsText);
 		}
 
-		public virtual string GetXcodeMinVersionParam()
+		public virtual string GetXcodeMinVersionParam(UnrealArch Architecture)
 		{
-			return "iphoneos-version-min";
+			return (Architecture != UnrealArch.IOSSimulator) ? "iphoneos-version-min" : "iphonesimulator-version-min";
 		}
 
 		public string GetAdditionalLinkerFlags(CppConfiguration InConfiguration)
@@ -313,9 +313,9 @@ namespace UnrealBuildTool
 			Arguments.Add($" -isysroot \"{SDKPath}\"");
 
 			Arguments.Add("-dead_strip");
-			if (GetXcodeMinVersionParam() != "")
+			if (GetXcodeMinVersionParam(LinkEnvironment.Architecture) != "") 
 			{
-				Arguments.Add("-m" + GetXcodeMinVersionParam() + "=" + ProjectSettings.RuntimeVersion);
+				Arguments.Add("-m" + GetXcodeMinVersionParam(LinkEnvironment.Architecture) + "=" + ProjectSettings.RuntimeVersion);
 			}
 			Arguments.Add("-Wl-no_pie");
 			Arguments.Add("-stdlib=libc++");
@@ -1589,6 +1589,7 @@ namespace UnrealBuildTool
 							" -scheme '" + SchemeName + "'" +
 								" -sdk " + GetCodesignPlatformName(Target.Platform) +
 							" -destination generic/platform=" + AppleExports.GetDestinationPlatform(Target.Platform) +
+							// " -destination \"platform=iOS Simulator,OS=16.0,name=iPhone 14\"" + //AKSIM 
 								(!String.IsNullOrEmpty(TeamUUID) ? " DEVELOPMENT_TEAM=" + TeamUUID : "");
 					}
 
