NonUnity command line arguments should be per-target while build configuration XML settings should be global.

diff --git a/Engine/Source/Programs/UnrealBuildTool/Configuration/BuildConfiguration.cs b/Engine/Source/Programs/UnrealBuildTool/Configuration/BuildConfiguration.cs
--- a/Engine/Source/Programs/UnrealBuildTool/Configuration/BuildConfiguration.cs
+++ b/Engine/Source/Programs/UnrealBuildTool/Configuration/BuildConfiguration.cs
@@ -238,39 +238,31 @@ namespace UnrealBuildTool
 		/// <summary>
 		/// Whether to unify C++ code into larger files for faster compilation.
 		/// </summary>
-		[CommandLine("-DisableUnity", Value = "false")]
 		[XmlConfigFile(Category = "BuildConfiguration")]
-		public bool bUseUnityBuild = true;
+		public bool? bUseUnityBuild;
 
 		/// <summary>
 		/// Whether to force C++ source files to be combined into larger files for faster compilation.
 		/// </summary>
-		[CommandLine("-ForceUnity")]
 		[XmlConfigFile(Category = "BuildConfiguration")]
-		public bool bForceUnityBuild = false;
-
-		/// <summary>
-		/// Enables "include what you use" mode.
-		/// </summary>
-		[CommandLine("-IWYU")]
-		public bool bIWYU = false;
+		public bool? bForceUnityBuild;
 
 		/// <summary>
 		/// Intermediate environment. Determines if the intermediates end up in a different folder than normal.
 		/// </summary>
-		public UnrealIntermediateEnvironment IntermediateEnvironment
+		public UnrealIntermediateEnvironment? IntermediateEnvironment
 		{
 			get
 			{
-				if (bIWYU)
+				if (bForceUnityBuild == true)
 				{
-					return UnrealIntermediateEnvironment.IWYU;
+					return UnrealIntermediateEnvironment.Default;
 				}
-				if (!bUseUnityBuild && !bForceUnityBuild)
+				if (bUseUnityBuild == false)
 				{
 					return UnrealIntermediateEnvironment.NonUnity;
 				}
-				return UnrealIntermediateEnvironment.Default;
+				return null;
 			}
 		}
 	}
diff --git a/Engine/Source/Programs/UnrealBuildTool/Configuration/TargetDescriptor.cs b/Engine/Source/Programs/UnrealBuildTool/Configuration/TargetDescriptor.cs
--- a/Engine/Source/Programs/UnrealBuildTool/Configuration/TargetDescriptor.cs
+++ b/Engine/Source/Programs/UnrealBuildTool/Configuration/TargetDescriptor.cs
@@ -128,9 +128,47 @@ namespace UnrealBuildTool
 		public bool bRebuild;
 
 		/// <summary>
-		/// The intermediate environment to use. This primarily controls the target's intermediate path.
+		/// Whether to unify C++ code into larger files for faster compilation.
 		/// </summary>
-		public UnrealIntermediateEnvironment IntermediateEnvironment = UnrealIntermediateEnvironment.Default;
+		[CommandLine("-DisableUnity", Value = "false")]
+		public bool bUseUnityBuild = true;
+
+		/// <summary>
+		/// Whether to force C++ source files to be combined into larger files for faster compilation.
+		/// </summary>
+		[CommandLine("-ForceUnity")]
+		public bool bForceUnityBuild = false;
+
+		/// <summary>
+		/// Enables "include what you use" mode.
+		/// </summary>
+		[CommandLine("-IWYU")]
+		public bool bIWYU = false;
+
+		/// <summary>
+		/// Intermediate environment. Determines if the intermediates end up in a different folder than normal.
+		/// </summary>
+		public UnrealIntermediateEnvironment IntermediateEnvironment
+		{
+			get
+			{
+				if (IntermediateEnvironmentOverride.HasValue)
+				{
+					return IntermediateEnvironmentOverride.Value;
+				}
+				if (bIWYU)
+				{
+					return UnrealIntermediateEnvironment.IWYU;
+				}
+				if (!bUseUnityBuild && !bForceUnityBuild)
+				{
+					return UnrealIntermediateEnvironment.NonUnity;
+				}
+				return UnrealIntermediateEnvironment.Default;
+			}
+			set { IntermediateEnvironmentOverride = value; }
+		}
+		private UnrealIntermediateEnvironment? IntermediateEnvironmentOverride;
 
 		/// <summary>
 		/// Constructor
@@ -240,7 +278,10 @@ namespace UnrealBuildTool
 			ParseCommandLine(Arguments, BuildConfiguration.bUsePrecompiled, BuildConfiguration.bSkipRulesCompile, BuildConfiguration.bForceRulesCompile, TargetDescriptors, Logger);
 
 			// apply the intermediate environment from the build configuration
-			TargetDescriptors.ForEach(x => x.IntermediateEnvironment = BuildConfiguration.IntermediateEnvironment);
+			if (BuildConfiguration.IntermediateEnvironment.HasValue)
+			{
+				TargetDescriptors.ForEach(x => x.IntermediateEnvironment = BuildConfiguration.IntermediateEnvironment.Value);
+			}
 
 			return TargetDescriptors;
 		}
diff --git a/Engine/Source/Programs/UnrealBuildTool/Configuration/UEBuildTarget.cs b/Engine/Source/Programs/UnrealBuildTool/Configuration/UEBuildTarget.cs
--- a/Engine/Source/Programs/UnrealBuildTool/Configuration/UEBuildTarget.cs
+++ b/Engine/Source/Programs/UnrealBuildTool/Configuration/UEBuildTarget.cs
@@ -1179,7 +1179,7 @@ namespace UnrealBuildTool
 		/// <returns></returns>
 		public static UEBuildTarget Create(TargetDescriptor Descriptor, BuildConfiguration BuildConfiguration, ILogger Logger)
 		{
-			return Create(Descriptor, BuildConfiguration.bSkipRulesCompile, BuildConfiguration.bForceRulesCompile, BuildConfiguration.bUsePrecompiled, BuildConfiguration.IntermediateEnvironment, Logger);
+			return Create(Descriptor, BuildConfiguration.bSkipRulesCompile, BuildConfiguration.bForceRulesCompile, BuildConfiguration.bUsePrecompiled, Descriptor.IntermediateEnvironment, Logger);
 		}
 
 		/// <summary>
diff --git a/Engine/Source/Programs/UnrealBuildTool/Modes/BuildMode.cs b/Engine/Source/Programs/UnrealBuildTool/Modes/BuildMode.cs
--- a/Engine/Source/Programs/UnrealBuildTool/Modes/BuildMode.cs
+++ b/Engine/Source/Programs/UnrealBuildTool/Modes/BuildMode.cs
@@ -1067,7 +1067,7 @@ namespace UnrealBuildTool
 			FileReference? MakefileLocation = null;
 			if (BuildConfiguration.bUseUBTMakefiles)
 			{
-				MakefileLocation = TargetMakefile.GetLocation(TargetDescriptor.ProjectFile, TargetDescriptor.Name, TargetDescriptor.Platform, TargetDescriptor.Architectures, TargetDescriptor.Configuration, BuildConfiguration.IntermediateEnvironment);
+				MakefileLocation = TargetMakefile.GetLocation(TargetDescriptor.ProjectFile, TargetDescriptor.Name, TargetDescriptor.Platform, TargetDescriptor.Architectures, TargetDescriptor.Configuration, TargetDescriptor.IntermediateEnvironment);
 			}
 
 			// Try to load an existing makefile
