Added the system include paths to the common include path logic to propagate the heavily used system include paths up to the common list.
Removed some redundant common path logic that didn't take into account env value limits.

diff --git a/Engine/Source/Programs/UnrealBuildTool/ProjectFiles/VisualStudio/VCProject.cs b/Engine/Source/Programs/UnrealBuildTool/ProjectFiles/VisualStudio/VCProject.cs
--- a/Engine/Source/Programs/UnrealBuildTool/ProjectFiles/VisualStudio/VCProject.cs
+++ b/Engine/Source/Programs/UnrealBuildTool/ProjectFiles/VisualStudio/VCProject.cs
@@ -1126,6 +1126,10 @@ namespace UnrealBuildTool
 						{
 							IncludePathToCount.AddOrUpdate(IncludePath, _ => Pair.Value, (k, v) => v + Pair.Value);
 						}
+						foreach (DirectoryReference IncludePath in OutBuildEnvironment.SystemIncludePaths.AbsolutePaths)
+						{
+							IncludePathToCount.AddOrUpdate(IncludePath, _ => Pair.Value, (k, v) => v + Pair.Value);
+						}
 						return;
 					}
 				});
@@ -1272,16 +1276,6 @@ namespace UnrealBuildTool
 			string CommonForcedIncludes = string.Empty;
 			string CommonAdditionalOptions = string.Empty;
 			{
-				if (DirectoryToIncludeSearchPaths.Any())
-				{
-					string IncludePathToCheck = DirectoryToIncludeSearchPaths.Values.First();
-					if (DirectoryToIncludeSearchPaths.Values.All(x => x == IncludePathToCheck))
-					{
-						SharedIncludeSearchPaths.Append(IncludePathToCheck);
-						DirectoryToIncludeSearchPaths.Clear();
-					}
-				}
-
 				if (DirectoryToForceIncludePaths.Any())
 				{
 					string ForceIncludePathToCheck = DirectoryToForceIncludePaths.Values.First();
@@ -1550,14 +1544,14 @@ namespace UnrealBuildTool
 						if (TryGetBuildEnvironment(Directory, out BuildEnvironment? BuildEnvironment))
 						{
 							StringBuilder ClCompileInfo = new();
-							if (DirectoryToIncludeSearchPaths.Any())
+							if (DirectoryToIncludeSearchPaths.TryGetValue(Directory, out string? DirectoryToIncludeSearchPathValue) && !string.IsNullOrEmpty(DirectoryToIncludeSearchPathValue))
 							{
-								ClCompileInfo.AppendLine($"      <AdditionalIncludeDirectories>$({DirectoryToIncludeSearchPaths[Directory]})</AdditionalIncludeDirectories>");
+								ClCompileInfo.AppendLine($"      <AdditionalIncludeDirectories>$({DirectoryToIncludeSearchPathValue})</AdditionalIncludeDirectories>");
 							}
 
-							if (DirectoryToForceIncludePaths.Any())
+							if (DirectoryToForceIncludePaths.TryGetValue(Directory, out string? DirectoryToForceIncludePathValue) && !string.IsNullOrEmpty(DirectoryToForceIncludePathValue))
 							{
-								ClCompileInfo.AppendLine($"      <ForcedIncludeFiles>$({DirectoryToForceIncludePaths[Directory]})</ForcedIncludeFiles>");
+								ClCompileInfo.AppendLine($"      <ForcedIncludeFiles>$({DirectoryToForceIncludePathValue})</ForcedIncludeFiles>");
 							}
 
 							if (DirectoryToPchFile.Any())
