Additional UBT changes to support iOS Simulator:
 - update Xcode build script to set correct arch when compiling within Xcode
 - set correct SDKRoot depending on sim vs device

diff --git a/Engine/Build/BatchFiles/Mac/XcodeBuild.sh b/Engine/Build/BatchFiles/Mac/XcodeBuild.sh
--- a/Engine/Build/BatchFiles/Mac/XcodeBuild.sh
+++ b/Engine/Build/BatchFiles/Mac/XcodeBuild.sh
@@ -73,9 +73,13 @@ fi
 
 # Convert platform to UBT terms
 case $PLATFORM in
-	"iphoneos"|"IOS"|"iphonesimulator")
+	"iphoneos"|"IOS")
 		PLATFORM="IOS"
 	;;
+	"iphonesimulator"|"iossimulator")
+		PLATFORM="IOS"
+		ARCHS="iossimulator" 
+	;;
 	"appletvos")
 		PLATFORM="TVOS"
 	;;
diff --git a/Engine/Source/Programs/UnrealBuildTool/Platform/IOS/UEDeployIOS.cs b/Engine/Source/Programs/UnrealBuildTool/Platform/IOS/UEDeployIOS.cs
--- a/Engine/Source/Programs/UnrealBuildTool/Platform/IOS/UEDeployIOS.cs
+++ b/Engine/Source/Programs/UnrealBuildTool/Platform/IOS/UEDeployIOS.cs
@@ -804,7 +804,7 @@ namespace UnrealBuildTool
 
 			Text.AppendLine("\t<key>CFBundleSupportedPlatforms</key>");
 			Text.AppendLine("\t<array>");
-			Text.AppendLine("\t\t<string>iPhoneOS</string>");
+			Text.AppendLine("\t\t<string>iPhoneOS iPhoneSimulator</string>");
 			Text.AppendLine("\t</array>");
 			Text.AppendLine("\t<key>MinimumOSVersion</key>");
 			Text.AppendLine(String.Format("\t<string>{0}</string>", MinVersion));
diff --git a/Engine/Source/Programs/UnrealBuildTool/ProjectFiles/Xcode/XcodeProject.cs b/Engine/Source/Programs/UnrealBuildTool/ProjectFiles/Xcode/XcodeProject.cs
--- a/Engine/Source/Programs/UnrealBuildTool/ProjectFiles/Xcode/XcodeProject.cs
+++ b/Engine/Source/Programs/UnrealBuildTool/ProjectFiles/Xcode/XcodeProject.cs
@@ -1683,8 +1683,17 @@ namespace UnrealBuildTool.XcodeProjectXcconfig
 
 				if (Platform == UnrealTargetPlatform.IOS)
 				{
-					SDKRoot = "iphoneos";
-					SupportedPlatforms = "iphoneos"; // iphonesimulator
+					if (TargetRules.Architecture == UnrealArch.IOSSimulator)
+					{
+						SDKRoot = "iphonesimulator";
+						SupportedPlatforms = "iphonesimulator";
+					}
+					else
+					{	
+						SDKRoot = "iphoneos";
+						SupportedPlatforms = "iphoneos";
+					}
+				
 					DeploymentTargetKey = "IPHONEOS_DEPLOYMENT_TARGET";
 					SupportedDevices = UnrealData.IOSProjectSettings!.RuntimeDevices;
 					DeploymentTarget = UnrealData.IOSProjectSettings.RuntimeVersion;
