Temporarily disable ios simulator root due to Horde issue.

#rnx

diff --git a/Engine/Source/Programs/UnrealBuildTool/ProjectFiles/Xcode/XcodeProject.cs b/Engine/Source/Programs/UnrealBuildTool/ProjectFiles/Xcode/XcodeProject.cs
--- a/Engine/Source/Programs/UnrealBuildTool/ProjectFiles/Xcode/XcodeProject.cs
+++ b/Engine/Source/Programs/UnrealBuildTool/ProjectFiles/Xcode/XcodeProject.cs
@@ -1683,16 +1683,16 @@ namespace UnrealBuildTool.XcodeProjectXcconfig
 
 				if (Platform == UnrealTargetPlatform.IOS)
 				{
-					if (TargetRules.Architecture == UnrealArch.IOSSimulator)
-					{
-						SDKRoot = "iphonesimulator";
-						SupportedPlatforms = "iphonesimulator";
-					}
-					else
-					{	
+					// if (TargetRules.Architecture == UnrealArch.IOSSimulator)
+					// {
+					// 	SDKRoot = "iphonesimulator";
+					// 	SupportedPlatforms = "iphonesimulator";
+					// }
+					// else
+					// {	
 						SDKRoot = "iphoneos";
 						SupportedPlatforms = "iphoneos";
-					}
+					// }
 				
 					DeploymentTargetKey = "IPHONEOS_DEPLOYMENT_TARGET";
 					SupportedDevices = UnrealData.IOSProjectSettings!.RuntimeDevices;
