Fix remote buil on iOS (scheme name)

diff --git a/Engine/Source/Programs/UnrealBuildTool/Platform/IOS/IOSToolChain.cs b/Engine/Source/Programs/UnrealBuildTool/Platform/IOS/IOSToolChain.cs
--- a/Engine/Source/Programs/UnrealBuildTool/Platform/IOS/IOSToolChain.cs
+++ b/Engine/Source/Programs/UnrealBuildTool/Platform/IOS/IOSToolChain.cs
@@ -1553,7 +1553,7 @@ namespace UnrealBuildTool
 					}
 					else
 					{
-						SchemeName = Target.ProjectFile!.GetFileNameWithoutExtension();
+						SchemeName = Target.TargetName;
 					}
 
 					Console.WriteLine("Provisioning: {0}, {1}, {2}, {3}", MobileProvisionFile, MobileProvisionFile.GetFileName(), MobileProvisionUUID, BundleID);
