Added -ActionStallReportTime=seconds option to UBT so that UBT runs where no actions are completed within the number of seconds will dump action queue status.

diff --git a/Engine/Source/Programs/UnrealBuildTool/Executors/ImmediateActionQueue.cs b/Engine/Source/Programs/UnrealBuildTool/Executors/ImmediateActionQueue.cs
--- a/Engine/Source/Programs/UnrealBuildTool/Executors/ImmediateActionQueue.cs
+++ b/Engine/Source/Programs/UnrealBuildTool/Executors/ImmediateActionQueue.cs
@@ -107,6 +107,14 @@ namespace UnrealBuildTool
 	class ImmediateActionQueue : IDisposable
 	{
 
+		/// <summary>
+		/// Number of second of no completed actions to trigger an action stall report.
+		/// If zero, stall reports will not be enabled.
+		/// </summary>
+		[CommandLine("-ActionStallReportTime=")]
+		[XmlConfigFile(Category = "BuildConfiguration")]
+		public int ActionStallReportTime = 0;
+
 		/// <summary>
 		/// Running status of the action
 		/// </summary>
@@ -304,6 +312,21 @@ namespace UnrealBuildTool
 		/// </summary>
 		private int _lastActionChange = 1;
 
+		/// <summary>
+		/// If true, a action stall has been reported for the current change count
+		/// </summary>
+		private bool _lastActionStallReported = false;
+
+		/// <summary>
+		/// Time of the last change to the action count.  This is updated by the timer.
+		/// </summary>
+		private DateTime _lastActionChangeTime;
+
+		/// <summary>
+		/// Copy of _lastActionChange used to detect updates by the time.
+		/// </summary>
+		private int _lastActionStallChange = 1;
+
 		/// <summary>
 		/// Used to terminate the run with status
 		/// </summary>
@@ -327,6 +350,9 @@ namespace UnrealBuildTool
 		/// <param name="logger">Logging interface</param>
 		public ImmediateActionQueue(IEnumerable<LinkedAction> actions, IActionArtifactCache? actionArtifactCache, int maxActionArtifactCacheTasks, string progressWriterText, Action<string> writeToolOutput, ILogger logger)
 		{
+			CommandLine.ParseArguments(Environment.GetCommandLineArgs(), this, logger);
+			XmlConfig.ApplyTo(this);
+
 			int count = actions.Count();
 			Actions = new ActionState[count];
 
@@ -408,17 +434,44 @@ namespace UnrealBuildTool
 		/// </summary>
 		public void Start()
 		{
-			if (ShowCPUUtilization)
+			if (ShowCPUUtilization || ActionStallReportTime > 0)
 			{
+				_lastActionChangeTime = DateTime.Now;
 				_cpuUtilizationTimer = new(x =>
 				{
-					lock (_cpuUtilization)
+					if (ShowCPUUtilization)
 					{
-						if (Utils.GetTotalCpuUtilization(out float cpuUtilization))
+						lock (_cpuUtilization)
 						{
-							_cpuUtilization.Add(cpuUtilization);
+							if (Utils.GetTotalCpuUtilization(out float cpuUtilization))
+							{
+								_cpuUtilization.Add(cpuUtilization);
+							}
 						}
 					}
+
+					if (ActionStallReportTime > 0)
+					{
+						lock (Actions)
+						{
+
+							// If there has been an action count change, reset the timer and enable the report again
+							if (_lastActionStallChange != _lastActionChange)
+							{
+								_lastActionStallChange = _lastActionChange;
+								_lastActionChangeTime = DateTime.Now;
+								_lastActionStallReported = false;
+							}
+
+							// Otherwise, if we haven't already generated a report, test for a timeout in seconds and generate one on timeout.
+							else if (!_lastActionStallReported && (DateTime.Now - _lastActionChangeTime).TotalSeconds > ActionStallReportTime)
+							{
+								_lastActionStallReported = true;
+								GenerateStallReport();
+							}
+						}
+					}
+
 				}, null, 1000, 1000);
 			}
 
@@ -1159,5 +1212,51 @@ namespace UnrealBuildTool
 			}
 			return ActionReadyState.Ready;
 		}
+
+		private void GenerateStallReport()
+		{
+			Logger.LogInformation("Action stall detected:");
+			foreach (ImmediateActionQueueRunner runner in _runners)
+			{
+				Logger.LogInformation("Runner Type: {Type}, Running Actions: {ActionCount}", runner.Type.ToString(), runner.ActiveActions);
+				if (runner.ActiveActions > 0)
+				{
+					int count = 0;
+					foreach (ActionState state in Actions)
+					{
+						if (state.Runner == runner && state.Status == ActionStatus.Running)
+						{
+							string description = $"{(state.Action.CommandDescription ?? state.Action.CommandPath.GetFileNameWithoutExtension())} {state.Action.StatusDescription}".Trim();
+							Logger.LogInformation("    Action[{Index}]: {Description}", count++, description);
+						}
+					}
+				}
+			}
+			{
+				int queued = 0;
+				int running = 0;
+				int error = 0;
+				int finished = 0;
+				foreach (ActionState state in Actions)
+				{
+					switch (state.Status)
+					{
+						case ActionStatus.Error:
+							error++;
+							break;
+						case ActionStatus.Finished:
+							finished++;
+							break;
+						case ActionStatus.Queued:
+							queued++;
+							break;
+						case ActionStatus.Running:
+							running++;
+							break;
+					}
+				}
+				Logger.LogInformation("Queue Counts: Queued = {Queued}, Running = {Running}, Finished = {Finished}, Error = {Error}", queued, running, finished, error);
+			}
+		}
 	}
 }
