Make UBT handle gracefully lack of the PGO data files.

diff --git a/Engine/Source/Programs/UnrealBuildTool/Platform/Linux/UEBuildLinux.cs b/Engine/Source/Programs/UnrealBuildTool/Platform/Linux/UEBuildLinux.cs
--- a/Engine/Source/Programs/UnrealBuildTool/Platform/Linux/UEBuildLinux.cs
+++ b/Engine/Source/Programs/UnrealBuildTool/Platform/Linux/UEBuildLinux.cs
@@ -506,8 +506,22 @@ namespace UnrealBuildTool
 				CompileEnvironment.PGODirectory = CompileEnvironment.PGODirectory.Replace('\\', '/') + "/";
 				CompileEnvironment.PGOFilenamePrefix = string.Format("{0}-{1}-{2}.profdata", Target.Name, Target.Platform, Target.Configuration);
 
-				LinkEnvironment.PGODirectory = CompileEnvironment.PGODirectory;
-				LinkEnvironment.PGOFilenamePrefix = CompileEnvironment.PGOFilenamePrefix;
+				// check if the profdata file exists and disable if not
+				String PGOFilePath = Path.Combine(CompileEnvironment.PGODirectory, CompileEnvironment.PGOFilenamePrefix);
+				if (!File.Exists(PGOFilePath))
+				{
+					Logger.LogWarning("Warning: PGO file '{0}' does not exist, disabling optimization", PGOFilePath);
+					CompileEnvironment.bPGOOptimize = false;
+					LinkEnvironment.bPGOOptimize = false;
+
+					CompileEnvironment.PGODirectory = "";
+					CompileEnvironment.PGOFilenamePrefix = "";
+				}
+				else
+				{
+					LinkEnvironment.PGODirectory = CompileEnvironment.PGODirectory;
+					LinkEnvironment.PGOFilenamePrefix = CompileEnvironment.PGOFilenamePrefix;
+				}
 			}
 
 			LinkEnvironment.bCodeCoverage = CompileEnvironment.bCodeCoverage;
