Add UBT options that allow per-target and per-module choice between "precise" or "imprecise" FP semantics.

diff --git a/Engine/Source/Programs/UnrealBuildTool/Configuration/ModuleRules.cs b/Engine/Source/Programs/UnrealBuildTool/Configuration/ModuleRules.cs
--- a/Engine/Source/Programs/UnrealBuildTool/Configuration/ModuleRules.cs
+++ b/Engine/Source/Programs/UnrealBuildTool/Configuration/ModuleRules.cs
@@ -711,6 +711,17 @@ namespace UnrealBuildTool
 			set => OptimizationLevelOverride = value;
 		}
 
+		/// <summary>
+		/// Allows overriding the FP semantics for this module. This requires a private PCH (or NoPCHs, which is not recommended)
+		/// </summary>
+		public FPSemanticsMode FPSemantics
+		{
+			get => FPSemanticsPrivate ?? Target.FPSemantics;
+			set => FPSemanticsPrivate = value;
+		}
+
+		private FPSemanticsMode? FPSemanticsPrivate = null;
+
 		/// <summary>
 		/// Explicit private PCH for this module. Implies that this module will not use a shared PCH.
 		/// </summary>
diff --git a/Engine/Source/Programs/UnrealBuildTool/Configuration/TargetRules.cs b/Engine/Source/Programs/UnrealBuildTool/Configuration/TargetRules.cs
--- a/Engine/Source/Programs/UnrealBuildTool/Configuration/TargetRules.cs
+++ b/Engine/Source/Programs/UnrealBuildTool/Configuration/TargetRules.cs
@@ -277,6 +277,29 @@ namespace UnrealBuildTool
 		SizeAndSpeed
 	}
 
+	/// <summary>
+	/// Floating point math semantics
+	/// </summary>
+	public enum FPSemanticsMode
+	{
+		/// <summary>
+		/// Use the default semantics for the platform.
+		/// </summary>
+		Default,
+
+		/// <summary>
+		/// FP math is IEEE-754 compliant, assuming that FP exceptions are disabled and the rounding
+		/// mode is round-to-nearest-even.
+		/// </summary>
+		Precise,
+
+		/// <summary>
+		/// FP math isn't IEEE-754 compliant: the compiler is allowed to transform math expressions
+		/// in a ways that might result in differently rounded results from what IEEE-754 requires.
+		/// </summary>
+		Imprecise,
+	}
+
 	/// <summary>
 	/// Determines how the Gameplay Debugger plugin will be activated.
 	/// </summary>
@@ -1337,6 +1360,13 @@ namespace UnrealBuildTool
 		[CommandLine("-OptimizationLevel=")]
 		public OptimizationMode OptimizationLevel = OptimizationMode.Speed;
 
+		/// <summary>
+		/// Allows setting the FP semantics.
+		/// </summary>
+		[RequiresUniqueBuildEnvironment]
+		[CommandLine("-FPSemantics=")]
+		public FPSemanticsMode FPSemantics = FPSemanticsMode.Default;
+
 		/// <summary>
 		/// Whether to compile development automation tests.
 		/// </summary>
@@ -3114,6 +3144,8 @@ namespace UnrealBuildTool
 
 		public OptimizationMode OptimizationLevel => Inner.OptimizationLevel;
 
+		public FPSemanticsMode FPSemantics => Inner.FPSemantics;
+
 		public bool bRetainFramePointers => Inner.bRetainFramePointers;
 
 		public bool bForceCompileDevelopmentAutomationTests => Inner.bForceCompileDevelopmentAutomationTests;
diff --git a/Engine/Source/Programs/UnrealBuildTool/Configuration/UEBuildModuleCPP.cs b/Engine/Source/Programs/UnrealBuildTool/Configuration/UEBuildModuleCPP.cs
--- a/Engine/Source/Programs/UnrealBuildTool/Configuration/UEBuildModuleCPP.cs
+++ b/Engine/Source/Programs/UnrealBuildTool/Configuration/UEBuildModuleCPP.cs
@@ -1871,6 +1871,24 @@ namespace UnrealBuildTool
 				Result.OptimizationLevel = Rules.OptimizationLevel;
 			}
 
+			if (Result.FPSemantics != Rules.FPSemantics)
+			{
+				if (Rules.PrivatePCHHeaderFile == null)
+				{
+					if (Rules.PCHUsage != ModuleRules.PCHUsageMode.NoPCHs)
+					{
+						Logger.LogInformation($"  Overriding FPSemantics requires a private PCH. Disabling PCH usage for {Name}");
+						Rules.PCHUsage = ModuleRules.PCHUsageMode.NoPCHs;
+					}
+				}
+				else if (Rules.PCHUsage == ModuleRules.PCHUsageMode.UseSharedPCHs)
+				{
+					Logger.LogInformation($"  Overriding FPSemantics requires a private PCH. A private PCH exists but UseSharedPCHs was specified. Overriding to NoSharedPCHs for {Name}");
+					Rules.PCHUsage = ModuleRules.PCHUsageMode.NoSharedPCHs;
+				}
+				Result.FPSemantics = Rules.FPSemantics;
+			}
+
 			// If the module overrides the C++ language version, override it on the compile environment
 			if (Rules.CppStandard != null)
 			{
diff --git a/Engine/Source/Programs/UnrealBuildTool/Platform/Windows/VCToolChain.cs b/Engine/Source/Programs/UnrealBuildTool/Platform/Windows/VCToolChain.cs
--- a/Engine/Source/Programs/UnrealBuildTool/Platform/Windows/VCToolChain.cs
+++ b/Engine/Source/Programs/UnrealBuildTool/Platform/Windows/VCToolChain.cs
@@ -782,17 +782,22 @@ namespace UnrealBuildTool
 				Arguments.Add("/bigobj");
 			}
 
+			FPSemanticsMode FPSemantics = CompileEnvironment.FPSemantics;
 			if (Target.WindowsPlatform.Compiler.IsClang())
 			{
 				// FMath::Sqrt calls get inlined and when reciprical is taken, turned into an rsqrtss instruction,
 				// which is *too* imprecise for, e.g., TestVectorNormalize_Sqrt in UnrealMathTest.cpp
 				// TODO: Observed in clang 7.0, presumably the same in Intel C++ Compiler?
-				Arguments.Add("/fp:precise");
+				FPSemantics = FPSemanticsMode.Precise;
 			}
-			else
+
+			switch (FPSemantics)
 			{
-				// Relaxes floating point precision semantics to allow more optimization.
-				Arguments.Add("/fp:fast");
+				case FPSemanticsMode.Default: // Default is imprecise FP semantics.
+				case FPSemanticsMode.Imprecise: Arguments.Add("/fp:fast"); break;
+				case FPSemanticsMode.Precise: Arguments.Add("/fp:precise"); break;
+				default:
+					throw new BuildException($"Unsupported FP semantics: {FPSemantics}");
 			}
 
 			// Intel oneAPI compiler does not support /Zo
diff --git a/Engine/Source/Programs/UnrealBuildTool/System/CppCompileEnvironment.cs b/Engine/Source/Programs/UnrealBuildTool/System/CppCompileEnvironment.cs
--- a/Engine/Source/Programs/UnrealBuildTool/System/CppCompileEnvironment.cs
+++ b/Engine/Source/Programs/UnrealBuildTool/System/CppCompileEnvironment.cs
@@ -401,6 +401,11 @@ namespace UnrealBuildTool
 		/// </summary>
 		public OptimizationMode OptimizationLevel = OptimizationMode.Speed;
 
+		/// <summary>
+		/// Determines the FP semantics.
+		/// </summary>
+		public FPSemanticsMode FPSemantics = FPSemanticsMode.Default;
+
 		/// <summary>
 		/// True if debug info should be created.
 		/// </summary>
@@ -726,6 +731,7 @@ namespace UnrealBuildTool
 			bAllowAutoRTFMInstrumentation = Other.bAllowAutoRTFMInstrumentation;
 			bCodeCoverage = Other.bCodeCoverage;
 			OptimizationLevel = Other.OptimizationLevel;
+			FPSemantics = Other.FPSemantics;
 			bCreateDebugInfo = Other.bCreateDebugInfo;
 			bIsBuildingLibrary = Other.bIsBuildingLibrary;
 			bIsBuildingDLL = Other.bIsBuildingDLL;
diff --git a/Engine/Source/Programs/UnrealBuildTool/ToolChain/ClangToolChain.cs b/Engine/Source/Programs/UnrealBuildTool/ToolChain/ClangToolChain.cs
--- a/Engine/Source/Programs/UnrealBuildTool/ToolChain/ClangToolChain.cs
+++ b/Engine/Source/Programs/UnrealBuildTool/ToolChain/ClangToolChain.cs
@@ -584,6 +584,28 @@ namespace UnrealBuildTool
 			}
 		}
 
+		/// <summary>
+		/// Compile arguments for FP semantics
+		/// </summary>
+		/// <param name="CompileEnvironment"></param>
+		/// <param name="Arguments"></param>
+		protected virtual void GetCompileArguments_FPSemantics(CppCompileEnvironment CompileEnvironment, List<string> Arguments)
+		{
+			switch (CompileEnvironment.FPSemantics)
+			{
+				case FPSemanticsMode.Default: // Default to precise FP semantics.
+				case FPSemanticsMode.Precise:
+					// Clang defaults to -ffp-contract=on, which allows fusing multiplications and additions into FMAs.
+					Arguments.Add("-ffp-contract=off");
+					break;
+				case FPSemanticsMode.Imprecise:
+					Arguments.Add("-ffast-math");
+					break;
+				default:
+					throw new BuildException($"Unsupported FP semantics: {CompileEnvironment.FPSemantics}");
+			}
+		}
+
 		/// <summary>
 		/// Compile arguments for optimization settings, such as profile guided optimization and link time optimization
 		/// </summary>
@@ -771,6 +793,9 @@ namespace UnrealBuildTool
 			// Add warning and error flags to the argument list.
 			GetCompileArguments_WarningsAndErrors(CompileEnvironment, Arguments);
 
+			// Add FP semantics flags to the argument list.
+			GetCompileArguments_FPSemantics(CompileEnvironment, Arguments);
+
 			// Add optimization flags to the argument list.
 			GetCompileArguments_Optimizations(CompileEnvironment, Arguments);
 
