- Added architecture to GetDestinationPlatform() so -destination param to xcodebuild can get simulators (and brought it over from 5.4)
- Added some logging to UseModernXcode so it is clear which mode is being used, and why
- Fixed stub creation in modern to use local dummy identity (no need for keychain unlocking), and moved files into Payload directory inside the stub

diff --git a/Engine/Source/Programs/AutomationTool/Apple/ApplePlatform.Automation.cs b/Engine/Source/Programs/AutomationTool/Apple/ApplePlatform.Automation.cs
--- a/Engine/Source/Programs/AutomationTool/Apple/ApplePlatform.Automation.cs
+++ b/Engine/Source/Programs/AutomationTool/Apple/ApplePlatform.Automation.cs
@@ -103,7 +103,7 @@ public abstract class ApplePlatform : Platform
 					ExtraOptions += $" PRODUCT_NAME={ProductName}";
 				}
 
-				AppleExports.BuildWithStubXcodeProject(SC.RawProjectPath, Target.Platform, Target.Configuration, Target.TargetName, AppleExports.XcodeBuildMode.Stage, Logger, ExtraOptions);
+				AppleExports.BuildWithStubXcodeProject(SC.RawProjectPath, Target.Platform, Target.Architectures, Target.Configuration, Target.TargetName, AppleExports.XcodeBuildMode.Stage, Logger, ExtraOptions);
 			}
 		}
 	}
@@ -198,7 +198,7 @@ public abstract class ApplePlatform : Platform
 				// the archive will be created in the standard Archives location accessible via Xcode. Using -archive will copy it out into
 				// the specified location for use as needed
 				AppleExports.XcodeBuildMode BuildMode = Params.Distribution ? AppleExports.XcodeBuildMode.Distribute : AppleExports.XcodeBuildMode.Package;
-				if (AppleExports.BuildWithStubXcodeProject(Params.RawProjectPath, Receipt.Platform, Receipt.Configuration, Receipt.TargetName, BuildMode, Logger, ExtraOptions) == 0)
+				if (AppleExports.BuildWithStubXcodeProject(Params.RawProjectPath, Receipt.Platform, Receipt.Architectures, Receipt.Configuration, Receipt.TargetName, BuildMode, Logger, ExtraOptions) == 0)
 				{
 					Logger.LogInformation("=====================================================================================");
 					if (Params.Distribution)
diff --git a/Engine/Source/Programs/AutomationTool/IOS/IOSPlatform.Automation.cs b/Engine/Source/Programs/AutomationTool/IOS/IOSPlatform.Automation.cs
--- a/Engine/Source/Programs/AutomationTool/IOS/IOSPlatform.Automation.cs
+++ b/Engine/Source/Programs/AutomationTool/IOS/IOSPlatform.Automation.cs
@@ -1112,7 +1112,7 @@ public class IOSPlatform : ApplePlatform
 		Arguments += SchemeName != null ? SchemeName : GameName;
 		Arguments += "'";
 		Arguments += " -configuration \"" + (SchemeConfiguration != null ? SchemeConfiguration : TargetConfig.ToString()) + "\"";
-		Arguments += " -destination generic/platform=" + AppleExports.GetDestinationPlatform(TargetPlatformType);
+		Arguments += $" -destination generic/platform=\"{AppleExports.GetDestinationPlatform(TargetPlatformType, new UnrealArchitectures(UnrealArch.Arm64))}\"";
 		Arguments += " -sdk " + SDKName;
 
 		if (bAutomaticSigning)
diff --git a/Engine/Source/Programs/UnrealBuildTool/Platform/Apple/AppleExports.cs b/Engine/Source/Programs/UnrealBuildTool/Platform/Apple/AppleExports.cs
--- a/Engine/Source/Programs/UnrealBuildTool/Platform/Apple/AppleExports.cs
+++ b/Engine/Source/Programs/UnrealBuildTool/Platform/Apple/AppleExports.cs
@@ -28,10 +28,12 @@ namespace UnrealBuildTool
 				{
 					throw new BuildException("Both -modernxcode and -legacyxcode were specified, please use one or the other.");
 				}
+				Log.TraceInformationOnce("Forcing MDOERN XCODE with -modernxcode");
 				return true;
 			}
 			if (bForceLegacyXcode)
 			{
+				Log.TraceInformationOnce("Forcing LEGACY XCODE with -legacyxcode");
 				return false;
 			}
 
@@ -41,6 +43,11 @@ namespace UnrealBuildTool
 			{
 				ConfigHierarchy Ini = ConfigCache.ReadHierarchy(ConfigHierarchyType.Engine, ProjectFile?.Directory, UnrealTargetPlatform.Mac);
 				Ini.TryGetValue("/Script/MacTargetPlatform.XcodeProjectSettings", "bUseModernXcode", out bUseModernXcode);
+				Log.TraceInformationOnce("Choosing {0} XCODE based on .ini settings", bUseModernXcode ? "MODERN" : "LEGACY");
+			}
+			else
+			{
+				Log.TraceInformationOnce("Forcing LEGACY XCODE because host OS is not Mac");
 			}
 			return bUseModernXcode;
 		}
@@ -75,8 +82,9 @@ namespace UnrealBuildTool
 		/// Convert UnrealTargetPlatform to the platform that Xcode uses for -destination
 		/// </summary>
 		/// <param name="Platform"></param>
+		/// <param name="Architectures">The architecture we are targeting (really, this is for Simulators)</param>
 		/// <returns></returns>
-		public static string GetDestinationPlatform(UnrealTargetPlatform Platform)
+		public static string GetDestinationPlatform(UnrealTargetPlatform Platform, UnrealArchitectures Architectures)
 		{
 			if (Platform == UnrealTargetPlatform.Mac)
 			{
@@ -84,15 +92,15 @@ namespace UnrealBuildTool
 			}
 			else if (Platform == UnrealTargetPlatform.IOS)
 			{
-				return "iOS";
+				return Architectures.SingleArchitecture == UnrealArch.IOSSimulator ? "iOS Simulator" : "iOS";
 			}
 			else if (Platform == UnrealTargetPlatform.TVOS)
 			{
-				return "tvOS";
+				return Architectures.SingleArchitecture == UnrealArch.TVOSSimulator ? "tvOS Simulator" : "tvOS";
 			}
 			else if (Platform == UnrealTargetPlatform.VisionOS)
 			{
-				return "xrOS";
+				return Architectures.SingleArchitecture == UnrealArch.IOSSimulator ? "xrOS Simulator" : "xrOS";
 			}
 
 			throw new BuildException($"Unknown plaform {Platform}");
@@ -130,14 +138,16 @@ namespace UnrealBuildTool
 		/// </summary>
 		/// <param name="ProjectFile">Project to build</param>
 		/// <param name="Platform">Platform to build</param>
+		/// <param name="Architectures">The architecture we are targeting (really, this is for Simulator to pass in the -destination field)</param>
 		/// <param name="Configuration">Configuration to build</param>
 		/// <param name="TargetName">Target to build</param>
 		/// <param name="BuildMode">Sets an envvar used inside the xcode project to control certain features</param>
 		/// <param name="Logger"></param>
 		/// <param name="ExtraOptions">Any extra options to pass to xcodebuild</param>
+		/// <param name="bForceDummySigning">If true, force signing with the - identity</param>
 		/// <returns>xcode's exit code</returns>
-		public static int BuildWithStubXcodeProject(FileReference? ProjectFile, UnrealTargetPlatform Platform, UnrealTargetConfiguration Configuration,
-			string TargetName, XcodeBuildMode BuildMode, ILogger Logger, string ExtraOptions = "")
+		public static int BuildWithStubXcodeProject(FileReference? ProjectFile, UnrealTargetPlatform Platform, UnrealArchitectures Architectures, UnrealTargetConfiguration Configuration,
+			string TargetName, XcodeBuildMode BuildMode, ILogger Logger, string ExtraOptions = "", bool bForceDummySigning=false)
 		{
 			DirectoryReference? GeneratedProjectFile;
 			// we don't use distro flag when making a modern project
@@ -148,6 +158,12 @@ namespace UnrealBuildTool
 				return 1;
 			}
 
+			if (bForceDummySigning)
+			{
+				ExtraOptions += " CODE_SIGN_STYLE=Manual CODE_SIGN_IDENTITY=-";
+			}
+			else
+			{
 			// look for the special app store connect key information
 			ConfigHierarchy SharedPlatformIni = ConfigCache.ReadHierarchy(ConfigHierarchyType.Engine, ProjectFile?.Directory, UnrealTargetPlatform.Mac);
 			bool bUseAutomaticCodeSigning;
@@ -173,10 +189,11 @@ namespace UnrealBuildTool
 					}
 				}
 			}
+			}
 
 			// run xcodebuild on the generated project to make the .app
 			string XcodeBuildAction = BuildMode == XcodeBuildMode.Distribute ? "archive" : "build";
-			return AppleExports.FinalizeAppWithXcode(GeneratedProjectFile!, Platform, TargetName, Configuration.ToString(), XcodeBuildAction, ExtraOptions + $" UE_XCODE_BUILD_MODE={BuildMode}", Logger);
+			return AppleExports.FinalizeAppWithXcode(GeneratedProjectFile!, Platform, Architectures, TargetName, Configuration.ToString(), XcodeBuildAction, ExtraOptions + $" UE_XCODE_BUILD_MODE={BuildMode}", Logger);
 		}
 
 		/// <summary>
@@ -198,15 +215,16 @@ namespace UnrealBuildTool
 		/// </summary>
 		/// <param name="XcodeProject">The .xcworkspace file to build</param>
 		/// <param name="Platform">THe platform to make the .app for</param>
+		/// <param name="Architectures">The architecture we are targeting (really, this is for Simulator to pass in the -destination field)</param>
 		/// <param name="SchemeName">The name of the scheme (basically the target on the .xcworkspace)</param>
 		/// <param name="Configuration">Which configuration to make (Debug, etc)</param>
 		/// <param name="Action">Action (build, archive, etc)</param>
 		/// <param name="ExtraOptions">Extra options to pass to xcodebuild</param>
 		/// <param name="Logger">Logging object</param>
 		/// <returns>xcode's exit code</returns>
-		public static int FinalizeAppWithXcode(DirectoryReference XcodeProject, UnrealTargetPlatform Platform, string SchemeName, string Configuration, string Action, string ExtraOptions, ILogger Logger)
+		public static int FinalizeAppWithXcode(DirectoryReference XcodeProject, UnrealTargetPlatform Platform, UnrealArchitectures Architectures, string SchemeName, string Configuration, string Action, string ExtraOptions, ILogger Logger)
 		{
-			return AppleToolChain.FinalizeAppWithXcode(XcodeProject, Platform, SchemeName, Configuration, Action, ExtraOptions, Logger);
+			return AppleToolChain.FinalizeAppWithXcode(XcodeProject, Platform, Architectures, SchemeName, Configuration, Action, ExtraOptions, Logger);
 		}
 
 		/// <summary>
diff --git a/Engine/Source/Programs/UnrealBuildTool/Platform/IOS/IOSPostBuildSyncMode.cs b/Engine/Source/Programs/UnrealBuildTool/Platform/IOS/IOSPostBuildSyncMode.cs
--- a/Engine/Source/Programs/UnrealBuildTool/Platform/IOS/IOSPostBuildSyncMode.cs
+++ b/Engine/Source/Programs/UnrealBuildTool/Platform/IOS/IOSPostBuildSyncMode.cs
@@ -13,6 +13,7 @@ namespace UnrealBuildTool
 	class IOSPostBuildSyncTarget
 	{
 		public UnrealTargetPlatform Platform;
+		public UnrealArchitectures Architectures;
 		public UnrealTargetConfiguration Configuration;
 		public FileReference? ProjectFile;
 		public string TargetName;
@@ -34,6 +35,7 @@ namespace UnrealBuildTool
 		public IOSPostBuildSyncTarget(ReadOnlyTargetRules Target, FileReference OutputPath, DirectoryReference? ProjectIntermediateDirectory, List<string> UPLScripts, VersionNumber SdkVersion, Dictionary<string, DirectoryReference> FrameworkNameToSourceDir)
 		{
 			Platform = Target.Platform;
+			Architectures = Target.Architectures;
 			Configuration = Target.Configuration;
 			ProjectFile = Target.ProjectFile;
 			TargetName = Target.Name;
diff --git a/Engine/Source/Programs/UnrealBuildTool/Platform/IOS/IOSToolChain.cs b/Engine/Source/Programs/UnrealBuildTool/Platform/IOS/IOSToolChain.cs
--- a/Engine/Source/Programs/UnrealBuildTool/Platform/IOS/IOSToolChain.cs
+++ b/Engine/Source/Programs/UnrealBuildTool/Platform/IOS/IOSToolChain.cs
@@ -1575,7 +1575,7 @@ namespace UnrealBuildTool
 								" -configuration \"" + ConfigName + "\"" +
 							" -scheme '" + SchemeName + "'" +
 								" -sdk " + GetCodesignPlatformName(Target.Platform) +
-							" -destination generic/platform=" + AppleExports.GetDestinationPlatform(Target.Platform) +
+							$" -destination generic/platform=\"{AppleExports.GetDestinationPlatform(Target.Platform, Target.Architectures)}\"" +
 								" -derivedDataPath \"" + FrameworkDerivedDataDir + "\"" +
 							" CONFIGURATION_BUILD_DIR=\"" + FrameworkPayloadDirectory + "\"" +
 								(!String.IsNullOrEmpty(TeamUUID) ? " DEVELOPMENT_TEAM=" + TeamUUID : "");
@@ -1588,8 +1588,7 @@ namespace UnrealBuildTool
 								" -configuration \"" + ConfigName + "\"" +
 							" -scheme '" + SchemeName + "'" +
 								" -sdk " + GetCodesignPlatformName(Target.Platform) +
-							" -destination generic/platform=" + AppleExports.GetDestinationPlatform(Target.Platform) +
-							// " -destination \"platform=iOS Simulator,OS=16.0,name=iPhone 14\"" + //AKSIM 
+							$" -destination generic/platform=\"{AppleExports.GetDestinationPlatform(Target.Platform, Target.Architectures)}\"" +
 								(!String.IsNullOrEmpty(TeamUUID) ? " DEVELOPMENT_TEAM=" + TeamUUID : "");
 					}
 
diff --git a/Engine/Source/Programs/UnrealBuildTool/ToolChain/AppleToolChain.cs b/Engine/Source/Programs/UnrealBuildTool/ToolChain/AppleToolChain.cs
--- a/Engine/Source/Programs/UnrealBuildTool/ToolChain/AppleToolChain.cs
+++ b/Engine/Source/Programs/UnrealBuildTool/ToolChain/AppleToolChain.cs
@@ -74,7 +74,7 @@ namespace UnrealBuildTool
 			ToolchainDir = DirectoryReference.Combine(XcodeDeveloperDir, "Toolchains/XcodeDefault.xctoolchain/usr/bin");
 			SDKDir = DirectoryReference.Combine(XcodeDeveloperDir, $"Platforms/{OSPrefix}.platform/Developer/SDKs/{OSPrefix}.sdk");
 			if (SimulatorOSPrefix != null)
-			{   
+		{
 				SimulatorSDKDir = DirectoryReference.Combine(XcodeDeveloperDir, $"Platforms/{SimulatorOSPrefix}.platform/Developer/SDKs/{SimulatorOSPrefix}.sdk");
 			}
 
@@ -94,7 +94,7 @@ namespace UnrealBuildTool
 		{
 			// note that VisionOS uses IOSSimulator (as TVOS should eventually do as well)
 			if (Architecture == UnrealArch.IOSSimulator || Architecture == UnrealArch.TVOSSimulator)
-			{
+		{
 				return SimulatorSDKDir!;
 			}
 			return SDKDir;
@@ -106,7 +106,7 @@ namespace UnrealBuildTool
 		/// <param name="Architecture"></param>
 		/// <returns></returns>
 		public string GetTargetTuple(UnrealArch Architecture)
-		{
+			{
 			string Prefix = Architecture.AppleName;
 			string Suffix = (Architecture == UnrealArch.IOSSimulator || Architecture == UnrealArch.TVOSSimulator) ? "-simulator" : "";
 
@@ -115,30 +115,30 @@ namespace UnrealBuildTool
 
 		private static void SelectXcode(bool bVerbose, ILogger Logger)
 		{
-			// on the Mac, run xcode-select directly.
-			int ReturnCode;
+				// on the Mac, run xcode-select directly.
+				int ReturnCode;
 			string XcodeSelectResult = Utils.RunLocalProcessAndReturnStdOut("xcode-select", "--print-path", null, out ReturnCode);
-			if (ReturnCode != 0)
-			{
-				string? MinVersion = UEBuildPlatform.GetSDK(UnrealTargetPlatform.Mac)!.GetSDKInfo("Sdk")!.Min;
-				throw new BuildException($"We were unable to find your build tools (via 'xcode-select --print-path'). Please install Xcode, version {MinVersion} or later");
-			}
+				if (ReturnCode != 0)
+				{
+					string? MinVersion = UEBuildPlatform.GetSDK(UnrealTargetPlatform.Mac)!.GetSDKInfo("Sdk")!.Min;
+					throw new BuildException($"We were unable to find your build tools (via 'xcode-select --print-path'). Please install Xcode, version {MinVersion} or later");
+				}
 
 			_XcodeDeveloperDir = new DirectoryReference(XcodeSelectResult);
 
-			// make sure we get a full path
+				// make sure we get a full path
 			if (DirectoryReference.Exists(XcodeDeveloperDir) == false)
-			{
+				{
 				throw new BuildException("Selected Xcode ('{0}') doesn't exist, cannot continue.", XcodeDeveloperDir);
-			}
+				}
 
 			if (XcodeDeveloperDir.ContainsName("CommandLineTools", 0))
-			{
+				{
 				throw new BuildException($"Your Mac is set to use CommandLineTools for its build tools ({XcodeDeveloperDir}). Unreal expects Xcode as the build tools. Please install Xcode if it's not already, then do one of the following:\n" +
-					"  - Run Xcode, go to Settings, and in the Locations tab, choose your Xcode in Command Line Tools dropdown.\n" +
-					"  - In Terminal, run 'sudo xcode-select -s /Applications/Xcode.app' (or an alternate location if you installed Xcode to a non-standard location)\n" +
-					"Either way, you will need to enter your Mac password.");
-			}
+						"  - Run Xcode, go to Settings, and in the Locations tab, choose your Xcode in Command Line Tools dropdown.\n" +
+						"  - In Terminal, run 'sudo xcode-select -s /Applications/Xcode.app' (or an alternate location if you installed Xcode to a non-standard location)\n" +
+						"Either way, you will need to enter your Mac password.");
+				}
 
 			if (bVerbose && !XcodeDeveloperDir.FullName.StartsWith("/Applications/Xcode.app"))
 			{
@@ -163,60 +163,60 @@ namespace UnrealBuildTool
 		private static string SelectSDK(DirectoryReference BaseSDKDir, string OSPrefix, bool bVerbose, ILogger Logger)
 		{
 			string PlatformSDKVersion = "";
-			try
-			{
-				// loop over the subdirs and parse out the version
-				int MaxSDKVersionMajor = 0;
-				int MaxSDKVersionMinor = 0;
-				string? MaxSDKVersionString = null;
-				foreach (DirectoryReference SubDir in DirectoryReference.EnumerateDirectories(BaseSDKDir))
+				try
 				{
-					string SubDirName = Path.GetFileNameWithoutExtension(SubDir.GetDirectoryName());
-					if (SubDirName.StartsWith(OSPrefix))
+					// loop over the subdirs and parse out the version
+					int MaxSDKVersionMajor = 0;
+					int MaxSDKVersionMinor = 0;
+					string? MaxSDKVersionString = null;
+				foreach (DirectoryReference SubDir in DirectoryReference.EnumerateDirectories(BaseSDKDir))
 					{
-						// get the SDK version from the directory name
-						string SDKString = SubDirName.Replace(OSPrefix, "");
-						int Major = 0;
-						int Minor = 0;
-
-						// parse it into whole and fractional parts (since 10.10 > 10.9 in versions, but not in math)
-						try
+					string SubDirName = Path.GetFileNameWithoutExtension(SubDir.GetDirectoryName());
+						if (SubDirName.StartsWith(OSPrefix))
 						{
-							string[] Tokens = SDKString.Split(".".ToCharArray());
-							if (Tokens.Length == 2)
+							// get the SDK version from the directory name
+							string SDKString = SubDirName.Replace(OSPrefix, "");
+							int Major = 0;
+							int Minor = 0;
+
+							// parse it into whole and fractional parts (since 10.10 > 10.9 in versions, but not in math)
+							try
 							{
-								Major = Int32.Parse(Tokens[0]);
-								Minor = Int32.Parse(Tokens[1]);
+								string[] Tokens = SDKString.Split(".".ToCharArray());
+								if (Tokens.Length == 2)
+								{
+									Major = Int32.Parse(Tokens[0]);
+									Minor = Int32.Parse(Tokens[1]);
+								}
+							}
+							catch (Exception)
+							{
+								// weirdly formatted SDKs
+								continue;
 							}
-						}
-						catch (Exception)
-						{
-							// weirdly formatted SDKs
-							continue;
-						}
 
-						// update largest SDK version number
-						if (Major > MaxSDKVersionMajor || (Major == MaxSDKVersionMajor && Minor > MaxSDKVersionMinor))
-						{
-							MaxSDKVersionString = SDKString;
-							MaxSDKVersionMajor = Major;
-							MaxSDKVersionMinor = Minor;
+							// update largest SDK version number
+							if (Major > MaxSDKVersionMajor || (Major == MaxSDKVersionMajor && Minor > MaxSDKVersionMinor))
+							{
+								MaxSDKVersionString = SDKString;
+								MaxSDKVersionMajor = Major;
+								MaxSDKVersionMinor = Minor;
+							}
 						}
 					}
-				}
 
-				// use the largest version
-				if (MaxSDKVersionString != null)
+					// use the largest version
+					if (MaxSDKVersionString != null)
+					{
+						PlatformSDKVersion = MaxSDKVersionString;
+					}
+				}
+				catch (Exception Ex)
 				{
-					PlatformSDKVersion = MaxSDKVersionString;
+					// on any exception, just use the backup version
+					Logger.LogInformation("Triggered an exception while looking for SDK directory in Xcode.app");
+					Logger.LogInformation("{Ex}", Ex.ToString());
 				}
-			}
-			catch (Exception Ex)
-			{
-				// on any exception, just use the backup version
-				Logger.LogInformation("Triggered an exception while looking for SDK directory in Xcode.app");
-				Logger.LogInformation("{Ex}", Ex.ToString());
-			}
 
 			if (bVerbose && !ProjectFileGenerator.bGenerateProjectFiles)
 			{
@@ -822,7 +822,7 @@ namespace UnrealBuildTool
 			Environment.CurrentDirectory = CurrentCWD;
 		}
 
-		internal static int FinalizeAppWithXcode(DirectoryReference XcodeProject, UnrealTargetPlatform Platform, string SchemeName, string Configuration, string Action, string ExtraOptions, ILogger Logger)
+		internal static int FinalizeAppWithXcode(DirectoryReference XcodeProject, UnrealTargetPlatform Platform, UnrealArchitectures Architectures, string SchemeName, string Configuration, string Action, string ExtraOptions, ILogger Logger)
 		{
 			// Acquire a different mutex to the regular UBT instance, since this mode will be called as part of a build. We need the mutex to ensure that building two modular configurations 
 			// in parallel don't clash over writing shared *.modules files (eg. DebugGame and Development editors).
@@ -837,7 +837,7 @@ namespace UnrealBuildTool
 					$"-workspace \"{XcodeProject.FullName}\"",
 					$"-scheme \"{SchemeName}\"",
 					$"-configuration \"{Configuration}\"",
-					$"-destination generic/platform={AppleExports.GetDestinationPlatform(Platform)}",
+					$"-destination generic/platform=\"{AppleExports.GetDestinationPlatform(Platform, Architectures)}\"",
 					"-hideShellScriptEnvironment",
 					// xcode gets confused it we _just_ wrote out entitlements while generating the temp project, and it thinks it was modified _during_ building
 					// but it wasn't, it was written before the build started
@@ -889,6 +889,7 @@ namespace UnrealBuildTool
 	{
 		public FileReference? ProjectFile;
 		public UnrealTargetPlatform Platform;
+		public UnrealArchitectures Architectures;
 		public UnrealTargetConfiguration Configuration;
 		public string TargetName;
 
@@ -900,6 +901,7 @@ namespace UnrealBuildTool
 		{
 			Platform = Target.Platform;
 			Configuration = Target.Configuration;
+			Architectures = Target.Architectures;
 			ProjectFile = Target.ProjectFile;
 			TargetName = Target.Name;
 
@@ -951,14 +953,14 @@ namespace UnrealBuildTool
 				return 0;
 			}
 
-			int ExitCode = AppleExports.BuildWithStubXcodeProject(Target.ProjectFile, Target.Platform, Target.Configuration, Target.TargetName, AppleExports.XcodeBuildMode.PostBuildSync, Logger);
+			int ExitCode = AppleExports.BuildWithStubXcodeProject(Target.ProjectFile, Target.Platform, Target.Architectures, Target.Configuration, Target.TargetName, 
+				AppleExports.XcodeBuildMode.PostBuildSync, Logger, "", bForceDummySigning:Target.bCreateStubIPA);
 			if (ExitCode != 0)
 			{
 				Logger.LogError("ERROR: Failed to finalize the .app with Xcode. Check the log for more information");
 			}
 
-			if ((Target.Platform == UnrealTargetPlatform.IOS || Target.Platform == UnrealTargetPlatform.TVOS)
-				&& Target.bCreateStubIPA)
+			if (Target.Platform.IsInGroup(UnrealPlatformGroup.IOS) && Target.bCreateStubIPA)
 			{
 				IOSToolChain.PackageStub(Target.StubOutputPath.Directory.FullName, Target.TargetName, Target.StubOutputPath.GetFileNameWithoutExtension(), true);
 			}
