UnrealBuildTool: Check for mismatched PGO linker arguments in VCToolchain & delete duplicate code from other MSVC platforms

#rnx

diff --git a/Engine/Source/Programs/UnrealBuildTool/Platform/Windows/UEBuildWindows.cs b/Engine/Source/Programs/UnrealBuildTool/Platform/Windows/UEBuildWindows.cs
--- a/Engine/Source/Programs/UnrealBuildTool/Platform/Windows/UEBuildWindows.cs
+++ b/Engine/Source/Programs/UnrealBuildTool/Platform/Windows/UEBuildWindows.cs
@@ -189,6 +189,13 @@ namespace UnrealBuildTool
 		/// </summary>
 		TargetRules Target;
 
+		/// <summary>
+		/// If -PGOOptimize is specified but the linker flags have changed since the last -PGOProfile, this will emit a warning and build without PGO instead of failing during link with LNK1268. 
+		/// </summary>
+		[XmlConfigFile(Category = "WindowsPlatform")]
+		[CommandLine("-IgnoreStalePGOData")]
+		public bool bIgnoreStalePGOData = false;
+
 		/// <summary>
 		/// If specified along with -PGOProfile, then /FASTGENPROFILE will be used instead of /GENPROFILE. This usually means that the PGO data is generated faster, but the resulting data may not yield as efficient optimizations during -PGOOptimize
 		/// </summary>
@@ -639,6 +646,8 @@ namespace UnrealBuildTool
 		#region Read-only accessor properties 
 #pragma warning disable CS1591
 
+		public bool bIgnoreStalePGOData => Inner.bIgnoreStalePGOData;
+
 		public bool bUseFastGenProfile => Inner.bUseFastGenProfile;
 
 		public bool bPGONoExtraCounters => Inner.bPGONoExtraCounters;
diff --git a/Engine/Source/Programs/UnrealBuildTool/Platform/Windows/VCToolChain.cs b/Engine/Source/Programs/UnrealBuildTool/Platform/Windows/VCToolChain.cs
--- a/Engine/Source/Programs/UnrealBuildTool/Platform/Windows/VCToolChain.cs
+++ b/Engine/Source/Programs/UnrealBuildTool/Platform/Windows/VCToolChain.cs
@@ -2570,48 +2570,113 @@ namespace UnrealBuildTool
 			return true;
 		}
 
+		static readonly string[] PGOIgnoredLinkerSwitch =
+		{
+			"/USEPROFILE",
+			"/GENPROFILE",
+			"/FASTGENPROFILE",
+			"/OUT:",
+			"/PDB:",
+			"/NOLOGO",
+			"/LIBPATH:",
+			"/errorReport:",
+			"/d2:",
+			"/NATVIS:",
+			"/experimental:deterministic"
+		};
+		private IEnumerable<string> RemovePGOIgnoredSwitches(IEnumerable<string> SourceArguments)
+		{
+			return SourceArguments.Where(Argument => Argument.StartsWith("/") && !PGOIgnoredLinkerSwitch.Any(Switch => Argument.StartsWith(Switch, StringComparison.OrdinalIgnoreCase)));
+		}
+
 		protected virtual void AddPGOLinkArguments(LinkEnvironment LinkEnvironment, List<string> Arguments)
 		{
 			bool bPGOOptimize = LinkEnvironment.bPGOOptimize;
 			bool bPGOProfile = LinkEnvironment.bPGOProfile;
 
+			if (bPGOOptimize || bPGOProfile)
+			{
+				// serialize the linker arguments for comparing /USEPROFILE and /GENPROFILE
+				IEnumerable<string> PGOArguments = RemovePGOIgnoredSwitches(Arguments);
+				string PGOProfileLinkArgsFile = Path.Combine(LinkEnvironment.PGODirectory!, "PGOProfileLinkArgs.txt");
+
+				if (bPGOProfile)
+				{
+					// save the latest linker arguments for /GENPROFILE alongside the PGO data
+					Utils.WriteFileIfChanged(new FileReference(PGOProfileLinkArgsFile), PGOArguments, Logger);
+				}
+				else if (bPGOOptimize && Target.WindowsPlatform.bIgnoreStalePGOData && File.Exists(PGOProfileLinkArgsFile))
+				{
+					// compare latest linker arguments for /USEPROFILE to the last known /GENPROFILE and disable /USEPROFILE if they are different to avoid LNK1268
+					IEnumerable<string> PGOProfileLinkArgs = RemovePGOIgnoredSwitches(File.ReadAllLines(PGOProfileLinkArgsFile));
+					IEnumerable<string> AddedPGOArgs = PGOArguments.Except(PGOProfileLinkArgs);
+					IEnumerable<string> RemovedPGOArgs = PGOProfileLinkArgs.Except(PGOArguments);
+
+					if (AddedPGOArgs.Any() || RemovedPGOArgs.Any())
+					{
+						Logger.LogWarning("PGO Profile and PGO Optimize linker arguments do not match. Profile Guided Optimization will be disabled for {Config} until new PGO Profile data is generated.", Target.Configuration.ToString());
+						bPGOOptimize = false;
+
+						if (AddedPGOArgs.Any())
+						{
+							Logger.LogInformation("Added PGO args:");
+							foreach (string Arg in AddedPGOArgs)
+							{
+								Logger.LogInformation("  {Arg}", Arg);
+							}
+						}
+						if (RemovedPGOArgs.Any())
+						{
+							Logger.LogInformation("Removed PGO args:");
+							foreach (string Arg in RemovedPGOArgs)
+							{
+								Logger.LogInformation("  {Arg}", Arg);
+							}
+						}
+					}
+				}
+			}
+
+			// If PGO was requested, apply Link-time code generation even if PGO was disabled due to mismatched args
+			if (LinkEnvironment.bPGOOptimize || LinkEnvironment.bPGOProfile)
+			{
+				if (!Arguments.Contains("/LTCG"))
+				{
+					Arguments.Add("/LTCG");
+				}
+				Log.TraceInformationOnce("Enabling Link-time code generation (LTGC) as Profile Guided Optimization (PGO) was requested. Linking will take a while.");
+			}
+
 			if (bPGOOptimize)
 			{
 				if (PreparePGOFiles(LinkEnvironment))
 				{
-					//Arguments.Add("/USEPROFILE:PGD=" + Path.Combine(LinkEnvironment.PGODirectory, LinkEnvironment.PGOFilenamePrefix + ".pgd"));
-					Arguments.Add("/LTCG");
-					//Arguments.Add("/USEPROFILE:PGD=" + LinkEnvironment.PGOFilenamePrefix + ".pgd");
 					Arguments.Add("/USEPROFILE");
-					Log.TraceInformationOnce("Enabling Profile Guided Optimization (PGO). Linking will take a while.");
+					Log.TraceInformationOnce("Enabling using Profile Guided Optimization (PGO). Linking will take a while.");
 				}
 				else
 				{
-					Logger.LogWarning("PGO Optimize build will be disabled");
+					Logger.LogWarning("Unable to prepare Profile Guided Optimization (PGO) files. Using PGO Optimize build will be disabled.");
 					bPGOOptimize = false;
 				}
 			}
 			else if (bPGOProfile)
 			{
-				//Arguments.Add("/GENPROFILE:PGD=" + Path.Combine(LinkEnvironment.PGODirectory, LinkEnvironment.PGOFilenamePrefix + ".pgd"));
-				Arguments.Add("/LTCG");
-				//Arguments.Add("/GENPROFILE:PGD=" + LinkEnvironment.PGOFilenamePrefix + ".pgd");
-
 				if (Target.WindowsPlatform.bUseFastGenProfile)
 				{
+					Log.TraceInformationOnce("Enabling generating Fastgen Profile Guided Optimization (PGO). Linking will take a while.");
 					Arguments.Add("/FASTGENPROFILE");
-					Log.TraceInformationOnce("Enabling Fastgen Profile Guided Optimization (PGO). Linking will take a while.");
 				}
 				else
 				{
 					if (Target.WindowsPlatform.bPGONoExtraCounters)
 					{
-						Log.TraceInformationOnce("Enabling Profile Guided Optimization No Extra Counters(PGO). Linking will take a while.");
+						Log.TraceInformationOnce("Enabling generating Profile Guided Optimization No Extra Counters (PGO). Linking will take a while.");
 						Arguments.Add("/GENPROFILE:NOTRACKEH");
 					}
 					else
 					{
-						Log.TraceInformationOnce("Enabling Profile Guided Optimization (PGO). Linking will take a while.");
+						Log.TraceInformationOnce("Enabling generating Profile Guided Optimization (PGO). Linking will take a while.");
 						Arguments.Add("/GENPROFILE");
 					}
 				}
@@ -2620,10 +2685,10 @@ namespace UnrealBuildTool
 
 		protected virtual void ModifyFinalLinkArguments(LinkEnvironment LinkEnvironment, List<string> Arguments, bool bBuildImportLibraryOnly)
 		{
-			AddPGOLinkArguments(LinkEnvironment, Arguments);
-
 			// IMPLEMENT_MODULE_ is not required - it only exists to ensure developers add an IMPLEMENT_MODULE() declaration in code. These are always removed for PGO so that adding/removing a module won't invalidate PGC data.
 			Arguments.RemoveAll(Argument => Argument.StartsWith("/INCLUDE:IMPLEMENT_MODULE_"));
+
+			AddPGOLinkArguments(LinkEnvironment, Arguments);
 		}
 
 		private void ExportObjectFilePaths(LinkEnvironment LinkEnvironment, string FileName, VCEnvironment EnvVars)
