- Made IOS and Mac BuildPlatforms classes come from AppleBuildPlatform so we can make UBT regen makefiles when xcode-select location changes
- Fixed XcodeBuild.sh to set iossimulator architecture the expected way

diff --git a/Engine/Build/BatchFiles/Mac/XcodeBuild.sh b/Engine/Build/BatchFiles/Mac/XcodeBuild.sh
--- a/Engine/Build/BatchFiles/Mac/XcodeBuild.sh
+++ b/Engine/Build/BatchFiles/Mac/XcodeBuild.sh
@@ -78,7 +78,7 @@ case $PLATFORM in
 	;;
 	"iphonesimulator"|"iossimulator")
 		PLATFORM="IOS"
-		ARCHS="iossimulator" 
+		UBT_ARCHFLAG="-architecture=iossimulator" 
 	;;
 	"appletvos")
 		PLATFORM="TVOS"
diff --git a/Engine/Source/Programs/UnrealBuildTool/Platform/IOS/UEBuildIOS.cs b/Engine/Source/Programs/UnrealBuildTool/Platform/IOS/UEBuildIOS.cs
--- a/Engine/Source/Programs/UnrealBuildTool/Platform/IOS/UEBuildIOS.cs
+++ b/Engine/Source/Programs/UnrealBuildTool/Platform/IOS/UEBuildIOS.cs
@@ -143,6 +143,11 @@ namespace UnrealBuildTool
 		/// </summary>
 		[CommandLine("-EnableUBSan")]
 		public bool bEnableUndefinedBehaviorSanitizer = false;
+
+		/// <summary>
+		/// If true, we use SwiftUI's main function instead of Obj-C's main
+		/// </summary>
+		public bool bUseSwiftUIMain = false;
 	}
 
 	/// <summary>
@@ -197,6 +202,8 @@ namespace UnrealBuildTool
 
 		public bool bEnableUndefinedBehaviorSanitizer => Inner.bEnableUndefinedBehaviorSanitizer;
 
+		public bool bUseSwiftUIMain => Inner.bUseSwiftUIMain;
+
 #pragma warning restore CS1591
 		#endregion
 	}
@@ -654,7 +661,7 @@ namespace UnrealBuildTool
 		}
 	}
 
-	class IOSPlatform : UEBuildPlatform
+	class IOSPlatform : AppleBuildPlatform
 	{
 		List<IOSProjectSettings> CachedProjectSettings = new List<IOSProjectSettings>();
 		List<IOSProjectSettings> CachedProjectSettingsByBundle = new List<IOSProjectSettings>();
diff --git a/Engine/Source/Programs/UnrealBuildTool/Platform/Mac/UEBuildMac.cs b/Engine/Source/Programs/UnrealBuildTool/Platform/Mac/UEBuildMac.cs
--- a/Engine/Source/Programs/UnrealBuildTool/Platform/Mac/UEBuildMac.cs
+++ b/Engine/Source/Programs/UnrealBuildTool/Platform/Mac/UEBuildMac.cs
@@ -2,6 +2,7 @@
 
 using System;
 using System.Collections.Generic;
+using System.Text;
 using EpicGames.Core;
 using Microsoft.Extensions.Logging;
 using UnrealBuildBase;
@@ -260,7 +261,23 @@ namespace UnrealBuildTool
 		}
 	}
 
-	class MacPlatform : UEBuildPlatform
+	abstract class AppleBuildPlatform : UEBuildPlatform
+	{
+		public AppleBuildPlatform(UnrealTargetPlatform Platform, UEBuildPlatformSDK SDK, UnrealArchitectureConfig ArchitectureConfig, ILogger Logger)
+			: base(Platform, SDK, ArchitectureConfig, Logger)
+		{
+
+		}
+
+		public override void GetExternalBuildMetadata(FileReference? ProjectFile, StringBuilder Metadata)
+		{
+			base.GetExternalBuildMetadata(ProjectFile, Metadata);
+			
+			Metadata.AppendLine("xcode-select: {0}", AppleToolChainSettings.XcodeDeveloperDir);
+		}
+	}
+
+	class MacPlatform : AppleBuildPlatform
 	{
 		public MacPlatform(UEBuildPlatformSDK InSDK, ILogger InLogger)
 			: base(UnrealTargetPlatform.Mac, InSDK, new MacArchitectureConfig(), InLogger)
