- Couple of changes that got lost for .stub generation (modern was missing the Payload internal directory, and legacy needed to use targetname for the schemename, not project name)

diff --git a/Engine/Source/Programs/UnrealBuildTool/Platform/IOS/IOSToolChain.cs b/Engine/Source/Programs/UnrealBuildTool/Platform/IOS/IOSToolChain.cs
--- a/Engine/Source/Programs/UnrealBuildTool/Platform/IOS/IOSToolChain.cs
+++ b/Engine/Source/Programs/UnrealBuildTool/Platform/IOS/IOSToolChain.cs
@@ -741,7 +741,7 @@ namespace UnrealBuildTool
 			}
 
 			// set up the directories
-			string ZipWorkingDir = String.Format(bUseModernXcode ? "{0}.app/" : "Payload/{0}.app/", GameName);
+			string ZipWorkingDir = String.Format("Payload/{0}.app/", GameName);
 			string ZipSourceDir = String.Format(bUseModernXcode ? "{0}/{1}.app" : "{0}/Payload/{1}.app", BinaryPath, GameName);
 
 			// create the file
@@ -1545,16 +1545,7 @@ namespace UnrealBuildTool
 						ConfigName += " " + Target.TargetType.ToString();
 					}
 
-					string SchemeName;
-					if (AppName == "UnrealGame" || AppName == "UnrealClient")
-					{
-
-						SchemeName = "UnrealGame";
-					}
-					else
-					{
-						SchemeName = Target.TargetName;
-					}
+					string SchemeName = Target.TargetName;
 
 					Console.WriteLine("Provisioning: {0}, {1}, {2}, {3}", MobileProvisionFile, MobileProvisionFile.GetFileName(), MobileProvisionUUID, BundleID);
 
