UnrealBuildTool: Correctly mark extra dynamic libraries, executables, map (for Microsoft platforms) and symbol files in the build output list as their correct BuildProductType, rather than the generic RequiredResource

Reported via UDN

diff --git a/Engine/Source/Programs/UnrealBuildTool/Configuration/UEBuildTarget.cs b/Engine/Source/Programs/UnrealBuildTool/Configuration/UEBuildTarget.cs
--- a/Engine/Source/Programs/UnrealBuildTool/Configuration/UEBuildTarget.cs
+++ b/Engine/Source/Programs/UnrealBuildTool/Configuration/UEBuildTarget.cs
@@ -2642,7 +2642,37 @@ namespace UnrealBuildTool
 				Binary.GetBuildProducts(Rules, TargetToolChain, BinaryBuildProducts, GlobalLinkEnvironment.bCreateDebugInfo);
 				BuildProducts.AddRange(BinaryBuildProducts);
 			}
-			BuildProducts.AddRange(RuntimeDependencyTargetFileToSourceFile.Select(x => new KeyValuePair<FileReference, BuildProductType>(x.Key, BuildProductType.RequiredResource)));
+
+			{
+				UEBuildPlatform BuildPlatform = UEBuildPlatform.GetBuildPlatform(Platform);
+				string DynamicLinkLibraryExt = BuildPlatform.GetBinaryExtension(UEBuildBinaryType.DynamicLinkLibrary);
+				string ExecutableExt = BuildPlatform.GetBinaryExtension(UEBuildBinaryType.Executable);
+				HashSet<string> DebugInfoExtensions = new(BuildPlatform.GetDebugInfoExtensions(Rules, UEBuildBinaryType.DynamicLinkLibrary)
+					.Concat(BuildPlatform.GetDebugInfoExtensions(Rules, UEBuildBinaryType.Executable)));
+				HashSet<string> MapInfoExtensions = Platform.IsInGroup(UnrealPlatformGroup.Microsoft) ? new(new string[]{ ".map" , ".objpaths" }) : new();
+				BuildProducts.AddRange(RuntimeDependencyTargetFileToSourceFile.Select(x =>
+				{
+					string Ext = x.Key.GetExtension();
+					BuildProductType ProductType = BuildProductType.RequiredResource;
+					if (DynamicLinkLibraryExt.Equals(Ext, StringComparison.OrdinalIgnoreCase))
+					{
+						ProductType = BuildProductType.DynamicLibrary;
+					}
+					else if (ExecutableExt.Equals(Ext, StringComparison.OrdinalIgnoreCase))
+					{
+						ProductType = BuildProductType.Executable;
+					}
+					else if (DebugInfoExtensions.Contains(Ext, StringComparer.OrdinalIgnoreCase))
+					{
+						ProductType = BuildProductType.SymbolFile;
+					}
+					else if (MapInfoExtensions.Contains(Ext, StringComparer.OrdinalIgnoreCase))
+					{
+						ProductType = BuildProductType.MapFile;
+					}
+					return new KeyValuePair<FileReference, BuildProductType>(x.Key, ProductType);
+				}));
+			}
 
 			// Remove any installed build products that don't exist. They may be part of an optional install.
 			if (Unreal.IsEngineInstalled())
