UnrealBuildTool: Fix generate clang database not setting disable unity

#rnx

diff --git a/Engine/Source/Programs/UnrealBuildTool/Modes/GenerateClangDatabase.cs b/Engine/Source/Programs/UnrealBuildTool/Modes/GenerateClangDatabase.cs
--- a/Engine/Source/Programs/UnrealBuildTool/Modes/GenerateClangDatabase.cs
+++ b/Engine/Source/Programs/UnrealBuildTool/Modes/GenerateClangDatabase.cs
@@ -145,7 +145,8 @@ namespace UnrealBuildTool
 				foreach (TargetDescriptor TargetDescriptor in TargetDescriptors)
 				{
 					// Disable PCHs and unity builds for the target
-					TargetDescriptor.AdditionalArguments = TargetDescriptor.AdditionalArguments.Append(new string[] { "-NoPCH", "-DisableUnity" });
+					TargetDescriptor.bUseUnityBuild = false;
+					TargetDescriptor.AdditionalArguments = TargetDescriptor.AdditionalArguments.Append(new string[] { "-NoPCH" });
 
 					// Create a makefile for the target
 					UEBuildTarget Target = UEBuildTarget.Create(TargetDescriptor, BuildConfiguration, Logger);
