UnrealBuildTool: Fix unity setting for profile unity sizes mode

#rnx

diff --git a/Engine/Source/Programs/UnrealBuildTool/Modes/ProfileUnitySizesMode.cs b/Engine/Source/Programs/UnrealBuildTool/Modes/ProfileUnitySizesMode.cs
--- a/Engine/Source/Programs/UnrealBuildTool/Modes/ProfileUnitySizesMode.cs
+++ b/Engine/Source/Programs/UnrealBuildTool/Modes/ProfileUnitySizesMode.cs
@@ -322,11 +322,7 @@ namespace UnrealBuildTool
 
 					// Add the module name to the cmdline
 					TargetDescriptor.AdditionalArguments = TargetDescriptor.AdditionalArguments.Append(new string[] { $"-BytesPerUnityCPP={UnitySize}", "-DisableModuleNumIncludedBytesPerUnityCPPOverride" });
-
-					if (bDisableUnity)
-					{
-						TargetDescriptor.AdditionalArguments = TargetDescriptor.AdditionalArguments.Append(new string[] { $"-DisableUnity" });
-					}
+					TargetDescriptor.bUseUnityBuild = !bDisableUnity;
 				}
 
 				using (ISourceFileWorkingSet WorkingSet = new EmptySourceFileWorkingSet())
