UnrealBuildTool: BuildSettingsVersion.V4 enable msvc strict conformance mode by default

diff --git a/Engine/Source/Programs/UnrealBuildTool/Configuration/TargetRules.cs b/Engine/Source/Programs/UnrealBuildTool/Configuration/TargetRules.cs
--- a/Engine/Source/Programs/UnrealBuildTool/Configuration/TargetRules.cs
+++ b/Engine/Source/Programs/UnrealBuildTool/Configuration/TargetRules.cs
@@ -135,6 +135,7 @@ namespace UnrealBuildTool
 		/// <summary>
 		/// New defaults for 5.3:
 		/// * TargetRules.CppStandard = CppStandard.Default has changed from Cpp17 to Cpp20
+		/// * TargetRules.WindowsPlatform.bStrictConformanceMode = true
 		/// </summary>
 		V4,
 
@@ -2860,6 +2861,7 @@ namespace UnrealBuildTool
 				if (BuildSettingsVersion.V4 <= LatestVersion && DefaultBuildSettings < BuildSettingsVersion.V4)
 				{
 					ModifiedSettings.Add(Tuple.Create(String.Format("{0} = CppStandardVersion.Default", nameof(CppStandard)), "Updates C++ Standard to C++20 (Previously: CppStandardVersion.Cpp17)."));
+					ModifiedSettings.Add(Tuple.Create(String.Format("{0}.{1} = true", nameof(WindowsPlatform), nameof(WindowsPlatform.bStrictConformanceMode)), "Updates MSVC strict conformance mode to true (Previously: false)."));
 				}
 
 				if (ModifiedSettings.Count > 0)
diff --git a/Engine/Source/Programs/UnrealBuildTool/Platform/Windows/UEBuildWindows.cs b/Engine/Source/Programs/UnrealBuildTool/Platform/Windows/UEBuildWindows.cs
--- a/Engine/Source/Programs/UnrealBuildTool/Platform/Windows/UEBuildWindows.cs
+++ b/Engine/Source/Programs/UnrealBuildTool/Platform/Windows/UEBuildWindows.cs
@@ -394,7 +394,12 @@ namespace UnrealBuildTool
 		/// </summary>
 		[XmlConfigFile(Category = "WindowsPlatform")]
 		[CommandLine("-Strict")]
-		public bool bStrictConformanceMode = false;
+		public bool bStrictConformanceMode
+		{
+			get => bStrictConformanceModePrivate ?? Target.DefaultBuildSettings >= BuildSettingsVersion.V4;
+			set => bStrictConformanceModePrivate = value;
+		}
+		private bool? bStrictConformanceModePrivate;
 
 		/// <summary>
 		/// Enables updated __cplusplus macro (/Zc:__cplusplus).
