Fixed issue where some executor output might appear below the summary.

diff --git a/Engine/Source/Programs/UnrealBuildTool/Executors/ImmediateActionQueue.cs b/Engine/Source/Programs/UnrealBuildTool/Executors/ImmediateActionQueue.cs
--- a/Engine/Source/Programs/UnrealBuildTool/Executors/ImmediateActionQueue.cs
+++ b/Engine/Source/Programs/UnrealBuildTool/Executors/ImmediateActionQueue.cs
@@ -272,6 +272,11 @@ namespace UnrealBuildTool
 		/// </summary>
 		private readonly Action<string> _writeToolOutput;
 
+		/// <summary>
+		/// Flush the tool output after logging has completed
+		/// </summary>
+		private readonly System.Action _flushToolOutput;
+
 		/// <summary>
 		/// Timer used to collect CPU utilization
 		/// </summary>
@@ -347,8 +352,9 @@ namespace UnrealBuildTool
 		/// <param name="maxActionArtifactCacheTasks">Max number of concurrent artifact cache tasks</param>
 		/// <param name="progressWriterText">Text to be displayed with the progress writer</param>
 		/// <param name="writeToolOutput">Action to invoke when writing tool output</param>
+		/// <param name="flushToolOutput">Action to invoke when flushing tool output</param>
 		/// <param name="logger">Logging interface</param>
-		public ImmediateActionQueue(IEnumerable<LinkedAction> actions, IActionArtifactCache? actionArtifactCache, int maxActionArtifactCacheTasks, string progressWriterText, Action<string> writeToolOutput, ILogger logger)
+		public ImmediateActionQueue(IEnumerable<LinkedAction> actions, IActionArtifactCache? actionArtifactCache, int maxActionArtifactCacheTasks, string progressWriterText, Action<string> writeToolOutput, System.Action flushToolOutput, ILogger logger)
 		{
 			CommandLine.ParseArguments(Environment.GetCommandLineArgs(), this, logger);
 			XmlConfig.ApplyTo(this);
@@ -360,6 +366,7 @@ namespace UnrealBuildTool
 			ProgressWriter = new(progressWriterText, false, logger);
 			_actionArtifactCache = actionArtifactCache;
 			_writeToolOutput = writeToolOutput;
+			_flushToolOutput = flushToolOutput;
 
 			bool readArtifacts = _actionArtifactCache != null && _actionArtifactCache.EnableReads;
 			ActionPhase initialPhase = readArtifacts ? ActionPhase.ArtifactCheck : ActionPhase.Compile;
@@ -1120,6 +1127,7 @@ namespace UnrealBuildTool
 				loggingTask = _actionsToLogTask;
 			}
 			loggingTask?.Wait();
+			_flushToolOutput();
 
 			if (ShowCPUUtilization)
 			{
diff --git a/Engine/Source/Programs/UnrealBuildTool/Executors/ParallelExecutor.cs b/Engine/Source/Programs/UnrealBuildTool/Executors/ParallelExecutor.cs
--- a/Engine/Source/Programs/UnrealBuildTool/Executors/ParallelExecutor.cs
+++ b/Engine/Source/Programs/UnrealBuildTool/Executors/ParallelExecutor.cs
@@ -158,7 +158,7 @@ namespace UnrealBuildTool
 		/// <returns>Action queue</returns>
 		public ImmediateActionQueue CreateActionQueue(IEnumerable<LinkedAction> actionsToExecute, IActionArtifactCache? actionArtifactCache, ILogger logger)
 		{
-			return new(actionsToExecute, actionArtifactCache, NumParallelProcesses, "Compiling C++ source code...", x => WriteToolOutput(x), logger)
+			return new(actionsToExecute, actionArtifactCache, NumParallelProcesses, "Compiling C++ source code...", x => WriteToolOutput(x), () => FlushToolOutput(), logger)
 			{
 				ShowCompilationTimes = bShowCompilationTimes,
 				ShowCPUUtilization = bShowCPUUtilization,
