- Removed static lib version checking because it's not as useful now, and cannot work with VisionOS libs

diff --git a/Engine/Source/Programs/UnrealBuildTool/Platform/IOS/UEBuildIOS.cs b/Engine/Source/Programs/UnrealBuildTool/Platform/IOS/UEBuildIOS.cs
--- a/Engine/Source/Programs/UnrealBuildTool/Platform/IOS/UEBuildIOS.cs
+++ b/Engine/Source/Programs/UnrealBuildTool/Platform/IOS/UEBuildIOS.cs
@@ -770,6 +770,8 @@ namespace UnrealBuildTool
 		{
 			base.ValidateModule(Module, Target);
 
+			// @todo temporarily disabling due to VisionOS rquiring newer Xcode than this will allow - we may remove this entirely
+#if false
 			if (BuildHostPlatform.Current.Platform == UnrealTargetPlatform.Mac && !Target.IOSPlatform.bSkipClangValidation)
 			{
 				IOSPlatformSDK SDK = (IOSPlatformSDK?)GetSDK() ?? new IOSPlatformSDK(Logger);
@@ -811,6 +813,7 @@ namespace UnrealBuildTool
 					}
 				}
 			}
+#endif
 		}
 
 		/// <summary>
