[UBT]
* Fixed so -rebuild works again.. did not fix the code fully because it was broken before and did not delete the "Inc" folder because the code produced the wrong path. I don't want to fix this because deleting that path would make -rebuild of non-unity also force rebuild of unity version

diff --git a/Engine/Source/Programs/UnrealBuildTool/Modes/CleanMode.cs b/Engine/Source/Programs/UnrealBuildTool/Modes/CleanMode.cs
--- a/Engine/Source/Programs/UnrealBuildTool/Modes/CleanMode.cs
+++ b/Engine/Source/Programs/UnrealBuildTool/Modes/CleanMode.cs
@@ -149,13 +149,16 @@ namespace UnrealBuildTool
 					{
 						foreach (string NamePrefix in NamePrefixes)
 						{
-							DirectoryReference GeneratedCodeDir = DirectoryReference.Combine(BaseDir, UEBuildTarget.GetPlatformIntermediateFolder(Target.Platform, Target.Architectures, false), NamePrefix, "Inc");
+							string NamePrefixWithEnv = UEBuildTarget.GetTargetIntermediateFolderName(NamePrefix, TargetDescriptor.IntermediateEnvironment);
+							// This is actually wrong.. the generated code is not in this dir.. if changing "Target.Architectures" parameter to null it will delete the right files.
+							// However, this also means that it will cause a rebuild for both unity, nonunity and iwyu targets since they share this folder.
+							DirectoryReference GeneratedCodeDir = DirectoryReference.Combine(BaseDir, UEBuildTarget.GetPlatformIntermediateFolder(Target.Platform, Target.Architectures, false), NamePrefixWithEnv, "Inc");
 							if (DirectoryReference.Exists(GeneratedCodeDir))
 							{
 								DirectoriesToDelete.Add(GeneratedCodeDir);
 							}
 
-							DirectoryReference IntermediateDir = DirectoryReference.Combine(BaseDir, UEBuildTarget.GetPlatformIntermediateFolder(Target.Platform, Target.Architectures, false), NamePrefix, Target.Configuration.ToString());
+							DirectoryReference IntermediateDir = DirectoryReference.Combine(BaseDir, UEBuildTarget.GetPlatformIntermediateFolder(Target.Platform, Target.Architectures, false), NamePrefixWithEnv, Target.Configuration.ToString());
 							if (DirectoryReference.Exists(IntermediateDir))
 							{
 								DirectoriesToDelete.Add(IntermediateDir);
