Add natvis files to build outputs and manifest to ensure they are copied to intermediate directory during precompile and tagged in horde builds.

diff --git a/Engine/Source/Programs/UnrealBuildTool/Configuration/UEBuildBinary.cs b/Engine/Source/Programs/UnrealBuildTool/Configuration/UEBuildBinary.cs
--- a/Engine/Source/Programs/UnrealBuildTool/Configuration/UEBuildBinary.cs
+++ b/Engine/Source/Programs/UnrealBuildTool/Configuration/UEBuildBinary.cs
@@ -469,6 +469,14 @@ namespace UnrealBuildTool
 							}
 						}
 					}
+					if (Module.NatvisSourceFile != null)
+					{
+						FileItem? Item = ToolChain.LinkDebuggerVisualizer(Module.NatvisSourceFile, Module.IntermediateDirectory);
+						if (Item != null)
+						{
+							BuildProducts.Add(Item.Location, BuildProductType.BuildResource);
+						}
+					}
 					if (Target.LinkType == TargetLinkType.Monolithic)
 					{
 						FileReference PrecompiledManifestLocation = Module.PrecompiledManifestLocation;
diff --git a/Engine/Source/Programs/UnrealBuildTool/Configuration/UEBuildModule.cs b/Engine/Source/Programs/UnrealBuildTool/Configuration/UEBuildModule.cs
--- a/Engine/Source/Programs/UnrealBuildTool/Configuration/UEBuildModule.cs
+++ b/Engine/Source/Programs/UnrealBuildTool/Configuration/UEBuildModule.cs
@@ -938,14 +938,19 @@ namespace UnrealBuildTool
 			return new List<FileItem>();
 		}
 
-		public void CopyDebuggerVisualizers(UEToolChain ToolChain, IActionGraphBuilder Graph, ILogger Logger)
+		public IEnumerable<FileItem> CopyDebuggerVisualizers(UEToolChain ToolChain, IActionGraphBuilder Graph, ILogger Logger)
 		{
 			if (NatvisSourceFile != null)
 			{
-				ToolChain.CopyDebuggerVisualizer(NatvisSourceFile, IntermediateDirectory, Graph);
+				FileItem? Item = ToolChain.CopyDebuggerVisualizer(NatvisSourceFile, IntermediateDirectory, Graph);
+				if (Item != null)
+				{
+					return new[] { Item }; 
+				}
 			}
+			return Enumerable.Empty<FileItem>();
 		}
-
+		
 		public void LinkDebuggerVisualizers(List<FileItem> OutFiles, UEToolChain ToolChain, ILogger Logger)
 		{
 			if (NatvisSourceFile != null)
diff --git a/Engine/Source/Programs/UnrealBuildTool/Configuration/UEBuildModuleCPP.cs b/Engine/Source/Programs/UnrealBuildTool/Configuration/UEBuildModuleCPP.cs
--- a/Engine/Source/Programs/UnrealBuildTool/Configuration/UEBuildModuleCPP.cs
+++ b/Engine/Source/Programs/UnrealBuildTool/Configuration/UEBuildModuleCPP.cs
@@ -724,6 +724,14 @@ namespace UnrealBuildTool
 
 				PrecompiledManifest Manifest = new PrecompiledManifest();
 				Manifest.OutputFiles.AddRange(LinkInputFiles.Select(x => x.Location));
+				if (NatvisSourceFile != null) 
+				{
+					FileItem? Item = ToolChain.LinkDebuggerVisualizer(NatvisSourceFile, IntermediateDirectory);
+					if (Item != null)
+					{
+						Manifest.OutputFiles.Add(Item.Location);
+					}
+				}	
 				Manifest.WriteIfModified(PrecompiledManifestLocation);
 			}
 
diff --git a/Engine/Source/Programs/UnrealBuildTool/Configuration/UEBuildTarget.cs b/Engine/Source/Programs/UnrealBuildTool/Configuration/UEBuildTarget.cs
--- a/Engine/Source/Programs/UnrealBuildTool/Configuration/UEBuildTarget.cs
+++ b/Engine/Source/Programs/UnrealBuildTool/Configuration/UEBuildTarget.cs
@@ -2566,7 +2566,8 @@ namespace UnrealBuildTool
 			// Copy debugger visualizer files for each module to their intermediate directory 
 			foreach (UEBuildModule Module in Modules.Values)
 			{
-				Module.CopyDebuggerVisualizers(TargetToolChain, MakefileBuilder, Logger);
+				IEnumerable<FileItem> Items = Module.CopyDebuggerVisualizers(TargetToolChain, MakefileBuilder, Logger);
+				Makefile.OutputItems.AddRange(Items);
 			}
 
 			// Build the target's binaries.
