Support for .natstepfilter

#rnx

diff --git a/Engine/Source/Programs/UnrealBuildTool/Configuration/UEBuildBinary.cs b/Engine/Source/Programs/UnrealBuildTool/Configuration/UEBuildBinary.cs
--- a/Engine/Source/Programs/UnrealBuildTool/Configuration/UEBuildBinary.cs
+++ b/Engine/Source/Programs/UnrealBuildTool/Configuration/UEBuildBinary.cs
@@ -469,9 +469,9 @@ namespace UnrealBuildTool
 							}
 						}
 					}
-					if (Module.NatvisSourceFile != null)
+					foreach (FileItem NatvisSourceFile in Module.NatvisFiles)
 					{
-						FileItem? Item = ToolChain.LinkDebuggerVisualizer(Module.NatvisSourceFile, Module.IntermediateDirectory);
+						FileItem? Item = ToolChain.LinkDebuggerVisualizer(NatvisSourceFile, Module.IntermediateDirectory);
 						if (Item != null)
 						{
 							BuildProducts.Add(Item.Location, BuildProductType.BuildResource);
diff --git a/Engine/Source/Programs/UnrealBuildTool/Configuration/UEBuildModule.cs b/Engine/Source/Programs/UnrealBuildTool/Configuration/UEBuildModule.cs
--- a/Engine/Source/Programs/UnrealBuildTool/Configuration/UEBuildModule.cs
+++ b/Engine/Source/Programs/UnrealBuildTool/Configuration/UEBuildModule.cs
@@ -192,7 +192,7 @@ namespace UnrealBuildTool
 		/// </summary>
 		public bool bDependsOnVerse = false;
 
-		public FileItem? NatvisSourceFile;
+		public List<FileItem> NatvisFiles = new();
 
 		/// <summary>
 		/// Constructor
@@ -290,14 +290,14 @@ namespace UnrealBuildTool
 			AliasRestrictedFolders = new Dictionary<string, string>(Rules.AliasRestrictedFolders);
 
 			DependenciesToSkipPerArchitecture = Rules.DependenciesToSkipPerArchitecture;
-			FileReference PossibleNatvisFile = FileReference.Combine(Rules.File.Directory, $"{Name}.natvis");
-			if (FileReference.Exists(PossibleNatvisFile))
-			{
-				NatvisSourceFile = FileItem.GetItemByFileReference(PossibleNatvisFile);
-			}
 
 			// get the module directories from the module
 			ModuleDirectories = Rules.GetAllModuleDirectories();
+
+			foreach (DirectoryItem Directory in ModuleDirectories.Select(x => DirectoryItem.GetItemByDirectoryReference(x)))
+			{
+				NatvisFiles.AddRange(Directory.EnumerateFiles().Where(x => x.HasExtension(".natvis") || x.HasExtension(".natstepfilter")));
+			}
 		}
 
 		/// <summary>
@@ -940,20 +940,21 @@ namespace UnrealBuildTool
 
 		public IEnumerable<FileItem> CopyDebuggerVisualizers(UEToolChain ToolChain, IActionGraphBuilder Graph, ILogger Logger)
 		{
-			if (NatvisSourceFile != null)
+			List<FileItem> Results = new();
+			foreach (FileItem NatvisSourceFile in NatvisFiles)
 			{
 				FileItem? Item = ToolChain.CopyDebuggerVisualizer(NatvisSourceFile, IntermediateDirectory, Graph);
 				if (Item != null)
 				{
-					return new[] { Item }; 
+					Results.Add(Item);
 				}
 			}
-			return Enumerable.Empty<FileItem>();
+			return Results;
 		}
 		
 		public void LinkDebuggerVisualizers(List<FileItem> OutFiles, UEToolChain ToolChain, ILogger Logger)
 		{
-			if (NatvisSourceFile != null)
+			foreach (FileItem NatvisSourceFile in NatvisFiles)
 			{
 				FileItem? Item = ToolChain.LinkDebuggerVisualizer(NatvisSourceFile, IntermediateDirectory);
 				if (Item != null)
diff --git a/Engine/Source/Programs/UnrealBuildTool/Configuration/UEBuildModuleCPP.cs b/Engine/Source/Programs/UnrealBuildTool/Configuration/UEBuildModuleCPP.cs
--- a/Engine/Source/Programs/UnrealBuildTool/Configuration/UEBuildModuleCPP.cs
+++ b/Engine/Source/Programs/UnrealBuildTool/Configuration/UEBuildModuleCPP.cs
@@ -724,7 +724,7 @@ namespace UnrealBuildTool
 
 				PrecompiledManifest Manifest = new PrecompiledManifest();
 				Manifest.OutputFiles.AddRange(LinkInputFiles.Select(x => x.Location));
-				if (NatvisSourceFile != null) 
+				foreach (FileItem NatvisSourceFile in NatvisFiles)
 				{
 					FileItem? Item = ToolChain.LinkDebuggerVisualizer(NatvisSourceFile, IntermediateDirectory);
 					if (Item != null)
diff --git a/Engine/Source/Programs/UnrealBuildTool/Platform/Windows/UEBuildWindows.cs b/Engine/Source/Programs/UnrealBuildTool/Platform/Windows/UEBuildWindows.cs
--- a/Engine/Source/Programs/UnrealBuildTool/Platform/Windows/UEBuildWindows.cs
+++ b/Engine/Source/Programs/UnrealBuildTool/Platform/Windows/UEBuildWindows.cs
@@ -1258,7 +1258,9 @@ namespace UnrealBuildTool
 				|| IsBuildProductName(FileName, NamePrefixes, NameSuffixes, ".exp")
 				|| IsBuildProductName(FileName, NamePrefixes, NameSuffixes, ".obj")
 				|| IsBuildProductName(FileName, NamePrefixes, NameSuffixes, ".map")
-				|| IsBuildProductName(FileName, NamePrefixes, NameSuffixes, ".objpaths");
+				|| IsBuildProductName(FileName, NamePrefixes, NameSuffixes, ".objpaths")
+				|| IsBuildProductName(FileName, NamePrefixes, NameSuffixes, ".natvis")
+				|| IsBuildProductName(FileName, NamePrefixes, NameSuffixes, ".natstepfilter");
 		}
 
 		/// <summary>
diff --git a/Engine/Source/Programs/UnrealBuildTool/Platform/Windows/VCToolChain.cs b/Engine/Source/Programs/UnrealBuildTool/Platform/Windows/VCToolChain.cs
--- a/Engine/Source/Programs/UnrealBuildTool/Platform/Windows/VCToolChain.cs
+++ b/Engine/Source/Programs/UnrealBuildTool/Platform/Windows/VCToolChain.cs
@@ -1389,7 +1389,7 @@ namespace UnrealBuildTool
 
 		public override FileItem? CopyDebuggerVisualizer(FileItem SourceFile, DirectoryReference IntermediateDirectory, IActionGraphBuilder Graph)
 		{
-			if (SourceFile.HasExtension(".natvis"))
+			if (SourceFile.HasExtension(".natvis") || SourceFile.HasExtension(".natstepfilter"))
 			{
 				FileItem IntermediateFile = FileItem.GetItemByFileReference(FileReference.Combine(IntermediateDirectory, SourceFile.Name));
 				Graph.CreateCopyAction(SourceFile, IntermediateFile);
@@ -1400,7 +1400,7 @@ namespace UnrealBuildTool
 
 		public override FileItem? LinkDebuggerVisualizer(FileItem SourceFile, DirectoryReference IntermediateDirectory)
 		{
-			if (SourceFile.HasExtension(".natvis"))
+			if (SourceFile.HasExtension(".natvis") || SourceFile.HasExtension(".natstepfilter"))
 			{
 				FileItem IntermediateFile = FileItem.GetItemByFileReference(FileReference.Combine(IntermediateDirectory, SourceFile.Name));
 				return IntermediateFile;
diff --git a/Engine/Source/Programs/UnrealBuildTool/ProjectFiles/VisualStudio/VCProjectFileGenerator.cs b/Engine/Source/Programs/UnrealBuildTool/ProjectFiles/VisualStudio/VCProjectFileGenerator.cs
--- a/Engine/Source/Programs/UnrealBuildTool/ProjectFiles/VisualStudio/VCProjectFileGenerator.cs
+++ b/Engine/Source/Programs/UnrealBuildTool/ProjectFiles/VisualStudio/VCProjectFileGenerator.cs
@@ -763,11 +763,13 @@ namespace UnrealBuildTool
 
 				// Get the path to the visualizers file. Try to make it relative to the solution directory, but fall back to a full path if it's a foreign project.
 				FileReference VisualizersFile = FileReference.Combine(Unreal.EngineDirectory, "Extras", "VisualStudioDebugging", "Unreal.natvis");
+				FileReference VisualizersStepFile = FileReference.Combine(Unreal.EngineDirectory, "Extras", "VisualStudioDebugging", "Unreal.natstepfilter");
 
 				// Add the visualizers at the solution level. Doesn't seem to be picked up from a makefile project in VS2017 15.8.5.
 				VCSolutionFileContent.AppendLine(String.Format("Project(\"{0}\") = \"Visualizers\", \"Visualizers\", \"{{1CCEC849-CC72-4C59-8C36-2F7C38706D4C}}\"", SolutionFolderEntryGUID));
 				VCSolutionFileContent.AppendLine("\tProjectSection(SolutionItems) = preProject");
 				VCSolutionFileContent.AppendLine("\t\t{0} = {0}", VisualizersFile.MakeRelativeTo(PrimaryProjectPath));
+				VCSolutionFileContent.AppendLine("\t\t{0} = {0}", VisualizersStepFile.MakeRelativeTo(PrimaryProjectPath));
 				VCSolutionFileContent.AppendLine("\tEndProjectSection");
 				VCSolutionFileContent.AppendLine("EndProject");
 			}
diff --git a/Engine/Source/Programs/UnrealBuildTool/System/LinkEnvironment.cs b/Engine/Source/Programs/UnrealBuildTool/System/LinkEnvironment.cs
--- a/Engine/Source/Programs/UnrealBuildTool/System/LinkEnvironment.cs
+++ b/Engine/Source/Programs/UnrealBuildTool/System/LinkEnvironment.cs
@@ -431,6 +431,7 @@ namespace UnrealBuildTool
 			DefaultResourceFiles.AddRange(Other.DefaultResourceFiles);
 			CommonResourceFiles.AddRange(Other.CommonResourceFiles);
 			IncludeFunctions.AddRange(Other.IncludeFunctions);
+			DebuggerVisualizerFiles.AddRange(Other.DebuggerVisualizerFiles);
 			ModuleDefinitionFile = Other.ModuleDefinitionFile;
 			AdditionalProperties.AddRange(Other.AdditionalProperties);
 
