Enable profile guided optimization on Intel ICX compiler 2023.1

diff --git a/Engine/Source/Programs/UnrealBuildTool/Platform/Windows/UEBuildWindows.cs b/Engine/Source/Programs/UnrealBuildTool/Platform/Windows/UEBuildWindows.cs
--- a/Engine/Source/Programs/UnrealBuildTool/Platform/Windows/UEBuildWindows.cs
+++ b/Engine/Source/Programs/UnrealBuildTool/Platform/Windows/UEBuildWindows.cs
@@ -168,6 +168,16 @@ namespace UnrealBuildTool
 			return Compiler == WindowsCompiler.Clang || Compiler == WindowsCompiler.ClangRTFM || Compiler == WindowsCompiler.Intel;
 		}
 
+		/// <summary>
+		/// Returns if this compiler toolchain based on Intel
+		/// </summary>
+		/// <param name="Compiler">The compiler to check</param>
+		/// <returns>true if Intel based</returns>
+		public static bool IsIntel(this WindowsCompiler Compiler)
+		{
+			return Compiler == WindowsCompiler.Intel;
+		}
+
 		/// <summary>
 		/// Returns if this compiler toolchain based on MSVC
 		/// </summary>
diff --git a/Engine/Source/Programs/UnrealBuildTool/Platform/Windows/VCToolChain.cs b/Engine/Source/Programs/UnrealBuildTool/Platform/Windows/VCToolChain.cs
--- a/Engine/Source/Programs/UnrealBuildTool/Platform/Windows/VCToolChain.cs
+++ b/Engine/Source/Programs/UnrealBuildTool/Platform/Windows/VCToolChain.cs
@@ -672,12 +672,33 @@ namespace UnrealBuildTool
 				CompileEnvironment.bPGOProfile ||
 				CompileEnvironment.bPGOOptimize ||
 				CompileEnvironment.bAllowLTCG;
-			if (bEnableLTCG)
+			if (bEnableLTCG && !Target.WindowsPlatform.Compiler.IsIntel())
 			{
 				// Enable link-time code generation.
 				Arguments.Add("/GL");
 			}
 
+			if (Target.WindowsPlatform.Compiler.IsIntel())
+			{
+				if (CompileEnvironment.bAllowLTCG)
+				{
+					// Enable link-time code generation.
+					Arguments.Add("-flto");
+				}
+
+				if (CompileEnvironment.bPGOProfile)
+				{
+					// Generate instrumented code.
+					Arguments.Add($"-fprofile-instr-generate=\"{CompileEnvironment.PGOFilenamePrefix}\"-%p-%m.profraw");
+				}
+				else if (CompileEnvironment.bPGOOptimize)
+				{
+					// Use a merged profdata file.
+					Arguments.Add("-Wno-profile-instr-out-of-date");
+					Arguments.Add($"-fprofile-instr-use=\"{Path.Combine(CompileEnvironment.PGODirectory!, CompileEnvironment.PGOFilenamePrefix!)}\".profdata");
+				}
+			}
+
 			//
 			//	PC
 			//
@@ -1861,7 +1882,7 @@ namespace UnrealBuildTool
 			}
 
 			// copy PGO support binaries for Windows from $(VC_PGO_RunTime_Dir)
-			if (Target.bPGOProfile && Target.Platform.IsInGroup(UnrealPlatformGroup.Windows))
+			if (Target.bPGOProfile && Target.Platform.IsInGroup(UnrealPlatformGroup.Windows) && !Target.WindowsPlatform.Compiler.IsIntel())
 			{
 				string[] PGOFiles = {
 					"pgort140.dll",
@@ -2594,91 +2615,106 @@ namespace UnrealBuildTool
 			bool bPGOOptimize = LinkEnvironment.bPGOOptimize;
 			bool bPGOProfile = LinkEnvironment.bPGOProfile;
 
-			if (bPGOOptimize || bPGOProfile)
+			if (!Target.WindowsPlatform.Compiler.IsIntel())
 			{
-				// serialize the linker arguments for comparing /USEPROFILE and /GENPROFILE
-				IEnumerable<string> PGOArguments = RemovePGOIgnoredSwitches(Arguments);
-				string PGOProfileLinkArgsFile = Path.Combine(LinkEnvironment.PGODirectory!, "PGOProfileLinkArgs.txt");
-
-				if (bPGOProfile)
-				{
-					// save the latest linker arguments for /GENPROFILE alongside the PGO data
-					Utils.WriteFileIfChanged(new FileReference(PGOProfileLinkArgsFile), PGOArguments, Logger);
-				}
-				else if (bPGOOptimize && Target.WindowsPlatform.bIgnoreStalePGOData && File.Exists(PGOProfileLinkArgsFile))
+				if (bPGOOptimize || bPGOProfile)
 				{
-					// compare latest linker arguments for /USEPROFILE to the last known /GENPROFILE and disable /USEPROFILE if they are different to avoid LNK1268
-					IEnumerable<string> PGOProfileLinkArgs = RemovePGOIgnoredSwitches(File.ReadAllLines(PGOProfileLinkArgsFile));
-					IEnumerable<string> AddedPGOArgs = PGOArguments.Except(PGOProfileLinkArgs);
-					IEnumerable<string> RemovedPGOArgs = PGOProfileLinkArgs.Except(PGOArguments);
+					// serialize the linker arguments for comparing /USEPROFILE and /GENPROFILE
+					IEnumerable<string> PGOArguments = RemovePGOIgnoredSwitches(Arguments);
+					string PGOProfileLinkArgsFile = Path.Combine(LinkEnvironment.PGODirectory!, "PGOProfileLinkArgs.txt");
 
-					if (AddedPGOArgs.Any() || RemovedPGOArgs.Any())
+					if (bPGOProfile)
 					{
-						Logger.LogWarning("PGO Profile and PGO Optimize linker arguments do not match. Profile Guided Optimization will be disabled for {Config} until new PGO Profile data is generated.", Target.Configuration.ToString());
-						bPGOOptimize = false;
+						// save the latest linker arguments for /GENPROFILE alongside the PGO data
+						Utils.WriteFileIfChanged(new FileReference(PGOProfileLinkArgsFile), PGOArguments, Logger);
+					}
+					else if (bPGOOptimize && Target.WindowsPlatform.bIgnoreStalePGOData && File.Exists(PGOProfileLinkArgsFile))
+					{
+						// compare latest linker arguments for /USEPROFILE to the last known /GENPROFILE and disable /USEPROFILE if they are different to avoid LNK1268
+						IEnumerable<string> PGOProfileLinkArgs = RemovePGOIgnoredSwitches(File.ReadAllLines(PGOProfileLinkArgsFile));
+						IEnumerable<string> AddedPGOArgs = PGOArguments.Except(PGOProfileLinkArgs);
+						IEnumerable<string> RemovedPGOArgs = PGOProfileLinkArgs.Except(PGOArguments);
 
-						if (AddedPGOArgs.Any())
+						if (AddedPGOArgs.Any() || RemovedPGOArgs.Any())
 						{
-							Logger.LogInformation("Added PGO args:");
-							foreach (string Arg in AddedPGOArgs)
+							Logger.LogWarning("PGO Profile and PGO Optimize linker arguments do not match. Profile Guided Optimization will be disabled for {Config} until new PGO Profile data is generated.", Target.Configuration.ToString());
+							bPGOOptimize = false;
+
+							if (AddedPGOArgs.Any())
 							{
-								Logger.LogInformation("  {Arg}", Arg);
+								Logger.LogInformation("Added PGO args:");
+								foreach (string Arg in AddedPGOArgs)
+								{
+									Logger.LogInformation("  {Arg}", Arg);
+								}
 							}
-						}
-						if (RemovedPGOArgs.Any())
-						{
-							Logger.LogInformation("Removed PGO args:");
-							foreach (string Arg in RemovedPGOArgs)
+							if (RemovedPGOArgs.Any())
 							{
-								Logger.LogInformation("  {Arg}", Arg);
+								Logger.LogInformation("Removed PGO args:");
+								foreach (string Arg in RemovedPGOArgs)
+								{
+									Logger.LogInformation("  {Arg}", Arg);
+								}
 							}
 						}
 					}
 				}
-			}
 
-			// If PGO was requested, apply Link-time code generation even if PGO was disabled due to mismatched args
-			if (LinkEnvironment.bPGOOptimize || LinkEnvironment.bPGOProfile)
-			{
-				if (!Arguments.Contains("/LTCG"))
+				// If PGO was requested, apply Link-time code generation even if PGO was disabled due to mismatched args
+				if (LinkEnvironment.bPGOOptimize || LinkEnvironment.bPGOProfile)
 				{
-					Arguments.Add("/LTCG");
+					if (!Arguments.Contains("/LTCG"))
+					{
+						Arguments.Add("/LTCG");
+					}
+					Log.TraceInformationOnce("Enabling Link-time code generation (LTGC) as Profile Guided Optimization (PGO) was requested. Linking will take a while.");
 				}
-				Log.TraceInformationOnce("Enabling Link-time code generation (LTGC) as Profile Guided Optimization (PGO) was requested. Linking will take a while.");
-			}
 
-			if (bPGOOptimize)
-			{
-				if (PreparePGOFiles(LinkEnvironment))
+				if (bPGOOptimize)
 				{
-					Arguments.Add("/USEPROFILE");
-					Log.TraceInformationOnce("Enabling using Profile Guided Optimization (PGO). Linking will take a while.");
+					if (PreparePGOFiles(LinkEnvironment))
+					{
+						Arguments.Add("/USEPROFILE");
+						Log.TraceInformationOnce("Enabling using Profile Guided Optimization (PGO). Linking will take a while.");
+					}
+					else
+					{
+						Logger.LogWarning("Unable to prepare Profile Guided Optimization (PGO) files. Using PGO Optimize build will be disabled.");
+						bPGOOptimize = false;
+					}
 				}
-				else
+				else if (bPGOProfile)
 				{
-					Logger.LogWarning("Unable to prepare Profile Guided Optimization (PGO) files. Using PGO Optimize build will be disabled.");
-					bPGOOptimize = false;
+					if (Target.WindowsPlatform.bUseFastGenProfile)
+					{
+						Log.TraceInformationOnce("Enabling generating Fastgen Profile Guided Optimization (PGO). Linking will take a while.");
+						Arguments.Add("/FASTGENPROFILE");
+					}
+					else
+					{
+						if (Target.WindowsPlatform.bPGONoExtraCounters)
+						{
+							Log.TraceInformationOnce("Enabling generating Profile Guided Optimization No Extra Counters (PGO). Linking will take a while.");
+							Arguments.Add("/GENPROFILE:NOTRACKEH");
+						}
+						else
+						{
+							Log.TraceInformationOnce("Enabling generating Profile Guided Optimization (PGO). Linking will take a while.");
+							Arguments.Add("/GENPROFILE");
+						}
+					}
 				}
 			}
-			else if (bPGOProfile)
+			else
 			{
-				if (Target.WindowsPlatform.bUseFastGenProfile)
+				// No link arguments used for PGO on Intel compiler
+				if (bPGOOptimize)
 				{
-					Log.TraceInformationOnce("Enabling generating Fastgen Profile Guided Optimization (PGO). Linking will take a while.");
-					Arguments.Add("/FASTGENPROFILE");
+					Log.TraceInformationOnce("Enabling Profile Guided Optimization (PGO) on Intel Compiler. Linking will take a while.");
 				}
-				else
+				else if (bPGOProfile)
 				{
-					if (Target.WindowsPlatform.bPGONoExtraCounters)
-					{
-						Log.TraceInformationOnce("Enabling generating Profile Guided Optimization No Extra Counters (PGO). Linking will take a while.");
-						Arguments.Add("/GENPROFILE:NOTRACKEH");
-					}
-					else
-					{
-						Log.TraceInformationOnce("Enabling generating Profile Guided Optimization (PGO). Linking will take a while.");
-						Arguments.Add("/GENPROFILE");
-					}
+					Log.TraceInformationOnce("Enabling generating Profile Guided Optimization (PGO) on Intel Compiler. Linking will take a while.");
 				}
 			}
 		}
diff --git a/Engine/Source/Runtime/Core/Private/Windows/WindowsPlatformMisc.cpp b/Engine/Source/Runtime/Core/Private/Windows/WindowsPlatformMisc.cpp
--- a/Engine/Source/Runtime/Core/Private/Windows/WindowsPlatformMisc.cpp
+++ b/Engine/Source/Runtime/Core/Private/Windows/WindowsPlatformMisc.cpp
@@ -440,7 +440,7 @@ static void PureCallHandler()
 	}
 }
 
-#if ENABLE_PGO_PROFILE
+#if ENABLE_PGO_PROFILE && !defined(__INTEL_LLVM_COMPILER)
 void PGO_WriteFile()
 {
 	// NB. Using pgosweep.exe means the PGC file will be writable as soon as the title exits & we can control where it is written.
@@ -483,7 +483,7 @@ void PGO_WriteFile()
 		UE_LOG(LogWindows, Log, TEXT("pgosweep.exe exit code %d"), ExitCode);
 	}
 }
-#endif //ENABLE_PGO_PROFILE
+#endif //ENABLE_PGO_PROFILE && !__INTEL_LLVM_COMPILER
 
 
 /*-----------------------------------------------------------------------------
@@ -1018,7 +1018,7 @@ void FWindowsPlatformMisc::RequestExitWithStatus(bool Force, uint8 ReturnCode, c
 	UE_LOG(LogWindows, Log, TEXT("FPlatformMisc::RequestExitWithStatus(%i, %i, %s)"), Force, ReturnCode,
 		CallSite ? CallSite : TEXT("<NoCallSiteInfo>"));
 
-#if ENABLE_PGO_PROFILE
+#if ENABLE_PGO_PROFILE && !defined(__INTEL_LLVM_COMPILER)
 	// save current PGO profiling data and terminate immediately
 	PGO_WriteFile();
 	TerminateProcess(GetCurrentProcess(), 0);
