Don't copy debugger files to intermediate when building with installed engine if the file already exists.

diff --git a/Engine/Source/Programs/UnrealBuildTool/Platform/Windows/VCToolChain.cs b/Engine/Source/Programs/UnrealBuildTool/Platform/Windows/VCToolChain.cs
--- a/Engine/Source/Programs/UnrealBuildTool/Platform/Windows/VCToolChain.cs
+++ b/Engine/Source/Programs/UnrealBuildTool/Platform/Windows/VCToolChain.cs
@@ -1412,9 +1412,13 @@ namespace UnrealBuildTool
 		{
 			if (SourceFile.HasExtension(".natvis") || SourceFile.HasExtension(".natstepfilter"))
 			{
-				FileItem IntermediateFile = FileItem.GetItemByFileReference(FileReference.Combine(IntermediateDirectory, SourceFile.Name));
-				Graph.CreateCopyAction(SourceFile, IntermediateFile);
-				return IntermediateFile;
+				FileReference IntermediateFile = FileReference.Combine(IntermediateDirectory, SourceFile.Name);
+				if (!UnrealBuildTool.IsFileInstalled(IntermediateFile))
+				{
+					FileItem Item = FileItem.GetItemByFileReference(IntermediateFile);
+					Graph.CreateCopyAction(SourceFile, Item);
+					return Item; // Only return the item if we are responsible for copying it, for addition to makefile/manifests
+				}
 			}
 			return null;
 		}
