UnrealBuildTool: Set min cl-clang to 15 to match min linux version

#rnx

diff --git a/Engine/Source/Programs/UnrealBuildTool/Platform/Windows/MicrosoftPlatformSDK.Versions.cs b/Engine/Source/Programs/UnrealBuildTool/Platform/Windows/MicrosoftPlatformSDK.Versions.cs
--- a/Engine/Source/Programs/UnrealBuildTool/Platform/Windows/MicrosoftPlatformSDK.Versions.cs
+++ b/Engine/Source/Programs/UnrealBuildTool/Platform/Windows/MicrosoftPlatformSDK.Versions.cs
@@ -30,14 +30,12 @@ namespace UnrealBuildTool
 		{
 			VersionNumberRange.Parse("16.0.0", "16.999"), // VS2022 17.7.x runtime requires Clang 16
 			VersionNumberRange.Parse("15.0.0", "15.999"), // VS2022 17.5.x runtime requires Clang 15
-			VersionNumberRange.Parse("14.0.0", "14.999"), // VS2022 17.4.x runtime requires Clang 14
-			VersionNumberRange.Parse("13.0.0", "13.999"), // VS2019 16.11 runtime requires Clang 13
 		};
 
 		/// <summary>
 		/// The minimum supported Clang compiler
 		/// </summary>
-		static readonly VersionNumber MinimumClangVersion = new VersionNumber(13, 0, 0);
+		static readonly VersionNumber MinimumClangVersion = new VersionNumber(15, 0, 0);
 
 		/// <summary>
 		/// Ranges of tested compiler toolchains to be used, in order of preference. If multiple toolchains in a range are present, the latest version will be preferred.
