Add exported=false to GameActivity for splashscreen case
Include symbol formatter for Android lldb in release

#android

diff --git a/Engine/Build/InstalledEngineFilters.xml b/Engine/Build/InstalledEngineFilters.xml
--- a/Engine/Build/InstalledEngineFilters.xml
+++ b/Engine/Build/InstalledEngineFilters.xml
@@ -497,6 +497,7 @@
 		Engine/Binaries/DotNET/Android/...
 		Engine/Build/Android/...
 		Engine/Extras/Android/...
+		Engine/Extras/LLDBDataFormatters/...
 		Engine/Plugins/Runtime/OpenXR/*.xml
 		Engine/Source/ThirdParty/Android/...
 		Engine/Source/ThirdParty/AndroidPermission/...
diff --git a/Engine/Source/Programs/UnrealBuildTool/Platform/Android/UEDeployAndroid.cs b/Engine/Source/Programs/UnrealBuildTool/Platform/Android/UEDeployAndroid.cs
--- a/Engine/Source/Programs/UnrealBuildTool/Platform/Android/UEDeployAndroid.cs
+++ b/Engine/Source/Programs/UnrealBuildTool/Platform/Android/UEDeployAndroid.cs
@@ -2724,6 +2724,7 @@ namespace UnrealBuildTool
 				Text.AppendLine("\t\t\t</intent-filter>");
 				Text.AppendLine("\t\t</activity>");
 				Text.AppendLine("\t\t<activity android:name=\"com.epicgames.unreal.GameActivity\"");
+				Text.AppendLine("\t\t          android:exported=\"false\"");
 				Text.AppendLine("\t\t          android:label=\"@string/app_name\"");
 				Text.AppendLine("\t\t          android:theme=\"@style/UnrealSplashTheme\"");
 				Text.AppendLine(bAddDensity ? "\t\t          android:configChanges=\"mcc|mnc|uiMode|density|screenSize|smallestScreenSize|screenLayout|orientation|keyboardHidden|keyboard|navigation|touchscreen|locale|fontScale|layoutDirection\""
