Allow bGatherThirdPartySource to be set from UBT xml config.

diff --git a/Engine/Source/Programs/UnrealBuildTool/ProjectFiles/ProjectFileGenerator.cs b/Engine/Source/Programs/UnrealBuildTool/ProjectFiles/ProjectFileGenerator.cs
--- a/Engine/Source/Programs/UnrealBuildTool/ProjectFiles/ProjectFileGenerator.cs
+++ b/Engine/Source/Programs/UnrealBuildTool/ProjectFiles/ProjectFileGenerator.cs
@@ -402,6 +402,7 @@ namespace UnrealBuildTool
 		/// generated projects.  This can make the generated projects quite a bit bigger, but makes it easier to open files
 		/// directly from the IDE.
 		/// </summary>
+		[XmlConfigFile]
 		bool bGatherThirdPartySource = false;
 
 		/// <summary>
