- Very very early support for UE running as a SwiftUI app:
  - .ini settings for using SwiftUI main function
  - this will eanble #defines, etc, to enable some embedded style functionality in UE, where SwiftUI will trigger when to create the engine/view/etc
  - assume fully immersive mode, which uses CompositorServices to get drawables from SwiftUI land
  - Added some code to FMetalViewport for using Composotir layer/drawable instead of CoreAnimation layer/drawable
  - If the .ini setting is false, then it will use the original VIsionOS support which works generally in the simulator

- Also removed -lc++ commandline which was causing "duplicate -lc++" warning when building in recent Xcode

diff --git a/Engine/Platforms/VisionOS/Config/VisionOSEngine.ini b/Engine/Platforms/VisionOS/Config/VisionOSEngine.ini
--- a/Engine/Platforms/VisionOS/Config/VisionOSEngine.ini
+++ b/Engine/Platforms/VisionOS/Config/VisionOSEngine.ini
@@ -0,0 +1,3 @@
+[SwiftIntegration]
+bUseSwiftUIMain=false
+bCreateBridgingHeader=false
diff --git a/Engine/Platforms/VisionOS/Source/Runtime/Launch/Launch_VisionOS.Build.cs b/Engine/Platforms/VisionOS/Source/Runtime/Launch/Launch_VisionOS.Build.cs
new file mode 100644
--- /dev/null
+++ b/Engine/Platforms/VisionOS/Source/Runtime/Launch/Launch_VisionOS.Build.cs
@@ -0,0 +1,12 @@
+// Copyright Epic Games, Inc. All Rights Reserved.
+
+using UnrealBuildTool;
+using System;
+using System.IO;
+
+public class Launch_VisionOS : Launch
+{
+	public Launch_VisionOS(ReadOnlyTargetRules Target) : base(Target)
+	{
+	}
+}
diff --git a/Engine/Platforms/VisionOS/Source/Runtime/Launch/Source/UESwift.swift b/Engine/Platforms/VisionOS/Source/Runtime/Launch/Source/UESwift.swift
new file mode 100644
--- /dev/null
+++ b/Engine/Platforms/VisionOS/Source/Runtime/Launch/Source/UESwift.swift
@@ -0,0 +1,76 @@
+import SwiftUI
+import CompositorServices
+
+struct SwiftUIView: View {
+	var onClick: () -> Void = {}
+
+	@Environment(\.openImmersiveSpace) var openImmersiveSpace
+
+	func Open()
+	{
+		Task
+		{
+			await openImmersiveSpace(id: "ImmersiveSpace")
+		}
+	}
+	
+	var body: some View {
+		
+	VStack(spacing: 8) {
+	  Text("SwiftUI in Unreal!")
+		.font(.title)
+		.bold()
+	  Button(action: {
+		Open()
+//		onClick()
+	  }, label: {
+		Text("Test Button")
+	  })
+	}
+  }
+}
+
+
+// unused at the moment, but this is how UE can open a SwiftUI view form the Obj-C side
+class HostingViewFactory: NSObject
+{
+	static var onStartImmersive: (LayerRenderer) -> Void = { layer in }
+
+	@objc static func MakeSwiftUIView(OnClick: @escaping (() -> Void)) -> UIViewController
+	{
+	  return UIHostingController(rootView: SwiftUIView(onClick: OnClick))
+	}
+}
+
+#if UE_USE_SWIFT_UI_MAIN
+
+@main
+struct UESwiftApp: App {
+	
+	@UIApplicationDelegateAdaptor(IOSAppDelegate.self) var delegate
+	
+	var body: some Scene
+	{
+		WindowGroup
+		{
+			SwiftUIView()
+		}
+
+		ImmersiveSpace(id: "ImmersiveSpace")
+		{
+			// this will make a CompositorLayer that can pull a Metal drawable out for UE to render to
+			CompositorLayer
+			{
+				// on button click, tell Unreal the layer is ready and can continue
+				layerRenderer in
+					FSwiftAppBootstrap.KickoffWithCompositingLayer(layerRenderer)
+			}
+		}
+	}
+}
+
+#endif
+
+//#Preview {
+//    SwiftUIView()
+//}
diff --git a/Engine/Source/Programs/UnrealBuildTool/Platform/Apple/AppleExports.cs b/Engine/Source/Programs/UnrealBuildTool/Platform/Apple/AppleExports.cs
--- a/Engine/Source/Programs/UnrealBuildTool/Platform/Apple/AppleExports.cs
+++ b/Engine/Source/Programs/UnrealBuildTool/Platform/Apple/AppleExports.cs
@@ -52,6 +52,20 @@ namespace UnrealBuildTool
 			return bUseModernXcode;
 		}
 
+		/// <summary>
+		/// Get the given project's Swift settings
+		/// </summary>
+		/// <param name="ProjectFile"></param>
+		/// <param name="Platform"></param>
+		/// <param name="bUseSwiftUIMain"></param>
+		/// <param name="bCreateBridgingHeader"></param>
+		public static void GetSwiftIntegrationSettings(FileReference? ProjectFile, UnrealTargetPlatform Platform, out bool bUseSwiftUIMain, out bool bCreateBridgingHeader)
+		{
+			ConfigHierarchy Ini = ConfigCache.ReadHierarchy(ConfigHierarchyType.Engine, ProjectFile?.Directory, Platform);
+			Ini.TryGetValue("SwiftIntegration", "bUseSwiftUIMain", out bUseSwiftUIMain);
+			Ini.TryGetValue("SwiftIntegration", "bCreateBridgingHeader", out bCreateBridgingHeader);
+		}
+
 		/// <summary>
 		///  This is a FilePath from UE settings, like /Game/Foo/Bar.txt
 		/// </summary>
diff --git a/Engine/Source/Programs/UnrealBuildTool/Platform/IOS/IOSToolChain.cs b/Engine/Source/Programs/UnrealBuildTool/Platform/IOS/IOSToolChain.cs
--- a/Engine/Source/Programs/UnrealBuildTool/Platform/IOS/IOSToolChain.cs
+++ b/Engine/Source/Programs/UnrealBuildTool/Platform/IOS/IOSToolChain.cs
@@ -32,7 +32,6 @@ namespace UnrealBuildTool
 	{
 		private static List<FileItem> BundleDependencies = new List<FileItem>();
 
-		public readonly ReadOnlyTargetRules? Target;
 		protected IOSProjectSettings ProjectSettings;
 
 		public IOSToolChain(ReadOnlyTargetRules? Target, IOSProjectSettings InProjectSettings, ClangToolChainOptions ToolchainOptions, ILogger InLogger)
@@ -41,9 +40,8 @@ namespace UnrealBuildTool
 		}
 
 		protected IOSToolChain(ReadOnlyTargetRules? Target, IOSProjectSettings InProjectSettings, Func<AppleToolChainSettings> InCreateSettings, ClangToolChainOptions ToolchainOptions, ILogger InLogger)
-			: base((Target == null) ? null : Target.ProjectFile, InCreateSettings, ToolchainOptions, InLogger)
+			: base(Target, InCreateSettings, ToolchainOptions, InLogger)
 		{
-			this.Target = Target;
 			ProjectSettings = InProjectSettings;
 		}
 
@@ -218,9 +216,10 @@ namespace UnrealBuildTool
 
 			Arguments.Add($"-isysroot \"{Settings.GetSDKPath(CompileEnvironment.Architecture)}\"");
 
-			if (GetXcodeMinVersionParam(CompileEnvironment.Architecture) != "")
+			string MinVersionParam = GetXcodeMinVersionParam(CompileEnvironment.Architecture);
+			if (MinVersionParam != "")
 			{
-				Arguments.Add("-m" + GetXcodeMinVersionParam(CompileEnvironment.Architecture) + "=" + ProjectSettings.RuntimeVersion);
+				Arguments.Add($"-m{MinVersionParam}={ProjectSettings.RuntimeVersion}");
 			}
 
 			// Add additional frameworks so that their headers can be found
@@ -313,9 +312,10 @@ namespace UnrealBuildTool
 			Arguments.Add($" -isysroot \"{SDKPath}\"");
 
 			Arguments.Add("-dead_strip");
-			if (GetXcodeMinVersionParam(LinkEnvironment.Architecture) != "") 
+			string MinVersionParam = GetXcodeMinVersionParam(LinkEnvironment.Architecture);
+			if (MinVersionParam != "")
 			{
-				Arguments.Add("-m" + GetXcodeMinVersionParam(LinkEnvironment.Architecture) + "=" + ProjectSettings.RuntimeVersion);
+				Arguments.Add($"-m{MinVersionParam}={ProjectSettings.RuntimeVersion}");
 			}
 			Arguments.Add("-Wl-no_pie");
 			Arguments.Add("-stdlib=libc++");
@@ -1081,7 +1081,7 @@ namespace UnrealBuildTool
 				OutputFiles.Add(StripCompleteFile);
 			}
 
-			if (!BinaryLinkEnvironment.bIsBuildingDLL)
+			if (!AppleExports.UseModernXcode(Target.ProjectFile) && !BinaryLinkEnvironment.bIsBuildingDLL)
 			{
 				// generate the asset catalog
 				bool bUserImagesExist = false;
diff --git a/Engine/Source/Programs/UnrealBuildTool/Platform/IOS/UEBuildIOS.cs b/Engine/Source/Programs/UnrealBuildTool/Platform/IOS/UEBuildIOS.cs
--- a/Engine/Source/Programs/UnrealBuildTool/Platform/IOS/UEBuildIOS.cs
+++ b/Engine/Source/Programs/UnrealBuildTool/Platform/IOS/UEBuildIOS.cs
@@ -143,11 +143,6 @@ namespace UnrealBuildTool
 		/// </summary>
 		[CommandLine("-EnableUBSan")]
 		public bool bEnableUndefinedBehaviorSanitizer = false;
-
-		/// <summary>
-		/// If true, we use SwiftUI's main function instead of Obj-C's main
-		/// </summary>
-		public bool bUseSwiftUIMain = false;
 	}
 
 	/// <summary>
@@ -202,8 +197,6 @@ namespace UnrealBuildTool
 
 		public bool bEnableUndefinedBehaviorSanitizer => Inner.bEnableUndefinedBehaviorSanitizer;
 
-		public bool bUseSwiftUIMain => Inner.bUseSwiftUIMain;
-
 #pragma warning restore CS1591
 		#endregion
 	}
diff --git a/Engine/Source/Programs/UnrealBuildTool/Platform/Mac/MacToolChain.cs b/Engine/Source/Programs/UnrealBuildTool/Platform/Mac/MacToolChain.cs
--- a/Engine/Source/Programs/UnrealBuildTool/Platform/Mac/MacToolChain.cs
+++ b/Engine/Source/Programs/UnrealBuildTool/Platform/Mac/MacToolChain.cs
@@ -45,8 +45,8 @@ namespace UnrealBuildTool
 	/// </summary>
 	class MacToolChain : AppleToolChain
 	{
-		public MacToolChain(FileReference? InProjectFile, ClangToolChainOptions InOptions, ILogger InLogger)
-			: base(InProjectFile, () => new MacToolChainSettings(false, InLogger), InOptions, InLogger)
+		public MacToolChain(ReadOnlyTargetRules? Target, ClangToolChainOptions InOptions, ILogger InLogger)
+			: base(Target, () => new MacToolChainSettings(false, InLogger), InOptions, InLogger)
 		{
 		}
 
@@ -277,8 +277,6 @@ namespace UnrealBuildTool
 
 			// Needed to make sure install_name_tool will be able to update paths in Mach-O headers
 			Arguments.Add("-headerpad_max_install_names");
-
-			Arguments.Add("-lc++");
 		}
 
 		void GetArchiveArguments_Global(LinkEnvironment LinkEnvironment, List<string> Arguments)
diff --git a/Engine/Source/Programs/UnrealBuildTool/Platform/Mac/UEBuildMac.cs b/Engine/Source/Programs/UnrealBuildTool/Platform/Mac/UEBuildMac.cs
--- a/Engine/Source/Programs/UnrealBuildTool/Platform/Mac/UEBuildMac.cs
+++ b/Engine/Source/Programs/UnrealBuildTool/Platform/Mac/UEBuildMac.cs
@@ -598,7 +598,7 @@ namespace UnrealBuildTool
 				Options |= ClangToolChainOptions.OutputDylib;
 			}
 
-			return new MacToolChain(Target.ProjectFile, Options, Logger);
+			return new MacToolChain(Target, Options, Logger);
 		}
 
 		/// <inheritdoc/>
diff --git a/Engine/Source/Programs/UnrealBuildTool/ProjectFiles/Xcode/XcodeProject.cs b/Engine/Source/Programs/UnrealBuildTool/ProjectFiles/Xcode/XcodeProject.cs
--- a/Engine/Source/Programs/UnrealBuildTool/ProjectFiles/Xcode/XcodeProject.cs
+++ b/Engine/Source/Programs/UnrealBuildTool/ProjectFiles/Xcode/XcodeProject.cs
@@ -1729,6 +1729,15 @@ namespace UnrealBuildTool.XcodeProjectXcconfig
 
 					ExtraConfigLines.Add($"VALID_ARCHS = arm64");
 					ExtraConfigLines.Add($"ARCHS = arm64");
+
+					// if we are doing immersive with SwiftUI we need a plist key
+					bool bUseSwiftUIMain;
+					AppleExports.GetSwiftIntegrationSettings(UnrealData.UProjectFileLocation, Platform, out bUseSwiftUIMain, out _);
+					if (bUseSwiftUIMain)
+					{
+						ExtraConfigLines.Add($"INFOPLIST_KEY_UIApplicationSceneManifest_Generation = YES");
+					}
+
 				}
 				else
 				{
diff --git a/Engine/Source/Programs/UnrealBuildTool/ToolChain/AppleToolChain.cs b/Engine/Source/Programs/UnrealBuildTool/ToolChain/AppleToolChain.cs
--- a/Engine/Source/Programs/UnrealBuildTool/ToolChain/AppleToolChain.cs
+++ b/Engine/Source/Programs/UnrealBuildTool/ToolChain/AppleToolChain.cs
@@ -244,13 +244,21 @@ namespace UnrealBuildTool
 		public Lazy<AppleToolChainSettings> ToolChainSettings;
 
 		protected FileReference? ProjectFile;
+		public readonly ReadOnlyTargetRules? Target;
+
+		// cache some ini settings
+		readonly bool bUseSwiftUIMain;
+		readonly bool bCreateSwiftBridgingHeader;
 
 		protected bool bUseModernXcode => AppleExports.UseModernXcode(ProjectFile);
 
-		public AppleToolChain(FileReference? InProjectFile, Func<AppleToolChainSettings> InCreateSettings, ClangToolChainOptions InOptions, ILogger InLogger) : base(InOptions, InLogger)
+		public AppleToolChain(ReadOnlyTargetRules? Target, Func<AppleToolChainSettings> InCreateSettings, ClangToolChainOptions InOptions, ILogger InLogger) : base(InOptions, InLogger)
 		{
-			ProjectFile = InProjectFile;
+			this.Target = Target;
+			ProjectFile = Target?.ProjectFile;
 			ToolChainSettings = new Lazy<AppleToolChainSettings>(InCreateSettings);
+
+			AppleExports.GetSwiftIntegrationSettings(ProjectFile, Target == null ? UnrealTargetPlatform.Mac : Target.Platform, out bUseSwiftUIMain, out bCreateSwiftBridgingHeader);
 		}
 
 		/// <summary>
@@ -549,13 +557,15 @@ namespace UnrealBuildTool
 			// @todo hack fix this better - it's so that files can see the generated -Swift.h headers - we don't know when we need them at this point
 			if (CompileEnvironment.Platform == UnrealTargetPlatform.VisionOS)
 			{
-				if (!SourceFile.HasExtension(".swift") && SourceFile.Location.ContainsName("Launch", 0))
+				if (bCreateSwiftBridgingHeader && !SourceFile.HasExtension(".swift") && SourceFile.Location.ContainsName("Launch", 0))
 				{
 					FileItem OutputInteropHeader = GetBridgingHeader("UESwift", OutputDir);
 					Arguments.Add(GetUserIncludePathArgument(OutputInteropHeader.GetDirectoryItem().Location));
 					CompileAction.PrerequisiteItems.Add(OutputInteropHeader);
 				}
+
 			}
+			Arguments.Add("-DUE_USE_SWIFT_UI_MAIN=" + (bUseSwiftUIMain ? "1" : "0"));
 			return Output;
 		}
 
@@ -584,7 +594,7 @@ namespace UnrealBuildTool
 
 			// output file settings
 			Arguments.Add("-emit-object"); // same as -c
-			Arguments.Add($"-parse-as-library"); // don't create a main function
+			Arguments.Add($"-parse-as-library"); // don't create a main function (even if we use the SwiftUI @main, this still seems to be correctly working)
 			Arguments.Add($"-o \"{OutputFile}\"");
 
 			// platform settings
@@ -596,9 +606,17 @@ namespace UnrealBuildTool
 			Arguments.Add("-aarch64-use-tbi");
 			Arguments.Add("-stack-check");
 			Arguments.Add($"-swift-version 5");
-			
+
+			Arguments.Add("-enable-objc-interop");
+			Arguments.Add("-cxx-interoperability-mode=default");
+			Arguments.Add($"-import-objc-header {Unreal.EngineDirectory}/Source/Runtime/Launch/Private/IOS/UECppToSwift.h");
 			Arguments.Add($"-module-name Launch");
 
+			if (bUseSwiftUIMain)
+			{
+				Arguments.Add("-DUE_USE_SWIFT_UI_MAIN");
+			}
+
 			Action CompileAction = Graph.CreateAction(ActionType.Compile);
 			CompileAction.Weight = CompileActionWeight;
 			CompileAction.CommandArguments = String.Join(" ", Arguments);
@@ -618,43 +636,55 @@ namespace UnrealBuildTool
 			CompileAction.bIsGCCCompiler = true;
 			CompileAction.bCanExecuteRemotely = true;
 
+			if (bCreateSwiftBridgingHeader)
+			{
+				FileItem OutputInteropHeader = GetBridgingHeader(SourceFile.FullName, OutputDir);
 
-			FileItem OutputInteropHeader = GetBridgingHeader(SourceFile.FullName, OutputDir);
+				// obj-c bridging header settings
+				Arguments.Clear();
+				Arguments.Add($"\"{SourceFile}\"");
+				Arguments.Add("-parse"); // this will allow it to generate the header without writing out any .o/executable
+				Arguments.Add("-emit-objc-header");
+				Arguments.Add($"-emit-objc-header-path \"{OutputInteropHeader}\"");
+				Arguments.Add("-parse-as-library");
 
-			// obj-c bridging header settings
-			Arguments.Clear();
-			Arguments.Add($"\"{SourceFile}\"");
-			Arguments.Add("-parse"); // this will allow it to generate the header without writing out any .o/executable
-			Arguments.Add("-emit-objc-header");
-			Arguments.Add($"-emit-objc-header-path \"{OutputInteropHeader}\"");
+				//Arguments.Add("-enable-objc-interop");
+				//Arguments.Add("-cxx-interoperability-mode=default");
+				Arguments.Add($"-import-objc-header {Unreal.EngineDirectory}/Source/Runtime/Launch/Private/IOS/UECppToSwift.h");
 
-			// platform settings
-			Arguments.Add($"-target {ToolChainSettings.Value.GetTargetTuple(CompileEnvironment.Architecture)}");
-			Arguments.Add($"-sdk {ToolChainSettings.Value.GetSDKPath(CompileEnvironment.Architecture)}");
+				// platform settings
+				Arguments.Add($"-target {ToolChainSettings.Value.GetTargetTuple(CompileEnvironment.Architecture)}");
+				Arguments.Add($"-sdk {ToolChainSettings.Value.GetSDKPath(CompileEnvironment.Architecture)}");
 
-			Arguments.Add($"-swift-version 5");
+				Arguments.Add($"-swift-version 5");
 
-			Arguments.Add($"-module-name Launch");
+				Arguments.Add($"-module-name Launch");
 
-			// now make an action to export the swift code as a header Obj-C bridging
-			Action HeaderAction = Graph.CreateAction(ActionType.CompileModuleInterface);
-			HeaderAction.CommandPath = FileReference.Combine(ToolChainSettings.Value.ToolchainDir, "swiftc");
-			HeaderAction.CommandArguments = String.Join(" ", Arguments);
-			HeaderAction.PrerequisiteItems.Add(SourceFile);
-			HeaderAction.ProducedItems.Add(OutputInteropHeader);
-			HeaderAction.CommandDescription = "Generate Header";
-			HeaderAction.StatusDescription = Path.GetFileName(OutputInteropHeader.AbsolutePath);
-			if (ArchConfig.Mode != UnrealArchitectureMode.SingleArchitecture)
-			{
-				string ReadableArch = ArchConfig.ConvertToReadableArchitecture(CompileEnvironment.Architecture);
-				CompileAction.CommandDescription += $" [{ReadableArch}]";
-			}
+				if (bUseSwiftUIMain)
+				{
+					Arguments.Add("-DUE_USE_SWIFT_UI_MAIN");
+				}
 
-			HeaderAction.WorkingDirectory = CompileAction.WorkingDirectory;
-			HeaderAction.bIsGCCCompiler = CompileAction.bIsGCCCompiler;
-			HeaderAction.bCanExecuteRemotely = CompileAction.bCanExecuteRemotely;
+				// now make an action to export the swift code as a header Obj-C bridging
+				Action HeaderAction = Graph.CreateAction(ActionType.CompileModuleInterface);
+				HeaderAction.CommandPath = FileReference.Combine(ToolChainSettings.Value.ToolchainDir, "swiftc");
+				HeaderAction.CommandArguments = String.Join(" ", Arguments);
+				HeaderAction.PrerequisiteItems.Add(SourceFile);
+				HeaderAction.ProducedItems.Add(OutputInteropHeader);
+				HeaderAction.CommandDescription = "Generate Header";
+				HeaderAction.StatusDescription = Path.GetFileName(OutputInteropHeader.AbsolutePath);
+				if (ArchConfig.Mode != UnrealArchitectureMode.SingleArchitecture)
+				{
+					string ReadableArch = ArchConfig.ConvertToReadableArchitecture(CompileEnvironment.Architecture);
+					CompileAction.CommandDescription += $" [{ReadableArch}]";
+				}
 
-			Console.WriteLine($"{CompileAction.CommandPath} {CompileAction.CommandArguments}");
+				HeaderAction.WorkingDirectory = CompileAction.WorkingDirectory;
+				HeaderAction.bIsGCCCompiler = CompileAction.bIsGCCCompiler;
+				HeaderAction.bCanExecuteRemotely = CompileAction.bCanExecuteRemotely;
+			}
+			//Console.WriteLine($"{CompileAction.CommandPath} {CompileAction.CommandArguments}");
+			//Console.WriteLine($"{HeaderAction.CommandPath} {HeaderAction.CommandArguments}");
 
 			// this is likely ignored, but the Compile action is the important one
 			return CompileAction;
@@ -845,7 +875,7 @@ namespace UnrealBuildTool
 					ExtraOptions,
 					//$"-sdk {SDKName}",
 				};
-	
+
 				Process LocalProcess = new Process();
 				LocalProcess.StartInfo = new ProcessStartInfo("/usr/bin/env", String.Join(" ", Arguments));
 				LocalProcess.OutputDataReceived += (Sender, Args) => { LocalProcessOutput(Args, false, Logger); };
diff --git a/Engine/Source/Runtime/Apple/MetalRHI/Private/MetalContext.cpp b/Engine/Source/Runtime/Apple/MetalRHI/Private/MetalContext.cpp
--- a/Engine/Source/Runtime/Apple/MetalRHI/Private/MetalContext.cpp
+++ b/Engine/Source/Runtime/Apple/MetalRHI/Private/MetalContext.cpp
@@ -313,7 +313,10 @@ mtlpp::PrimitiveTopologyClass TranslatePrimitiveTopology(uint32 PrimitiveType)
 FMetalDeviceContext* FMetalDeviceContext::CreateDeviceContext()
 {
 	uint32 DeviceIndex = 0;
-#if PLATFORM_IOS
+#if PLATFORM_VISIONOS && UE_USE_SWIFT_UI_MAIN
+	// get the device from the compositor layer
+	mtlpp::Device Device = cp_layer_renderer_get_device([IOSAppDelegate GetDelegate].SwiftLayer);
+#elif PLATFORM_IOS
 	mtlpp::Device Device = mtlpp::Device([IOSAppDelegate GetDelegate].IOSView->MetalDevice);
 #else
 	mtlpp::Device Device = GetMTLDevice(DeviceIndex);
diff --git a/Engine/Source/Runtime/Apple/MetalRHI/Private/MetalViewport.cpp b/Engine/Source/Runtime/Apple/MetalRHI/Private/MetalViewport.cpp
--- a/Engine/Source/Runtime/Apple/MetalRHI/Private/MetalViewport.cpp
+++ b/Engine/Source/Runtime/Apple/MetalRHI/Private/MetalViewport.cpp
@@ -81,6 +81,12 @@ FMetalViewport::FMetalViewport(void* WindowHandle, uint32 InSizeX, uint32 InSize
 	, CustomPresent{nullptr}
 #endif
 {
+#if PLATFORM_VISIONOS
+	// look to see if we need to hook up to a Swift compositor renderer
+	SwiftLayer = [IOSAppDelegate GetDelegate].SwiftLayer;
+	NewFrame();
+#endif
+
 #if PLATFORM_MAC
 	MainThreadCall(^{
 		FCocoaWindow* Window = (FCocoaWindow*)WindowHandle;
@@ -292,6 +298,14 @@ TRefCountPtr<FMetalSurface> FMetalViewport::GetBackBuffer(EMetalViewportAccessFl
 
 id<CAMetalDrawable> FMetalViewport::GetDrawable(EMetalViewportAccessFlag Accessor)
 {
+#if PLATFORM_VISIONOS
+	// no CAMetalDrawable in Swift mode
+	if (SwiftLayer != nullptr)
+	{
+		return nil;
+	}
+#endif
+	
 	SCOPE_CYCLE_COUNTER(STAT_MetalMakeDrawableTime);
     if (!Drawable
 #if !PLATFORM_MAC
@@ -362,8 +376,41 @@ id<CAMetalDrawable> FMetalViewport::GetDrawable(EMetalViewportAccessFlag Accesso
 	return Drawable;
 }
 
+void FMetalViewport::NewFrame()
+{
+#if PLATFORM_VISIONOS
+	if (SwiftLayer)
+	{
+		// end the last frame if there was one
+		if (SwiftLayerFrame != nullptr)
+		{
+			cp_frame_end_submission(SwiftLayerFrame);
+		}
+		
+		// get next frame
+		SwiftLayerFrame = cp_layer_renderer_query_next_frame(SwiftLayer);
+	}
+#endif
+}
+
 FMetalTexture FMetalViewport::GetDrawableTexture(EMetalViewportAccessFlag Accessor)
 {
+#if PLATFORM_VISIONOS
+	if (SwiftLayerFrame != nullptr)
+	{
+		// get the next drawable
+		cp_drawable_t SwiftDrawable = cp_frame_query_drawable(SwiftLayerFrame);
+		if (SwiftDrawable == nullptr)
+		{
+			UE_LOG(LogMetal, Fatal, TEXT("Failed to get next drawable"));
+		}
+		
+		// get the color texture out and use that with the RHI
+		DrawableTextures[Accessor] = cp_drawable_get_color_texture(SwiftDrawable, 0);
+		return DrawableTextures[Accessor];
+	}
+#endif
+	
 	id<CAMetalDrawable> CurrentDrawable = GetDrawable(Accessor);
 #if METAL_DEBUG_OPTIONS
 	@autoreleasepool
@@ -447,15 +494,36 @@ void FMetalViewport::Present(FMetalCommandQueue& CommandQueue, bool bLockToVsync
 #endif
 			if (FrameAvailable > 0 && (InDisplayID == 0 || (DisplayID == InDisplayID && !bIsInLiveResize)))
 			{
+#if PLATFORM_VISIONOS
+				if (SwiftLayerFrame != nullptr)
+				{
+					// tell the compositor service we are going to submit soon and need the drawable
+					cp_frame_start_submission(SwiftLayerFrame);
+				}
+#endif
+
 				FPlatformAtomics::InterlockedDecrement(&FrameAvailable);
 				id<CAMetalDrawable> LocalDrawable = [GetDrawable(EMetalViewportAccessDisplayLink) retain];
+				FMetalTexture DrawableTexture;
+#if PLATFORM_VISIONOS
+				// it seems like this should be usable in all cases, instead of LocalDrawable.texture, but it's causing View to be overwritten with the drawable, somehow!
+				if (SwiftLayer != nullptr)
+				{
+					DrawableTexture = GetDrawableTexture(EMetalViewportAccessDisplayLink);
+				}
+				else
+#endif
+				{
+					DrawableTexture = LocalDrawable.texture;
+				}
+				
 				{
 					FScopeLock BlockLock(&Mutex);
 #if PLATFORM_MAC
 					bIsInLiveResize = View.inLiveResize;
 #endif
 					
-					if (LocalDrawable && LocalDrawable.texture && (InDisplayID == 0 || !bIsInLiveResize))
+					if (DrawableTexture && (InDisplayID == 0 || !bIsInLiveResize))
 					{
 						mtlpp::CommandBuffer CurrentCommandBuffer = CommandQueue.CreateCommandBuffer();
 						check(CurrentCommandBuffer);
@@ -471,7 +539,7 @@ void FMetalViewport::Present(FMetalCommandQueue& CommandQueue, bool bLockToVsync
 							check(IsValidRef(Texture));
 							
 							FMetalTexture Src = Texture->Texture;
-							FMetalTexture Dst = LocalDrawable.texture;
+							FMetalTexture Dst = DrawableTexture;
 							
 							NSUInteger Width = FMath::Min(Src.GetWidth(), Dst.GetWidth());
 							NSUInteger Height = FMath::Min(Src.GetHeight(), Dst.GetHeight());
@@ -539,6 +607,13 @@ void FMetalViewport::Present(FMetalCommandQueue& CommandQueue, bool bLockToVsync
 #else // PLATFORM_MAC
 						CurrentCommandBuffer.AddCompletedHandler(C);
 
+#if PLATFORM_VISIONOS
+						if (SwiftLayer != nullptr)
+						{
+							cp_drawable_encode_present(cp_frame_query_drawable(SwiftLayerFrame), CurrentCommandBuffer.GetPtr());
+						}
+						else
+#endif
 						if (MinPresentDuration && GEnablePresentPacing)
 						{
 							CurrentCommandBuffer.PresentAfterMinimumDuration(LocalDrawable, 1.0f/(float)FramePace);
@@ -588,6 +663,9 @@ void FMetalViewport::Swap()
 		BackBuffer[0] = BB1;
 		BackBuffer[1] = BB0;
 	}
+	
+	// move on to next frame (only used in SwiftUI/compositor mode)
+	NewFrame();
 }
 
 /*=============================================================================
diff --git a/Engine/Source/Runtime/Apple/MetalRHI/Public/MetalViewport.h b/Engine/Source/Runtime/Apple/MetalRHI/Public/MetalViewport.h
--- a/Engine/Source/Runtime/Apple/MetalRHI/Public/MetalViewport.h
+++ b/Engine/Source/Runtime/Apple/MetalRHI/Public/MetalViewport.h
@@ -16,6 +16,10 @@ THIRD_PARTY_INCLUDES_START
 #include "mtlpp.hpp"
 THIRD_PARTY_INCLUDES_END
 
+#if PLATFORM_VISIONOS
+#import <CompositorServices/CompositorServices.h>
+#endif
+
 enum EMetalViewportAccessFlag
 {
 	EMetalViewportAccessRHI,
@@ -63,7 +67,14 @@ public:
 private:
 	uint32 GetViewportIndex(EMetalViewportAccessFlag Accessor) const;
 
+	void NewFrame();
+	
 private:
+#if PLATFORM_VISIONOS
+	CP_OBJECT_cp_layer_renderer* SwiftLayer = nullptr;
+	cp_frame_t SwiftLayerFrame = nullptr;
+#endif
+	
 	id<CAMetalDrawable> Drawable;
 	TRefCountPtr<FMetalSurface> BackBuffer[2];
 	mutable FCriticalSection Mutex;
diff --git a/Engine/Source/Runtime/ApplicationCore/Private/IOS/IOSAppDelegate.cpp b/Engine/Source/Runtime/ApplicationCore/Private/IOS/IOSAppDelegate.cpp
--- a/Engine/Source/Runtime/ApplicationCore/Private/IOS/IOSAppDelegate.cpp
+++ b/Engine/Source/Runtime/ApplicationCore/Private/IOS/IOSAppDelegate.cpp
@@ -259,6 +259,10 @@ static IOSAppDelegate* CachedDelegate = nil;
 
 -(id)init
 {
+#if UE_USE_SWIFT_UI_MAIN
+	// will get the command later
+	FIOSCommandLineHelper::InitCommandArgs("");
+#endif
 	self = [super init];
 	CachedDelegate = self;
 	memset(GEnabledAudioFeatures, 0, sizeof(GEnabledAudioFeatures));
diff --git a/Engine/Source/Runtime/ApplicationCore/Private/IOS/IOSWindow.cpp b/Engine/Source/Runtime/ApplicationCore/Private/IOS/IOSWindow.cpp
--- a/Engine/Source/Runtime/ApplicationCore/Private/IOS/IOSWindow.cpp
+++ b/Engine/Source/Runtime/ApplicationCore/Private/IOS/IOSWindow.cpp
@@ -134,6 +134,15 @@ FPlatformRect FIOSWindow::GetScreenRect()
 		ScreenRect.Bottom = FMath::TruncToInt((Frame.origin.y + Frame.size.height) * Scale);
 		ScreenRect.Right = FMath::TruncToInt((Frame.origin.x + Frame.size.width) * Scale);
 	}
+	else if (AppDelegate.SwiftLayer != nullptr)
+	{
+		// using the last viewport since this will work with 1 eye or 2 eye setup, to get the full size of the screen
+		CGRect LastViewport = [[AppDelegate.SwiftLayerViewports lastObject] CGRectValue];
+		ScreenRect.Left = 0;
+		ScreenRect.Top = 0;
+		ScreenRect.Right = (int)(LastViewport.origin.x + LastViewport.size.width);
+		ScreenRect.Bottom = (int)(LastViewport.origin.y + LastViewport.size.height);
+	}
 
 	return ScreenRect;
 }
diff --git a/Engine/Source/Runtime/ApplicationCore/Public/IOS/IOSAppDelegate.h b/Engine/Source/Runtime/ApplicationCore/Public/IOS/IOSAppDelegate.h
--- a/Engine/Source/Runtime/ApplicationCore/Public/IOS/IOSAppDelegate.h
+++ b/Engine/Source/Runtime/ApplicationCore/Public/IOS/IOSAppDelegate.h
@@ -3,17 +3,25 @@
 #pragma once
 
 #import <UIKit/UIKit.h>
-#import <AVFoundation/AVAudioSession.h>
 #import <GameKit/GKGameCenterViewController.h>
+
+#ifndef SWIFT_IMPORT
+#import <AVFoundation/AVAudioSession.h>
 #include "Delegates/Delegate.h"
 #include "Logging/LogMacros.h"
 #include "Containers/UnrealString.h"
-#if !PLATFORM_VISIONOS
+#endif
+
+#if PLATFORM_VISIONOS
+#import <CompositorServices/CompositorServices.h>
+#else
 #import <UserNotifications/UserNotifications.h>
 #endif
 
 #define USE_MUTE_SWITCH_DETECTION 0
 
+#ifndef SWIFT_IMPORT
+
 enum class EAudioFeature : uint8
 {
 	Playback, // Audio Not affected by the ringer switch
@@ -72,13 +80,10 @@ private:
 	static TArray<FFilterDelegateAndHandle> PushNotificationFilters;
 };
 
-@class FIOSView;
-@class IOSViewController;
-@class SlateOpenGLESViewController;
-@class IOSAppDelegate;
-
 DECLARE_LOG_CATEGORY_EXTERN(LogIOSAudioSession, Log, All);
 
+@class IOSAppDelegate;
+
 namespace FAppEntry
 {
 	void PlatformInit();
@@ -102,8 +107,18 @@ namespace FAppEntry
 	extern int32	gLaunchLocalNotificationFireDate;
 }
 
+#endif
+
+@class FIOSView;
+@class IOSViewController;
+@class SlateOpenGLESViewController;
+
+
+#ifndef SWIFT_IMPORT
 APPLICATIONCORE_API
-@interface IOSAppDelegate : UIResponder <UIApplicationDelegate,
+#endif
+@interface IOSAppDelegate : UIResponder <
+	UIApplicationDelegate,
 #if !UE_BUILD_SHIPPING
 	UIGestureRecognizerDelegate,
 #endif
@@ -119,7 +134,12 @@ APPLICATIONCORE_API
 /** Window object */
 @property (strong, retain, nonatomic) UIWindow *Window;
 
-/** Main GL View */
+// support Compositor on other platforms?
+#if PLATFORM_VISIONOS
+@property (strong) CP_OBJECT_cp_layer_renderer* SwiftLayer;
+@property (retain) NSArray* SwiftLayerViewports;
+#endif
+
 @property (retain) FIOSView* IOSView;
 
 @property class FIOSApplication* IOSApplication;
@@ -203,8 +223,10 @@ APPLICATIONCORE_API
 -(bool)IsIdleTimerEnabled;
 -(void)EnableIdleTimer:(bool)bEnable;
 -(void)StartGameThread;
+#ifndef SWIFT_IMPORT
 /** Uses the TaskGraph to execute a function on the game thread, and then blocks until the function is executed. */
 +(bool)WaitAndRunOnGameThread:(TUniqueFunction<void()>)Function;
+#endif
 -(void)NoUrlCommandLine;
 
 -(void)LoadScreenResolutionModifiers;
@@ -247,9 +269,11 @@ APPLICATIONCORE_API
 - (void)EnableHighQualityVoiceChat:(bool)bEnable;
 - (bool)IsVoiceChatEnabled;
 
+#ifndef SWIFT_IMPORT
 /** Enable/Disable an EAudioFeature. This is reference counted, so a feature must be disabled as many times as it has been enabled to actually be disabled. */
 - (void)SetFeature:(EAudioFeature)Feature Active:(bool)bIsActive;
 - (bool)IsFeatureActive:(EAudioFeature)Mode;
+#endif
 
 @property (atomic) bool bAudioActive;
 @property (atomic) bool bVoiceChatEnabled;
diff --git a/Engine/Source/Runtime/Launch/Private/IOS/LaunchIOS.cpp b/Engine/Source/Runtime/Launch/Private/IOS/LaunchIOS.cpp
--- a/Engine/Source/Runtime/Launch/Private/IOS/LaunchIOS.cpp
+++ b/Engine/Source/Runtime/Launch/Private/IOS/LaunchIOS.cpp
@@ -289,13 +289,10 @@ void FAppEntry::DecrementAudioSuspendCounters()
     }
 }
 
-#if PLATFORM_VISIONOS
-// SwiftUI test integration
-#include "UESwift-Swift.h"
-#endif
-
 void FAppEntry::PreInit(IOSAppDelegate* AppDelegate, UIApplication* Application)
 {
+	// SwiftUI apps handle this differently
+#if !UE_USE_SWIFT_UI_MAIN
 	// make a controller object
 	IOSViewController* IOSController = [[IOSViewController alloc] init];
 	
@@ -318,18 +315,13 @@ void FAppEntry::PreInit(IOSAppDelegate* AppDelegate, UIApplication* Application)
 	Application.applicationIconBadgeNumber = 0;
 #endif
 	
-#if PLATFORM_VISIONOS
-	UIViewController* HostingController = [HostingViewFactory MakeSwiftUIViewOnClick:
-	^{
-		// close on click
-		[[IOSController presentedViewController] dismissViewControllerAnimated:YES completion:nil];
-	}];
-	[IOSController presentViewController:HostingController animated:YES completion:nil];
 #endif
 }
 
 static void MainThreadInit()
 {
+	// SwiftUI apps handle this differently
+#if !UE_USE_SWIFT_UI_MAIN
 	IOSAppDelegate* AppDelegate = [IOSAppDelegate GetDelegate];
 
 #if PLATFORM_VISIONOS
@@ -375,6 +367,7 @@ static void MainThreadInit()
 	// initialize the backbuffer of the view (so the RHI can use it)
 	[AppDelegate.IOSView CreateFramebuffer];
 #endif
+#endif
 }
 
 
@@ -393,6 +386,13 @@ void FAppEntry::PlatformInit()
 	// wait until the GLView is fully initialized, so the RHI can be initialized
 	IOSAppDelegate* AppDelegate = [IOSAppDelegate GetDelegate];
 
+#if UE_USE_SWIFT_UI_MAIN
+	while (!AppDelegate.SwiftLayer)
+	{
+		FPlatformProcess::Sleep(0.005f);
+	}
+#else
+	
 	while (!AppDelegate.IOSView || !AppDelegate.IOSView->bIsInitialized)
 	{
 #if BUILD_EMBEDDED_APP
@@ -403,7 +403,7 @@ void FAppEntry::PlatformInit()
 #endif
 		FPlatformProcess::Sleep(0.005f);
 	}
-
+#endif
 	// Set GSystemResolution now that we have the size.
 	FDisplayMetrics DisplayMetrics;
 	FDisplayMetrics::RebuildDisplayMetrics(DisplayMetrics);
@@ -416,9 +416,9 @@ extern TcpConsoleListener *ConsoleListener;
 void FAppEntry::Init()
 {
 	SCOPED_BOOT_TIMING("FAppEntry::Init()");
-
+	
 	FPlatformProcess::SetRealTimeMode();
-
+	
 	//extern TCHAR GCmdLine[16384];
 	GEngineLoop.PreInit(FCommandLine::Get());
 
@@ -556,8 +556,7 @@ int32	FAppEntry::gLaunchLocalNotificationFireDate;
 
 FString GSavedCommandLine;
 
-#if !BUILD_EMBEDDED_APP
-
+#if !BUILD_EMBEDDED_APP && !UE_USE_SWIFT_UI_MAIN
 int main(int argc, char *argv[])
 {
     for(int Option = 1; Option < argc; Option++)
@@ -590,3 +589,72 @@ int main(int argc, char *argv[])
 }
 
 #endif
+
+
+#if UE_USE_SWIFT_UI_MAIN
+
+#include "UECppToSwift.h"
+
+void FSwiftAppBootstrap::KickoffWithCompositingLayer(CP_OBJECT_cp_layer_renderer* Layer)
+{
+	IOSAppDelegate* AppDelegate = [IOSAppDelegate GetDelegate];
+	
+	cp_layer_renderer_properties_t Props = cp_layer_renderer_get_properties(Layer);
+	int NumViews = cp_layer_renderer_properties_get_view_count(Props);
+	int NumTopologies = cp_layer_renderer_properties_get_texture_topology_count(Props);
+
+	NSMutableArray* Viewports = [NSMutableArray arrayWithCapacity:NumViews];
+
+	// get the texture topology
+	// @todo when Apple adds the API to actually get the size, use this instead of the mess below (docs indicate you can get
+	// get the width/height, but there's no functions to get them 
+//	for (int TopoIndex = 0; TopoIndex < NumToplogies; TopoIndex++)
+//	{
+//		cp_texture_topology_t Topology = cp_layer_renderer_properties_get_texture_topology(Props, TopoIndex);
+//	
+////		NSValue* VPValue = [NSValue valueWithCGRect:CGRectMake(Viewport.originX, Viewport.originY, Viewport.width, Viewport.height)];
+////		[Viewports addObject:VPValue];
+//	}
+
+	{
+		cp_frame_t SwiftLayerFrame = cp_layer_renderer_query_next_frame(Layer);
+		cp_drawable_t SwiftDrawable = cp_frame_query_drawable(SwiftLayerFrame); 
+		{
+			cp_view_t View = cp_drawable_get_view(SwiftDrawable, 0);
+			cp_view_texture_map_t TextureMap = cp_view_get_view_texture_map(View);
+			MTLViewport Viewport = cp_view_texture_map_get_viewport(TextureMap);
+			
+			NSValue* VPValue = [NSValue valueWithCGRect:CGRectMake(Viewport.originX, Viewport.originY, Viewport.width, Viewport.height)];
+			[Viewports addObject:VPValue];
+		}
+		cp_frame_start_submission(SwiftLayerFrame);
+		id<MTLDevice> Device = cp_layer_renderer_get_device(Layer);
+		id<MTLCommandQueue> CommandQueue = [[Device newCommandQueue] autorelease];
+		id<MTLCommandBuffer> CommandBuffer = [CommandQueue commandBuffer];
+		cp_drawable_encode_present(cp_frame_query_drawable(SwiftLayerFrame), CommandBuffer);
+		[CommandBuffer commit];
+		cp_frame_end_submission(SwiftLayerFrame);
+	}
+
+	// cache the viewports in the delegate so code later can get it when asking about the screen bounds
+	AppDelegate.SwiftLayerViewports = Viewports;
+
+	// grab the commandline, skipping over the executable path (index 0), and pass it to unreal
+	NSArray* Arguments = [[NSProcessInfo processInfo] arguments];
+	Arguments = [Arguments subarrayWithRange:NSMakeRange(1, [Arguments count] - 1)];
+	GSavedCommandLine = [Arguments componentsJoinedByString:@" "];
+
+		// Xcode strips quotes, so replace ^ with "
+	GSavedCommandLine = GSavedCommandLine.Replace(TEXT("+"), TEXT("\""));
+	
+	FIOSCommandLineHelper::InitCommandArgs(FString());
+	
+	CGRect FirstViewport = [[AppDelegate.SwiftLayerViewports firstObject] CGRectValue];
+	
+	FPlatformMisc::LowLevelOutputDebugStringf(TEXT("Kicking off UE with Swift Layer. Commandline: %s\n"), *GSavedCommandLine);
+	FPlatformMisc::LowLevelOutputDebugStringf(TEXT("NumViews: %d, Full size = %f x %f\n", NumViews, FirstViewport.size.width, FirstViewport.size.height);
+
+	AppDelegate.SwiftLayer = Layer;
+}
+
+#endif
diff --git a/Engine/Source/Runtime/Launch/Private/IOS/UECppToSwift.h b/Engine/Source/Runtime/Launch/Private/IOS/UECppToSwift.h
new file mode 100644
--- /dev/null
+++ b/Engine/Source/Runtime/Launch/Private/IOS/UECppToSwift.h
@@ -0,0 +1,18 @@
+// Copyright Epic Games, Inc. All Rights Reserved.
+
+#pragma once
+
+#import <CompositorServices/CompositorServices.h>
+
+#define SWIFT_IMPORT
+// @todo do this better
+#define PLATFORM_VISIONOS 1
+
+#include "../../../ApplicationCore/Public/IOS/IOSAppDelegate.h"
+
+class FSwiftAppBootstrap
+{
+public:
+	static void KickoffWithCompositingLayer(CP_OBJECT_cp_layer_renderer* Layer);
+};
+
diff --git a/Engine/Source/Runtime/Launch/Private/VisionOS/UESwift.swift b/Engine/Source/Runtime/Launch/Private/VisionOS/UESwift.swift
deleted file mode 100644
