- Fixed remote building of IOS (from WIndows) with Modern Xcode
- Updated IPhonePackager.exe

diff --git a/Engine/Source/Programs/IOS/iPhonePackager/CodeSigning.cs b/Engine/Source/Programs/IOS/iPhonePackager/CodeSigning.cs
--- a/Engine/Source/Programs/IOS/iPhonePackager/CodeSigning.cs
+++ b/Engine/Source/Programs/IOS/iPhonePackager/CodeSigning.cs
@@ -910,8 +910,11 @@ namespace iPhonePackager
 					if (CodeSigningBlobLC != null)
 					{
 						RequirementsBlob OldRequirements = CodeSigningBlobLC.Payload.GetBlobByMagic(AbstractBlob.CSMAGIC_REQUIREMENTS_TABLE) as RequirementsBlob;
-						RequirementBlob OldRequire = OldRequirements.GetBlobByKey(0x00003) as RequirementBlob;
-						OldExpression = OldRequire.Expression;
+						if (OldRequirements.GetBlobByKey(0x00003) != null)
+						{
+							RequirementBlob OldRequire = OldRequirements.GetBlobByKey(0x00003) as RequirementBlob;
+							OldExpression = OldRequire.Expression;
+						}
 					}
 					FinalRequirementsBlob = RequirementsBlob.CreateEmpty();
 					FinalRequirementsBlob.Add(0x00003, RequirementBlob.CreateFromCertificate(SigningCert, CFBundleIdentifier, OldExpression));
diff --git a/Engine/Source/Programs/UnrealBuildTool/Platform/Apple/AppleExports.cs b/Engine/Source/Programs/UnrealBuildTool/Platform/Apple/AppleExports.cs
--- a/Engine/Source/Programs/UnrealBuildTool/Platform/Apple/AppleExports.cs
+++ b/Engine/Source/Programs/UnrealBuildTool/Platform/Apple/AppleExports.cs
@@ -174,7 +174,7 @@ namespace UnrealBuildTool
 
 			if (bForceDummySigning)
 			{
-				ExtraOptions += " CODE_SIGN_STYLE=Manual CODE_SIGN_IDENTITY=-";
+				ExtraOptions += " CODE_SIGN_STYLE=Manual CODE_SIGN_IDENTITY= PROVISIONING_PROFILE_SPECIFIER=";
 			}
 			else
 			{
diff --git a/Engine/Source/Programs/UnrealBuildTool/Platform/IOS/IOSToolChain.cs b/Engine/Source/Programs/UnrealBuildTool/Platform/IOS/IOSToolChain.cs
--- a/Engine/Source/Programs/UnrealBuildTool/Platform/IOS/IOSToolChain.cs
+++ b/Engine/Source/Programs/UnrealBuildTool/Platform/IOS/IOSToolChain.cs
@@ -744,6 +744,18 @@ namespace UnrealBuildTool
 			string ZipWorkingDir = String.Format("Payload/{0}.app/", GameName);
 			string ZipSourceDir = String.Format(bUseModernXcode ? "{0}/{1}.app" : "{0}/Payload/{1}.app", BinaryPath, GameName);
 
+			if (bUseModernXcode)
+			{
+				// with modern, we don't codesign when building the source .app due to keychain issues over ssh and Xcode wants provision setup, even if using the - identity
+				// so, now codesign right at the end with the dummy signature so that there is signature at all, which IPP on windows will replace
+				int ExitCode;
+				string Output = Utils.RunLocalProcessAndReturnStdOut("/bin/sh", $"-c 'codesign -f -s - \"{ZipSourceDir}\"", null, out ExitCode);
+				if (ExitCode != 0)
+				{
+					throw new BuildException($"Failed to codesign with dummy identity, which should never happen [ExitCode = {ExitCode}, output = {Output}");
+				}
+			}
+
 			// create the file
 			using (ZipFile Zip = new ZipFile())
 			{
