UnrealBuildTool: Ensure generated files are written to the correct intermediate folder name

#rnx

diff --git a/Engine/Source/Programs/UnrealBuildTool/Configuration/UEBuildTarget.cs b/Engine/Source/Programs/UnrealBuildTool/Configuration/UEBuildTarget.cs
--- a/Engine/Source/Programs/UnrealBuildTool/Configuration/UEBuildTarget.cs
+++ b/Engine/Source/Programs/UnrealBuildTool/Configuration/UEBuildTarget.cs
@@ -1102,15 +1102,15 @@ namespace UnrealBuildTool
 	public enum UnrealIntermediateEnvironment
 	{
 		/// <summary>
-		/// 
+		/// Default environment
 		/// </summary>
 		Default,
 		/// <summary>
-		/// 
+		/// Include what you use
 		/// </summary>
 		IWYU,
 		/// <summary>
-		/// 
+		/// Nonunity
 		/// </summary>
 		NonUnity,
 	}
@@ -1645,7 +1645,7 @@ namespace UnrealBuildTool
 			Platform = InDescriptor.Platform;
 			Configuration = InDescriptor.Configuration;
 			Architectures = InDescriptor.Architectures;
-			IntermediateEnvironment = InDescriptor.IntermediateEnvironment;
+			IntermediateEnvironment = InDescriptor.IntermediateEnvironment != UnrealIntermediateEnvironment.Default ? InDescriptor.IntermediateEnvironment : InRules.IntermediateEnvironment;
 			Rules = InRules;
 			RulesAssembly = InRulesAssembly;
 			TargetType = Rules.Type;
@@ -1705,11 +1705,11 @@ namespace UnrealBuildTool
 			}
 			else if (Configuration == UnrealTargetConfiguration.DebugGame)
 			{
-				EngineIntermediateDirectory = DirectoryReference.Combine(Unreal.EngineDirectory, PlatformIntermediateFolder, AppName, UnrealTargetConfiguration.Development.ToString());
+				EngineIntermediateDirectory = DirectoryReference.Combine(Unreal.EngineDirectory, PlatformIntermediateFolder, GetTargetIntermediateFolderName(TargetName, IntermediateEnvironment), UnrealTargetConfiguration.Development.ToString());
 			}
 			else
 			{
-				EngineIntermediateDirectory = DirectoryReference.Combine(Unreal.EngineDirectory, PlatformIntermediateFolder, AppName, Configuration.ToString());
+				EngineIntermediateDirectory = DirectoryReference.Combine(Unreal.EngineDirectory, PlatformIntermediateFolder, GetTargetIntermediateFolderName(TargetName, IntermediateEnvironment), Configuration.ToString());
 			}
 
 			// Get the receipt path for this target
@@ -4636,7 +4636,7 @@ namespace UnrealBuildTool
 			DirectoryReference IntermediateDirectory;
 			if (LaunchModule.RulesFile.IsUnderDirectory(Unreal.EngineDirectory) && !ShouldCompileMonolithic())
 			{
-				IntermediateDirectory = DirectoryReference.Combine(Unreal.EngineDirectory, PlatformIntermediateFolder, AppName, Configuration.ToString());
+				IntermediateDirectory = DirectoryReference.Combine(Unreal.EngineDirectory, PlatformIntermediateFolder, UEBuildTarget.GetTargetIntermediateFolderName(AppName, IntermediateEnvironment), Configuration.ToString());
 			}
 			else
 			{
@@ -5350,7 +5350,7 @@ namespace UnrealBuildTool
 					}
 
 					// Get the subfolder containing generated code - we don't need architeceture information since these are shared between all arches for a platform
-					GeneratedCodeDirectory = DirectoryReference.Combine(GeneratedCodeDirectory, PlatformIntermediateFolderNoArch, AppName, "Inc");
+					GeneratedCodeDirectory = DirectoryReference.Combine(GeneratedCodeDirectory, PlatformIntermediateFolderNoArch, GetTargetIntermediateFolderName(AppName, IntermediateEnvironment), "Inc");
 
 					// Append the binaries subfolder, if present. We rely on this to ensure that build products can be filtered correctly.
 					if (RulesObject.BinariesSubFolder != null)
diff --git a/Engine/Source/Programs/UnrealBuildTool/Modes/CleanMode.cs b/Engine/Source/Programs/UnrealBuildTool/Modes/CleanMode.cs
--- a/Engine/Source/Programs/UnrealBuildTool/Modes/CleanMode.cs
+++ b/Engine/Source/Programs/UnrealBuildTool/Modes/CleanMode.cs
@@ -93,6 +93,9 @@ namespace UnrealBuildTool
 					// Create the rules object
 					ReadOnlyTargetRules Target = new ReadOnlyTargetRules(RulesAssembly.CreateTargetRules(TargetDescriptor.Name, TargetDescriptor.Platform, TargetDescriptor.Configuration, TargetDescriptor.Architectures, TargetDescriptor.ProjectFile, TargetDescriptor.AdditionalArguments, Logger));
 
+					// Get the intermediate environment for this target
+					UnrealIntermediateEnvironment IntermediateEnvironment = TargetDescriptor.IntermediateEnvironment != UnrealIntermediateEnvironment.Default ? TargetDescriptor.IntermediateEnvironment : Target.IntermediateEnvironment;
+
 					if (!bSkipPreBuildTargets && Target.PreBuildTargets.Count > 0)
 					{
 						foreach (TargetInfo PreBuildTarget in Target.PreBuildTargets)
@@ -141,7 +144,7 @@ namespace UnrealBuildTool
 					}
 
 					// Add all the makefiles and caches to be deleted
-					FilesToDelete.Add(TargetMakefile.GetLocation(Target.ProjectFile, Target.Name, Target.Platform, Target.Architectures, Target.Configuration, TargetDescriptor.IntermediateEnvironment));
+					FilesToDelete.Add(TargetMakefile.GetLocation(Target.ProjectFile, Target.Name, Target.Platform, Target.Architectures, Target.Configuration, IntermediateEnvironment));
 					FilesToDelete.UnionWith(SourceFileMetadataCache.GetFilesToClean(Target.ProjectFile));
 
 					// Add all the intermediate folders to be deleted
@@ -149,7 +152,7 @@ namespace UnrealBuildTool
 					{
 						foreach (string NamePrefix in NamePrefixes)
 						{
-							string NamePrefixWithEnv = UEBuildTarget.GetTargetIntermediateFolderName(NamePrefix, TargetDescriptor.IntermediateEnvironment);
+							string NamePrefixWithEnv = UEBuildTarget.GetTargetIntermediateFolderName(NamePrefix, IntermediateEnvironment);
 							// This is actually wrong.. the generated code is not in this dir.. if changing "Target.Architectures" parameter to null it will delete the right files.
 							// However, this also means that it will cause a rebuild for both unity, nonunity and iwyu targets since they share this folder.
 							DirectoryReference GeneratedCodeDir = DirectoryReference.Combine(BaseDir, UEBuildTarget.GetPlatformIntermediateFolder(Target.Platform, Target.Architectures, false), NamePrefixWithEnv, "Inc");
