UnrealBuildTool: Write files for static analysis to a separate intermediate folder
Fix up to date checks for clang static analysis. Note that clang analysis warnings can't be made into error, so subsequent builds will not replay warnings for up to date actions

#rnx

diff --git a/Engine/Source/Programs/UnrealBuildTool/Configuration/TargetRules.cs b/Engine/Source/Programs/UnrealBuildTool/Configuration/TargetRules.cs
--- a/Engine/Source/Programs/UnrealBuildTool/Configuration/TargetRules.cs
+++ b/Engine/Source/Programs/UnrealBuildTool/Configuration/TargetRules.cs
@@ -1480,7 +1480,7 @@ namespace UnrealBuildTool
 		/// </summary>
 		public bool bUseUnityBuild
 		{
-			get => bUseUnityBuildOverride ?? IntermediateEnvironment == UnrealIntermediateEnvironment.Default && !bEnableCppModules;
+			get => bUseUnityBuildOverride ?? IntermediateEnvironment.IsUnity() && !bEnableCppModules;
 			set => bUseUnityBuildOverride = value;
 		}
 
@@ -2567,6 +2567,10 @@ namespace UnrealBuildTool
 				{
 					IntermediateEnvironment = UnrealIntermediateEnvironment.IWYU;
 				}
+				else if (StaticAnalyzer != StaticAnalyzer.None)
+				{
+					IntermediateEnvironment = UnrealIntermediateEnvironment.Analyze;
+				}
 				else if (!bUseUnityBuild)
 				{
 					IntermediateEnvironment = UnrealIntermediateEnvironment.NonUnity;
diff --git a/Engine/Source/Programs/UnrealBuildTool/Configuration/UEBuildTarget.cs b/Engine/Source/Programs/UnrealBuildTool/Configuration/UEBuildTarget.cs
--- a/Engine/Source/Programs/UnrealBuildTool/Configuration/UEBuildTarget.cs
+++ b/Engine/Source/Programs/UnrealBuildTool/Configuration/UEBuildTarget.cs
@@ -1113,6 +1113,32 @@ namespace UnrealBuildTool
 		/// Nonunity
 		/// </summary>
 		NonUnity,
+		/// <summary>
+		/// Static code analysis
+		/// </summary>
+		Analyze,
+	}
+
+	/// <summary>
+	/// Extension methods for UnrealIntermediateEnvironment enum
+	/// </summary>
+	public static class UEBuildTargetExtensions
+	{
+		/// <summary>
+		/// Returns if this environment should default to a unity build
+		/// </summary>
+		/// <param name="Environment">The environment to check</param>
+		/// <returns>true if shold build unity by default</returns>
+		public static bool IsUnity(this UnrealIntermediateEnvironment Environment)
+		{
+			switch (Environment)
+			{
+				case UnrealIntermediateEnvironment.IWYU:
+				case UnrealIntermediateEnvironment.NonUnity:
+					return false;
+			}
+			return true;
+		}
 	}
 
 	/// <summary>
@@ -1757,6 +1783,9 @@ namespace UnrealBuildTool
 				case UnrealIntermediateEnvironment.NonUnity:
 					TargetFolderName += "NU";
 					break;
+				case UnrealIntermediateEnvironment.Analyze:
+					TargetFolderName += "SA";
+					break;
 			}
 			return TargetFolderName;
 		}
diff --git a/Engine/Source/Programs/UnrealBuildTool/Platform/Windows/VCCompileAction.cs b/Engine/Source/Programs/UnrealBuildTool/Platform/Windows/VCCompileAction.cs
--- a/Engine/Source/Programs/UnrealBuildTool/Platform/Windows/VCCompileAction.cs
+++ b/Engine/Source/Programs/UnrealBuildTool/Platform/Windows/VCCompileAction.cs
@@ -164,6 +164,11 @@ namespace UnrealBuildTool
 		/// </summary>
 		public UnrealArch Architecture { get; set; }
 
+		/// <summary>
+		/// Whether this compile is static code analysis (used for display)
+		/// </summary>
+		public bool bIsAnalyzing { get; set; }
+
 		#region Public IAction implementation
 
 		/// <summary>
@@ -197,7 +202,7 @@ namespace UnrealBuildTool
 
 		IEnumerable<FileItem> IExternalAction.DeleteItems => DeleteItems;
 		public DirectoryReference WorkingDirectory => Unreal.EngineSourceDirectory;
-		string IExternalAction.CommandDescription => $"Compile [{Architecture}]";
+		string IExternalAction.CommandDescription => $"{(bIsAnalyzing ? "Analyze" : "Compile")} [{Architecture}]";
 		bool IExternalAction.bIsGCCCompiler => false;
 		bool IExternalAction.bProducesImportLibrary => false;
 		string IExternalAction.StatusDescription => (SourceFile == null) ? "Compiling" : SourceFile.Location.GetFileName();
diff --git a/Engine/Source/Programs/UnrealBuildTool/Platform/Windows/VCToolChain.cs b/Engine/Source/Programs/UnrealBuildTool/Platform/Windows/VCToolChain.cs
--- a/Engine/Source/Programs/UnrealBuildTool/Platform/Windows/VCToolChain.cs
+++ b/Engine/Source/Programs/UnrealBuildTool/Platform/Windows/VCToolChain.cs
@@ -1497,6 +1497,13 @@ namespace UnrealBuildTool
 					CompileAction.PreprocessedFile = FileItem.GetItemByFileReference(FileReference.Combine(OutputDir, FileName + ".i"));
 					CompileAction.ResponseFile = FileItem.GetItemByFileReference(GetResponseFileName(CompileEnvironment, CompileAction.PreprocessedFile));
 				}
+				else if (Target.WindowsPlatform.Compiler.IsClang() && Target.StaticAnalyzer == StaticAnalyzer.Default)
+				{
+					// Clang analysis does not actually create an object, use the dependency list as the response filename
+					string DependencyListFilename = FileName + ".d";
+					FileItem DependencyListFile = FileItem.GetItemByFileReference(FileReference.Combine(OutputDir, DependencyListFilename));
+					CompileAction.ResponseFile = FileItem.GetItemByFileReference(GetResponseFileName(CompileEnvironment, DependencyListFile));
+				}
 				else
 				{
 					// Add the object file to the produced item list.
@@ -1504,7 +1511,7 @@ namespace UnrealBuildTool
 					FileItem ObjectFile = FileItem.GetItemByFileReference(FileReference.Combine(OutputDir, ObjectLeafFilename));
 
 					CompileAction.ObjectFile = ObjectFile;
-					CompileAction.ResponseFile = FileItem.GetItemByFileReference(GetResponseFileName(CompileEnvironment, ObjectFile)); ;
+					CompileAction.ResponseFile = FileItem.GetItemByFileReference(GetResponseFileName(CompileEnvironment, ObjectFile));
 
 					if (Target.WindowsPlatform.ObjSrcMapFile != null)
 					{
@@ -1689,6 +1696,8 @@ namespace UnrealBuildTool
 					CompileAction.WriteResponseFile(Graph, Logger);
 				}
 
+				CompileAction.bIsAnalyzing = Target.StaticAnalyzer != StaticAnalyzer.None;
+
 				// When compiling with SN-DBS, modules that contain a #import must be built locally
 				CompileAction.bCanExecuteRemotelyWithSNDBS = CompileAction.bCanExecuteRemotely;
 				if (CompileEnvironment.bBuildLocallyWithSNDBS == true)
@@ -1703,6 +1712,11 @@ namespace UnrealBuildTool
 
 			CPPOutput Result = new CPPOutput();
 			Result.ObjectFiles.AddRange(Actions.Where(x => x.ObjectFile != null || x.PreprocessedFile != null).Select(x => x.ObjectFile != null ? x.ObjectFile! : x.PreprocessedFile!));
+			// Clang static analysis doesn't create object files, so treat the dependency list file as the output
+			if (Target.WindowsPlatform.Compiler.IsClang() && Target.StaticAnalyzer == StaticAnalyzer.Default)
+			{
+				Result.ObjectFiles.AddRange(Actions.Where(x => x.DependencyListFile != null ).Select(x => x.DependencyListFile!));
+			}
 			Result.CompiledModuleInterfaces.AddRange(Actions.Where(x => x.CompiledModuleInterfaceFile != null).Select(x => x.CompiledModuleInterfaceFile!));
 			Result.PrecompiledHeaderFile = Actions.Select(x => x.CreatePchFile).Where(x => x != null).FirstOrDefault();
 			return Result;
diff --git a/Engine/Source/Programs/UnrealBuildTool/ToolChain/ClangToolChain.cs b/Engine/Source/Programs/UnrealBuildTool/ToolChain/ClangToolChain.cs
--- a/Engine/Source/Programs/UnrealBuildTool/ToolChain/ClangToolChain.cs
+++ b/Engine/Source/Programs/UnrealBuildTool/ToolChain/ClangToolChain.cs
@@ -924,6 +924,12 @@ namespace UnrealBuildTool
 				// this is parsed by external tools wishing to open this file directly.
 				Logger.LogInformation("PreProcessPath: {File}", OutputFile.AbsolutePath);
 			}
+			else if (IsAnalyzing(CompileEnvironment))
+			{
+				// Clang analysis does not actually create an object, use the dependency list as the response filename
+				OutputFile = FileItem.GetItemByFileReference(FileReference.Combine(OutputDir, GetFileNameFromExtension(FileName, ".d")));
+				CompileResult.ObjectFiles.Add(OutputFile);
+			}
 			else
 			{
 				string ObjectFileExtension = (CompileEnvironment.AdditionalArguments != null && CompileEnvironment.AdditionalArguments.Contains("-emit-llvm")) ? ".bc" : ".o";
@@ -937,6 +943,15 @@ namespace UnrealBuildTool
 			// Add the source file path to the command-line.
 			Arguments.Add(GetSourceFileArgument(SourceFile));
 
+			// Generate the included header dependency list
+			if (!PreprocessDepends)
+			{
+				FileItem DependencyListFile = FileItem.GetItemByFileReference(FileReference.Combine(OutputDir, GetFileNameFromExtension(FileName, ".d")));
+				Arguments.Add(GetDepencenciesListFileArgument(DependencyListFile));
+				CompileAction.DependencyListFile = DependencyListFile;
+				CompileAction.ProducedItems.Add(DependencyListFile);
+			}
+
 			if (!IsAnalyzing(CompileEnvironment))
 			{
 				// Generate the timing info
@@ -947,23 +962,10 @@ namespace UnrealBuildTool
 					CompileAction.ProducedItems.Add(TraceFile);
 				}
 
-				// Generate the included header dependency list
-				if (!PreprocessDepends)
-				{
-					FileItem DependencyListFile = FileItem.GetItemByFileReference(FileReference.Combine(OutputDir, GetFileNameFromExtension(FileName, ".d")));
-					Arguments.Add(GetDepencenciesListFileArgument(DependencyListFile));
-					CompileAction.DependencyListFile = DependencyListFile;
-					CompileAction.ProducedItems.Add(DependencyListFile);
-				}
-			}
-			else
-			{
-				OutputFile = FileItem.GetItemByFileReference(FileReference.Combine(OutputDir, GetFileNameFromExtension(FileName, ".analysis")));
+				// Add the parameters needed to compile the output file to the command-line.
+				Arguments.Add(GetOutputFileArgument(OutputFile));
 			}
 
-			// Add the parameters needed to compile the output file to the command-line.
-			Arguments.Add(GetOutputFileArgument(OutputFile));
-
 			return OutputFile;
 		}
 
