Added support to have multiple architectures per platform in the VS sln/projects.
Updated the project generators to search the VS platform directory when HasVisualStudioSupport is called to find out whether a platform is supported. This fixes several issues where platforms were being generated like they were supported by the host machine but actually weren't.

diff --git a/Engine/Source/Programs/UnrealBuildTool/Platform/Android/AndroidProjectGenerator.cs b/Engine/Source/Programs/UnrealBuildTool/Platform/Android/AndroidProjectGenerator.cs
--- a/Engine/Source/Programs/UnrealBuildTool/Platform/Android/AndroidProjectGenerator.cs
+++ b/Engine/Source/Programs/UnrealBuildTool/Platform/Android/AndroidProjectGenerator.cs
@@ -72,26 +72,16 @@ namespace UnrealBuildTool
 			yield return UnrealTargetPlatform.Android;
 		}
 
-		/// <summary>
-		/// Whether this build platform has native support for VisualStudio
-		/// </summary>
-		/// <param name="InPlatform">  The UnrealTargetPlatform being built</param>
-		/// <param name="InConfiguration"> The UnrealTargetConfiguration being built</param>
-		/// <param name="ProjectFileFormat"></param>
-		/// <returns>bool    true if native VisualStudio support (or custom VSI) is available</returns>
-		public override bool HasVisualStudioSupport(UnrealTargetPlatform InPlatform, UnrealTargetConfiguration InConfiguration, VCProjectFileFormat ProjectFileFormat)
+		/// <inheritdoc/>
+		public override bool HasVisualStudioSupport(UnrealTargetPlatform InPlatform, UnrealTargetConfiguration InConfiguration, VCProjectFileFormat ProjectFileFormat, DirectoryReference InProjectDir, UnrealArch? InArch)
 		{
 			// Debugging, etc. are dependent on the TADP being installed
 			return AGDEInstalled;
 		}
 
-		/// <summary>
-		/// Return the VisualStudio platform name for this build platform
-		/// </summary>
-		/// <param name="InPlatform">  The UnrealTargetPlatform being built</param>
-		/// <param name="InConfiguration"> The UnrealTargetConfiguration being built</param>
-		/// <returns>string    The name of the platform that VisualStudio recognizes</returns>
-		public override string GetVisualStudioPlatformName(UnrealTargetPlatform InPlatform, UnrealTargetConfiguration InConfiguration)
+
+		/// <inheritdoc/>
+		public override string GetVisualStudioPlatformName(UnrealTargetPlatform InPlatform, UnrealTargetConfiguration InConfiguration, DirectoryReference InProjectDir, UnrealArch? InArch)
 		{
 			string PlatformName = InPlatform.ToString();
 
diff --git a/Engine/Source/Programs/UnrealBuildTool/Platform/IOS/IOSProjectGenerator.cs b/Engine/Source/Programs/UnrealBuildTool/Platform/IOS/IOSProjectGenerator.cs
--- a/Engine/Source/Programs/UnrealBuildTool/Platform/IOS/IOSProjectGenerator.cs
+++ b/Engine/Source/Programs/UnrealBuildTool/Platform/IOS/IOSProjectGenerator.cs
@@ -29,17 +29,8 @@ namespace UnrealBuildTool
 			yield return UnrealTargetPlatform.IOS;
 		}
 
-		///
-		///	VisualStudio project generation functions
-		///	
-		/// <summary>
-		/// Whether this build platform has native support for VisualStudio
-		/// </summary>
-		/// <param name="InPlatform">  The UnrealTargetPlatform being built</param>
-		/// <param name="InConfiguration"> The UnrealTargetConfiguration being built</param>
-		/// <param name="ProjectFileFormat"></param>
-		/// <returns>bool    true if native VisualStudio support (or custom VSI) is available</returns>
-		public override bool HasVisualStudioSupport(UnrealTargetPlatform InPlatform, UnrealTargetConfiguration InConfiguration, VCProjectFileFormat ProjectFileFormat)
+		/// <inheritdoc/>
+		public override bool HasVisualStudioSupport(UnrealTargetPlatform InPlatform, UnrealTargetConfiguration InConfiguration, VCProjectFileFormat ProjectFileFormat, DirectoryReference InProjectDir, UnrealArch? InArch)
 		{
 			// iOS is not supported in VisualStudio
 			return false;
diff --git a/Engine/Source/Programs/UnrealBuildTool/Platform/Linux/LinuxProjectGenerator.cs b/Engine/Source/Programs/UnrealBuildTool/Platform/Linux/LinuxProjectGenerator.cs
--- a/Engine/Source/Programs/UnrealBuildTool/Platform/Linux/LinuxProjectGenerator.cs
+++ b/Engine/Source/Programs/UnrealBuildTool/Platform/Linux/LinuxProjectGenerator.cs
@@ -30,17 +30,8 @@ namespace UnrealBuildTool
 			yield return UnrealTargetPlatform.LinuxArm64;
 		}
 
-		///
-		///	VisualStudio project generation functions
-		///	
-		/// <summary>
-		/// Whether this build platform has native support for VisualStudio
-		/// </summary>
-		/// <param name="InPlatform">  The UnrealTargetPlatform being built</param>
-		/// <param name="InConfiguration"> The UnrealTargetConfiguration being built</param>
-		/// <param name="ProjectFileFormat"></param>
-		/// <returns>bool    true if native VisualStudio support (or custom VSI) is available</returns>
-		public override bool HasVisualStudioSupport(UnrealTargetPlatform InPlatform, UnrealTargetConfiguration InConfiguration, VCProjectFileFormat ProjectFileFormat)
+		/// <inheritdoc/>
+		public override bool HasVisualStudioSupport(UnrealTargetPlatform InPlatform, UnrealTargetConfiguration InConfiguration, VCProjectFileFormat ProjectFileFormat, DirectoryReference InProjectDir, UnrealArch? InArch)
 		{
 			return false;
 		}
diff --git a/Engine/Source/Programs/UnrealBuildTool/Platform/Mac/MacProjectGenerator.cs b/Engine/Source/Programs/UnrealBuildTool/Platform/Mac/MacProjectGenerator.cs
--- a/Engine/Source/Programs/UnrealBuildTool/Platform/Mac/MacProjectGenerator.cs
+++ b/Engine/Source/Programs/UnrealBuildTool/Platform/Mac/MacProjectGenerator.cs
@@ -29,17 +29,8 @@ namespace UnrealBuildTool
 			yield return UnrealTargetPlatform.Mac;
 		}
 
-		///
-		///	VisualStudio project generation functions
-		///	
-		/// <summary>
-		/// Whether this build platform has native support for VisualStudio
-		/// </summary>
-		/// <param name="InPlatform">  The UnrealTargetPlatform being built</param>
-		/// <param name="InConfiguration"> The UnrealTargetConfiguration being built</param>
-		/// <param name="ProjectFileFormat"></param>
-		/// <returns>bool    true if native VisualStudio support (or custom VSI) is available</returns>
-		public override bool HasVisualStudioSupport(UnrealTargetPlatform InPlatform, UnrealTargetConfiguration InConfiguration, VCProjectFileFormat ProjectFileFormat)
+		/// <inheritdoc/>
+		public override bool HasVisualStudioSupport(UnrealTargetPlatform InPlatform, UnrealTargetConfiguration InConfiguration, VCProjectFileFormat ProjectFileFormat, DirectoryReference InProjectDir, UnrealArch? InArch)
 		{
 			// Mac is not supported in VisualStudio
 			return false;
diff --git a/Engine/Source/Programs/UnrealBuildTool/Platform/Windows/WindowsProjectGenerator.cs b/Engine/Source/Programs/UnrealBuildTool/Platform/Windows/WindowsProjectGenerator.cs
--- a/Engine/Source/Programs/UnrealBuildTool/Platform/Windows/WindowsProjectGenerator.cs
+++ b/Engine/Source/Programs/UnrealBuildTool/Platform/Windows/WindowsProjectGenerator.cs
@@ -22,32 +22,31 @@ namespace UnrealBuildTool
 		{
 		}
 
-		/// <summary>
-		/// Enumerate all the platforms that this generator supports
-		/// </summary>
+		/// <inheritdoc/>
 		public override IEnumerable<UnrealTargetPlatform> GetPlatforms()
 		{
 			yield return UnrealTargetPlatform.Win64;
 		}
 
-		///
-		///	VisualStudio project generation functions
-		///	
-		/// <summary>
-		/// Return the VisualStudio platform name for this build platform
-		/// </summary>
-		/// <param name="InPlatform">  The UnrealTargetPlatform being built</param>
-		/// <param name="InConfiguration"> The UnrealTargetConfiguration being built</param>
-		/// <returns>string    The name of the platform that VisualStudio recognizes</returns>
-		public override string GetVisualStudioPlatformName(UnrealTargetPlatform InPlatform, UnrealTargetConfiguration InConfiguration)
+		/// <inheritdoc/>
+		public override string GetVisualStudioPlatformName(UnrealTargetPlatform InPlatform, UnrealTargetConfiguration InConfiguration, DirectoryReference InProjectDir, UnrealArch? InArch)
 		{
 			if (InPlatform == UnrealTargetPlatform.Win64)
 			{
+				if (InArch == UnrealArch.Arm64)
+				{
+					return "arm64";
+				}
+				else if (InArch == UnrealArch.Arm64ec)
+				{
+					return "arm64ec";
+				}
 				return "x64";
 			}
 			return InPlatform.ToString();
 		}
 
+		/// <inheritdoc/>
 		public override string GetVisualStudioUserFileStrings(UnrealTargetPlatform InPlatform, UnrealTargetConfiguration InConfiguration, string InConditionString, TargetRules InTargetRules, FileReference TargetRulesPath, FileReference ProjectFilePath, string ProjectName, string? ForeignUProjectPath)
 		{
 			StringBuilder VCUserFileContent = new StringBuilder();
@@ -75,6 +74,7 @@ namespace UnrealBuildTool
 			return VCUserFileContent.ToString();
 		}
 
+		/// <inheritdoc/>
 		public override bool RequiresVSUserFileGeneration()
 		{
 			return true;
diff --git a/Engine/Source/Programs/UnrealBuildTool/ProjectFiles/PlatformProjectGenerator.cs b/Engine/Source/Programs/UnrealBuildTool/ProjectFiles/PlatformProjectGenerator.cs
--- a/Engine/Source/Programs/UnrealBuildTool/ProjectFiles/PlatformProjectGenerator.cs
+++ b/Engine/Source/Programs/UnrealBuildTool/ProjectFiles/PlatformProjectGenerator.cs
@@ -1,6 +1,7 @@
 // Copyright Epic Games, Inc. All Rights Reserved.
 
 using System;
+using System.Collections.Concurrent;
 using System.Collections.Generic;
 using System.Text;
 using EpicGames.Core;
@@ -47,20 +48,60 @@ namespace UnrealBuildTool
 			return false;
 		}
 
-		///
-		///	VisualStudio project generation functions
-		///	
+		ConcurrentDictionary<int, bool> FoundVSISupportDict = new();
+
 		/// <summary>
-		/// Whether this build platform has native support for VisualStudio
+		/// Checks the local VS install directories to see if the platform is supported.
 		/// </summary>
 		/// <param name="InPlatform">  The UnrealTargetPlatform being built</param>
 		/// <param name="InConfiguration"> The UnrealTargetConfiguration being built</param>
 		/// <param name="ProjectFileFormat">The visual studio project file format being generated</param>
-		/// <returns>bool    true if native VisualStudio support (or custom VSI) is available</returns>
-		public virtual bool HasVisualStudioSupport(UnrealTargetPlatform InPlatform, UnrealTargetConfiguration InConfiguration, VCProjectFileFormat ProjectFileFormat)
+		/// <param name="InProjectDir">  The root directory for the current project</param>
+		/// <param name="InArch">  The architecture of the platform</param>
+		/// <returns>bool true if native VisualStudio support (or custom VSI) is available</returns>
+		public virtual bool HasVisualStudioSupport(UnrealTargetPlatform InPlatform, UnrealTargetConfiguration InConfiguration, VCProjectFileFormat ProjectFileFormat, DirectoryReference InProjectDir, UnrealArch? InArch)
 		{
-			// By default, we assume this is true
-			return true;
+			int HashResult = 691;
+			HashResult *= InPlatform.ToString().GetHashCode();
+			HashResult *= InConfiguration.ToString().GetHashCode();
+			HashResult *= ProjectFileFormat.ToString().GetHashCode();
+			HashResult *= InProjectDir.GetHashCode();
+			HashResult *= InArch != null ? InArch.ToString()!.GetHashCode() : 1;
+
+			return FoundVSISupportDict.GetOrAdd(HashResult, _ =>
+			{
+				WindowsCompiler VSCompiler;
+				string VCVersion;
+
+				switch (ProjectFileFormat)
+				{
+					case VCProjectFileFormat.VisualStudio2019:
+						VSCompiler = WindowsCompiler.VisualStudio2019;
+						VCVersion = "v160";
+						break;
+					case VCProjectFileFormat.VisualStudio2022:
+						VSCompiler = WindowsCompiler.VisualStudio2022;
+						VCVersion = "v170";
+						break;
+					default:
+						// Unknown VS Version
+						return false;
+				}
+
+				IEnumerable<DirectoryReference>? InstallDirs = WindowsPlatform.TryGetVSInstallDirs(VSCompiler, Logger);
+				if (InstallDirs != null)
+				{
+					foreach (DirectoryReference VSInstallDir in InstallDirs)
+					{
+						DirectoryReference PlatformsPath = new DirectoryReference(System.IO.Path.Combine(VSInstallDir.FullName, "MSBuild\\Microsoft\\VC\\", VCVersion, "Platforms", GetVisualStudioPlatformName(InPlatform, InConfiguration, InProjectDir, InArch)));
+						if (DirectoryReference.Exists(PlatformsPath))
+						{
+							return true;
+						}
+					}
+				}
+				return false;
+			});
 		}
 
 		/// <summary>
@@ -68,26 +109,15 @@ namespace UnrealBuildTool
 		/// </summary>
 		/// <param name="InPlatform">  The UnrealTargetPlatform being built</param>
 		/// <param name="InConfiguration"> The UnrealTargetConfiguration being built</param>
+		/// <param name="InProjectDir">  The root directory for the current project</param>
+		/// <param name="InArch">  The architecture of the platform</param>
 		/// <returns>string    The name of the platform that VisualStudio recognizes</returns>
-		public virtual string GetVisualStudioPlatformName(UnrealTargetPlatform InPlatform, UnrealTargetConfiguration InConfiguration)
+		public virtual string GetVisualStudioPlatformName(UnrealTargetPlatform InPlatform, UnrealTargetConfiguration InConfiguration, DirectoryReference InProjectDir, UnrealArch? InArch)
 		{
 			// By default, return the platform string
 			return InPlatform.ToString();
 		}
 
-		/// <summary>
-		/// Return the VisualStudio platform name for this build platform
-		/// </summary>
-		/// <param name="InPlatform">  The UnrealTargetPlatform being built</param>
-		/// <param name="InConfiguration"> The UnrealTargetConfiguration being built</param>
-		/// <param name="InProjectDir">  The root directory for the current project</param>
-		/// <returns>string    The name of the platform that VisualStudio recognizes</returns>
-		public virtual string GetVisualStudioPlatformName(UnrealTargetPlatform InPlatform, UnrealTargetConfiguration InConfiguration, DirectoryReference InProjectDir)
-		{
-			// By default, return the project-independent platform string
-			return GetVisualStudioPlatformName(InPlatform, InConfiguration);
-		}
-
 		/// <summary>
 		/// Return whether we need a distinct platform name - for cases where there are two targets of the same VS solution platform
 		/// </summary>
diff --git a/Engine/Source/Programs/UnrealBuildTool/ProjectFiles/VisualStudio/VCProject.cs b/Engine/Source/Programs/UnrealBuildTool/ProjectFiles/VisualStudio/VCProject.cs
--- a/Engine/Source/Programs/UnrealBuildTool/ProjectFiles/VisualStudio/VCProject.cs
+++ b/Engine/Source/Programs/UnrealBuildTool/ProjectFiles/VisualStudio/VCProject.cs
@@ -6,6 +6,7 @@ using System.Collections.Generic;
 using System.Diagnostics.CodeAnalysis;
 using System.IO;
 using System.Linq;
+using System.Runtime.InteropServices;
 using System.Security;
 using System.Text;
 using System.Text.RegularExpressions;
@@ -161,9 +162,10 @@ namespace UnrealBuildTool
 		/// <param name="SolutionConfiguration">The solution configuration</param>
 		/// <param name="SolutionPlatform">The solution platform</param>
 		/// <param name="PlatformProjectGenerators">Set of platform project generators</param>
+		/// <param name="Architecture">The target architecture</param>
 		/// <param name="Logger">Logger for output</param>
 		/// <returns>Project context matching the given solution context</returns>
-		public abstract MSBuildProjectContext GetMatchingProjectContext(TargetType SolutionTarget, UnrealTargetConfiguration SolutionConfiguration, UnrealTargetPlatform SolutionPlatform, PlatformProjectGeneratorCollection PlatformProjectGenerators, ILogger Logger);
+		public abstract MSBuildProjectContext GetMatchingProjectContext(TargetType SolutionTarget, UnrealTargetConfiguration SolutionConfiguration, UnrealTargetPlatform SolutionPlatform, PlatformProjectGeneratorCollection PlatformProjectGenerators, UnrealArch? Architecture, ILogger Logger);
 
 		/// <summary>
 		/// Checks to see if the specified solution platform and configuration is able to map to this project
@@ -284,7 +286,7 @@ namespace UnrealBuildTool
 		/// This is the platform name that Visual Studio is always guaranteed to support.  We'll use this as
 		/// a platform for any project configurations where our actual platform is not supported by the
 		/// installed version of Visual Studio (e.g, "iOS")
-		public const string DefaultPlatformName = "Win64";
+		public const string DefaultPlatformName = "x64";
 
 		// This is the GUID that Visual Studio uses to identify a C++ project file in the solution
 		public override string ProjectTypeGUID => "{8BC9CEB8-8B4A-11D0-8D11-00A0C91BC942}";
@@ -309,67 +311,8 @@ namespace UnrealBuildTool
 			this.Settings = Settings;
 		}
 
-		/// <summary>
-		/// Given a target platform and configuration, generates a platform and configuration name string to use in Visual Studio projects.
-		/// Unlike with solution configurations, Visual Studio project configurations only support certain types of platforms, so we'll
-		/// generate a configuration name that has the platform "built in", and use a default platform type
-		/// </summary>
-		/// <param name="Platform">Actual platform</param>
-		/// <param name="Configuration">Actual configuration</param>
-		/// <param name="TargetConfigurationName">The configuration name from the target rules, or null if we don't have one</param>
-		/// <param name="PlatformProjectGenerators">Set of platform project generators</param>
-		/// <param name="Logger"></param>
-		/// <param name="ProjectPlatformName">Name of platform string to use for Visual Studio project</param>
-		/// <param name="ProjectConfigurationName">Name of configuration string to use for Visual Studio project</param>
-		private void MakeProjectPlatformAndConfigurationNames(UnrealTargetPlatform Platform, UnrealTargetConfiguration Configuration, TargetType TargetConfigurationName, PlatformProjectGeneratorCollection PlatformProjectGenerators, ILogger Logger, out string ProjectPlatformName, out string ProjectConfigurationName)
-		{
-			PlatformProjectGenerator? PlatformProjectGenerator = PlatformProjectGenerators.GetPlatformProjectGenerator(Platform, bInAllowFailure: true);
-
-			// Check to see if this platform is supported directly by Visual Studio projects.
-			bool HasActualVSPlatform = (PlatformProjectGenerator != null) ? PlatformProjectGenerator.HasVisualStudioSupport(Platform, Configuration, ProjectFileFormat) : false;
-
-			if (HasActualVSPlatform)
-			{
-				// Great!  Visual Studio supports this platform natively, so we don't need to make up
-				// a fake project configuration name.
-
-				// Allow the platform to specify the name used in VisualStudio.
-				// Note that the actual name of the platform on the Visual Studio side may be different than what
-				// UnrealBuildTool calls it (e.g. "Win64" -> "x64".) GetVisualStudioPlatformName() will figure this out.
-				ProjectConfigurationName = Configuration.ToString();
-				ProjectPlatformName = PlatformProjectGenerator!.GetVisualStudioPlatformName(Platform, Configuration, BaseDir);
-
-				// The project generator may require a distinct configuration name - typically when two UnrealTargetPlatforms need the same ProjectPlatformName - otherwise the properties overwrite each oter.
-				if (PlatformProjectGenerator!.RequiresDistinctVisualStudioConfigurationName(Platform, Configuration, BaseDir))
-				{
-					ProjectConfigurationName = Platform.ToString() + "_" + Configuration.ToString();
-				}
-			}
-			else
-			{
-				// Visual Studio doesn't natively support this platform, so we fake it by mapping it to
-				// a project configuration that has the platform name in that configuration as a suffix,
-				// and then using "Win32" as the actual VS platform name
-				ProjectConfigurationName = Platform.ToString() + "_" + Configuration.ToString();
-				ProjectPlatformName = DefaultPlatformName;
-			}
-
-			if (!bMakeProjectPerTarget && TargetConfigurationName != TargetType.Game)
-			{
-				ProjectConfigurationName += "_" + TargetConfigurationName.ToString();
-			}
-		}
-
-		/// <summary>
-		/// Get the project context for the given solution context
-		/// </summary>
-		/// <param name="SolutionTarget">The solution target type</param>
-		/// <param name="SolutionConfiguration">The solution configuration</param>
-		/// <param name="SolutionPlatform">The solution platform</param>
-		/// <param name="PlatformProjectGenerators">Set of platform project generations</param>
-		/// <param name="Logger"></param>
-		/// <returns>Project context matching the given solution context</returns>
-		public override MSBuildProjectContext GetMatchingProjectContext(TargetType SolutionTarget, UnrealTargetConfiguration SolutionConfiguration, UnrealTargetPlatform SolutionPlatform, PlatformProjectGeneratorCollection PlatformProjectGenerators, ILogger Logger)
+		/// <inheritdoc/>
+		public override MSBuildProjectContext GetMatchingProjectContext(TargetType SolutionTarget, UnrealTargetConfiguration SolutionConfiguration, UnrealTargetPlatform SolutionPlatform, PlatformProjectGeneratorCollection PlatformProjectGenerators, UnrealArch? Architecture, ILogger Logger)
 		{
 			// Stub projects always build in the same configuration
 			if (IsStubProject)
@@ -470,7 +413,24 @@ namespace UnrealBuildTool
 					}
 					else
 					{
-						MakeProjectPlatformAndConfigurationNames(ProjectPlatform, ProjectConfiguration, TargetConfigurationName, PlatformProjectGenerators, Logger, out ProjectPlatformName, out ProjectConfigurationName);
+						if (ProjectConfigAndTargetCombinations == null)
+						{
+							throw new BuildException("Project config and target combinations has not been populated.");
+						}
+
+						ProjectConfigAndTargetCombination? Combination = ProjectConfigAndTargetCombinations.FirstOrDefault(Combination =>
+							Combination.Platform == ProjectPlatform &&
+							Combination.Configuration == ProjectConfiguration &&
+							Combination.ProjectTarget == MatchingProjectTarget &&
+							Combination.Architecture == Architecture);
+
+						if (Combination == null)
+						{
+							throw new BuildException("Could not find the project config/platform combination in the generated list.");
+						}
+
+						ProjectPlatformName = Combination.ProjectPlatformName;
+						ProjectConfigurationName = Combination.ProjectConfigurationName;
 					}
 
 					// Set whether this project configuration should be built when the user initiates "build solution"
@@ -503,26 +463,28 @@ namespace UnrealBuildTool
 
 		class ProjectConfigAndTargetCombination
 		{
-			public UnrealTargetPlatform? Platform;
-			public UnrealTargetConfiguration Configuration;
-			public string ProjectPlatformName;
-			public string ProjectConfigurationName;
-			public ProjectTarget? ProjectTarget;
+			readonly public UnrealTargetPlatform? Platform;
+			readonly public UnrealTargetConfiguration Configuration;
+			readonly public string ProjectPlatformName;
+			readonly public string ProjectConfigurationName;
+			readonly public ProjectTarget? ProjectTarget;
+			readonly public UnrealArch? Architecture;
 
-			public ProjectConfigAndTargetCombination(UnrealTargetPlatform? InPlatform, UnrealTargetConfiguration InConfiguration, string InProjectPlatformName, string InProjectConfigurationName, ProjectTarget? InProjectTarget)
+			public ProjectConfigAndTargetCombination(UnrealTargetPlatform? InPlatform, UnrealTargetConfiguration InConfiguration, string InProjectPlatformName, string InProjectConfigurationName, ProjectTarget? InProjectTarget, UnrealArch? InArchitecture)
 			{
 				Platform = InPlatform;
 				Configuration = InConfiguration;
 				ProjectPlatformName = InProjectPlatformName;
 				ProjectConfigurationName = InProjectConfigurationName;
 				ProjectTarget = InProjectTarget;
+				Architecture = InArchitecture;
 			}
 
 			public string? ProjectConfigurationAndPlatformName => (ProjectPlatformName == null) ? null : (ProjectConfigurationName + "|" + ProjectPlatformName);
 
 			public override string ToString()
 			{
-				return String.Format("{0} {1} {2}", ProjectTarget, Platform, Configuration);
+				return String.Format("{0} {1} {2}", ProjectTarget, Platform, Configuration, Architecture != null ? " " + Architecture : string.Empty);
 			}
 		}
 
@@ -675,6 +637,7 @@ namespace UnrealBuildTool
 			//no need to do this more than once
 			if (ProjectConfigAndTargetCombinations == null)
 			{
+				HashSet<string> ProjectConfigAndTargets = new();
 				// Build up a list of platforms and configurations this project will support.  In this list, Unknown simply
 				// means that we should use the default "stub" project platform and configuration name.
 
@@ -682,7 +645,7 @@ namespace UnrealBuildTool
 				ProjectConfigAndTargetCombinations = new List<ProjectConfigAndTargetCombination>();
 				if (IsStubProject)
 				{
-					ProjectConfigAndTargetCombination StubCombination = new ProjectConfigAndTargetCombination(UnrealTargetPlatform.Parse(StubProjectPlatformName), UnrealTargetConfiguration.Unknown, StubProjectPlatformName, StubProjectConfigurationName, null);
+					ProjectConfigAndTargetCombination StubCombination = new ProjectConfigAndTargetCombination(UnrealTargetPlatform.Parse(StubProjectPlatformName), UnrealTargetConfiguration.Unknown, StubProjectPlatformName, StubProjectConfigurationName, null, null);
 					ProjectConfigAndTargetCombinations.Add(StubCombination);
 				}
 				else
@@ -712,13 +675,58 @@ namespace UnrealBuildTool
 
 								foreach (ProjectTarget ProjectTarget in ProjectTargets.OfType<ProjectTarget>())
 								{
-									if (IsValidProjectPlatformAndConfiguration(ProjectTarget, Platform, Configuration, Logger))
+									if (!IsValidProjectPlatformAndConfiguration(ProjectTarget, Platform, Configuration, Logger))
 									{
-										string ProjectPlatformName, ProjectConfigurationName;
-										MakeProjectPlatformAndConfigurationNames(Platform, Configuration, ProjectTarget.TargetRules!.Type, PlatformProjectGenerators, Logger, out ProjectPlatformName, out ProjectConfigurationName);
+										continue;
+									}
+
+									foreach (UnrealArch? Architecture in VCProjectFileGenerator.GetPlatformArchitectures(BuildPlatform))
+									{
+										PlatformProjectGenerator? PlatformProjectGenerator = PlatformProjectGenerators.GetPlatformProjectGenerator(Platform, bInAllowFailure: true);
+										string ProjectPlatformName;
+										string ProjectConfigurationName = Configuration.ToString();
+										bool CreateDistinctConfigName = false;
+
+										// Check to see if this platform is supported directly by Visual Studio projects.
+										if (PlatformProjectGenerator != null && PlatformProjectGenerator.HasVisualStudioSupport(Platform, Configuration, ProjectFileFormat, BaseDir, Architecture))
+										{
+											// Allow the platform to specify the name used in VisualStudio.
+											// Note that the actual name of the platform on the Visual Studio side may be different than what
+											// UnrealBuildTool calls it (e.g. "Win64" -> "x64".) GetVisualStudioPlatformName() will figure this out.
+											ProjectPlatformName = PlatformProjectGenerator.GetVisualStudioPlatformName(Platform, Configuration, BaseDir, Architecture);
+
+											// The project generator may require a distinct configuration name - typically when two UnrealTargetPlatforms need the same ProjectPlatformName - otherwise the properties overwrite each oter.
+											if (PlatformProjectGenerator.RequiresDistinctVisualStudioConfigurationName(Platform, Configuration, BaseDir))
+											{
+												CreateDistinctConfigName = true;
+											}
+										}
+										else
+										{
+											// Visual Studio doesn't natively support this platform, so we fake it by mapping it to
+											// a project configuration that has the platform name in that configuration as a suffix,
+											// and then using "x64" as the actual VS platform name
+											CreateDistinctConfigName = true;
+											ProjectPlatformName = DefaultPlatformName;
+										}
+
+										if (CreateDistinctConfigName)
+										{
+											ProjectConfigurationName = string.Format("{0}{1}_{2}", Platform.ToString(), Architecture != null ? "_" + Architecture.ToString() : string.Empty, Configuration.ToString());
+										}
+
+										TargetType TargetConfigurationType = ProjectTarget.TargetRules!.Type;
+										if (!bMakeProjectPerTarget && TargetConfigurationType != TargetType.Game)
+										{
+											ProjectConfigurationName += "_" + TargetConfigurationType.ToString();
+										}
 
-										ProjectConfigAndTargetCombination Combination = new ProjectConfigAndTargetCombination(Platform, Configuration, ProjectPlatformName, ProjectConfigurationName, ProjectTarget);
-										ProjectConfigAndTargetCombinations.Add(Combination);
+										if (ProjectConfigAndTargetCombinations.Any(Combination => Combination.ProjectPlatformName == ProjectPlatformName && Combination.ProjectConfigurationName == ProjectConfigurationName))
+										{
+											throw new BuildException("'{0}' '{1} is already in the platform/config list. This means a platform generator is not marking that the config needs to be distinct.", ProjectPlatformName, ProjectConfigurationName);
+										}
+
+										ProjectConfigAndTargetCombinations.Add(new ProjectConfigAndTargetCombination(Platform, Configuration, ProjectPlatformName, ProjectConfigurationName, ProjectTarget, Architecture));
 									}
 								}
 							}
@@ -974,7 +982,7 @@ namespace UnrealBuildTool
 				foreach (UnrealTargetPlatform Platform in ProjectPlatforms)
 				{
 					PlatformProjectGenerator? ProjGenerator = PlatformProjectGenerators.GetPlatformProjectGenerator(Platform, true);
-					if (ProjGenerator != null && ProjGenerator.HasVisualStudioSupport(Platform, UnrealTargetConfiguration.Development, ProjectFileFormat))
+					if (ProjGenerator != null && ProjGenerator.HasVisualStudioSupport(Platform, UnrealTargetConfiguration.Development, ProjectFileFormat, BaseDir, null))
 					{
 						ProjGenerator.GetAdditionalVisualStudioPropertyGroups(Platform, ProjectFileFormat, VCProjectFileContent);
 					}
@@ -1003,7 +1011,7 @@ namespace UnrealBuildTool
 				foreach (UnrealTargetPlatform Platform in ProjectPlatforms)
 				{
 					PlatformProjectGenerator? ProjGenerator = PlatformProjectGenerators.GetPlatformProjectGenerator(Platform, true);
-					if (ProjGenerator != null && ProjGenerator.HasVisualStudioSupport(Platform, UnrealTargetConfiguration.Development, ProjectFileFormat))
+					if (ProjGenerator != null && ProjGenerator.HasVisualStudioSupport(Platform, UnrealTargetConfiguration.Development, ProjectFileFormat, BaseDir, null))
 					{
 						ProjGenerator.GetVisualStudioGlobalProperties(Platform, VCProjectFileContent);
 					}
@@ -1628,7 +1636,7 @@ namespace UnrealBuildTool
 				foreach (UnrealTargetPlatform Platform in ProjectPlatforms)
 				{
 					PlatformProjectGenerator? ProjGenerator = PlatformProjectGenerators.GetPlatformProjectGenerator(Platform, true);
-					if (ProjGenerator != null && ProjGenerator.HasVisualStudioSupport(Platform, UnrealTargetConfiguration.Development, ProjectFileFormat))
+					if (ProjGenerator != null && ProjGenerator.HasVisualStudioSupport(Platform, UnrealTargetConfiguration.Development, ProjectFileFormat, BaseDir, null))
 					{
 						// @todo projectfiles: Serious hacks here because we are trying to emit one-time platform-specific sections that need information
 						//    about a target type, but the project file may contain many types of targets!  Some of this logic will need to move into
@@ -1684,7 +1692,7 @@ namespace UnrealBuildTool
 				foreach (UnrealTargetPlatform Platform in ProjectPlatforms)
 				{
 					PlatformProjectGenerator? ProjGenerator = PlatformProjectGenerators.GetPlatformProjectGenerator(Platform, true);
-					if (ProjGenerator != null && ProjGenerator.HasVisualStudioSupport(Platform, UnrealTargetConfiguration.Development, ProjectFileFormat))
+					if (ProjGenerator != null && ProjGenerator.HasVisualStudioSupport(Platform, UnrealTargetConfiguration.Development, ProjectFileFormat, BaseDir, null))
 					{
 						ProjGenerator.GetVisualStudioTargetOverrides(Platform, ProjectFileFormat, VCProjectFileContent);
 					}
@@ -1921,15 +1929,17 @@ namespace UnrealBuildTool
 			private readonly UnrealTargetConfiguration Configuration;
 			private readonly UnrealTargetPlatform Platform;
 			private readonly ProjectTarget ProjectTarget;
+			private readonly UnrealArch? Architecture;
 
 			public BuildCommandBuilder(UnrealTargetConfiguration InConfiguration, UnrealTargetPlatform InPlatform,
-				ProjectTarget InProjectTarget, string InUProjectPath, string? InBuildToolOverride = null)
+				ProjectTarget InProjectTarget, string InUProjectPath, string? InBuildToolOverride = null, UnrealArch? InArchitecture = null)
 			{
 				Configuration = InConfiguration;
 				Platform = InPlatform;
 				ProjectTarget = InProjectTarget;
 				UProjectPath = InUProjectPath;
 				BuildToolOverride = InBuildToolOverride;
+				Architecture = InArchitecture;
 
 				DirectoryReference BatchFilesDirectory = DirectoryReference.Combine(Unreal.EngineDirectory, "Build", "BatchFiles");
 				BuildScript = FileReference.Combine(BatchFilesDirectory, "Build.bat");
@@ -2003,6 +2013,11 @@ namespace UnrealBuildTool
 					BuildArguments.Append(ProjectGenerator.GetExtraBuildArguments(Platform, Configuration));
 				}
 
+				if (Architecture != null)
+				{
+					BuildArguments.AppendFormat(" -architecture={0}", Architecture);
+				}
+
 				return BuildArguments.ToString();
 			}
 		}
@@ -2010,7 +2025,7 @@ namespace UnrealBuildTool
 		private BuildCommandBuilder CreateArgumentsBuilder(ProjectConfigAndTargetCombination Combination, string UProjectPath, PlatformProjectGenerator? ProjGenerator)
 		{
 			BuildCommandBuilder Builder = new BuildCommandBuilder(Combination.Configuration, Combination.Platform!.Value,
-				Combination.ProjectTarget!, UProjectPath, BuildToolOverride)
+				Combination.ProjectTarget!, UProjectPath, BuildToolOverride, Combination.Architecture)
 			{
 				ProjectGenerator = ProjGenerator,
 				bEditorDependsOnShaderCompileWorker = Settings.bEditorDependsOnShaderCompileWorker,
@@ -2346,16 +2361,8 @@ namespace UnrealBuildTool
 			return Info.IsDotNETCoreProject();
 		}
 
-		/// <summary>
-		/// Get the project context for the given solution context
-		/// </summary>
-		/// <param name="SolutionTarget">The solution target type</param>
-		/// <param name="SolutionConfiguration">The solution configuration</param>
-		/// <param name="SolutionPlatform">The solution platform</param>
-		/// <param name="PlatformProjectGenerators">Set of platform project generators</param>
-		/// <param name="Logger"></param>
-		/// <returns>Project context matching the given solution context</returns>
-		public override MSBuildProjectContext GetMatchingProjectContext(TargetType SolutionTarget, UnrealTargetConfiguration SolutionConfiguration, UnrealTargetPlatform SolutionPlatform, PlatformProjectGeneratorCollection PlatformProjectGenerators, ILogger Logger)
+		/// <inheritdoc/>
+		public override MSBuildProjectContext GetMatchingProjectContext(TargetType SolutionTarget, UnrealTargetConfiguration SolutionConfiguration, UnrealTargetPlatform SolutionPlatform, PlatformProjectGeneratorCollection PlatformProjectGenerators, UnrealArch? Architecture, ILogger Logger)
 		{
 			// Find the matching platform name
 			string ProjectPlatformName;
diff --git a/Engine/Source/Programs/UnrealBuildTool/ProjectFiles/VisualStudio/VCProjectFileGenerator.cs b/Engine/Source/Programs/UnrealBuildTool/ProjectFiles/VisualStudio/VCProjectFileGenerator.cs
--- a/Engine/Source/Programs/UnrealBuildTool/ProjectFiles/VisualStudio/VCProjectFileGenerator.cs
+++ b/Engine/Source/Programs/UnrealBuildTool/ProjectFiles/VisualStudio/VCProjectFileGenerator.cs
@@ -262,6 +262,23 @@ namespace UnrealBuildTool
 			VCProjectFileContent.AppendLine("    <PlatformToolset>{0}</PlatformToolset>", PlatformToolsetVersionString);
 		}
 
+		/// <summary>
+		/// Returns a list of architectures to generate unique VS platforms for.
+		/// </summary>
+		public static IList<UnrealArch?> GetPlatformArchitectures(UEBuildPlatform BuildPlatform)
+		{
+			List<UnrealArch?> Architectures = new();
+			if (BuildPlatform.ArchitectureConfig.Mode == UnrealArchitectureMode.OneTargetPerArchitecture)
+			{
+				Architectures.AddRange(BuildPlatform.ArchitectureConfig.AllSupportedArchitectures.Architectures.Cast<UnrealArch?>());
+			}
+			else
+			{
+				Architectures.Add(null);
+			}
+			return Architectures;
+		}
+
 		/// <inheritdoc/>
 		protected override void ConfigureProjectFileGeneration(String[] Arguments, ref bool IncludeAllPlatforms, ILogger Logger)
 		{
@@ -399,9 +416,14 @@ namespace UnrealBuildTool
 			/// </summary>
 			public TargetType TargetConfigurationName;
 
+			/// <summary>
+			/// The target architecture
+			/// </summary>
+			public UnrealArch? Architecture;
+
 			public override string ToString()
 			{
-				return String.Format("{0}={1} {2} {3}", VCSolutionConfigAndPlatformName, Configuration, Platform, TargetConfigurationName);
+				return String.Format("{0}={1} {2} {3}{4}", VCSolutionConfigAndPlatformName, Configuration, Platform, TargetConfigurationName, Architecture != null ? " " + Architecture : string.Empty);
 			}
 
 			public VCSolutionConfigCombination(string VCSolutionConfigAndPlatformName)
@@ -488,91 +510,88 @@ namespace UnrealBuildTool
 
 		private void WriteCommonPropsFile(ILogger Logger)
 		{
-			using (ProgressWriter Progress = new ProgressWriter("Creating common properties file...", true, Logger))
-			{
-				StringBuilder VCCommonTargetFileContent = new StringBuilder();
-				VCCommonTargetFileContent.AppendLine("<?xml version=\"1.0\" encoding=\"utf-8\"?>");
-				VCCommonTargetFileContent.AppendLine("<Project xmlns=\"http://schemas.microsoft.com/developer/msbuild/2003\">");
-
-				// Project globals (project GUID, project type, SCC bindings, etc)
-				{
-					string ToolVersionString = GetProjectFileToolVersionString(Settings.ProjectFileFormat);
-					VCCommonTargetFileContent.AppendLine("  <PropertyGroup Label=\"Globals\">");
-					VCCommonTargetFileContent.AppendLine("    <Keyword>MakeFileProj</Keyword>");
-					AppendPlatformToolsetProperty(VCCommonTargetFileContent, Settings.ProjectFileFormat);
-					VCCommonTargetFileContent.AppendLine("    <MinimumVisualStudioVersion>{0}</MinimumVisualStudioVersion>", ToolVersionString);
-					VCCommonTargetFileContent.AppendLine("    <VCProjectVersion>{0}</VCProjectVersion>", ToolVersionString);
-					VCCommonTargetFileContent.AppendLine("    <NMakeUseOemCodePage>true</NMakeUseOemCodePage>"); // Fixes mojibake with non-Latin character sets (UE-102825)
-					VCCommonTargetFileContent.AppendLine("    <TargetRuntime>Native</TargetRuntime>");
-					VCCommonTargetFileContent.AppendLine("  </PropertyGroup>");
-				}
+			StringBuilder VCCommonTargetFileContent = new StringBuilder();
+			VCCommonTargetFileContent.AppendLine("<?xml version=\"1.0\" encoding=\"utf-8\"?>");
+			VCCommonTargetFileContent.AppendLine("<Project xmlns=\"http://schemas.microsoft.com/developer/msbuild/2003\">");
 
-				// Write the default configuration info
-				VCCommonTargetFileContent.AppendLine("  <PropertyGroup Label=\"Configuration\">");
-				VCCommonTargetFileContent.AppendLine($"    <ConfigurationType>{PlatformProjectGenerator.DefaultPlatformConfigurationType}</ConfigurationType>");
+			// Project globals (project GUID, project type, SCC bindings, etc)
+			{
+				string ToolVersionString = GetProjectFileToolVersionString(Settings.ProjectFileFormat);
+				VCCommonTargetFileContent.AppendLine("  <PropertyGroup Label=\"Globals\">");
+				VCCommonTargetFileContent.AppendLine("    <Keyword>MakeFileProj</Keyword>");
 				AppendPlatformToolsetProperty(VCCommonTargetFileContent, Settings.ProjectFileFormat);
+				VCCommonTargetFileContent.AppendLine("    <MinimumVisualStudioVersion>{0}</MinimumVisualStudioVersion>", ToolVersionString);
+				VCCommonTargetFileContent.AppendLine("    <VCProjectVersion>{0}</VCProjectVersion>", ToolVersionString);
+				VCCommonTargetFileContent.AppendLine("    <NMakeUseOemCodePage>true</NMakeUseOemCodePage>"); // Fixes mojibake with non-Latin character sets (UE-102825)
+				VCCommonTargetFileContent.AppendLine("    <TargetRuntime>Native</TargetRuntime>");
 				VCCommonTargetFileContent.AppendLine("  </PropertyGroup>");
+			}
 
-				VCCommonTargetFileContent.AppendLine("  <Import Project=\"$(VCTargetsPath)\\Microsoft.Cpp.Default.props\" />");
-				VCCommonTargetFileContent.AppendLine("  <Import Project=\"$(VCTargetsPath)\\Microsoft.Cpp.props\" />");
+			// Write the default configuration info
+			VCCommonTargetFileContent.AppendLine("  <PropertyGroup Label=\"Configuration\">");
+			VCCommonTargetFileContent.AppendLine($"    <ConfigurationType>{PlatformProjectGenerator.DefaultPlatformConfigurationType}</ConfigurationType>");
+			AppendPlatformToolsetProperty(VCCommonTargetFileContent, Settings.ProjectFileFormat);
+			VCCommonTargetFileContent.AppendLine("  </PropertyGroup>");
 
-				// Write the common and invalid configuration values
+			VCCommonTargetFileContent.AppendLine("  <Import Project=\"$(VCTargetsPath)\\Microsoft.Cpp.Default.props\" />");
+			VCCommonTargetFileContent.AppendLine("  <Import Project=\"$(VCTargetsPath)\\Microsoft.Cpp.props\" />");
+
+			// Write the common and invalid configuration values
+			{
+				const string InvalidMessage = "echo The selected platform/configuration is not valid for this target.";
+
+				string ProjectRelativeUnusedDirectory = ProjectFile.NormalizeProjectPath(DirectoryReference.Combine(Unreal.EngineDirectory, "Intermediate", "Build", "Unused"));
+
+				VCCommonTargetFileContent.AppendLine("  <PropertyGroup>");
+
+				DirectoryReference BatchFilesDirectory = DirectoryReference.Combine(Unreal.EngineDirectory, "Build", "BatchFiles");
+				VCCommonTargetFileContent.AppendLine("    <BuildBatchScript>{0}</BuildBatchScript>", ProjectFile.EscapePath(ProjectFile.NormalizeProjectPath(FileReference.Combine(BatchFilesDirectory, "Build.bat"))));
+				VCCommonTargetFileContent.AppendLine("    <RebuildBatchScript>{0}</RebuildBatchScript>", ProjectFile.EscapePath(ProjectFile.NormalizeProjectPath(FileReference.Combine(BatchFilesDirectory, "Rebuild.bat"))));
+				VCCommonTargetFileContent.AppendLine("    <CleanBatchScript>{0}</CleanBatchScript>", ProjectFile.EscapePath(ProjectFile.NormalizeProjectPath(FileReference.Combine(BatchFilesDirectory, "Clean.bat"))));
+				VCCommonTargetFileContent.AppendLine("    <NMakeBuildCommandLine>{0}</NMakeBuildCommandLine>", InvalidMessage);
+				VCCommonTargetFileContent.AppendLine("    <NMakeReBuildCommandLine>{0}</NMakeReBuildCommandLine>", InvalidMessage);
+				VCCommonTargetFileContent.AppendLine("    <NMakeCleanCommandLine>{0}</NMakeCleanCommandLine>", InvalidMessage);
+				VCCommonTargetFileContent.AppendLine("    <NMakeOutput>Invalid Output</NMakeOutput>", InvalidMessage);
+				VCCommonTargetFileContent.AppendLine("    <OutDir>{0}{1}</OutDir>", ProjectRelativeUnusedDirectory, Path.DirectorySeparatorChar);
+				VCCommonTargetFileContent.AppendLine("    <IntDir>{0}{1}</IntDir>", ProjectRelativeUnusedDirectory, Path.DirectorySeparatorChar);
+				// NOTE: We are intentionally overriding defaults for these paths with empty strings.  We never want Visual Studio's
+				//       defaults for these fields to be propagated, since they are version-sensitive paths that may not reflect
+				//       the environment that UBT is building in.  We'll set these environment variables ourselves!
+				// NOTE: We don't touch 'ExecutablePath' because that would result in Visual Studio clobbering the system "Path"
+				//       environment variable
+				VCCommonTargetFileContent.AppendLine("    <IncludePath />");
+				VCCommonTargetFileContent.AppendLine("    <ReferencePath />");
+				VCCommonTargetFileContent.AppendLine("    <LibraryPath />");
+				VCCommonTargetFileContent.AppendLine("    <LibraryWPath />");
+				VCCommonTargetFileContent.AppendLine("    <SourcePath />");
+				VCCommonTargetFileContent.AppendLine("    <ExcludePath />");
+
+				// Add all the default system include paths
+				if (OperatingSystem.IsWindows())
 				{
-					const string InvalidMessage = "echo The selected platform/configuration is not valid for this target.";
-
-					string ProjectRelativeUnusedDirectory = ProjectFile.NormalizeProjectPath(DirectoryReference.Combine(Unreal.EngineDirectory, "Intermediate", "Build", "Unused"));
-
-					VCCommonTargetFileContent.AppendLine("  <PropertyGroup>");
-
-					DirectoryReference BatchFilesDirectory = DirectoryReference.Combine(Unreal.EngineDirectory, "Build", "BatchFiles");
-					VCCommonTargetFileContent.AppendLine("    <BuildBatchScript>{0}</BuildBatchScript>", ProjectFile.EscapePath(ProjectFile.NormalizeProjectPath(FileReference.Combine(BatchFilesDirectory, "Build.bat"))));
-					VCCommonTargetFileContent.AppendLine("    <RebuildBatchScript>{0}</RebuildBatchScript>", ProjectFile.EscapePath(ProjectFile.NormalizeProjectPath(FileReference.Combine(BatchFilesDirectory, "Rebuild.bat"))));
-					VCCommonTargetFileContent.AppendLine("    <CleanBatchScript>{0}</CleanBatchScript>", ProjectFile.EscapePath(ProjectFile.NormalizeProjectPath(FileReference.Combine(BatchFilesDirectory, "Clean.bat"))));
-					VCCommonTargetFileContent.AppendLine("    <NMakeBuildCommandLine>{0}</NMakeBuildCommandLine>", InvalidMessage);
-					VCCommonTargetFileContent.AppendLine("    <NMakeReBuildCommandLine>{0}</NMakeReBuildCommandLine>", InvalidMessage);
-					VCCommonTargetFileContent.AppendLine("    <NMakeCleanCommandLine>{0}</NMakeCleanCommandLine>", InvalidMessage);
-					VCCommonTargetFileContent.AppendLine("    <NMakeOutput>Invalid Output</NMakeOutput>", InvalidMessage);
-					VCCommonTargetFileContent.AppendLine("    <OutDir>{0}{1}</OutDir>", ProjectRelativeUnusedDirectory, Path.DirectorySeparatorChar);
-					VCCommonTargetFileContent.AppendLine("    <IntDir>{0}{1}</IntDir>", ProjectRelativeUnusedDirectory, Path.DirectorySeparatorChar);
-					// NOTE: We are intentionally overriding defaults for these paths with empty strings.  We never want Visual Studio's
-					//       defaults for these fields to be propagated, since they are version-sensitive paths that may not reflect
-					//       the environment that UBT is building in.  We'll set these environment variables ourselves!
-					// NOTE: We don't touch 'ExecutablePath' because that would result in Visual Studio clobbering the system "Path"
-					//       environment variable
-					VCCommonTargetFileContent.AppendLine("    <IncludePath />");
-					VCCommonTargetFileContent.AppendLine("    <ReferencePath />");
-					VCCommonTargetFileContent.AppendLine("    <LibraryPath />");
-					VCCommonTargetFileContent.AppendLine("    <LibraryWPath />");
-					VCCommonTargetFileContent.AppendLine("    <SourcePath />");
-					VCCommonTargetFileContent.AppendLine("    <ExcludePath />");
-
-					// Add all the default system include paths
-					if (OperatingSystem.IsWindows())
+					if (SupportedPlatforms.Contains(UnrealTargetPlatform.Win64))
 					{
-						if (SupportedPlatforms.Contains(UnrealTargetPlatform.Win64))
-						{
-							VCCommonTargetFileContent.AppendLine("    <DefaultSystemIncludePaths>{0}</DefaultSystemIncludePaths>", VCToolChain.GetVCIncludePaths(UnrealTargetPlatform.Win64, GetCompilerForIntellisense(Settings.ProjectFileFormat), null, null, Logger));
-						}
-					}
-					else
-					{
-						Logger.LogInformation("Unable to compute VC include paths on non-Windows host");
-						VCCommonTargetFileContent.AppendLine("    <DefaultSystemIncludePaths />");
+						VCCommonTargetFileContent.AppendLine("    <DefaultSystemIncludePaths>{0}</DefaultSystemIncludePaths>", VCToolChain.GetVCIncludePaths(UnrealTargetPlatform.Win64, GetCompilerForIntellisense(Settings.ProjectFileFormat), null, null, Logger));
 					}
+				}
+				else
+				{
+					Logger.LogInformation("Unable to compute VC include paths on non-Windows host");
+					VCCommonTargetFileContent.AppendLine("    <DefaultSystemIncludePaths />");
+				}
 
-					VCCommonTargetFileContent.AppendLine("  </PropertyGroup>");
+				VCCommonTargetFileContent.AppendLine("  </PropertyGroup>");
 
-				}
+			}
 
-				// Write default import group
-				VCCommonTargetFileContent.AppendLine("  <ImportGroup Label=\"PropertySheets\">");
-				VCCommonTargetFileContent.AppendLine("    <Import Project=\"$(UserRootDir)\\Microsoft.Cpp.$(Platform).user.props\" Condition=\"exists('$(UserRootDir)\\Microsoft.Cpp.$(Platform).user.props')\" Label=\"LocalAppDataPlatform\" />");
-				VCCommonTargetFileContent.AppendLine("  </ImportGroup>");
+			// Write default import group
+			VCCommonTargetFileContent.AppendLine("  <ImportGroup Label=\"PropertySheets\">");
+			VCCommonTargetFileContent.AppendLine("    <Import Project=\"$(UserRootDir)\\Microsoft.Cpp.$(Platform).user.props\" Condition=\"exists('$(UserRootDir)\\Microsoft.Cpp.$(Platform).user.props')\" Label=\"LocalAppDataPlatform\" />");
+			VCCommonTargetFileContent.AppendLine("  </ImportGroup>");
 
-				VCCommonTargetFileContent.AppendLine("</Project>");
+			VCCommonTargetFileContent.AppendLine("</Project>");
 
-				Utils.WriteFileIfChanged(FileReference.Combine(IntermediateProjectFilesPath, "UECommon.props"), VCCommonTargetFileContent.ToString(), Logger);
-			}
+			Utils.WriteFileIfChanged(FileReference.Combine(IntermediateProjectFilesPath, "UECommon.props"), VCCommonTargetFileContent.ToString(), Logger);
 		}
 
 		/// <summary>
@@ -763,13 +782,11 @@ namespace UnrealBuildTool
 
 				// Get the path to the visualizers file. Try to make it relative to the solution directory, but fall back to a full path if it's a foreign project.
 				FileReference VisualizersFile = FileReference.Combine(Unreal.EngineDirectory, "Extras", "VisualStudioDebugging", "Unreal.natvis");
-				FileReference VisualizersStepFile = FileReference.Combine(Unreal.EngineDirectory, "Extras", "VisualStudioDebugging", "Unreal.natstepfilter");
 
 				// Add the visualizers at the solution level. Doesn't seem to be picked up from a makefile project in VS2017 15.8.5.
 				VCSolutionFileContent.AppendLine(String.Format("Project(\"{0}\") = \"Visualizers\", \"Visualizers\", \"{{1CCEC849-CC72-4C59-8C36-2F7C38706D4C}}\"", SolutionFolderEntryGUID));
 				VCSolutionFileContent.AppendLine("\tProjectSection(SolutionItems) = preProject");
 				VCSolutionFileContent.AppendLine("\t\t{0} = {0}", VisualizersFile.MakeRelativeTo(PrimaryProjectPath));
-				VCSolutionFileContent.AppendLine("\t\t{0} = {0}", VisualizersStepFile.MakeRelativeTo(PrimaryProjectPath));
 				VCSolutionFileContent.AppendLine("\tEndProjectSection");
 				VCSolutionFileContent.AppendLine("EndProject");
 			}
@@ -823,7 +840,7 @@ namespace UnrealBuildTool
 							foreach (VCSolutionConfigCombination SolutionConfigCombination in SolutionConfigCombinations)
 							{
 								// Get the context for the current solution context
-								MSBuildProjectContext ProjectContext = CurProject.GetMatchingProjectContext(SolutionConfigCombination.TargetConfigurationName, SolutionConfigCombination.Configuration, SolutionConfigCombination.Platform, PlatformProjectGenerators, Logger);
+								MSBuildProjectContext ProjectContext = CurProject.GetMatchingProjectContext(SolutionConfigCombination.TargetConfigurationName, SolutionConfigCombination.Configuration, SolutionConfigCombination.Platform, PlatformProjectGenerators, SolutionConfigCombination.Architecture, Logger);
 
 								// Override the configuration to build for UBT
 								if (Settings.bBuildUBTInDebug && CurProject == UBTProject)
@@ -994,8 +1011,7 @@ namespace UnrealBuildTool
 		{
 			OutValidPlatforms = new HashSet<UnrealTargetPlatform>();
 			OutSolutionConfigs = new List<VCSolutionConfigCombination>();
-			Dictionary<string, Tuple<UnrealTargetConfiguration, TargetType>> SolutionConfigurationsValidForProjects =
-				new Dictionary<string, Tuple<UnrealTargetConfiguration, TargetType>>();
+			Dictionary<string, Tuple<UnrealTargetConfiguration, Tuple<ProjectTarget, TargetType>>> SolutionConfigurationsValidForProjects = new();
 
 			foreach (UnrealTargetConfiguration CurConfiguration in AllConfigurations)
 			{
@@ -1016,7 +1032,7 @@ namespace UnrealBuildTool
 									}
 
 									// Figure out the set of valid target configuration names
-									foreach (Project ProjectTarget in CurProject.ProjectTargets)
+									foreach (ProjectTarget ProjectTarget in CurProject.ProjectTargets.OfType<ProjectTarget>())
 									{
 										if (VCProjectFile.IsValidProjectPlatformAndConfiguration(ProjectTarget, CurPlatform,
 												CurConfiguration, Logger))
@@ -1033,7 +1049,7 @@ namespace UnrealBuildTool
 											string SolutionConfigName =
 												MakeSolutionConfigurationName(CurConfiguration, TargetType, bMakeProjectPerTarget);
 											SolutionConfigurationsValidForProjects[SolutionConfigName] =
-												new Tuple<UnrealTargetConfiguration, TargetType>(CurConfiguration, TargetType);
+												new Tuple<UnrealTargetConfiguration, Tuple<ProjectTarget, TargetType>>(CurConfiguration, new Tuple<ProjectTarget, TargetType>(ProjectTarget, TargetType));
 										}
 									}
 								}
@@ -1045,25 +1061,41 @@ namespace UnrealBuildTool
 
 			foreach (UnrealTargetPlatform CurPlatform in OutValidPlatforms)
 			{
-				foreach (KeyValuePair<string, Tuple<UnrealTargetConfiguration, TargetType>> SolutionConfigKeyValue in
-						 SolutionConfigurationsValidForProjects)
+				UEBuildPlatform? BuildPlatform;
+				if (UEBuildPlatform.TryGetBuildPlatform(CurPlatform, out BuildPlatform))
 				{
-					// e.g.  "Development|Win64 = Development|Win64"
-					string SolutionConfigName = SolutionConfigKeyValue.Key;
-					UnrealTargetConfiguration Configuration = SolutionConfigKeyValue.Value.Item1;
-					TargetType TargetType = SolutionConfigKeyValue.Value.Item2;
-
-					string SolutionPlatformName = CurPlatform.ToString();
+					foreach (KeyValuePair<string, Tuple<UnrealTargetConfiguration, Tuple<ProjectTarget, TargetType>>> SolutionConfigKeyValue in
+						SolutionConfigurationsValidForProjects)
+					{
+						ProjectTarget ProjectTarget = SolutionConfigKeyValue.Value.Item2.Item1;
 
-					string SolutionConfigAndPlatformPair = SolutionConfigName + "|" + SolutionPlatformName;
-					OutSolutionConfigs.Add(
-						new VCSolutionConfigCombination(SolutionConfigAndPlatformPair)
+						IList<UnrealArch?> Architectures = GetPlatformArchitectures(BuildPlatform);
+						bool SingleArchitecture = Architectures.Count == 1;
+						foreach (UnrealArch? Arch in Architectures)
 						{
-							Configuration = Configuration,
-							Platform = CurPlatform,
-							TargetConfigurationName = TargetType
+							// e.g.  "Development|Win64 = Development|Win64"
+							string SolutionConfigName = SolutionConfigKeyValue.Key;
+							UnrealTargetConfiguration Configuration = SolutionConfigKeyValue.Value.Item1;
+							TargetType TargetType = SolutionConfigKeyValue.Value.Item2.Item2;
+
+							string SolutionPlatformName = CurPlatform.ToString();
+							if (!SingleArchitecture)
+							{
+								SolutionPlatformName += $"-{Arch}";
+							}
+
+							string SolutionConfigAndPlatformPair = SolutionConfigName + "|" + SolutionPlatformName;
+							OutSolutionConfigs.Add(
+								new VCSolutionConfigCombination(SolutionConfigAndPlatformPair)
+								{
+									Configuration = Configuration,
+									Platform = CurPlatform,
+									TargetConfigurationName = TargetType,
+									Architecture = Arch
+								}
+							);
 						}
-					);
+					}
 				}
 			}
 
