Change GameActivity exported to true

#android

diff --git a/Engine/Source/Programs/UnrealBuildTool/Platform/Android/UEDeployAndroid.cs b/Engine/Source/Programs/UnrealBuildTool/Platform/Android/UEDeployAndroid.cs
--- a/Engine/Source/Programs/UnrealBuildTool/Platform/Android/UEDeployAndroid.cs
+++ b/Engine/Source/Programs/UnrealBuildTool/Platform/Android/UEDeployAndroid.cs
@@ -2724,7 +2724,7 @@ namespace UnrealBuildTool
 				Text.AppendLine("\t\t\t</intent-filter>");
 				Text.AppendLine("\t\t</activity>");
 				Text.AppendLine("\t\t<activity android:name=\"com.epicgames.unreal.GameActivity\"");
-				Text.AppendLine("\t\t          android:exported=\"false\"");
+				Text.AppendLine("\t\t          android:exported=\"true\"");
 				Text.AppendLine("\t\t          android:label=\"@string/app_name\"");
 				Text.AppendLine("\t\t          android:theme=\"@style/UnrealSplashTheme\"");
 				Text.AppendLine(bAddDensity ? "\t\t          android:configChanges=\"mcc|mnc|uiMode|density|screenSize|smallestScreenSize|screenLayout|orientation|keyboardHidden|keyboard|navigation|touchscreen|locale|fontScale|layoutDirection\""
