Fix error due to missing executable extension when cross-compiling for Linux using the AutoRTFM compiler with Box enabled.

diff --git a/Engine/Source/Programs/UnrealBuildTool/Platform/Linux/LinuxToolChain.cs b/Engine/Source/Programs/UnrealBuildTool/Platform/Linux/LinuxToolChain.cs
--- a/Engine/Source/Programs/UnrealBuildTool/Platform/Linux/LinuxToolChain.cs
+++ b/Engine/Source/Programs/UnrealBuildTool/Platform/Linux/LinuxToolChain.cs
@@ -130,7 +130,7 @@ namespace UnrealBuildTool
 					DirectoryReference AutoRTFMDir = DirectoryReference.Combine(Unreal.EngineDirectory, "Restricted", "NotForLicensees", "Binaries", BuildHostPlatform.Current.Platform.ToString(), "AutoRTFM", "bin");
 
 					// set up the path to our toolchain
-					ClangPath = FileReference.Combine(AutoRTFMDir, $"verse-clang-cl");
+					ClangPath = FileReference.Combine(AutoRTFMDir, $"verse-clang-cl{BuildHostPlatform.Current.BinarySuffix}");
 				}
 
 				// When cross-compiling on Windows, use old FixDeps. It is slow, but it does not have timing issues
