Updated the exe name logic to take into account the project/platform architecture when it is set.

diff --git a/Engine/Source/Programs/UnrealBuildTool/ProjectFiles/VisualStudio/VCProject.cs b/Engine/Source/Programs/UnrealBuildTool/ProjectFiles/VisualStudio/VCProject.cs
--- a/Engine/Source/Programs/UnrealBuildTool/ProjectFiles/VisualStudio/VCProject.cs
+++ b/Engine/Source/Programs/UnrealBuildTool/ProjectFiles/VisualStudio/VCProject.cs
@@ -2155,6 +2155,14 @@ namespace UnrealBuildTool
 					{
 						NMakePath += TargetRulesObject.Architecture.ToString();
 					}
+					else if (Combination.Architecture != null) // support the case where the project/platform combination explicitly sets an architecture(e.g. Win64)
+					{
+						UnrealArch Architecture = (UnrealArch)Combination.Architecture;
+						if (UnrealArchitectureConfig.ForPlatform(Platform).RequiresArchitectureFilenames(new UnrealArchitectures(Architecture)))
+						{
+							NMakePath += Architecture.ToString();
+						}
+					}
 					NMakePath += BuildPlatform.GetBinaryExtension(UEBuildBinaryType.Executable);
 
 					VCProjectFileContent.AppendLine("  <PropertyGroup {0}>", ConditionString);
