Deprecating non-TObjectPtr FReferenceCollector::AddReferencedObject(s) API.

diff --git a/Engine/Source/Programs/UnrealBuildTool/Configuration/TargetRules.cs b/Engine/Source/Programs/UnrealBuildTool/Configuration/TargetRules.cs
--- a/Engine/Source/Programs/UnrealBuildTool/Configuration/TargetRules.cs
+++ b/Engine/Source/Programs/UnrealBuildTool/Configuration/TargetRules.cs
@@ -1586,6 +1586,18 @@ namespace UnrealBuildTool
 		[XmlConfigFile(Category = "BuildConfiguration", Name = nameof(DefaultWarningLevel))]
 		private WarningLevel DefaultWarningLevelPrivate;
 
+
+		/// <summary>
+		/// Require TObjectPtr for FReferenceCollector API's. (Needed for compatibility with incremental GC.)
+		/// </summary>
+		[RequiresUniqueBuildEnvironment]
+		public bool bRequireObjectPtrForAddReferencedObjects
+		{
+				get { return bRequireObjectPtrForAddReferencedObjectsPrivate ?? true; }
+				set { bRequireObjectPtrForAddReferencedObjectsPrivate = value; }
+		}
+		private bool? bRequireObjectPtrForAddReferencedObjectsPrivate;		
+
 		/// <summary>
 		/// Level to report deprecation warnings as errors
 		/// </summary>
diff --git a/Engine/Source/Programs/UnrealBuildTool/System/RulesAssembly.cs b/Engine/Source/Programs/UnrealBuildTool/System/RulesAssembly.cs
--- a/Engine/Source/Programs/UnrealBuildTool/System/RulesAssembly.cs
+++ b/Engine/Source/Programs/UnrealBuildTool/System/RulesAssembly.cs
@@ -678,6 +678,11 @@ namespace UnrealBuildTool
 				Rules.GlobalDefinitions.Add("DISABLE_UNVERIFIED_CERTIFICATE_LOADING=1");
 			}
 
+			if (Rules.bRequireObjectPtrForAddReferencedObjects)
+			{
+				Rules.GlobalDefinitions.Add("UE_REFERENCE_COLLECTOR_REQUIRE_OBJECTPTR=1");
+			}
+
 			// if the Target has opted in only some platforms, disable any plugins of other platforms (there may be editor, etc, modules that
 			// will just add themselves, with no other reference to be able to remove them, other than disabling them here)
 			if (Rules.OptedInModulePlatforms != null)
