Turned off the creation of the additional platform/architectures options until I figure out a way to not have VS reset what current config/platform was previously selected.

diff --git a/Engine/Source/Programs/UnrealBuildTool/ProjectFiles/VisualStudio/VCProjectFileGenerator.cs b/Engine/Source/Programs/UnrealBuildTool/ProjectFiles/VisualStudio/VCProjectFileGenerator.cs
--- a/Engine/Source/Programs/UnrealBuildTool/ProjectFiles/VisualStudio/VCProjectFileGenerator.cs
+++ b/Engine/Source/Programs/UnrealBuildTool/ProjectFiles/VisualStudio/VCProjectFileGenerator.cs
@@ -268,11 +268,11 @@ namespace UnrealBuildTool
 		public static IList<UnrealArch?> GetPlatformArchitectures(UEBuildPlatform BuildPlatform)
 		{
 			List<UnrealArch?> Architectures = new();
-			if (BuildPlatform.ArchitectureConfig.Mode == UnrealArchitectureMode.OneTargetPerArchitecture)
-			{
-				Architectures.AddRange(BuildPlatform.ArchitectureConfig.AllSupportedArchitectures.Architectures.Cast<UnrealArch?>());
-			}
-			else
+			//if (BuildPlatform.ArchitectureConfig.Mode == UnrealArchitectureMode.OneTargetPerArchitecture)
+			//{
+			//	Architectures.AddRange(BuildPlatform.ArchitectureConfig.AllSupportedArchitectures.Architectures.Cast<UnrealArch?>());
+			//}
+			//else
 			{
 				Architectures.Add(null);
 			}
