- Enabled Modern Xcode mode by default for everyone!
- Fixed a naming issue (GameName vs ExeName) with stub creation for non-Development builds

diff --git a/Engine/Config/BaseEngine.ini b/Engine/Config/BaseEngine.ini
--- a/Engine/Config/BaseEngine.ini
+++ b/Engine/Config/BaseEngine.ini
@@ -3012,7 +3012,7 @@ bBuildAllSupportedOnBuildMachine=true
 bEnableRayTracing=false
 
 [/Script/MacTargetPlatform.XcodeProjectSettings]
-bUseModernXcode=false
+bUseModernXcode=true
 bUseAutomaticCodeSigning=true
 bMacSignToRunLocally=true
 CodeSigningTeam=
diff --git a/Engine/Source/Programs/UnrealBuildTool/Platform/IOS/IOSToolChain.cs b/Engine/Source/Programs/UnrealBuildTool/Platform/IOS/IOSToolChain.cs
--- a/Engine/Source/Programs/UnrealBuildTool/Platform/IOS/IOSToolChain.cs
+++ b/Engine/Source/Programs/UnrealBuildTool/Platform/IOS/IOSToolChain.cs
@@ -740,9 +740,9 @@ namespace UnrealBuildTool
 				Directory.CreateDirectory(DestSubdir);
 			}
 
-			// set up the directories
-			string ZipWorkingDir = String.Format("Payload/{0}.app/", GameName);
-			string ZipSourceDir = String.Format(bUseModernXcode ? "{0}/{1}.app" : "{0}/Payload/{1}.app", BinaryPath, GameName);
+			// set up the directories (modern has .app named with the config names in them, like UnrealGame-IOS-Shipping, which the ExeName will be)
+			string ZipWorkingDir = bUseModernXcode ? $"Payload/{ExeName}.app/" : $"Payload/{GameName}.app/";
+			string ZipSourceDir = bUseModernXcode ? $"{BinaryPath}/{ExeName}.app" : $"{BinaryPath}/Payload/{GameName}.app";
 
 			if (bUseModernXcode)
 			{
