- Changed how .dylibs in the Contents/MacOS directory are copied into Mac .apps (done directly instead of needing to be staged first)

diff --git a/Engine/Source/Programs/UnrealBuildTool/ProjectFiles/Xcode/XcodeProject.cs b/Engine/Source/Programs/UnrealBuildTool/ProjectFiles/Xcode/XcodeProject.cs
--- a/Engine/Source/Programs/UnrealBuildTool/ProjectFiles/Xcode/XcodeProject.cs
+++ b/Engine/Source/Programs/UnrealBuildTool/ProjectFiles/Xcode/XcodeProject.cs
@@ -1245,7 +1245,7 @@ namespace UnrealBuildTool.XcodeProjectXcconfig
 				References.Add(ResourcesBuildPhase);
 
 				ProcessFrameworks(ResourcesBuildPhase, Project, ProjectFile, Logger);
-				ProcessScripts(ResourcesBuildPhase, Project);
+				ProcessScripts(ResourcesBuildPhase, Project, ProjectFile);
 				ProcessAssets(ResourcesBuildPhase);
 			}
 
@@ -1286,7 +1286,8 @@ namespace UnrealBuildTool.XcodeProjectXcconfig
 		/// </summary>
 		/// <param name="ResourcesBuildPhase"></param>
 		/// <param name="Project"></param>
-		protected void ProcessScripts(XcodeResourcesBuildPhase ResourcesBuildPhase, XcodeProject Project)
+		/// <param name="ProjectFile"></param>
+		protected void ProcessScripts(XcodeResourcesBuildPhase ResourcesBuildPhase, XcodeProject Project, XcodeProjectFile ProjectFile)
 		{
 			List<string> CopyScript = new();
 
@@ -1298,19 +1299,36 @@ namespace UnrealBuildTool.XcodeProjectXcconfig
 			{
 				"set -eo pipefail",
 
-				"# Copy the executable into .app if needed (-ef checks if two files/paths are equivalent)",
 				"SRC_EXE=\\\"${UE_BINARIES_DIR}/${UE_UBT_BINARY_SUBPATH}\\\"",
 				"DEST_EXE=\\\"${CONFIGURATION_BUILD_DIR}/${EXECUTABLE_PATH}\\\"",
+				"DEST_EXE_DIR=`dirname ${DEST_EXE}`",
 				"",
-				"echo ${SRC_EXE} /// ${DEST_EXE}",
-				"if [[ ! ${SRC_EXE} -ef ${DEST_EXE} ]]; then",
-				"  echo Copying executable...",
-				"  mkdir -p `dirname ${DEST_EXE}`",
-				"  ditto \\\"${SRC_EXE}\\\" \\\"${DEST_EXE}\\\"",
-				"fi",
-				""
+				"echo Copying executable and any standalone dylibs into ",
+				"mkdir -p ${DEST_EXE_DIR}",
+				"ditto \\\"${SRC_EXE}\\\" \\\"${DEST_EXE}\\\"",
 			});
 
+			IEnumerable<ModuleRules.RuntimeDependency>? Dylibs;
+			Tuple<ProjectFile, UnrealTargetPlatform> DylibKey = Tuple.Create((ProjectFile)ProjectFile, Platform);
+			if (XcodeProjectFileGenerator.TargetRawDylibs.TryGetValue(DylibKey, out Dylibs))
+			{
+				foreach (ModuleRules.RuntimeDependency Dylib in Dylibs)
+				{
+					// make it absolute if it was relative (to Engine/Source)
+					string FixedSource = Dylib.SourcePath!;
+					if (!FixedSource.StartsWith("/"))
+					{
+						FixedSource = $"${{UE_ENGINE_DIR}}/Source/{FixedSource}";
+					}
+
+					// make it relative to Binaries dir
+					string FixedDest = Dylib.Path.Replace("$(BinaryOutputDir)", "");
+					CopyScript.Add($"ditto \\\"{FixedSource}\\\" \\\"${{DEST_EXE_DIR}}{FixedDest}\\\"");
+				}
+			}
+
+			CopyScript.Add("");
+
 			// Editor just need the above script to copy executable into .app
 			if (Project.UnrealData.TargetRules.Type == TargetType.Editor)
 			{
@@ -1370,22 +1388,6 @@ namespace UnrealBuildTool.XcodeProjectXcconfig
 					$"echo \\\"Syncing ${{STAGED_DIR}}{SyncSourceSubdir} to ${{CONFIGURATION_BUILD_DIR}}/${{CONTENTS_FOLDER_PATH}}{SyncDestSubdir}\\\"",
 					$"rsync -a --delete --exclude=Info.plist --exclude=/Manifest_* --exclude=*.app \\\"${{STAGED_DIR}}{SyncSourceSubdir}/\\\" \\\"${{CONFIGURATION_BUILD_DIR}}/${{CONTENTS_FOLDER_PATH}}{SyncDestSubdir}\\\"",
 				});
-
-				// copy any dylibs from the staged stub .app into the real one
-				if (Platform == UnrealTargetPlatform.Mac)
-				{
-					// when doing a Content only project the Binaries are under Engine, not the project
-					string EngineOrProject = (Project.UnrealData.bIsContentOnlyProject || Project.UnrealData.UProjectFileLocation == null) ? "Engine" : Project.UnrealData.ProductName;
-					CopyScript.AddRange(new string[]
-					{
-						// copy from whatever .app is in Binaries/Mac - this is so we can make a Shipping app from a Development staged directory
-						"pushd \\\"${STAGED_DIR}/${UE_STAGED_BINARIES_DIR_BASE}/Binaries/Mac/\\\" > /dev/null",
-						"APPS=($(echo *.app))",
-						"echo Syncing $PWD/${APPS[0]}/${BUNDLE_CONTENTS_FOLDER_PATH} to ${CONFIGURATION_BUILD_DIR}/${CONTENTS_FOLDER_PATH}/\\\"",
-						"rsync -a --include='*/' --include='*.dylib' --exclude='*' ${APPS[0]}/${BUNDLE_CONTENTS_FOLDER_PATH}MacOS/\\\" \\\"${CONFIGURATION_BUILD_DIR}/${CONTENTS_FOLDER_PATH}/MacOS\\\"",
-						"popd > /dev/null",
-					});
-				}
 			}
 
 			// run this script every time, but xcode will show a warning if there isn't _some_ output
diff --git a/Engine/Source/Programs/UnrealBuildTool/ProjectFiles/Xcode/XcodeProjectFileGenerator.cs b/Engine/Source/Programs/UnrealBuildTool/ProjectFiles/Xcode/XcodeProjectFileGenerator.cs
--- a/Engine/Source/Programs/UnrealBuildTool/ProjectFiles/Xcode/XcodeProjectFileGenerator.cs
+++ b/Engine/Source/Programs/UnrealBuildTool/ProjectFiles/Xcode/XcodeProjectFileGenerator.cs
@@ -72,6 +72,7 @@ namespace UnrealBuildTool
 		public static List<UnrealTargetPlatform?> NullPlatformList = new() { null };
 		internal static Dictionary<Tuple<ProjectFile, UnrealTargetPlatform>, IEnumerable<UEBuildFramework>> TargetFrameworks = new();
 		internal static Dictionary<Tuple<ProjectFile, UnrealTargetPlatform>, IEnumerable<UEBuildBundleResource>> TargetBundles = new();
+		internal static Dictionary<Tuple<ProjectFile, UnrealTargetPlatform>, IEnumerable<ModuleRules.RuntimeDependency>> TargetRawDylibs = new();
 
 		/// <summary>
 		/// Should we generate only a run project (no build/index targets)
@@ -437,6 +438,7 @@ namespace UnrealBuildTool
 			XcodeProjectFileGenerator.RunTargetPlatforms.Clear();
 			XcodeProjectFileGenerator.TargetFrameworks.Clear();
 			XcodeProjectFileGenerator.TargetBundles.Clear();
+			XcodeProjectFileGenerator.TargetRawDylibs.Clear();
 			XcodeProjectFileGenerator.bGenerateRunOnlyProject = false;
 
 			if (ProjectPlatforms.Count > 0)
@@ -542,6 +544,7 @@ namespace UnrealBuildTool
 
 						List<UEBuildFramework> Frameworks = new();
 						List<UEBuildBundleResource> Bundles = new();
+						List<ModuleRules.RuntimeDependency> Dylibs = new();
 						// Generate a compile environment for each module in the binary
 						CppCompileEnvironment GlobalCompileEnvironment = Target.CreateCompileEnvironmentForProjectFiles(Logger);
 						foreach (UEBuildBinary Binary in Target.Binaries)
@@ -552,6 +555,11 @@ namespace UnrealBuildTool
 								CppCompileEnvironment CompileEnvironment = Module.CreateModuleCompileEnvironment(Target.Rules, BinaryCompileEnvironment, Logger);
 								Frameworks.AddRange(CompileEnvironment.AdditionalFrameworks);
 							}
+							// walk over CPP and External modules looking for dylibs that need to be copied into the .app directly (not in a framework)
+							foreach (UEBuildModule Module in Binary.Modules)
+							{
+								Dylibs.AddRange(Module.Rules.RuntimeDependencies.Inner.Where(x => x.Path.StartsWith("$(BinaryOutputDir)") && Path.GetExtension(x.SourcePath) == ".dylib"));
+							}
 						}
 
 						// track frameworks if we found any
@@ -569,6 +577,13 @@ namespace UnrealBuildTool
 								TargetBundles.Add(Tuple.Create(TargetProjectFile, Platform), Bundles.Distinct());
 							}
 						}
+						if (Dylibs.Count > 0)
+						{
+							lock (TargetRawDylibs)
+							{
+								TargetRawDylibs.Add(Tuple.Create(TargetProjectFile, Platform), Dylibs);
+							}
+						}
 					}
 					catch (Exception Ex)
 					{
