- Fixed making .apps using a StagingDirectory override with code-based projects

diff --git a/Engine/Source/Programs/UnrealBuildTool/ProjectFiles/Xcode/XcodeProject.cs b/Engine/Source/Programs/UnrealBuildTool/ProjectFiles/Xcode/XcodeProject.cs
--- a/Engine/Source/Programs/UnrealBuildTool/ProjectFiles/Xcode/XcodeProject.cs
+++ b/Engine/Source/Programs/UnrealBuildTool/ProjectFiles/Xcode/XcodeProject.cs
@@ -1301,10 +1301,10 @@ namespace UnrealBuildTool.XcodeProjectXcconfig
 
 				"SRC_EXE=\\\"${UE_BINARIES_DIR}/${UE_UBT_BINARY_SUBPATH}\\\"",
 				"DEST_EXE=\\\"${CONFIGURATION_BUILD_DIR}/${EXECUTABLE_PATH}\\\"",
-				"DEST_EXE_DIR=`dirname ${DEST_EXE}`",
+				"DEST_EXE_DIR=`dirname \\\"${DEST_EXE}\\\"`",
 				"",
-				"echo Copying executable and any standalone dylibs into ",
-				"mkdir -p ${DEST_EXE_DIR}",
+				"echo Copying executable and any standalone dylibs into ${DEST_EXE_DIR}",
+				"mkdir -p \\\"${DEST_EXE_DIR}\\\"",
 				"ditto \\\"${SRC_EXE}\\\" \\\"${DEST_EXE}\\\"",
 			});
 
@@ -1343,7 +1343,7 @@ namespace UnrealBuildTool.XcodeProjectXcconfig
 			if (TargetType != TargetType.Editor)
 			{
 				bool bIsEngineBuild = Project.UnrealData.UProjectFileLocation == null;
-				string DefaultStageDir = bIsEngineBuild ? "${UE_OVERRIDE_STAGE_DIR}" : "${UE_PROJECT_DIR}/Saved/StagedBuilds/${UE_TARGET_PLATFORM_NAME}";
+				string DefaultStageDir = bIsEngineBuild ? "" : "${UE_PROJECT_DIR}/Saved/StagedBuilds/${UE_TARGET_PLATFORM_NAME}";
 				string SyncSourceSubdir = (Platform == UnrealTargetPlatform.Mac) ? "" : "/cookeddata";
 				string SyncDestSubdir = (Platform == UnrealTargetPlatform.Mac) ? "/UE" : "/cookeddata";
 				CopyScript.AddRange(new string[]
@@ -1352,7 +1352,11 @@ namespace UnrealBuildTool.XcodeProjectXcconfig
 					"if [[ ${UE_SKIP_STAGEDDATA_SYNC} -eq 1 ]]; then exit 0; fi",
 					"",
 					"# When building engine projects, like UnrealGame, we don't have data to stage unless something has specified UE_OVERRIDE_STAGE_DIR",
+					"if [[ -z ${UE_OVERRIDE_STAGE_DIR} ]]; then ",
 					$"STAGED_DIR=\\\"{DefaultStageDir}\\\"",
+					"else",
+					"  STAGED_DIR=\\\"${UE_OVERRIDE_STAGE_DIR}\\\"",
+					"fi",
 					"if [[ -z ${STAGED_DIR} ]]; then exit 0; fi",
 				});
 
