UnrealBuildTool: Fix regression with shared environment path

#rnx

diff --git a/Engine/Source/Programs/UnrealBuildTool/Configuration/UEBuildTarget.cs b/Engine/Source/Programs/UnrealBuildTool/Configuration/UEBuildTarget.cs
--- a/Engine/Source/Programs/UnrealBuildTool/Configuration/UEBuildTarget.cs
+++ b/Engine/Source/Programs/UnrealBuildTool/Configuration/UEBuildTarget.cs
@@ -1731,11 +1731,11 @@ namespace UnrealBuildTool
 			}
 			else if (Configuration == UnrealTargetConfiguration.DebugGame)
 			{
-				EngineIntermediateDirectory = DirectoryReference.Combine(Unreal.EngineDirectory, PlatformIntermediateFolder, GetTargetIntermediateFolderName(TargetName, IntermediateEnvironment), UnrealTargetConfiguration.Development.ToString());
+				EngineIntermediateDirectory = DirectoryReference.Combine(Unreal.EngineDirectory, PlatformIntermediateFolder, GetTargetIntermediateFolderName(AppName, IntermediateEnvironment), UnrealTargetConfiguration.Development.ToString());
 			}
 			else
 			{
-				EngineIntermediateDirectory = DirectoryReference.Combine(Unreal.EngineDirectory, PlatformIntermediateFolder, GetTargetIntermediateFolderName(TargetName, IntermediateEnvironment), Configuration.ToString());
+				EngineIntermediateDirectory = DirectoryReference.Combine(Unreal.EngineDirectory, PlatformIntermediateFolder, GetTargetIntermediateFolderName(AppName, IntermediateEnvironment), Configuration.ToString());
 			}
 
 			// Get the receipt path for this target
