UBT: Use ld-classic in order to avoid the new linker since it seems to have some bugs with our code base.
Remove redundant -lc++

diff --git a/Engine/Source/Programs/UnrealBuildTool/Platform/IOS/IOSToolChain.cs b/Engine/Source/Programs/UnrealBuildTool/Platform/IOS/IOSToolChain.cs
--- a/Engine/Source/Programs/UnrealBuildTool/Platform/IOS/IOSToolChain.cs
+++ b/Engine/Source/Programs/UnrealBuildTool/Platform/IOS/IOSToolChain.cs
@@ -304,8 +304,9 @@ namespace UnrealBuildTool
 			}
 		}
 
-		void GetLinkArguments_Global(LinkEnvironment LinkEnvironment, List<string> Arguments)
+		protected override void GetLinkArguments_Global(LinkEnvironment LinkEnvironment, List<string> Arguments)
 		{
+			base.GetLinkArguments_Global(LinkEnvironment, Arguments);
 			Arguments.Add(FormatArchitectureArg(LinkEnvironment.Architectures));
 
 			DirectoryReference SDKPath = Settings.GetSDKPath(LinkEnvironment.Architecture);
diff --git a/Engine/Source/Programs/UnrealBuildTool/Platform/Mac/MacToolChain.cs b/Engine/Source/Programs/UnrealBuildTool/Platform/Mac/MacToolChain.cs
--- a/Engine/Source/Programs/UnrealBuildTool/Platform/Mac/MacToolChain.cs
+++ b/Engine/Source/Programs/UnrealBuildTool/Platform/Mac/MacToolChain.cs
@@ -225,8 +225,9 @@ namespace UnrealBuildTool
 			return Result;
 		}
 
-		void GetLinkArguments_Global(LinkEnvironment LinkEnvironment, List<string> Arguments)
+		protected override void GetLinkArguments_Global(LinkEnvironment LinkEnvironment, List<string> Arguments)
 		{
+			base.GetLinkArguments_Global(LinkEnvironment, Arguments);
 			// Pass through architecture and OS info		
 			Arguments.Add(FormatArchitectureArg(LinkEnvironment.Architectures));
 			Arguments.Add(String.Format("-isysroot \"{0}\"", ToolChainSettings.Value.GetSDKPath(LinkEnvironment.Architecture)));
diff --git a/Engine/Source/Programs/UnrealBuildTool/ToolChain/AppleToolChain.cs b/Engine/Source/Programs/UnrealBuildTool/ToolChain/AppleToolChain.cs
--- a/Engine/Source/Programs/UnrealBuildTool/ToolChain/AppleToolChain.cs
+++ b/Engine/Source/Programs/UnrealBuildTool/ToolChain/AppleToolChain.cs
@@ -7,6 +7,7 @@ using System.IO;
 using System.Linq;
 using System.Runtime.InteropServices;
 using System.Text;
+using System.Text.RegularExpressions;
 using System.Threading.Tasks;
 using EpicGames.Core;
 using Microsoft.Extensions.Logging;
@@ -769,6 +770,27 @@ namespace UnrealBuildTool
 			return OutputFiles;
 		}
 
+		protected virtual void GetLinkArguments_Global(LinkEnvironment LinkEnvironment, List<string> Arguments)
+		{
+			// Temp solution for UE-191350
+
+			// This Regex will extract a 2-4 segment version code from string containing a 2-5 segment code 
+			// ie: if given string: "Apple clang version 14.0.0 (clang-1400.0.17.3.1)"
+			//     it'll extract: "1400.0.17.3"  (note the dropped 5th segment)
+			Match M = Regex.Match(Info.ClangVersionString, @"(\(clang-(?<ver>\d+.\d+(.(\d+))?(.(\d+))?)(.(\d+))?\))");
+			if (M.Success)
+			{
+				string LibString = M.Groups["ver"].ToString();
+				Version? LibVersion = new Version(LibString);
+				// The Apple's new linker in Xcode 15 beta 5 (clang version 1500.0.38.1) has issues with some templated classes and dynamic linking.
+				// Fall back to using the classic.
+				if (LibVersion.CompareTo(new Version(1500, 0, 38, 1)) >= 0)
+				{
+					Arguments.Add(" -ld_classic");
+				}
+			}
+		}
+
 		#region Stub Xcode Projects
 
 		internal static bool GenerateProjectFiles(FileReference? ProjectFile, string[] Arguments, string? SingleTargetName, ILogger Logger, out DirectoryReference? XcodeProjectFile)
