Cleaned up the project generator API by passing a struct of commonly used arguments. This makes it easier for us to add/remove arguments.

diff --git a/Engine/Source/Programs/UnrealBuildTool/Platform/Android/AndroidProjectGenerator.cs b/Engine/Source/Programs/UnrealBuildTool/Platform/Android/AndroidProjectGenerator.cs
--- a/Engine/Source/Programs/UnrealBuildTool/Platform/Android/AndroidProjectGenerator.cs
+++ b/Engine/Source/Programs/UnrealBuildTool/Platform/Android/AndroidProjectGenerator.cs
@@ -73,7 +73,7 @@ namespace UnrealBuildTool
 		}
 
 		/// <inheritdoc/>
-		public override bool HasVisualStudioSupport(UnrealTargetPlatform InPlatform, UnrealTargetConfiguration InConfiguration, VCProjectFileFormat ProjectFileFormat, DirectoryReference InProjectDir, UnrealArch? InArch)
+		public override bool HasVisualStudioSupport(VSSettings InVSSettings)
 		{
 			// Debugging, etc. are dependent on the TADP being installed
 			return AGDEInstalled;
@@ -81,11 +81,11 @@ namespace UnrealBuildTool
 
 
 		/// <inheritdoc/>
-		public override string GetVisualStudioPlatformName(UnrealTargetPlatform InPlatform, UnrealTargetConfiguration InConfiguration, DirectoryReference InProjectDir, UnrealArch? InArch)
+		public override string GetVisualStudioPlatformName(VSSettings InVSSettings)
 		{
-			string PlatformName = InPlatform.ToString();
+			string PlatformName = InVSSettings.Platform.ToString();
 
-			if (InPlatform == UnrealTargetPlatform.Android && AGDEInstalled)
+			if (InVSSettings.Platform == UnrealTargetPlatform.Android && AGDEInstalled)
 			{
 				PlatformName = "Android-arm64-v8a";
 			}
@@ -93,34 +93,17 @@ namespace UnrealBuildTool
 			return PlatformName;
 		}
 
-		/// <summary>
-		/// Return any custom property group lines
-		/// </summary>
-		/// <param name="InPlatform">  The UnrealTargetPlatform being built</param>
-		/// <param name="ProjectFileFormat"></param>
-		/// <param name="ProjectFileBuilder">String builder for the project file</param>
-		public override void GetAdditionalVisualStudioPropertyGroups(UnrealTargetPlatform InPlatform, VCProjectFileFormat ProjectFileFormat, StringBuilder ProjectFileBuilder)
+		/// <inheritdoc/>
+		public override void GetAdditionalVisualStudioPropertyGroups(VSSettings InVSSettings, StringBuilder ProjectFileBuilder)
 		{
 			if (AGDEInstalled)
 			{
-				base.GetAdditionalVisualStudioPropertyGroups(InPlatform, ProjectFileFormat, ProjectFileBuilder);
+				base.GetAdditionalVisualStudioPropertyGroups(InVSSettings, ProjectFileBuilder);
 			}
 		}
 
-		/// <summary>
-		/// Return any custom paths for VisualStudio this platform requires
-		/// This include ReferencePath, LibraryPath, LibraryWPath, IncludePath and ExecutablePath.
-		/// </summary>
-		/// <param name="InPlatform">The UnrealTargetPlatform being built</param>
-		/// <param name="InConfiguration">The configuration being built</param>
-		/// <param name="TargetType">The type of target (game or program)</param>
-		/// <param name="TargetRulesPath">Path to the target.cs file</param>
-		/// <param name="ProjectFilePath">Path to the project file</param>
-		/// <param name="NMakeOutputPath"></param>
-		/// <param name="InProjectFileFormat">Format for the generated project files</param>
-		/// <param name="ProjectFileBuilder">String builder for the project file</param>
-		/// <returns>The custom path lines for the project file; Empty string if it doesn't require one</returns>
-		public override void GetVisualStudioPathsEntries(UnrealTargetPlatform InPlatform, UnrealTargetConfiguration InConfiguration, TargetType TargetType, FileReference TargetRulesPath, FileReference ProjectFilePath, FileReference NMakeOutputPath, VCProjectFileFormat InProjectFileFormat, StringBuilder ProjectFileBuilder)
+		/// <inheritdoc/>
+		public override void GetVisualStudioPathsEntries(VSSettings InVSSettings, TargetType TargetType, FileReference TargetRulesPath, FileReference ProjectFilePath, FileReference NMakeOutputPath, StringBuilder ProjectFileBuilder)
 		{
 			if (AGDEInstalled)
 			{
@@ -135,21 +118,21 @@ namespace UnrealBuildTool
 			}
 			else
 			{
-				base.GetVisualStudioPathsEntries(InPlatform, InConfiguration, TargetType, TargetRulesPath, ProjectFilePath, NMakeOutputPath, InProjectFileFormat, ProjectFileBuilder);
+				base.GetVisualStudioPathsEntries(InVSSettings, TargetType, TargetRulesPath, ProjectFilePath, NMakeOutputPath, ProjectFileBuilder);
 			}
 		}
 
-		public override string GetExtraBuildArguments(UnrealTargetPlatform InPlatform, UnrealTargetConfiguration InConfiguration)
+		public override string GetExtraBuildArguments(VSSettings InVSSettings)
 		{
 			// do not need to check InPlatform since it will always be UnrealTargetPlatform.Android
-			return (AGDEInstalled ? " -Architectures=arm64 -ForceAPKGeneration" : "") + base.GetExtraBuildArguments(InPlatform, InConfiguration);
+			return (AGDEInstalled ? " -Architectures=arm64 -ForceAPKGeneration" : "") + base.GetExtraBuildArguments(InVSSettings);
 		}
 
-		public override string GetVisualStudioUserFileStrings(UnrealTargetPlatform InPlatform, UnrealTargetConfiguration InConfiguration,
+		public override string GetVisualStudioUserFileStrings(VSSettings InVSSettings,
 			string InConditionString, TargetRules InTargetRules, FileReference TargetRulesPath, FileReference ProjectFilePath)
 		{
 			if (AGDEInstalled
-				&& (InPlatform == UnrealTargetPlatform.Android)
+				&& (InVSSettings.Platform == UnrealTargetPlatform.Android)
 				&& ((InTargetRules.Type == TargetRules.TargetType.Client) || (InTargetRules.Type == TargetRules.TargetType.Game)))
 			{
 				string UserFileEntry = "<PropertyGroup " + InConditionString + ">\n";
@@ -161,7 +144,7 @@ namespace UnrealBuildTool
 				return UserFileEntry;
 			}
 
-			return base.GetVisualStudioUserFileStrings(InPlatform, InConfiguration, InConditionString, InTargetRules, TargetRulesPath, ProjectFilePath);
+			return base.GetVisualStudioUserFileStrings(InVSSettings, InConditionString, InTargetRules, TargetRulesPath, ProjectFilePath);
 		}
 	}
 }
diff --git a/Engine/Source/Programs/UnrealBuildTool/Platform/IOS/IOSProjectGenerator.cs b/Engine/Source/Programs/UnrealBuildTool/Platform/IOS/IOSProjectGenerator.cs
--- a/Engine/Source/Programs/UnrealBuildTool/Platform/IOS/IOSProjectGenerator.cs
+++ b/Engine/Source/Programs/UnrealBuildTool/Platform/IOS/IOSProjectGenerator.cs
@@ -30,7 +30,7 @@ namespace UnrealBuildTool
 		}
 
 		/// <inheritdoc/>
-		public override bool HasVisualStudioSupport(UnrealTargetPlatform InPlatform, UnrealTargetConfiguration InConfiguration, VCProjectFileFormat ProjectFileFormat, DirectoryReference InProjectDir, UnrealArch? InArch)
+		public override bool HasVisualStudioSupport(VSSettings InVSSettings)
 		{
 			// iOS is not supported in VisualStudio
 			return false;
diff --git a/Engine/Source/Programs/UnrealBuildTool/Platform/Linux/LinuxProjectGenerator.cs b/Engine/Source/Programs/UnrealBuildTool/Platform/Linux/LinuxProjectGenerator.cs
--- a/Engine/Source/Programs/UnrealBuildTool/Platform/Linux/LinuxProjectGenerator.cs
+++ b/Engine/Source/Programs/UnrealBuildTool/Platform/Linux/LinuxProjectGenerator.cs
@@ -31,7 +31,7 @@ namespace UnrealBuildTool
 		}
 
 		/// <inheritdoc/>
-		public override bool HasVisualStudioSupport(UnrealTargetPlatform InPlatform, UnrealTargetConfiguration InConfiguration, VCProjectFileFormat ProjectFileFormat, DirectoryReference InProjectDir, UnrealArch? InArch)
+		public override bool HasVisualStudioSupport(VSSettings InVSSettings)
 		{
 			return false;
 		}
diff --git a/Engine/Source/Programs/UnrealBuildTool/Platform/Mac/MacProjectGenerator.cs b/Engine/Source/Programs/UnrealBuildTool/Platform/Mac/MacProjectGenerator.cs
--- a/Engine/Source/Programs/UnrealBuildTool/Platform/Mac/MacProjectGenerator.cs
+++ b/Engine/Source/Programs/UnrealBuildTool/Platform/Mac/MacProjectGenerator.cs
@@ -30,7 +30,7 @@ namespace UnrealBuildTool
 		}
 
 		/// <inheritdoc/>
-		public override bool HasVisualStudioSupport(UnrealTargetPlatform InPlatform, UnrealTargetConfiguration InConfiguration, VCProjectFileFormat ProjectFileFormat, DirectoryReference InProjectDir, UnrealArch? InArch)
+		public override bool HasVisualStudioSupport(VSSettings InVSSettings)
 		{
 			// Mac is not supported in VisualStudio
 			return false;
diff --git a/Engine/Source/Programs/UnrealBuildTool/Platform/Windows/WindowsProjectGenerator.cs b/Engine/Source/Programs/UnrealBuildTool/Platform/Windows/WindowsProjectGenerator.cs
--- a/Engine/Source/Programs/UnrealBuildTool/Platform/Windows/WindowsProjectGenerator.cs
+++ b/Engine/Source/Programs/UnrealBuildTool/Platform/Windows/WindowsProjectGenerator.cs
@@ -29,25 +29,25 @@ namespace UnrealBuildTool
 		}
 
 		/// <inheritdoc/>
-		public override string GetVisualStudioPlatformName(UnrealTargetPlatform InPlatform, UnrealTargetConfiguration InConfiguration, DirectoryReference InProjectDir, UnrealArch? InArch)
+		public override string GetVisualStudioPlatformName(VSSettings InVSSettings)
 		{
-			if (InPlatform == UnrealTargetPlatform.Win64)
+			if (InVSSettings.Platform == UnrealTargetPlatform.Win64)
 			{
-				if (InArch == UnrealArch.Arm64)
+				if (InVSSettings.Architecture == UnrealArch.Arm64)
 				{
 					return "arm64";
 				}
-				else if (InArch == UnrealArch.Arm64ec)
+				else if (InVSSettings.Architecture == UnrealArch.Arm64ec)
 				{
 					return "arm64ec";
 				}
 				return "x64";
 			}
-			return InPlatform.ToString();
+			return InVSSettings.Platform.ToString();
 		}
 
 		/// <inheritdoc/>
-		public override string GetVisualStudioUserFileStrings(UnrealTargetPlatform InPlatform, UnrealTargetConfiguration InConfiguration, string InConditionString, TargetRules InTargetRules, FileReference TargetRulesPath, FileReference ProjectFilePath, string ProjectName, string? ForeignUProjectPath)
+		public override string GetVisualStudioUserFileStrings(VSSettings InVSSettings, string InConditionString, TargetRules InTargetRules, FileReference TargetRulesPath, FileReference ProjectFilePath, string ProjectName, string? ForeignUProjectPath)
 		{
 			StringBuilder VCUserFileContent = new StringBuilder();
 
diff --git a/Engine/Source/Programs/UnrealBuildTool/ProjectFiles/PlatformProjectGenerator.cs b/Engine/Source/Programs/UnrealBuildTool/ProjectFiles/PlatformProjectGenerator.cs
--- a/Engine/Source/Programs/UnrealBuildTool/ProjectFiles/PlatformProjectGenerator.cs
+++ b/Engine/Source/Programs/UnrealBuildTool/ProjectFiles/PlatformProjectGenerator.cs
@@ -50,30 +50,41 @@ namespace UnrealBuildTool
 
 		ConcurrentDictionary<int, bool> FoundVSISupportDict = new();
 
+		public struct VSSettings
+		{
+			public UnrealTargetPlatform Platform;
+			public UnrealTargetConfiguration Configuration = UnrealTargetConfiguration.Development;
+			public VCProjectFileFormat ProjectFileFormat = VCProjectFileFormat.Default;
+			public UnrealArch? Architecture = null;
+
+			public VSSettings(UnrealTargetPlatform InPlatform, UnrealTargetConfiguration InConfiguration, VCProjectFileFormat InProjectFileFormat, UnrealArch? InArchitecture)
+			{
+				Platform = InPlatform;
+				Configuration = InConfiguration;
+				ProjectFileFormat = InProjectFileFormat;
+				Architecture = InArchitecture;
+			}
+		};
+
 		/// <summary>
 		/// Checks the local VS install directories to see if the platform is supported.
 		/// </summary>
-		/// <param name="InPlatform">  The UnrealTargetPlatform being built</param>
-		/// <param name="InConfiguration"> The UnrealTargetConfiguration being built</param>
-		/// <param name="ProjectFileFormat">The visual studio project file format being generated</param>
-		/// <param name="InProjectDir">  The root directory for the current project</param>
-		/// <param name="InArch">  The architecture of the platform</param>
+		/// <param name="InVSSettings">The ProjectFileSettings that contains the project platform/configuration/etc info being built</param>
 		/// <returns>bool true if native VisualStudio support (or custom VSI) is available</returns>
-		public virtual bool HasVisualStudioSupport(UnrealTargetPlatform InPlatform, UnrealTargetConfiguration InConfiguration, VCProjectFileFormat ProjectFileFormat, DirectoryReference InProjectDir, UnrealArch? InArch)
+		public virtual bool HasVisualStudioSupport(VSSettings InVSSettings)
 		{
 			int HashResult = 691;
-			HashResult *= InPlatform.ToString().GetHashCode();
-			HashResult *= InConfiguration.ToString().GetHashCode();
-			HashResult *= ProjectFileFormat.ToString().GetHashCode();
-			HashResult *= InProjectDir.GetHashCode();
-			HashResult *= InArch != null ? InArch.ToString()!.GetHashCode() : 1;
+			HashResult *= InVSSettings.Platform.ToString().GetHashCode();
+			HashResult *= InVSSettings.Configuration.ToString().GetHashCode();
+			HashResult *= InVSSettings.ProjectFileFormat.ToString().GetHashCode();
+			HashResult *= InVSSettings.Architecture != null ? InVSSettings.Architecture.ToString()!.GetHashCode() : 1;
 
 			return FoundVSISupportDict.GetOrAdd(HashResult, _ =>
 			{
 				WindowsCompiler VSCompiler;
 				string VCVersion;
 
-				switch (ProjectFileFormat)
+				switch (InVSSettings.ProjectFileFormat)
 				{
 					case VCProjectFileFormat.VisualStudio2019:
 						VSCompiler = WindowsCompiler.VisualStudio2019;
@@ -93,7 +104,7 @@ namespace UnrealBuildTool
 				{
 					foreach (DirectoryReference VSInstallDir in InstallDirs)
 					{
-						DirectoryReference PlatformsPath = new DirectoryReference(System.IO.Path.Combine(VSInstallDir.FullName, "MSBuild\\Microsoft\\VC\\", VCVersion, "Platforms", GetVisualStudioPlatformName(InPlatform, InConfiguration, InProjectDir, InArch)));
+						DirectoryReference PlatformsPath = new DirectoryReference(System.IO.Path.Combine(VSInstallDir.FullName, "MSBuild\\Microsoft\\VC\\", VCVersion, "Platforms", GetVisualStudioPlatformName(InVSSettings)));
 						if (DirectoryReference.Exists(PlatformsPath))
 						{
 							return true;
@@ -107,25 +118,20 @@ namespace UnrealBuildTool
 		/// <summary>
 		/// Return the VisualStudio platform name for this build platform
 		/// </summary>
-		/// <param name="InPlatform">  The UnrealTargetPlatform being built</param>
-		/// <param name="InConfiguration"> The UnrealTargetConfiguration being built</param>
-		/// <param name="InProjectDir">  The root directory for the current project</param>
-		/// <param name="InArch">  The architecture of the platform</param>
+		/// <param name="InVSSettings">The ProjectFileSettings that contains the project platform/configuration/etc info being built</param>
 		/// <returns>string    The name of the platform that VisualStudio recognizes</returns>
-		public virtual string GetVisualStudioPlatformName(UnrealTargetPlatform InPlatform, UnrealTargetConfiguration InConfiguration, DirectoryReference InProjectDir, UnrealArch? InArch)
+		public virtual string GetVisualStudioPlatformName(VSSettings InVSSettings)
 		{
 			// By default, return the platform string
-			return InPlatform.ToString();
+			return InVSSettings.Platform.ToString();
 		}
 
 		/// <summary>
 		/// Return whether we need a distinct platform name - for cases where there are two targets of the same VS solution platform
 		/// </summary>
-		/// <param name="InPlatform">  The UnrealTargetPlatform being built</param>
-		/// <param name="InConfiguration"> The UnrealTargetConfiguration being built</param>
-		/// <param name="InProjectDir">  The root directory for the current project</param>
+		/// <param name="InVSSettings">The ProjectFileSettings that contains the project platform/configuration/etc info being built</param>
 		/// <returns>bool</returns>
-		public virtual bool RequiresDistinctVisualStudioConfigurationName(UnrealTargetPlatform InPlatform, UnrealTargetConfiguration InConfiguration, DirectoryReference InProjectDir)
+		public virtual bool RequiresDistinctVisualStudioConfigurationName(VSSettings InVSSettings)
 		{
 			return false;
 		}
@@ -154,33 +160,29 @@ namespace UnrealBuildTool
 		/// <summary>
 		/// Return the platform toolset string to write into the project configuration
 		/// </summary>
-		/// <param name="InPlatform">  The UnrealTargetPlatform being built</param>
-		/// <param name="InConfiguration"> The UnrealTargetConfiguration being built</param>
-		/// <param name="InProjectFileFormat">The visual studio project file format being generated</param>
+		/// <param name="InVSSettings">The ProjectFileSettings that contains the project platform/configuration/etc info being built</param>
 		/// <param name="ProjectFileBuilder">String builder for the project file</param>
 		/// <returns>string    The custom configuration section for the project file; Empty string if it doesn't require one</returns>
-		public virtual void GetVisualStudioPlatformToolsetString(UnrealTargetPlatform InPlatform, UnrealTargetConfiguration InConfiguration, VCProjectFileFormat InProjectFileFormat, StringBuilder ProjectFileBuilder)
+		public virtual void GetVisualStudioPlatformToolsetString(VSSettings InVSSettings, StringBuilder ProjectFileBuilder)
 		{
 		}
 
 		/// <summary>
 		/// Return any custom property group lines
 		/// </summary>
-		/// <param name="InPlatform">  The UnrealTargetPlatform being built</param>
-		/// <param name="InProjectFileFormat">The visual studio project file format being generated</param>
+		/// <param name="InVSSettings">The ProjectFileSettings that contains the project platform/configuration/etc info being built</param>
 		/// <param name="ProjectFileBuilder">String builder for the project file</param>
 		/// <returns>string    The custom property import lines for the project file; Empty string if it doesn't require one</returns>
-		public virtual void GetAdditionalVisualStudioPropertyGroups(UnrealTargetPlatform InPlatform, VCProjectFileFormat InProjectFileFormat, StringBuilder ProjectFileBuilder)
+		public virtual void GetAdditionalVisualStudioPropertyGroups(VSSettings InVSSettings, StringBuilder ProjectFileBuilder)
 		{
 		}
 
 		/// <summary>
 		/// Return any custom property group lines
 		/// </summary>
-		/// <param name="InPlatform">  The UnrealTargetPlatform being built</param>
-		/// <param name="InProjectFileFormat">The visual studio project file format being generated</param>
+		/// <param name="InVSSettings">The ProjectFileSettings that contains the project platform/configuration/etc info being built</param>
 		/// <returns>string    The platform configuration type.  Defaults to "Makefile" unless overridden</returns>
-		public virtual string GetVisualStudioPlatformConfigurationType(UnrealTargetPlatform InPlatform, VCProjectFileFormat InProjectFileFormat)
+		public virtual string GetVisualStudioPlatformConfigurationType(VSSettings InVSSettings)
 		{
 			return DefaultPlatformConfigurationType;
 		}
@@ -189,63 +191,58 @@ namespace UnrealBuildTool
 		/// Return any custom paths for VisualStudio this platform requires
 		/// This include ReferencePath, LibraryPath, LibraryWPath, IncludePath and ExecutablePath.
 		/// </summary>
-		/// <param name="InPlatform">The UnrealTargetPlatform being built</param>
-		/// <param name="InConfiguration">The configuration being built</param>
+		/// <param name="InVSSettings">The ProjectFileSettings that contains the project platform/configuration/etc info being built</param>
 		/// <param name="TargetType">The type of target (game or program)</param>
 		/// <param name="TargetRulesPath">Path to the .target.cs file</param>
 		/// <param name="ProjectFilePath"></param>
 		/// <param name="NMakeOutputPath"></param>
-		/// <param name="InProjectFileFormat">The visual studio project file format being generated</param>
 		/// <param name="ProjectFileBuilder">String builder for the project file</param>
 		/// <returns>The custom path lines for the project file; Empty string if it doesn't require one</returns>
-		public virtual void GetVisualStudioPathsEntries(UnrealTargetPlatform InPlatform, UnrealTargetConfiguration InConfiguration, TargetType TargetType, FileReference TargetRulesPath, FileReference ProjectFilePath, FileReference NMakeOutputPath, VCProjectFileFormat InProjectFileFormat, StringBuilder ProjectFileBuilder)
+		public virtual void GetVisualStudioPathsEntries(VSSettings InVSSettings, TargetType TargetType, FileReference TargetRulesPath, FileReference ProjectFilePath, FileReference NMakeOutputPath, StringBuilder ProjectFileBuilder)
 		{
 		}
 
 		/// <summary>
 		/// Return any custom property settings. These will be included in the ImportGroup section
 		/// </summary>
-		/// <param name="InPlatform">  The UnrealTargetPlatform being built</param>
+		/// <param name="InVSSettings">The ProjectFileSettings that contains the project platform/configuration/etc info being built</param>
 		/// <param name="ProjectFileBuilder">String builder for the project file</param>
 		/// <returns>string    The custom property import lines for the project file; Empty string if it doesn't require one</returns>
-		public virtual void GetVisualStudioImportGroupProperties(UnrealTargetPlatform InPlatform, StringBuilder ProjectFileBuilder)
+		public virtual void GetVisualStudioImportGroupProperties(VSSettings InVSSettings, StringBuilder ProjectFileBuilder)
 		{
 		}
 
 		/// <summary>
 		/// Return any custom property settings. These will be included right after Global properties to make values available to all other imports.
 		/// </summary>
-		/// <param name="InPlatform">  The UnrealTargetPlatform being built</param>
+		/// <param name="InVSSettings">The ProjectFileSettings that contains the project platform/configuration/etc info being built</param>
 		/// <param name="ProjectFileBuilder">String builder for the project file</param>
 		/// <returns>string    The custom property import lines for the project file; Empty string if it doesn't require one</returns>
-		public virtual void GetVisualStudioGlobalProperties(UnrealTargetPlatform InPlatform, StringBuilder ProjectFileBuilder)
+		public virtual void GetVisualStudioGlobalProperties(VSSettings InVSSettings, StringBuilder ProjectFileBuilder)
 		{
 		}
 
 		/// <summary>
 		/// Return any custom target overrides. These will be included last in the project file so they have the opportunity to override any existing settings.
 		/// </summary>
-		/// <param name="InPlatform">  The UnrealTargetPlatform being built</param>
-		/// <param name="InProjectFileFormat">The visual studio project file format being generated</param>
+		/// <param name="InVSSettings">The ProjectFileSettings that contains the project platform/configuration/etc info being built</param>
 		/// <param name="ProjectFileBuilder">String builder for the project file</param>
 		/// <returns>string    The custom property import lines for the project file; Empty string if it doesn't require one</returns>
-		public virtual void GetVisualStudioTargetOverrides(UnrealTargetPlatform InPlatform, VCProjectFileFormat InProjectFileFormat, StringBuilder ProjectFileBuilder)
+		public virtual void GetVisualStudioTargetOverrides(VSSettings InVSSettings, StringBuilder ProjectFileBuilder)
 		{
 		}
 
 		/// <summary>
 		/// Return any custom layout directory sections
 		/// </summary>
-		/// <param name="InPlatform">The UnrealTargetPlatform being built</param>
-		/// <param name="InConfiguration">The configuration being built</param>
+		/// <param name="InVSSettings">The ProjectFileSettings that contains the project platform/configuration/etc info being built</param>
 		/// <param name="InConditionString"></param>
 		/// <param name="TargetType">The type of target (game or program)</param>
-		/// <param name="InProjectFileFormat">The visual studio project file format being generated</param>
 		/// <param name="NMakeOutputPath"></param>
 		/// <param name="ProjectFilePath"></param>
 		/// <param name="TargetRulesPath"></param>
 		/// <returns>string    The custom property import lines for the project file; Empty string if it doesn't require one</returns>
-		public virtual string GetVisualStudioLayoutDirSection(UnrealTargetPlatform InPlatform, UnrealTargetConfiguration InConfiguration, string InConditionString, TargetType TargetType, FileReference TargetRulesPath, FileReference ProjectFilePath, FileReference NMakeOutputPath, VCProjectFileFormat InProjectFileFormat)
+		public virtual string GetVisualStudioLayoutDirSection(VSSettings InVSSettings, string InConditionString, TargetType TargetType, FileReference TargetRulesPath, FileReference ProjectFilePath, FileReference NMakeOutputPath)
 		{
 			return "";
 		}
@@ -253,13 +250,12 @@ namespace UnrealBuildTool
 		/// <summary>
 		/// Get the output manifest section, if required
 		/// </summary>
-		/// <param name="InPlatform">The UnrealTargetPlatform being built</param>
+		/// <param name="InVSSettings">The ProjectFileSettings that contains the project platform/configuration/etc info being built</param>
 		/// <param name="TargetType">The type of the target being built</param>
 		/// <param name="TargetRulesPath">Path to the .target.cs file</param>
 		/// <param name="ProjectFilePath">Path to the project file</param>
-		/// <param name="InProjectFileFormat">The visual studio project file format being generated</param>
 		/// <returns>The output manifest section for the project file; Empty string if it doesn't require one</returns>
-		public virtual string GetVisualStudioOutputManifestSection(UnrealTargetPlatform InPlatform, TargetType TargetType, FileReference TargetRulesPath, FileReference ProjectFilePath, VCProjectFileFormat InProjectFileFormat)
+		public virtual string GetVisualStudioOutputManifestSection(VSSettings InVSSettings, TargetType TargetType, FileReference TargetRulesPath, FileReference ProjectFilePath)
 		{
 			return "";
 		}
@@ -268,7 +264,7 @@ namespace UnrealBuildTool
 		/// Get whether this platform deploys
 		/// </summary>
 		/// <returns>bool  true if the 'Deploy' option should be enabled</returns>
-		public virtual bool GetVisualStudioDeploymentEnabled(UnrealTargetPlatform InPlatform, UnrealTargetConfiguration InConfiguration)
+		public virtual bool GetVisualStudioDeploymentEnabled(VSSettings InVSSettings)
 		{
 			return false;
 		}
@@ -276,14 +272,13 @@ namespace UnrealBuildTool
 		/// <summary>
 		/// Get the text to insert into the user file for the given platform/configuration/target
 		/// </summary>
-		/// <param name="InPlatform">The platform being added</param>
-		/// <param name="InConfiguration">The configuration being added</param>
+		/// <param name="InVSSettings">The ProjectFileSettings that contains the project platform/configuration/etc info being built</param>
 		/// <param name="InConditionString">The condition string </param>
 		/// <param name="InTargetRules">The target rules </param>
 		/// <param name="TargetRulesPath">The target rules path</param>
 		/// <param name="ProjectFilePath">The project file path</param>
 		/// <returns>The string to append to the user file</returns>
-		public virtual string GetVisualStudioUserFileStrings(UnrealTargetPlatform InPlatform, UnrealTargetConfiguration InConfiguration,
+		public virtual string GetVisualStudioUserFileStrings(VSSettings InVSSettings,
 			string InConditionString, TargetRules InTargetRules, FileReference TargetRulesPath, FileReference ProjectFilePath)
 		{
 			return "";
@@ -292,8 +287,7 @@ namespace UnrealBuildTool
 		/// <summary>
 		/// Get the text to insert into the user file for the given platform/configuration/target
 		/// </summary>
-		/// <param name="InPlatform">The platform being added</param>
-		/// <param name="InConfiguration">The configuration being added</param>
+		/// <param name="InVSSettings">The ProjectFileSettings that contains the project platform/configuration/etc info being built</param>
 		/// <param name="InConditionString">The condition string </param>
 		/// <param name="InTargetRules">The target rules </param>
 		/// <param name="TargetRulesPath">The target rules path</param>
@@ -301,10 +295,10 @@ namespace UnrealBuildTool
 		/// <param name="ProjectName">The name of the project</param>
 		/// <param name="ForeignUProjectPath">Path to foreign .uproject file, if any</param>
 		/// <returns>The string to append to the user file</returns>
-		public virtual string GetVisualStudioUserFileStrings(UnrealTargetPlatform InPlatform, UnrealTargetConfiguration InConfiguration,
+		public virtual string GetVisualStudioUserFileStrings(VSSettings InVSSettings,
 			string InConditionString, TargetRules InTargetRules, FileReference TargetRulesPath, FileReference ProjectFilePath, string ProjectName, string? ForeignUProjectPath)
 		{
-			return GetVisualStudioUserFileStrings(InPlatform, InConfiguration, InConditionString, InTargetRules, TargetRulesPath, ProjectFilePath);
+			return GetVisualStudioUserFileStrings(InVSSettings, InConditionString, InTargetRules, TargetRulesPath, ProjectFilePath);
 		}
 
 		/// <summary>
@@ -340,11 +334,10 @@ namespace UnrealBuildTool
 		}
 
 		/// <summary>
-		/// Gets the text to insert into the Project files as additional build arugments for given platform/configuration
+		/// Gets the text to insert into the Project files as additional build arguments for given platform/configuration
 		/// </summary>
-		/// <param name="InPlatform">The platform being added</param>
-		/// <param name="InConfiguration">The configuration being added</param>
-		public virtual string GetExtraBuildArguments(UnrealTargetPlatform InPlatform, UnrealTargetConfiguration InConfiguration)
+		/// <param name="InVSSettings">The ProjectFileSettings that contains the project platform/configuration/etc info being built</param>
+		public virtual string GetExtraBuildArguments(VSSettings InVSSettings)
 		{
 			return "";
 		}
diff --git a/Engine/Source/Programs/UnrealBuildTool/ProjectFiles/Rider/RiderProjectFile.cs b/Engine/Source/Programs/UnrealBuildTool/ProjectFiles/Rider/RiderProjectFile.cs
--- a/Engine/Source/Programs/UnrealBuildTool/ProjectFiles/Rider/RiderProjectFile.cs
+++ b/Engine/Source/Programs/UnrealBuildTool/ProjectFiles/Rider/RiderProjectFile.cs
@@ -617,7 +617,8 @@ namespace UnrealBuildTool
 			if (HostPlatform.IsInGroup(UnrealPlatformGroup.Windows))
 			{
 				PlatformProjectGenerator? ProjGenerator = PlatformProjectGenerators.GetPlatformProjectGenerator(Platform, true);
-				VCProjectFile.BuildCommandBuilder BuildCommandBuilder = new VCProjectFile.BuildCommandBuilder(Configuration, Platform, ProjectTarget, UProjectPath)
+				VCProjectFile.BuildCommandBuilder BuildCommandBuilder = new VCProjectFile.BuildCommandBuilder(
+					new PlatformProjectGenerator.VSSettings(Platform, Configuration, VCProjectFileFormat.Default, null), ProjectTarget, UProjectPath)
 				{
 					ProjectGenerator = ProjGenerator,
 					bIsForeignProject = IsForeignProject
diff --git a/Engine/Source/Programs/UnrealBuildTool/ProjectFiles/VisualStudio/VCProject.cs b/Engine/Source/Programs/UnrealBuildTool/ProjectFiles/VisualStudio/VCProject.cs
--- a/Engine/Source/Programs/UnrealBuildTool/ProjectFiles/VisualStudio/VCProject.cs
+++ b/Engine/Source/Programs/UnrealBuildTool/ProjectFiles/VisualStudio/VCProject.cs
@@ -16,6 +16,8 @@ using EpicGames.Core;
 using Microsoft.Extensions.Logging;
 using UnrealBuildBase;
 using Newtonsoft.Json.Linq;
+using Microsoft.CodeAnalysis;
+using static UnrealBuildTool.PlatformProjectGenerator;
 
 namespace UnrealBuildTool
 {
@@ -449,7 +451,7 @@ namespace UnrealBuildTool
 
 							PlatformProjectGenerator? ProjGen = PlatformProjectGenerators.GetPlatformProjectGenerator(SolutionPlatform, true);
 							if (MatchingProjectTarget.ProjectDeploys ||
-								((ProjGen != null) && (ProjGen.GetVisualStudioDeploymentEnabled(ProjectPlatform, ProjectConfiguration) == true)))
+								((ProjGen != null) && (ProjGen.GetVisualStudioDeploymentEnabled(new VSSettings(ProjectPlatform, ProjectConfiguration, ProjectFileFormat, Architecture)) == true)))
 							{
 								bDeployByDefault = true;
 							}
@@ -687,16 +689,18 @@ namespace UnrealBuildTool
 										string ProjectConfigurationName = Configuration.ToString();
 										bool CreateDistinctConfigName = false;
 
+										VSSettings VSSettings = new(Platform, Configuration, ProjectFileFormat, Architecture);
+
 										// Check to see if this platform is supported directly by Visual Studio projects.
-										if (PlatformProjectGenerator != null && PlatformProjectGenerator.HasVisualStudioSupport(Platform, Configuration, ProjectFileFormat, BaseDir, Architecture))
+										if (PlatformProjectGenerator != null && PlatformProjectGenerator.HasVisualStudioSupport(VSSettings))
 										{
 											// Allow the platform to specify the name used in VisualStudio.
 											// Note that the actual name of the platform on the Visual Studio side may be different than what
 											// UnrealBuildTool calls it (e.g. "Win64" -> "x64".) GetVisualStudioPlatformName() will figure this out.
-											ProjectPlatformName = PlatformProjectGenerator.GetVisualStudioPlatformName(Platform, Configuration, BaseDir, Architecture);
+											ProjectPlatformName = PlatformProjectGenerator.GetVisualStudioPlatformName(VSSettings);
 
 											// The project generator may require a distinct configuration name - typically when two UnrealTargetPlatforms need the same ProjectPlatformName - otherwise the properties overwrite each oter.
-											if (PlatformProjectGenerator.RequiresDistinctVisualStudioConfigurationName(Platform, Configuration, BaseDir))
+											if (PlatformProjectGenerator.RequiresDistinctVisualStudioConfigurationName(VSSettings))
 											{
 												CreateDistinctConfigName = true;
 											}
@@ -982,9 +986,10 @@ namespace UnrealBuildTool
 				foreach (UnrealTargetPlatform Platform in ProjectPlatforms)
 				{
 					PlatformProjectGenerator? ProjGenerator = PlatformProjectGenerators.GetPlatformProjectGenerator(Platform, true);
-					if (ProjGenerator != null && ProjGenerator.HasVisualStudioSupport(Platform, UnrealTargetConfiguration.Development, ProjectFileFormat, BaseDir, null))
+					VSSettings VSSettings = new(Platform, UnrealTargetConfiguration.Development, ProjectFileFormat, null);
+					if (ProjGenerator != null && ProjGenerator.HasVisualStudioSupport(VSSettings))
 					{
-						ProjGenerator.GetAdditionalVisualStudioPropertyGroups(Platform, ProjectFileFormat, VCProjectFileContent);
+						ProjGenerator.GetAdditionalVisualStudioPropertyGroups(VSSettings, VCProjectFileContent);
 					}
 				}
 			}
@@ -1010,10 +1015,11 @@ namespace UnrealBuildTool
 			{
 				foreach (UnrealTargetPlatform Platform in ProjectPlatforms)
 				{
+					VSSettings VSSettings = new(Platform, UnrealTargetConfiguration.Development, ProjectFileFormat, null);
 					PlatformProjectGenerator? ProjGenerator = PlatformProjectGenerators.GetPlatformProjectGenerator(Platform, true);
-					if (ProjGenerator != null && ProjGenerator.HasVisualStudioSupport(Platform, UnrealTargetConfiguration.Development, ProjectFileFormat, BaseDir, null))
+					if (ProjGenerator != null && ProjGenerator.HasVisualStudioSupport(VSSettings))
 					{
-						ProjGenerator.GetVisualStudioGlobalProperties(Platform, VCProjectFileContent);
+						ProjGenerator.GetVisualStudioGlobalProperties(VSSettings, VCProjectFileContent);
 					}
 				}
 			}
@@ -1087,9 +1093,10 @@ namespace UnrealBuildTool
 						continue;
 					}
 
+					PlatformProjectGenerator.VSSettings VSSettings = new(TargetPlatform, TargetConfiguration, ProjectFileFormat, null);
 					StringBuilder PlatformToolsetString = new StringBuilder();
-					ProjGenerator.GetVisualStudioPlatformToolsetString(TargetPlatform, TargetConfiguration, ProjectFileFormat, PlatformToolsetString);
-					string PlatformConfigurationType = ProjGenerator.GetVisualStudioPlatformConfigurationType(TargetPlatform, ProjectFileFormat);
+					ProjGenerator.GetVisualStudioPlatformToolsetString(VSSettings, PlatformToolsetString);
+					string PlatformConfigurationType = ProjGenerator.GetVisualStudioPlatformConfigurationType(VSSettings);
 
 					// if we are using the defaults set earlier then skip writing this
 					if (PlatformConfigurationType == PlatformProjectGenerator.DefaultPlatformConfigurationType && PlatformToolsetString.Length == 0)
@@ -1635,8 +1642,9 @@ namespace UnrealBuildTool
 			{
 				foreach (UnrealTargetPlatform Platform in ProjectPlatforms)
 				{
+					VSSettings VSSettings = new(Platform, UnrealTargetConfiguration.Development, ProjectFileFormat, null);
 					PlatformProjectGenerator? ProjGenerator = PlatformProjectGenerators.GetPlatformProjectGenerator(Platform, true);
-					if (ProjGenerator != null && ProjGenerator.HasVisualStudioSupport(Platform, UnrealTargetConfiguration.Development, ProjectFileFormat, BaseDir, null))
+					if (ProjGenerator != null && ProjGenerator.HasVisualStudioSupport(VSSettings))
 					{
 						// @todo projectfiles: Serious hacks here because we are trying to emit one-time platform-specific sections that need information
 						//    about a target type, but the project file may contain many types of targets!  Some of this logic will need to move into
@@ -1656,7 +1664,7 @@ namespace UnrealBuildTool
 
 						if (HackTargetFilePath != null)
 						{
-							OutputManifestString += ProjGenerator.GetVisualStudioOutputManifestSection(Platform, HackTargetType, HackTargetFilePath, ProjectFilePath, ProjectFileFormat);
+							OutputManifestString += ProjGenerator.GetVisualStudioOutputManifestSection(VSSettings, HackTargetType, HackTargetFilePath, ProjectFilePath);
 						}
 					}
 				}
@@ -1692,9 +1700,10 @@ namespace UnrealBuildTool
 				foreach (UnrealTargetPlatform Platform in ProjectPlatforms)
 				{
 					PlatformProjectGenerator? ProjGenerator = PlatformProjectGenerators.GetPlatformProjectGenerator(Platform, true);
-					if (ProjGenerator != null && ProjGenerator.HasVisualStudioSupport(Platform, UnrealTargetConfiguration.Development, ProjectFileFormat, BaseDir, null))
+					VSSettings VSSettings = new(Platform, UnrealTargetConfiguration.Development, ProjectFileFormat, null);
+					if (ProjGenerator != null && ProjGenerator.HasVisualStudioSupport(VSSettings))
 					{
-						ProjGenerator.GetVisualStudioTargetOverrides(Platform, ProjectFileFormat, VCProjectFileContent);
+						ProjGenerator.GetVisualStudioTargetOverrides(VSSettings, VCProjectFileContent);
 					}
 				}
 			}
@@ -1926,20 +1935,15 @@ namespace UnrealBuildTool
 
 			private readonly string UProjectPath;
 
-			private readonly UnrealTargetConfiguration Configuration;
-			private readonly UnrealTargetPlatform Platform;
+			private readonly VSSettings VSSettings;
 			private readonly ProjectTarget ProjectTarget;
-			private readonly UnrealArch? Architecture;
 
-			public BuildCommandBuilder(UnrealTargetConfiguration InConfiguration, UnrealTargetPlatform InPlatform,
-				ProjectTarget InProjectTarget, string InUProjectPath, string? InBuildToolOverride = null, UnrealArch? InArchitecture = null)
+			public BuildCommandBuilder(VSSettings InVSSettings, ProjectTarget InProjectTarget, string InUProjectPath, string? InBuildToolOverride = null)
 			{
-				Configuration = InConfiguration;
-				Platform = InPlatform;
+				VSSettings = InVSSettings;
 				ProjectTarget = InProjectTarget;
 				UProjectPath = InUProjectPath;
 				BuildToolOverride = InBuildToolOverride;
-				Architecture = InArchitecture;
 
 				DirectoryReference BatchFilesDirectory = DirectoryReference.Combine(Unreal.EngineDirectory, "Build", "BatchFiles");
 				BuildScript = FileReference.Combine(BatchFilesDirectory, "Build.bat");
@@ -1954,7 +1958,7 @@ namespace UnrealBuildTool
 
 				StringBuilder BuildArguments = new StringBuilder();
 
-				BuildArguments.AppendFormat("{0} {1} {2}", TargetName, Platform.ToString(), Configuration.ToString());
+				BuildArguments.AppendFormat("{0} {1} {2}", TargetName, VSSettings.Platform.ToString(), VSSettings.Configuration.ToString());
 				if (UProjectPath.Length > 0)
 				{
 					BuildArguments.AppendFormat(" -Project={0}", UProjectPath);
@@ -2010,12 +2014,12 @@ namespace UnrealBuildTool
 
 				if (ProjectGenerator != null)
 				{
-					BuildArguments.Append(ProjectGenerator.GetExtraBuildArguments(Platform, Configuration));
+					BuildArguments.Append(ProjectGenerator.GetExtraBuildArguments(VSSettings));
 				}
 
-				if (Architecture != null)
+				if (VSSettings.Architecture != null)
 				{
-					BuildArguments.AppendFormat(" -architecture={0}", Architecture);
+					BuildArguments.AppendFormat(" -architecture={0}", VSSettings.Architecture);
 				}
 
 				return BuildArguments.ToString();
@@ -2024,8 +2028,8 @@ namespace UnrealBuildTool
 
 		private BuildCommandBuilder CreateArgumentsBuilder(ProjectConfigAndTargetCombination Combination, string UProjectPath, PlatformProjectGenerator? ProjGenerator)
 		{
-			BuildCommandBuilder Builder = new BuildCommandBuilder(Combination.Configuration, Combination.Platform!.Value,
-				Combination.ProjectTarget!, UProjectPath, BuildToolOverride, Combination.Architecture)
+			BuildCommandBuilder Builder = new BuildCommandBuilder(new VSSettings(Combination.Platform!.Value, Combination.Configuration, ProjectFileFormat, Combination.Architecture),
+				Combination.ProjectTarget!, UProjectPath, BuildToolOverride)
 			{
 				ProjectGenerator = ProjGenerator,
 				bEditorDependsOnShaderCompileWorker = Settings.bEditorDependsOnShaderCompileWorker,
@@ -2061,7 +2065,7 @@ namespace UnrealBuildTool
 				if (ProjGenerator != null)
 				{
 					StringBuilder CustomImportGroupInfo = new StringBuilder();
-					ProjGenerator.GetVisualStudioImportGroupProperties(Combination.Platform!.Value, CustomImportGroupInfo);
+					ProjGenerator.GetVisualStudioImportGroupProperties(new(Combination.Platform!.Value, Configuration, ProjectFileFormat, null), CustomImportGroupInfo);
 					if (CustomImportGroupInfo.Length != 0)
 					{
 						VCProjectFileContent.AppendLine("  <ImportGroup {0} Label=\"PropertySheets\">", ConditionString);
@@ -2102,6 +2106,7 @@ namespace UnrealBuildTool
 					string TargetName = TargetFilePath.GetFileNameWithoutAnyExtensions();
 					string UBTPlatformName = Platform.ToString();
 					string UBTConfigurationName = Configuration.ToString();
+					VSSettings VSSettings = new(Platform, Configuration, ProjectFileFormat, null);
 
 					// Setup output path
 					UEBuildPlatform BuildPlatform = UEBuildPlatform.GetBuildPlatform(Platform);
@@ -2170,7 +2175,7 @@ namespace UnrealBuildTool
 					if (ProjGenerator != null)
 					{
 						StringBuilder PathsStringBuilder = new StringBuilder();
-						ProjGenerator.GetVisualStudioPathsEntries(Platform, Configuration, TargetRulesObject.Type, TargetFilePath, ProjectFilePath, NMakePath, ProjectFileFormat, PathsStringBuilder);
+						ProjGenerator.GetVisualStudioPathsEntries(VSSettings, TargetRulesObject.Type, TargetFilePath, ProjectFilePath, NMakePath, PathsStringBuilder);
 						VCProjectFileContent.Append(PathsStringBuilder.ToString());
 					}
 
@@ -2197,7 +2202,7 @@ namespace UnrealBuildTool
 
 					if (ProjGenerator != null)
 					{
-						VCProjectFileContent.Append(ProjGenerator.GetVisualStudioLayoutDirSection(Platform, Configuration, ConditionString, Combination.ProjectTarget.TargetRules!.Type, Combination.ProjectTarget.TargetFilePath, ProjectFilePath, NMakePath, ProjectFileFormat));
+						VCProjectFileContent.Append(ProjGenerator.GetVisualStudioLayoutDirSection(VSSettings, ConditionString, Combination.ProjectTarget.TargetRules!.Type, Combination.ProjectTarget.TargetFilePath, ProjectFilePath, NMakePath));
 					}
 
 					VCProjectFileContent.AppendLine("  <ItemDefinitionGroup {0}>", ConditionString);
@@ -2214,7 +2219,7 @@ namespace UnrealBuildTool
 					if (ProjGenerator != null)
 					{
 						string? ForeignUProjectPath = (IsForeignProject && !String.IsNullOrEmpty(UProjectPath)) ? UProjectPath : null;
-						VCUserFileContent.Append(ProjGenerator.GetVisualStudioUserFileStrings(Combination.Platform!.Value, Configuration, ConditionString, TargetRulesObject, Combination.ProjectTarget.TargetFilePath, ProjectFilePath, ProjectName, ForeignUProjectPath));
+						VCUserFileContent.Append(ProjGenerator.GetVisualStudioUserFileStrings(new(Combination.Platform!.Value, Configuration, ProjectFileFormat, null), ConditionString, TargetRulesObject, Combination.ProjectTarget.TargetFilePath, ProjectFilePath, ProjectName, ForeignUProjectPath));
 					}
 				}
 			}
