[UBT]
* Added actions first produced item directory path as part of sorting for scheduling. This groups modules together nicely which will benefit when running things remotely where remotes takes batches of actions at the time meaning there is a higher chance each remote helper will need to download fewer variations of pch files

diff --git a/Engine/Source/Programs/UnrealBuildTool/System/Action.cs b/Engine/Source/Programs/UnrealBuildTool/System/Action.cs
--- a/Engine/Source/Programs/UnrealBuildTool/System/Action.cs
+++ b/Engine/Source/Programs/UnrealBuildTool/System/Action.cs
@@ -764,11 +764,19 @@ namespace UnrealBuildTool
 			{
 				return Math.Sign(B.NumTotalDependentActions - A.NumTotalDependentActions);
 			}
-			// Secondary sort criteria is number of pre-requisites.
-			else
+
+			// Secondary sort criteria is directory name of first produced item.
+			// This will sort actions from same module closer to each other
+			// Also sort in alphabetic order since when building with unity files it is likely that the last unity file is the fastest to compile
+			// and we want longer actions to run earlier
+			int Result = String.Compare(A.ProducedItems.FirstOrDefault()?.FullName, B.ProducedItems.FirstOrDefault()?.FullName, StringComparison.Ordinal);
+			if (Result != 0)
 			{
-				return Math.Sign(B.PrerequisiteItems.Count() - A.PrerequisiteItems.Count());
+				return Math.Sign(Result);
 			}
+
+			// Third sort criteria is number of pre-requisites.
+			return Math.Sign(B.PrerequisiteItems.Count() - A.PrerequisiteItems.Count());
 		}
 	}
 }
