Fixes regression where BenchmarkBuild script would not reliably measure build performance.

The fix is to restore the <= 5.1 behavior of UBT where the MaxParallelActions argument specifies the maximum number of actions to execute in parallel, not a cap on UBTs heuristic-based selection.

This change also modifies BenchmarkBuild so it will default to the available processor count if no -cores argument is specified. This ensures the behavior of repeated runs will be deterministic as UBT may select different values if the system free memory changes

#tests Ran UBT / BenchmarkBuild with arguments that forces the number of actions higher than the locally available core count

diff --git a/Engine/Source/Programs/AutomationTool/Scripts/BenchmarkBuild/BenchmarkBuildTests.cs b/Engine/Source/Programs/AutomationTool/Scripts/BenchmarkBuild/BenchmarkBuildTests.cs
--- a/Engine/Source/Programs/AutomationTool/Scripts/BenchmarkBuild/BenchmarkBuildTests.cs
+++ b/Engine/Source/Programs/AutomationTool/Scripts/BenchmarkBuild/BenchmarkBuildTests.cs
@@ -87,12 +87,12 @@ namespace AutomationTool.Benchmark
 				//TaskModifiers.Add(Arg);
 				Command.Params = new[] { Arg }; // need to also pass it to this
 
-				if (CoreCount > 0)
-				{
-					TaskModifiers.Add(string.Format("{0}c", CoreCount));
+				// If no cores were specified use the machines CPU count rather than letting UBT pick a value. The latter may 
+				// not be deterministic.
+				int NumCores = CoreCount > 0 ? CoreCount : Environment.ProcessorCount;
 
-					Command.UBTArgs += string.Format(" -MaxParallelActions={0}", CoreCount);
-				}
+				TaskModifiers.Add(string.Format("{0}c", NumCores));
+				Command.UBTArgs += string.Format(" -MaxParallelActions={0}", NumCores);
 			}
 			else
 			{
diff --git a/Engine/Source/Programs/UnrealBuildTool/System/Utils.cs b/Engine/Source/Programs/UnrealBuildTool/System/Utils.cs
--- a/Engine/Source/Programs/UnrealBuildTool/System/Utils.cs
+++ b/Engine/Source/Programs/UnrealBuildTool/System/Utils.cs
@@ -1348,6 +1348,14 @@ namespace UnrealBuildTool
 		/// <returns>Max number of actions to execute in parallel</returns>
 		public static int GetMaxActionsToExecuteInParallel(int MaxProcessorCount, double ProcessorCountMultiplier, bool ConsiderLogicalCores, long MemoryPerActionBytes)
 		{
+			// If non-zero then this is the number of actions that can be executed in parallel. This 
+			// matches the BuildConfiguration.MaxLocalActions documentation and the <= 5.1 behavior
+			if (MaxProcessorCount != 0)
+			{
+				Log.TraceInformationOnce($"Executing up to {MaxProcessorCount} actions based on MaxProcessorCount override");
+				return MaxProcessorCount;
+			}
+
 			// Get the number of logical processors
 			int NumLogicalCores = Utils.GetLogicalProcessorCount();
 
@@ -1400,11 +1408,6 @@ namespace UnrealBuildTool
 				}
 			}
 
-			if (MaxProcessorCount > 0 && MaxProcessorCount < MaxActionsToExecuteInParallel)
-			{
-				MaxActionsToExecuteInParallel = Math.Max(1, Math.Min(MaxActionsToExecuteInParallel, MaxProcessorCount));
-				Log.TraceInformationOnce($"  Requested max {MaxProcessorCount} action(s): limiting max parallel actions to {MaxActionsToExecuteInParallel}");
-			}
 			return MaxActionsToExecuteInParallel;
 		}
 
