Remove circular reference entries that are no longer accurate post CL# 25899269
#rnx

diff --git a/Engine/Source/Editor/UnrealEd/UnrealEd.Build.cs b/Engine/Source/Editor/UnrealEd/UnrealEd.Build.cs
--- a/Engine/Source/Editor/UnrealEd/UnrealEd.Build.cs
+++ b/Engine/Source/Editor/UnrealEd/UnrealEd.Build.cs
@@ -270,9 +270,7 @@ public class UnrealEd : ModuleRules
 				"MeshPaint",
 				"PropertyEditor",
 				"ToolMenusEditor",
-				"InputBindingEditor",
 				"ClothingSystemEditor",
-				"PluginWarden",
 				"MaterialShaderQualitySettings",
 				"EditorInteractiveToolsFramework"
 			}
diff --git a/Engine/Source/Programs/UnrealBuildTool/Configuration/UEBuildModuleCPP.cs b/Engine/Source/Programs/UnrealBuildTool/Configuration/UEBuildModuleCPP.cs
--- a/Engine/Source/Programs/UnrealBuildTool/Configuration/UEBuildModuleCPP.cs
+++ b/Engine/Source/Programs/UnrealBuildTool/Configuration/UEBuildModuleCPP.cs
@@ -177,10 +177,8 @@ namespace UnrealBuildTool
 			new KeyValuePair<string, string>("UnrealEd", "Documentation"),
 			new KeyValuePair<string, string>("UnrealEd", "EditorInteractiveToolsFramework"),
 			new KeyValuePair<string, string>("UnrealEd", "GraphEditor"),
-			new KeyValuePair<string, string>("UnrealEd", "InputBindingEditor"),
 			new KeyValuePair<string, string>("UnrealEd", "Kismet"),
 			new KeyValuePair<string, string>("UnrealEd", "MeshPaint"),
-			new KeyValuePair<string, string>("UnrealEd", "PluginWarden"),
 			new KeyValuePair<string, string>("UnrealEd", "PropertyEditor"),
 			new KeyValuePair<string, string>("UnrealEd", "ToolMenusEditor"),
 			new KeyValuePair<string, string>("UnrealEd", "ViewportInteraction"),
