Change UBT limiting of actions based on memory to use total ram on Mac instead of free ram.

This addresses issues with CPU resources on Mac going under utilized on systems where total RAM is only 2-3x the number of cores (e.g 8c/16GB, 10c/32GB) and the user has applications running (UE, XCode, browsers) that are occupying RAM.

Historically this check limited actions based on free memory to ensure action processes weren't forced to use swap, but there's strong evidence that limiting based on total system memory and relying on the OS to swap out other processes to make room is fine. That said only Mac has been extensively tested here in this scenario (and tends to have faster disk performance than other platforms). So for now we'll change Mac and leave Windows/Linux using the old behavior

Tests -

* Ran on an M1 8c/16GB Mac with Xcode & UE open and compared core utilization before (5) to after this change (8)
* Loaded up a bucket-load of apps, browser tabs, and a running game and verified engine built successfully 2x with no compiler OOMs

#swarm [REDACTED_HANDLE] [REDACTED_HANDLE]

diff --git a/Engine/Source/Programs/UnrealBuildTool/System/Utils.cs b/Engine/Source/Programs/UnrealBuildTool/System/Utils.cs
--- a/Engine/Source/Programs/UnrealBuildTool/System/Utils.cs
+++ b/Engine/Source/Programs/UnrealBuildTool/System/Utils.cs
@@ -1237,6 +1237,17 @@ namespace UnrealBuildTool
 		[DllImport("libc")]
 		static extern IntPtr mach_host_self();
 
+		/// <summary>
+		/// Gets the total system memory in bytes based on what is known to the garbage collector.
+		/// </summary>
+		/// <remarks>This will return a max of 2GB for a 32bit application.</remarks>
+		/// <returns>The total system memory free, in bytes.</returns>
+		public static long GetTotalSystemMemoryBytes()
+		{
+			GCMemoryInfo MemoryInfo = GC.GetGCMemoryInfo();
+			return MemoryInfo.TotalAvailableMemoryBytes;
+		}
+
 		/// <summary>
 		/// Gets the total memory bytes free, based on what is known to the garbage collector.
 		/// </summary>
@@ -1396,14 +1407,20 @@ namespace UnrealBuildTool
 			// Limit number of actions to execute if the system is memory starved.
 			if (MemoryPerActionBytes > 0)
 			{
-				long FreeMemoryBytes = GetFreeMemoryBytes();
-				if (FreeMemoryBytes != -1)
+				// The OS needs enough memory to serve all the action processes that will be spawned. Historically this check limited
+				// actions based on free memory to ensure action processes weren't forced to use swap, but there's strong evidence that
+				// limiting based on total system memory and relying on the OS to swap out other processes to make room is fine.
+				// That said only Mac has been extensively tested here in this scenario, and does tend to have faster disk access
+				// than other platforms. So for now we'll change Mac and leave Windows/Linux using the old behavior
+				long AvailableMemoryBytes = RuntimePlatform.IsMac ? GetTotalSystemMemoryBytes() : GetFreeMemoryBytes();
+
+				if (AvailableMemoryBytes != -1)
 				{
-					int TotalMemoryActions = Convert.ToInt32(FreeMemoryBytes / MemoryPerActionBytes);
+					int TotalMemoryActions = Convert.ToInt32(AvailableMemoryBytes / MemoryPerActionBytes);
 					if (TotalMemoryActions < MaxActionsToExecuteInParallel)
 					{
 						MaxActionsToExecuteInParallel = Math.Max(1, Math.Min(MaxActionsToExecuteInParallel, TotalMemoryActions));
-						Log.TraceInformationOnce($"  Requested {StringUtils.FormatBytesString(MemoryPerActionBytes)} free memory per action, {StringUtils.FormatBytesString(FreeMemoryBytes)} available: limiting max parallel actions to {MaxActionsToExecuteInParallel}");
+						Log.TraceInformationOnce($"  Requested {StringUtils.FormatBytesString(MemoryPerActionBytes)} memory per action, {StringUtils.FormatBytesString(AvailableMemoryBytes)} available: limiting max parallel actions to {MaxActionsToExecuteInParallel}");
 					}
 				}
 			}
