Support resizeability for foldable devices / multiscreen

#android

diff --git a/Engine/Config/BaseEngine.ini b/Engine/Config/BaseEngine.ini
--- a/Engine/Config/BaseEngine.ini
+++ b/Engine/Config/BaseEngine.ini
@@ -2932,6 +2932,8 @@ bEnableAdvancedBinaryCompression=True
 bDisableStackProtector=False
 bEnableMulticastSupport=False
 bEnableRayTracing=true
+bAllowResizing=false
+bSupportSizeChanges=false
 
 [/Script/AndroidPlatformEditor.AndroidSDKSettings]
 SDKAPILevel=latest
diff --git a/Engine/Source/Programs/UnrealBuildTool/Platform/Android/UEDeployAndroid.cs b/Engine/Source/Programs/UnrealBuildTool/Platform/Android/UEDeployAndroid.cs
--- a/Engine/Source/Programs/UnrealBuildTool/Platform/Android/UEDeployAndroid.cs
+++ b/Engine/Source/Programs/UnrealBuildTool/Platform/Android/UEDeployAndroid.cs
@@ -2499,6 +2499,10 @@ namespace UnrealBuildTool
 			Ini.GetBool("/Script/AndroidRuntimeSettings.AndroidRuntimeSettings", "bFullScreen", out EnableFullScreen);
 			bool bUseDisplayCutout;
 			Ini.GetBool("/Script/AndroidRuntimeSettings.AndroidRuntimeSettings", "bUseDisplayCutout", out bUseDisplayCutout);
+			bool bAllowResizing;
+			Ini.GetBool("/Script/AndroidRuntimeSettings.AndroidRuntimeSettings", "bAllowResizing", out bAllowResizing);
+			bool bSupportSizeChanges;
+			Ini.GetBool("/Script/AndroidRuntimeSettings.AndroidRuntimeSettings", "bSupportSizeChanges", out bSupportSizeChanges);
 			bool bRestoreNotificationsOnReboot = false;
 			Ini.GetBool("/Script/AndroidRuntimeSettings.AndroidRuntimeSettings", "bRestoreNotificationsOnReboot", out bRestoreNotificationsOnReboot);
 			List<string>? ExtraManifestNodeTags;
@@ -2714,7 +2718,7 @@ namespace UnrealBuildTool
 				Text.AppendLine("\t\t          android:launchMode=\"singleTask\"");
 				if (SDKLevelInt >= 24)
 				{
-					Text.AppendLine("\t\t          android:resizeableActivity=\"false\"");
+					Text.AppendLine("\t\t          android:resizeableActivity=\"{0}\"", bAllowResizing ? "true" : "false");
 				}
 				Text.AppendLine(String.Format("\t\t          android:screenOrientation=\"{0}\"", Orientation));
 				Text.AppendLine(String.Format("\t\t          android:debuggable=\"{0}\">", bIsForDistribution ? "false" : "true"));
@@ -2742,7 +2746,7 @@ namespace UnrealBuildTool
 			}
 			if (SDKLevelInt >= 24)
 			{
-				Text.AppendLine("\t\t          android:resizeableActivity=\"false\"");
+				Text.AppendLine("\t\t          android:resizeableActivity=\"{0}\"", bAllowResizing ? "true" : "false");
 			}
 			Text.AppendLine("\t\t          android:launchMode=\"singleTask\"");
 			Text.AppendLine(String.Format("\t\t          android:screenOrientation=\"{0}\"", Orientation));
@@ -2828,6 +2832,10 @@ namespace UnrealBuildTool
 			Text.AppendLine(String.Format("\t\t<meta-data android:name=\"com.epicgames.unreal.GameActivity.bSupportsVulkan\" android:value=\"{0}\"/>", bSupportsVulkan ? "true" : "false"));
 			Text.AppendLine(String.Format("\t\t<meta-data android:name=\"com.epicgames.unreal.GameActivity.PropagateAlpha\" android:value=\"{0}\"/>", PropagateAlpha));
 			Text.AppendLine(String.Format("\t\t<meta-data android:name=\"com.epicgames.unreal.GameActivity.StartupPermissions\" android:value=\"{0}\"/>", StartupPermissions));
+			if (TargetSDKVersion >= 29)
+			{
+				Text.AppendLine(String.Format("\t\t<meta-data android:name=\"android.supports_size_changes\" android:value=\"{0}\"/>", bSupportSizeChanges ? "true" : "false"));
+			}
 			Text.AppendLine("\t\t<meta-data android:name=\"com.google.android.gms.games.APP_ID\"");
 			Text.AppendLine("\t\t           android:value=\"@string/app_id\" />");
 			Text.AppendLine("\t\t<meta-data android:name=\"com.google.android.gms.version\"");
diff --git a/Engine/Source/Runtime/Android/AndroidRuntimeSettings/Classes/AndroidRuntimeSettings.h b/Engine/Source/Runtime/Android/AndroidRuntimeSettings/Classes/AndroidRuntimeSettings.h
--- a/Engine/Source/Runtime/Android/AndroidRuntimeSettings/Classes/AndroidRuntimeSettings.h
+++ b/Engine/Source/Runtime/Android/AndroidRuntimeSettings/Classes/AndroidRuntimeSettings.h
@@ -284,6 +284,14 @@ public:
 	UPROPERTY(GlobalConfig, EditAnywhere, Category = "APK Packaging", Meta = (DisplayName = "Use display cutout region?"))
 	bool bUseDisplayCutout;
 
+	// Allow resizing of the window on Android devices with splitscreen
+	UPROPERTY(GlobalConfig, EditAnywhere, Category = "APK Packaging", Meta = (DisplayName = "Allow splitscreen resizing?"))
+	bool bAllowResizing;
+
+	// Allow support for size change when foldable and flip devices change screen or layout on Android 10+
+	UPROPERTY(GlobalConfig, EditAnywhere, Category = "APK Packaging", Meta = (DisplayName = "Allow fold/flip size changes?"))
+	bool bSupportSizeChanges;
+
 	// Should we restore scheduled local notifications on reboot? This will add a receiver for boot complete and a permission to the manifest.
 	UPROPERTY(GlobalConfig, EditAnywhere, Category = "APK Packaging", Meta = (DisplayName = "Restore scheduled notifications on reboot"))
 	bool bRestoreNotificationsOnReboot;
