UnrealBuildTool: Don't throw an exception if clang is too old when generating projects

#rnx

diff --git a/Engine/Source/Programs/UnrealBuildTool/Platform/Windows/UEBuildWindows.cs b/Engine/Source/Programs/UnrealBuildTool/Platform/Windows/UEBuildWindows.cs
--- a/Engine/Source/Programs/UnrealBuildTool/Platform/Windows/UEBuildWindows.cs
+++ b/Engine/Source/Programs/UnrealBuildTool/Platform/Windows/UEBuildWindows.cs
@@ -1032,7 +1032,7 @@ namespace UnrealBuildTool
 			}
 
 			// Ensure we're using a recent enough version of Clang given the MSVC version
-			if (Target.WindowsPlatform.Compiler.IsClang())
+			if (Target.WindowsPlatform.Compiler.IsClang() && !UEBuildPlatformSDK.bSuppressSDKWarnings)
 			{
 				VersionNumber ClangVersion = Target.WindowsPlatform.Compiler == WindowsCompiler.Intel ? MicrosoftPlatformSDK.GetClangVersionForIntelCompiler(Target.WindowsPlatform.Environment.CompilerPath) : Target.WindowsPlatform.Environment.CompilerVersion;
 				VersionNumber MinimumClang = MicrosoftPlatformSDK.GetMinimumClangVersionForVcVersion(Target.WindowsPlatform.Environment.ToolChainVersion);
