- Fixed an issue in IOSPostBuildSync crashing with null pointer

diff --git a/Engine/Source/Programs/UnrealBuildTool/Platform/IOS/IOSPostBuildSyncMode.cs b/Engine/Source/Programs/UnrealBuildTool/Platform/IOS/IOSPostBuildSyncMode.cs
--- a/Engine/Source/Programs/UnrealBuildTool/Platform/IOS/IOSPostBuildSyncMode.cs
+++ b/Engine/Source/Programs/UnrealBuildTool/Platform/IOS/IOSPostBuildSyncMode.cs
@@ -75,6 +75,8 @@ namespace UnrealBuildTool
 			Arguments.ApplyTo(this);
 			Arguments.CheckAllArgumentsUsed();
 
+			IOSToolChainSettings.SelectXcode(false, Logger);
+
 			// Run the PostBuildSync command
 			IOSPostBuildSyncTarget Target = BinaryFormatterUtils.Load<IOSPostBuildSyncTarget>(InputFile!);
 			IOSToolChain.PostBuildSync(Target, Logger);
diff --git a/Engine/Source/Programs/UnrealBuildTool/ToolChain/AppleToolChain.cs b/Engine/Source/Programs/UnrealBuildTool/ToolChain/AppleToolChain.cs
--- a/Engine/Source/Programs/UnrealBuildTool/ToolChain/AppleToolChain.cs
+++ b/Engine/Source/Programs/UnrealBuildTool/ToolChain/AppleToolChain.cs
@@ -75,7 +75,7 @@ namespace UnrealBuildTool
 			ToolchainDir = DirectoryReference.Combine(XcodeDeveloperDir, "Toolchains/XcodeDefault.xctoolchain/usr/bin");
 			SDKDir = DirectoryReference.Combine(XcodeDeveloperDir, $"Platforms/{OSPrefix}.platform/Developer/SDKs/{OSPrefix}.sdk");
 			if (SimulatorOSPrefix != null)
-		{
+			{
 				SimulatorSDKDir = DirectoryReference.Combine(XcodeDeveloperDir, $"Platforms/{SimulatorOSPrefix}.platform/Developer/SDKs/{SimulatorOSPrefix}.sdk");
 			}
 
@@ -114,32 +114,38 @@ namespace UnrealBuildTool
 			return $"{Prefix}-apple-{TargetOSName}{SDKVersion}{Suffix}";
 		}
 
-		private static void SelectXcode(bool bVerbose, ILogger Logger)
+		/// <summary>
+		/// Find the Xcode developer directory
+		/// </summary>
+		/// <param name="bVerbose"></param>
+		/// <param name="Logger"></param>
+		/// <exception cref="BuildException"></exception>
+		public static void SelectXcode(bool bVerbose, ILogger Logger)
 		{
-				// on the Mac, run xcode-select directly.
-				int ReturnCode;
+			// on the Mac, run xcode-select directly.
+			int ReturnCode;
 			string XcodeSelectResult = Utils.RunLocalProcessAndReturnStdOut("xcode-select", "--print-path", null, out ReturnCode);
-				if (ReturnCode != 0)
-				{
-					string? MinVersion = UEBuildPlatform.GetSDK(UnrealTargetPlatform.Mac)!.GetSDKInfo("Sdk")!.Min;
-					throw new BuildException($"We were unable to find your build tools (via 'xcode-select --print-path'). Please install Xcode, version {MinVersion} or later");
-				}
+			if (ReturnCode != 0)
+			{
+				string? MinVersion = UEBuildPlatform.GetSDK(UnrealTargetPlatform.Mac)!.GetSDKInfo("Sdk")!.Min;
+				throw new BuildException($"We were unable to find your build tools (via 'xcode-select --print-path'). Please install Xcode, version {MinVersion} or later");
+			}
 
 			_XcodeDeveloperDir = new DirectoryReference(XcodeSelectResult);
 
-				// make sure we get a full path
+			// make sure we get a full path
 			if (DirectoryReference.Exists(XcodeDeveloperDir) == false)
-				{
+			{
 				throw new BuildException("Selected Xcode ('{0}') doesn't exist, cannot continue.", XcodeDeveloperDir);
-				}
+			}
 
 			if (XcodeDeveloperDir.ContainsName("CommandLineTools", 0))
-				{
+			{
 				throw new BuildException($"Your Mac is set to use CommandLineTools for its build tools ({XcodeDeveloperDir}). Unreal expects Xcode as the build tools. Please install Xcode if it's not already, then do one of the following:\n" +
 						"  - Run Xcode, go to Settings, and in the Locations tab, choose your Xcode in Command Line Tools dropdown.\n" +
 						"  - In Terminal, run 'sudo xcode-select -s /Applications/Xcode.app' (or an alternate location if you installed Xcode to a non-standard location)\n" +
 						"Either way, you will need to enter your Mac password.");
-				}
+			}
 
 			if (bVerbose && !XcodeDeveloperDir.FullName.StartsWith("/Applications/Xcode.app"))
 			{
