UnrealBuildTool: Temporarily disable clang vs msvc version check due to issues

#rnx

diff --git a/Engine/Source/Programs/UnrealBuildTool/Platform/Windows/UEBuildWindows.cs b/Engine/Source/Programs/UnrealBuildTool/Platform/Windows/UEBuildWindows.cs
--- a/Engine/Source/Programs/UnrealBuildTool/Platform/Windows/UEBuildWindows.cs
+++ b/Engine/Source/Programs/UnrealBuildTool/Platform/Windows/UEBuildWindows.cs
@@ -1032,7 +1032,8 @@ namespace UnrealBuildTool
 			}
 
 			// Ensure we're using a recent enough version of Clang given the MSVC version
-			if (Target.WindowsPlatform.Compiler.IsClang() && !UEBuildPlatformSDK.bSuppressSDKWarnings)
+			// Temporarily disabled due to issues
+			/*if (Target.WindowsPlatform.Compiler.IsClang() && !UEBuildPlatformSDK.bSuppressSDKWarnings)
 			{
 				VersionNumber ClangVersion = Target.WindowsPlatform.Compiler == WindowsCompiler.Intel ? MicrosoftPlatformSDK.GetClangVersionForIntelCompiler(Target.WindowsPlatform.Environment.CompilerPath) : Target.WindowsPlatform.Environment.CompilerVersion;
 				VersionNumber MinimumClang = MicrosoftPlatformSDK.GetMinimumClangVersionForVcVersion(Target.WindowsPlatform.Environment.ToolChainVersion);
@@ -1040,7 +1041,7 @@ namespace UnrealBuildTool
 				{
 					throw new BuildException("MSVC toolchain version {0} requires Clang compiler version {1} or later. The current Clang compiler version was detected as: {2}", Target.WindowsPlatform.Environment.ToolChainVersion, MinimumClang, ClangVersion);
 				}
-			}
+			}*/
 
 			// Ensure we're using VS2022 when compiling for the installed engine.
 			if (Unreal.IsEngineInstalled() && (Target.WindowsPlatform.ToolChain == WindowsCompiler.VisualStudio2019 || Target.WindowsPlatform.Environment.ToolChainVersion < new VersionNumber(14, 34, 0)))
