New temporary target rule property bUseVerseBPVM

#rnx

diff --git a/Engine/Source/Programs/UnrealBuildTool/Configuration/TargetRules.cs b/Engine/Source/Programs/UnrealBuildTool/Configuration/TargetRules.cs
--- a/Engine/Source/Programs/UnrealBuildTool/Configuration/TargetRules.cs
+++ b/Engine/Source/Programs/UnrealBuildTool/Configuration/TargetRules.cs
@@ -781,6 +781,15 @@ namespace UnrealBuildTool
 		[XmlConfigFile(Category = "BuildConfiguration")]
 		public bool bUseVerse = true;
 
+		/// <summary>
+		/// Whether to use the BPVM to run Verse.
+		/// </summary>
+		[RequiresUniqueBuildEnvironment]
+		[CommandLine("-NoUseVerseBPVM", Value = "false")]
+		[CommandLine("-UseVerseBPVM", Value = "true")]
+		[XmlConfigFile(Category = "BuildConfiguration")]
+		public bool bUseVerseBPVM = true;
+
 		/// <summary>
 		/// Whether to use the AutoRTFM Clang compiler.
 		/// </summary>
@@ -3024,6 +3033,8 @@ namespace UnrealBuildTool
 
 		public bool bUseVerse => Inner.bUseVerse;
 
+		public bool bUseVerseBPVM => Inner.bUseVerseBPVM;
+
 		public bool bUseAutoRTFMCompiler => Inner.bUseAutoRTFMCompiler;
 
 		[Obsolete("Deprecated in UE5.1 - No longer used as Chaos is always enabled.")]
diff --git a/Engine/Source/Programs/UnrealBuildTool/Configuration/UEBuildTarget.cs b/Engine/Source/Programs/UnrealBuildTool/Configuration/UEBuildTarget.cs
--- a/Engine/Source/Programs/UnrealBuildTool/Configuration/UEBuildTarget.cs
+++ b/Engine/Source/Programs/UnrealBuildTool/Configuration/UEBuildTarget.cs
@@ -2522,6 +2522,9 @@ namespace UnrealBuildTool
 					}
 				}
 
+				// Copy Verse BPVM usage flag
+				Makefile.bUseVerseBPVM = Rules.bUseVerseBPVM;
+
 				// NOTE: Even in Gather mode, we need to run VNI to make sure the files exist for the static action graph to be setup correctly.  This is because VNI generates .cpp
 				// files that are injected as top level prerequisites.  If VNI only emitted included header files, we wouldn't need to run it during the Gather phase at all.
 				if (Makefile.VNIModules.Count > 0)
@@ -4915,6 +4918,15 @@ namespace UnrealBuildTool
 				GlobalCompileEnvironment.Definitions.Add("UE_USE_VERSE_PATHS=0");
 			}
 
+			if (Rules.bUseVerseBPVM)
+			{
+				GlobalCompileEnvironment.Definitions.Add("WITH_VERSE_BPVM=1");
+			}
+			else
+			{
+				GlobalCompileEnvironment.Definitions.Add("WITH_VERSE_BPVM=0");
+			}
+
 			if (Rules.bCompileWithStatsWithoutEngine)
 			{
 				GlobalCompileEnvironment.Definitions.Add("USE_STATS_WITHOUT_ENGINE=1");
diff --git a/Engine/Source/Programs/UnrealBuildTool/System/TargetMakefile.cs b/Engine/Source/Programs/UnrealBuildTool/System/TargetMakefile.cs
--- a/Engine/Source/Programs/UnrealBuildTool/System/TargetMakefile.cs
+++ b/Engine/Source/Programs/UnrealBuildTool/System/TargetMakefile.cs
@@ -178,7 +178,7 @@ namespace UnrealBuildTool
 		/// <summary>
 		/// The version number to write
 		/// </summary>
-		public const int CurrentVersion = 35;
+		public const int CurrentVersion = 36;
 
 		/// <summary>
 		/// The time at which the makefile was created
@@ -341,6 +341,11 @@ namespace UnrealBuildTool
 		/// All build modules containing Verse code
 		/// </summary>
 		public List<VNIModuleInfo> VNIModules;
+
+		/// <summary>
+		/// If Verse should use the BPVM backend
+		/// </summary>
+		public bool bUseVerseBPVM = true;
 #endif
 
 		/// <summary>
@@ -482,6 +487,7 @@ namespace UnrealBuildTool
 			UObjectModuleHeaders = Reader.ReadList(() => new UHTModuleHeaderInfo(Reader))!;
 #if __VPROJECT_AVAILABLE__
 			VNIModules = Reader.ReadList(() => new VNIModuleInfo(Reader))!;
+			bUseVerseBPVM = Reader.ReadBool();
 #endif
 			PluginFiles = Reader.ReadHashSet(() => Reader.ReadFileItem())!;
 			ExternalDependencies = Reader.ReadHashSet(() => Reader.ReadFileItem())!;
@@ -527,6 +533,7 @@ namespace UnrealBuildTool
 			Writer.WriteList(UObjectModuleHeaders, x => x.Write(Writer));
 #if __VPROJECT_AVAILABLE__
 			Writer.WriteList(VNIModules, e => e.Write(Writer));
+			Writer.WriteBool(bUseVerseBPVM);
 #endif
 			Writer.WriteHashSet(PluginFiles, x => Writer.WriteFileItem(x));
 			Writer.WriteHashSet(ExternalDependencies, x => Writer.WriteFileItem(x));
