UnrealBuildTool: Don't short circuit adaptive builds if all cpp in a module are in the working set

#rnx

diff --git a/Engine/Source/Programs/UnrealBuildTool/System/Unity.cs b/Engine/Source/Programs/UnrealBuildTool/System/Unity.cs
--- a/Engine/Source/Programs/UnrealBuildTool/System/Unity.cs
+++ b/Engine/Source/Programs/UnrealBuildTool/System/Unity.cs
@@ -183,14 +183,8 @@ namespace UnrealBuildTool
 			// Optimization only makes sense if PCH files are enabled.
 			bool bForceIntoSingleUnityFile = Target.bStressTestUnity || (TotalBytesInCPPFiles < NumIncludedBytesPerUnityCPP * 2 && Target.bUsePCHFiles);
 
-			// Every single file in the module appears in the working set. Don't bother using adaptive unity for this module.
-			// Otherwise it would make full builds really slow.
+			// Even if every single file in the module appears in the working set adaptive unity should still be used even if it's slower.
 			GetAdaptiveFiles(Target, CPPFiles, HeaderFiles, CompileEnvironment, WorkingSet, BaseName, IntermediateDirectory, Graph, out NormalFiles, out AdaptiveFiles);
-			if (!NormalFiles.Where(file => !file.HasExtension(".gen.cpp")).Any())
-			{
-				NormalFiles = CPPFiles;
-				AdaptiveFiles.RemoveAll(new HashSet<FileItem>(NormalFiles).Contains);
-			}
 
 			// Build the list of unity files.
 			List<FileCollection> AllUnityFiles;
