UnrealBuildTool: Print action groups in parallel executor when building multi target

#rnx

diff --git a/Engine/Source/Programs/Shared/EpicGames.Build/Unreal.cs b/Engine/Source/Programs/Shared/EpicGames.Build/Unreal.cs
--- a/Engine/Source/Programs/Shared/EpicGames.Build/Unreal.cs
+++ b/Engine/Source/Programs/Shared/EpicGames.Build/Unreal.cs
@@ -185,6 +185,11 @@ namespace UnrealBuildBase
 		/// </summary>
 		static public readonly FileReference DotnetPath = FileReference.Combine(DotnetDirectory, "dotnet" + RuntimePlatform.ExeExtension);
 
+		/// <summary>
+		/// Whether we're running on a build machine
+		/// </summary>
+		static private bool? bIsBuildMachine;
+
 		/// <summary>
 		/// Whether we're running with engine installed
 		/// </summary>
@@ -200,6 +205,19 @@ namespace UnrealBuildBase
 			return DirectoryReference.Combine(EngineDirectory, FindRelativeDotnetDirectory(HostPlatform));
 		}
 
+		/// <summary>
+		/// Returns true of UnrealBuildTool is running on a build machine
+		/// </summary>
+		/// <returns>True if running on a build machine</returns>
+		static public bool IsBuildMachine()
+		{
+			if (!bIsBuildMachine.HasValue)
+			{
+				bIsBuildMachine = Environment.GetEnvironmentVariable("IsBuildMachine") == "1";
+			}
+			return bIsBuildMachine.Value;
+		}
+
 		/// <summary>
 		/// Returns true if UnrealBuildTool is running using installed Engine components
 		/// </summary>
diff --git a/Engine/Source/Programs/UnrealBuildTool/Executors/ImmediateActionQueue.cs b/Engine/Source/Programs/UnrealBuildTool/Executors/ImmediateActionQueue.cs
--- a/Engine/Source/Programs/UnrealBuildTool/Executors/ImmediateActionQueue.cs
+++ b/Engine/Source/Programs/UnrealBuildTool/Executors/ImmediateActionQueue.cs
@@ -337,6 +337,11 @@ namespace UnrealBuildTool
 		/// </summary>
 		private readonly TaskCompletionSource _doneTaskSource = new();
 
+		/// <summary>
+		/// The last action group printed in multi-target builds
+		/// </summary>
+		private string? LastGroupPrefix = null;
+
 		/// <summary>
 		/// If set, artifact cache used to retrieve previously compiled results and save new results
 		/// </summary>
@@ -1069,6 +1074,19 @@ namespace UnrealBuildTool
 						Console.CursorLeft = 0;
 					}
 				}
+				else
+				{
+					// If the action group has changed for a multi target build, write it to the log
+					if (action.GroupNames.Count > 0)
+					{
+						string ActionGroup = $"** For {String.Join(" + ", action.GroupNames)} **";
+						if (!ActionGroup.Equals(LastGroupPrefix, StringComparison.Ordinal))
+						{
+							LastGroupPrefix = ActionGroup;
+							_writeToolOutput(ActionGroup);
+						}
+					}
+				}
 
 				s_previousLineLength = message.Length;
 
diff --git a/Engine/Source/Programs/UnrealBuildTool/Executors/ParallelExecutor.cs b/Engine/Source/Programs/UnrealBuildTool/Executors/ParallelExecutor.cs
--- a/Engine/Source/Programs/UnrealBuildTool/Executors/ParallelExecutor.cs
+++ b/Engine/Source/Programs/UnrealBuildTool/Executors/ParallelExecutor.cs
@@ -9,6 +9,7 @@ using System.Threading;
 using System.Threading.Tasks;
 using EpicGames.Core;
 using Microsoft.Extensions.Logging;
+using UnrealBuildBase;
 using UnrealBuildTool.Artifacts;
 
 namespace UnrealBuildTool
@@ -62,13 +63,13 @@ namespace UnrealBuildTool
 		/// Whether to show compilation times along with worst offenders or not.
 		/// </summary>
 		[XmlConfigFile]
-		private static bool bShowCompilationTimes = false;
+		private static bool bShowCompilationTimes = Unreal.IsBuildMachine();
 
 		/// <summary>
 		/// Whether to show compilation times for each executed action
 		/// </summary>
 		[XmlConfigFile]
-		private static bool bShowPerActionCompilationTimes = false;
+		private static bool bShowPerActionCompilationTimes = Unreal.IsBuildMachine();
 
 		/// <summary>
 		/// Whether to log command lines for actions being executed
@@ -92,7 +93,7 @@ namespace UnrealBuildTool
 		/// Whether to show CPU utilization after the work is complete.
 		/// </summary>
 		[XmlConfigFile]
-		protected static bool bShowCPUUtilization = false;
+		protected static bool bShowCPUUtilization = Unreal.IsBuildMachine();
 
 		/// <summary>
 		/// Collapse non-error output lines
diff --git a/Engine/Source/Programs/UnrealBuildTool/Modes/BuildMode.cs b/Engine/Source/Programs/UnrealBuildTool/Modes/BuildMode.cs
--- a/Engine/Source/Programs/UnrealBuildTool/Modes/BuildMode.cs
+++ b/Engine/Source/Programs/UnrealBuildTool/Modes/BuildMode.cs
@@ -1241,6 +1241,7 @@ namespace UnrealBuildTool
 					if (TargetAction.IgnoreConflicts())
 					{
 						IgnoreConflictActions.Add(TargetAction);
+						TargetAction.GroupNames.Add(GroupPrefix);
 						continue;
 					}
 
diff --git a/Engine/Source/Programs/UnrealBuildTool/System/Action.cs b/Engine/Source/Programs/UnrealBuildTool/System/Action.cs
--- a/Engine/Source/Programs/UnrealBuildTool/System/Action.cs
+++ b/Engine/Source/Programs/UnrealBuildTool/System/Action.cs
@@ -679,7 +679,7 @@ namespace UnrealBuildTool
 		/// <summary>
 		/// If set, will be output whenever the group differs to the last executed action. Set when executing multiple targets at once.
 		/// </summary>
-		public List<string> GroupNames = new List<string>();
+		public SortedSet<string> GroupNames = new SortedSet<string>();
 
 		#region Wrapper implementation of IAction
 
