Remove natvis files from precompiled manifest so they aren't added to linker as if they were object files in builds against precompiled engine.

diff --git a/Engine/Source/Programs/UnrealBuildTool/Configuration/UEBuildModuleCPP.cs b/Engine/Source/Programs/UnrealBuildTool/Configuration/UEBuildModuleCPP.cs
--- a/Engine/Source/Programs/UnrealBuildTool/Configuration/UEBuildModuleCPP.cs
+++ b/Engine/Source/Programs/UnrealBuildTool/Configuration/UEBuildModuleCPP.cs
@@ -722,14 +722,6 @@ namespace UnrealBuildTool
 
 				PrecompiledManifest Manifest = new PrecompiledManifest();
 				Manifest.OutputFiles.AddRange(LinkInputFiles.Select(x => x.Location));
-				foreach (FileItem NatvisSourceFile in NatvisFiles)
-				{
-					FileItem? Item = ToolChain.LinkDebuggerVisualizer(NatvisSourceFile, IntermediateDirectory);
-					if (Item != null)
-					{
-						Manifest.OutputFiles.Add(Item.Location);
-					}
-				}	
 				Manifest.WriteIfModified(PrecompiledManifestLocation);
 			}
 
