Add target rules option for enabling trace.

diff --git a/Engine/Source/Programs/BlankProgram/BlankProgram.Build.cs b/Engine/Source/Programs/BlankProgram/BlankProgram.Build.cs
--- a/Engine/Source/Programs/BlankProgram/BlankProgram.Build.cs
+++ b/Engine/Source/Programs/BlankProgram/BlankProgram.Build.cs
@@ -13,9 +13,6 @@ public class BlankProgram : ModuleRules
 		// to link with CoreUObject module:
 		// PrivateDependencyModuleNames.Add("CoreUObject");
 
-		// to enable tracing:
-		// AppendStringToPublicDefinition("UE_TRACE_ENABLED", "1");
-
 		// to enable LLM tracing:
 		// GlobalDefinitions.Add("LLM_ENABLED_IN_CONFIG=1");
 		// GlobalDefinitions.Add("UE_MEMORY_TAGS_TRACE_ENABLED=1");
diff --git a/Engine/Source/Programs/BlankProgram/BlankProgram.Target.cs b/Engine/Source/Programs/BlankProgram/BlankProgram.Target.cs
--- a/Engine/Source/Programs/BlankProgram/BlankProgram.Target.cs
+++ b/Engine/Source/Programs/BlankProgram/BlankProgram.Target.cs
@@ -31,7 +31,7 @@ public class BlankProgramTarget : TargetRules
 		// bForceCompileDevelopmentAutomationTests = true;
 
 		// to enable tracing:
-		// GlobalDefinitions.Add("UE_TRACE_ENABLED=1");
+		// bEnableTrace = true;
 
 		// This app is a console application, not a Windows app (sets entry point to main(), instead of WinMain())
 		bIsBuildingConsoleApplication = true;
diff --git a/Engine/Source/Programs/DerivedDataTool/DerivedDataTool.Target.cs b/Engine/Source/Programs/DerivedDataTool/DerivedDataTool.Target.cs
--- a/Engine/Source/Programs/DerivedDataTool/DerivedDataTool.Target.cs
+++ b/Engine/Source/Programs/DerivedDataTool/DerivedDataTool.Target.cs
@@ -26,7 +26,7 @@ public class DerivedDataToolTarget : TargetRules
 		{
 			GlobalDefinitions.Add("LLM_ENABLED_IN_CONFIG=1");
 			GlobalDefinitions.Add("UE_MEMORY_TAGS_TRACE_ENABLED=1");
-			GlobalDefinitions.Add("UE_TRACE_ENABLED=1");
+			bEnableTrace = true;
 		}
 	}
 }
diff --git a/Engine/Source/Programs/ReplicationSystemLowLevelTests/ReplicationSystemLowLevelTests.Target.cs b/Engine/Source/Programs/ReplicationSystemLowLevelTests/ReplicationSystemLowLevelTests.Target.cs
--- a/Engine/Source/Programs/ReplicationSystemLowLevelTests/ReplicationSystemLowLevelTests.Target.cs
+++ b/Engine/Source/Programs/ReplicationSystemLowLevelTests/ReplicationSystemLowLevelTests.Target.cs
@@ -20,7 +20,7 @@ public class ReplicationSystemLowLevelTestsTarget : TestTargetRules
 		bWithPushModel = true;
 		bUseIris = true;
 
-		GlobalDefinitions.Add("UE_TRACE_ENABLED=1");
+		bEnableTrace = true;
 		GlobalDefinitions.Add("UE_NET_TEST_FAKE_REP_TAGS=1");
 	}
 }
diff --git a/Engine/Source/Programs/ReplicationSystemTest/ReplicationSystemTest.Target.cs b/Engine/Source/Programs/ReplicationSystemTest/ReplicationSystemTest.Target.cs
--- a/Engine/Source/Programs/ReplicationSystemTest/ReplicationSystemTest.Target.cs
+++ b/Engine/Source/Programs/ReplicationSystemTest/ReplicationSystemTest.Target.cs
@@ -39,7 +39,7 @@ public class ReplicationSystemTestTarget : TargetRules
 		bWithPushModel = true;		
 		bUseIris = true;
 
-		GlobalDefinitions.Add("UE_TRACE_ENABLED=1");
+		bEnableTrace = true;
 		GlobalDefinitions.Add("UE_NET_TEST_FAKE_REP_TAGS=1");
 	}
 }
diff --git a/Engine/Source/Programs/ShaderCompileWorker/ShaderCompileWorker.Target.cs b/Engine/Source/Programs/ShaderCompileWorker/ShaderCompileWorker.Target.cs
--- a/Engine/Source/Programs/ShaderCompileWorker/ShaderCompileWorker.Target.cs
+++ b/Engine/Source/Programs/ShaderCompileWorker/ShaderCompileWorker.Target.cs
@@ -64,7 +64,7 @@ public class ShaderCompileWorkerTarget : TargetRules
         {
 			GlobalDefinitions.Add("LLM_ENABLED_IN_CONFIG=1");
 			GlobalDefinitions.Add("UE_MEMORY_TAGS_TRACE_ENABLED=1");
-			GlobalDefinitions.Add("UE_TRACE_ENABLED=1");
+			bEnableTrace = true;
 		}
 	}
 }
diff --git a/Engine/Source/Programs/SwitchboardListener/SwitchboardListener.Target.cs b/Engine/Source/Programs/SwitchboardListener/SwitchboardListener.Target.cs
--- a/Engine/Source/Programs/SwitchboardListener/SwitchboardListener.Target.cs
+++ b/Engine/Source/Programs/SwitchboardListener/SwitchboardListener.Target.cs
@@ -24,6 +24,6 @@ public class SwitchboardListenerTarget : TargetRules
 		// The listener is meant to be a console application (no window), but on MacOS, to get a proper log console, a full application must be built.
 		bIsBuildingConsoleApplication = Target.Platform != UnrealTargetPlatform.Mac;
 
-		GlobalDefinitions.Add("UE_TRACE_ENABLED=1");
+		bEnableTrace = true;
 	}
 }
diff --git a/Engine/Source/Programs/SwitchboardListenerHelper/SwitchboardListenerHelper.Target.cs b/Engine/Source/Programs/SwitchboardListenerHelper/SwitchboardListenerHelper.Target.cs
--- a/Engine/Source/Programs/SwitchboardListenerHelper/SwitchboardListenerHelper.Target.cs
+++ b/Engine/Source/Programs/SwitchboardListenerHelper/SwitchboardListenerHelper.Target.cs
@@ -26,7 +26,7 @@ public class SwitchboardListenerHelperTarget : TargetRules
 		// The listener is meant to be a console application (no window), but on MacOS, to get a proper log console, a full application must be built.
 		bIsBuildingConsoleApplication = Target.Platform != UnrealTargetPlatform.Mac;
 
-		GlobalDefinitions.Add("UE_TRACE_ENABLED=1");
+		bEnableTrace = true;
 
 		if (Target.Platform == UnrealTargetPlatform.Win64)
 		{
diff --git a/Engine/Source/Programs/UnrealBuildTool/Configuration/TargetRules.cs b/Engine/Source/Programs/UnrealBuildTool/Configuration/TargetRules.cs
--- a/Engine/Source/Programs/UnrealBuildTool/Configuration/TargetRules.cs
+++ b/Engine/Source/Programs/UnrealBuildTool/Configuration/TargetRules.cs
@@ -1094,6 +1094,35 @@ namespace UnrealBuildTool
 		}
 		private bool bCompileAgainstApplicationCorePrivate = true;
 
+		/// <summary>
+		/// Enabled for builds that want to use against the Trace module for profiling and diagnostics.
+		/// </summary>
+		[RequiresUniqueBuildEnvironment]
+		public virtual bool bEnableTrace
+		{
+			get 
+			{
+				if (bEnableTracePrivate.HasValue)
+				{
+					return bEnableTracePrivate.Value;
+				}
+				if (Configuration == UnrealTargetConfiguration.Shipping || Type == TargetType.Program)
+				{
+					return false;
+				}
+				if (UEBuildPlatform.IsPlatformInGroup(Platform, UnrealPlatformGroup.Apple) 
+				|| UEBuildPlatform.IsPlatformInGroup(Platform, UnrealPlatformGroup.Unix)
+				|| UEBuildPlatform.IsPlatformInGroup(Platform, UnrealPlatformGroup.Android)
+				|| UEBuildPlatform.IsPlatformInGroup(Platform, UnrealPlatformGroup.Windows))
+				{
+					return true;
+				}
+				return false;
+			}  
+			set => bEnableTracePrivate = value;
+		}
+		private bool? bEnableTracePrivate;
+
 		/// <summary>
 		/// Manually specified value for bCompileAgainstEditor.
 		/// </summary>
@@ -3103,6 +3132,8 @@ namespace UnrealBuildTool
 
 		public bool bCompileAgainstApplicationCore => Inner.bCompileAgainstApplicationCore;
 
+		public bool bEnableTrace => Inner.bEnableTrace;
+
 		public bool bCompileAgainstEditor => Inner.bCompileAgainstEditor;
 
 		public bool bCompileRecast => Inner.bCompileRecast;
diff --git a/Engine/Source/Programs/UnrealBuildTool/Configuration/UEBuildTarget.cs b/Engine/Source/Programs/UnrealBuildTool/Configuration/UEBuildTarget.cs
--- a/Engine/Source/Programs/UnrealBuildTool/Configuration/UEBuildTarget.cs
+++ b/Engine/Source/Programs/UnrealBuildTool/Configuration/UEBuildTarget.cs
@@ -4906,6 +4906,15 @@ namespace UnrealBuildTool
 			{
 				GlobalCompileEnvironment.Definitions.Add("WITH_COREUOBJECT=0");
 			}
+			
+			if (Rules.bEnableTrace)
+			{
+				GlobalCompileEnvironment.Definitions.Add("UE_TRACE_ENABLED=1");
+			}
+			else
+			{
+				GlobalCompileEnvironment.Definitions.Add("UE_TRACE_ENABLED=0");
+			}
 
 			if (Rules.bUseVerse)
 			{
diff --git a/Engine/Source/Programs/UnrealMultiUserServer/UnrealMultiUserServer.Target.cs b/Engine/Source/Programs/UnrealMultiUserServer/UnrealMultiUserServer.Target.cs
--- a/Engine/Source/Programs/UnrealMultiUserServer/UnrealMultiUserServer.Target.cs
+++ b/Engine/Source/Programs/UnrealMultiUserServer/UnrealMultiUserServer.Target.cs
@@ -27,6 +27,6 @@ public class UnrealMultiUserServerTarget : TargetRules
 		bIsBuildingConsoleApplication = Target.Platform != UnrealTargetPlatform.Mac;
 
 		GlobalDefinitions.Add("UE_LOG_CONCERT_DEBUG_VERBOSITY_LEVEL=Log");
-		GlobalDefinitions.Add("UE_TRACE_ENABLED=1");
+		bEnableTrace = true;
 	}
 }
diff --git a/Engine/Source/Programs/UnrealMultiUserSlateServer/UnrealMultiUserSlateServer.Target.cs b/Engine/Source/Programs/UnrealMultiUserSlateServer/UnrealMultiUserSlateServer.Target.cs
--- a/Engine/Source/Programs/UnrealMultiUserSlateServer/UnrealMultiUserSlateServer.Target.cs
+++ b/Engine/Source/Programs/UnrealMultiUserSlateServer/UnrealMultiUserSlateServer.Target.cs
@@ -26,6 +26,6 @@ public class UnrealMultiUserSlateServerTarget : TargetRules
 		bIsBuildingConsoleApplication = false;
 
 		GlobalDefinitions.Add("UE_LOG_CONCERT_DEBUG_VERBOSITY_LEVEL=Log");
-		GlobalDefinitions.Add("UE_TRACE_ENABLED=1");
+		bEnableTrace = true;
 	}
 }
diff --git a/Engine/Source/Programs/UnrealPackageTool/UnrealPackageTool.Target.cs b/Engine/Source/Programs/UnrealPackageTool/UnrealPackageTool.Target.cs
--- a/Engine/Source/Programs/UnrealPackageTool/UnrealPackageTool.Target.cs
+++ b/Engine/Source/Programs/UnrealPackageTool/UnrealPackageTool.Target.cs
@@ -22,6 +22,6 @@ public class UnrealPackageToolTarget : TargetRules
 
 		bIsBuildingConsoleApplication = true;
 
-		GlobalDefinitions.Add("UE_TRACE_ENABLED=1");
+		bEnableTrace = true;
 	}
 }
diff --git a/Engine/Source/Programs/UnrealPak/UnrealPak.Target.cs b/Engine/Source/Programs/UnrealPak/UnrealPak.Target.cs
--- a/Engine/Source/Programs/UnrealPak/UnrealPak.Target.cs
+++ b/Engine/Source/Programs/UnrealPak/UnrealPak.Target.cs
@@ -38,7 +38,7 @@ public class UnrealPakTarget : TargetRules
 		WindowsPlatform.bMergeIdenticalCOMDATs = true;
 		WindowsPlatform.bStripUnreferencedSymbols = true;
 
-		GlobalDefinitions.Add("UE_TRACE_ENABLED=1");
+		bEnableTrace = true;
 		GlobalDefinitions.Add("UE_EXTERNAL_PROFILING_ENABLED=0");
 	}
 }
diff --git a/Engine/Source/Programs/UnrealVirtualizationTool/UnrealVirtualizationTool.Target.cs b/Engine/Source/Programs/UnrealVirtualizationTool/UnrealVirtualizationTool.Target.cs
--- a/Engine/Source/Programs/UnrealVirtualizationTool/UnrealVirtualizationTool.Target.cs
+++ b/Engine/Source/Programs/UnrealVirtualizationTool/UnrealVirtualizationTool.Target.cs
@@ -30,7 +30,7 @@ public class UnrealVirtualizationToolTarget : TargetRules
 		// Enable Developer plugins
 		bCompileWithPluginSupport = true;
 
-		GlobalDefinitions.Add("UE_TRACE_ENABLED=1");
+		bEnableTrace = true;
 		GlobalDefinitions.Add("UE_SUPPORT_FULL_PACKAGEPATH=1");
 	}
 }
diff --git a/Engine/Source/Runtime/TraceLog/TraceLog.Build.cs b/Engine/Source/Runtime/TraceLog/TraceLog.Build.cs
--- a/Engine/Source/Runtime/TraceLog/TraceLog.Build.cs
+++ b/Engine/Source/Runtime/TraceLog/TraceLog.Build.cs
@@ -13,21 +13,4 @@ public class TraceLog : ModuleRules
 
 		UnsafeTypeCastWarningLevel = WarningLevel.Error;
 	}
-
-	// used by platform extension derived classes. probably should become a project setting at some point!
-	protected void EnableTraceByDefault(ReadOnlyTargetRules Target)
-	{
-		if (Target.Configuration != UnrealTargetConfiguration.Shipping && Target.Type != TargetType.Program)
-		{
-			foreach (String Definition in Target.GlobalDefinitions)
-			{
-				if (Definition.Contains("UE_TRACE_ENABLED"))
-				{
-					// Define already set in Target.GlobalDefinitions
-					return;
-				}
-			}
-			PublicDefinitions.Add("UE_TRACE_ENABLED=1");
-		}
-	}
 }
