UnrealBuildTool: Clean up IsBuildMachine checks

#rnx

diff --git a/Engine/Source/Programs/Shared/EpicGames.Build/Unreal.cs b/Engine/Source/Programs/Shared/EpicGames.Build/Unreal.cs
--- a/Engine/Source/Programs/Shared/EpicGames.Build/Unreal.cs
+++ b/Engine/Source/Programs/Shared/EpicGames.Build/Unreal.cs
@@ -213,7 +213,7 @@ namespace UnrealBuildBase
 		{
 			if (!bIsBuildMachine.HasValue)
 			{
-				bIsBuildMachine = Environment.GetEnvironmentVariable("IsBuildMachine") == "1";
+				bIsBuildMachine = Environment.GetEnvironmentVariable("IsBuildMachine")?.Trim() == "1";
 			}
 			return bIsBuildMachine.Value;
 		}
diff --git a/Engine/Source/Programs/UnrealBuildTool/Configuration/TestModuleRules.cs b/Engine/Source/Programs/UnrealBuildTool/Configuration/TestModuleRules.cs
--- a/Engine/Source/Programs/UnrealBuildTool/Configuration/TestModuleRules.cs
+++ b/Engine/Source/Programs/UnrealBuildTool/Configuration/TestModuleRules.cs
@@ -170,8 +170,7 @@ namespace UnrealBuildTool
 				bUpdateBuildGraphPropertiesFile = TestTargetRules.bUpdateBuildGraphPropertiesFile;
 			}
 
-			string? bIsBuildMachineEnv = Environment.GetEnvironmentVariable("IsBuildMachine");
-			bool bIsBuildMachine = bIsBuildMachineEnv != null && bIsBuildMachineEnv.Trim().Equals(1.ToString());
+			bool bIsBuildMachine = Unreal.IsBuildMachine();
 			if (bIsBuildMachine || !bUpdateBuildGraphPropertiesFile || TestMetadata == null)
 			{
 				return;
diff --git a/Engine/Source/Programs/UnrealBuildTool/Platform/Android/UEDeployAndroid.cs b/Engine/Source/Programs/UnrealBuildTool/Platform/Android/UEDeployAndroid.cs
--- a/Engine/Source/Programs/UnrealBuildTool/Platform/Android/UEDeployAndroid.cs
+++ b/Engine/Source/Programs/UnrealBuildTool/Platform/Android/UEDeployAndroid.cs
@@ -935,7 +935,7 @@ namespace UnrealBuildTool
 				bool bUseChangeListAsStoreVersion = false;
 				Ini.GetBool("/Script/AndroidRuntimeSettings.AndroidRuntimeSettings", "bUseChangeListAsStoreVersion", out bUseChangeListAsStoreVersion);
 
-				bool IsBuildMachine = Environment.GetEnvironmentVariable("IsBuildMachine") == "1";
+				bool IsBuildMachine = Unreal.IsBuildMachine();
 				// override store version with changelist if enabled and is build machine
 				if (bUseChangeListAsStoreVersion && IsBuildMachine)
 				{
@@ -986,7 +986,7 @@ namespace UnrealBuildTool
 				string VersionDisplayName = "";
 				Ini.GetString("/Script/AndroidRuntimeSettings.AndroidRuntimeSettings", "VersionDisplayName", out VersionDisplayName);
 
-				if (Environment.GetEnvironmentVariable("IsBuildMachine") == "1")
+				if (Unreal.IsBuildMachine())
 				{
 					bool bAppendChangeListToVersionDisplayName = false;
 					Ini.GetBool("/Script/AndroidRuntimeSettings.AndroidRuntimeSettings", "bAppendChangeListToVersionDisplayName", out bAppendChangeListToVersionDisplayName);
@@ -4021,7 +4021,7 @@ namespace UnrealBuildTool
 			// bundles disabled for launch-on
 			bool bEnableBundle = BundleEnabled() && !bDisallowPackagingDataInApk;
 
-			bool bIsBuildMachine = Environment.GetEnvironmentVariable("IsBuildMachine") == "1";
+			bool bIsBuildMachine = Unreal.IsBuildMachine();
 
 			// do this here so we'll stop early if there is a problem with the SDK API level (cached so later calls will return the same)
 			string SDKAPILevel = GetSdkApiLevel(ToolChain);
diff --git a/Engine/Source/Programs/UnrealBuildTool/Platform/IOS/IOSToolChain.cs b/Engine/Source/Programs/UnrealBuildTool/Platform/IOS/IOSToolChain.cs
--- a/Engine/Source/Programs/UnrealBuildTool/Platform/IOS/IOSToolChain.cs
+++ b/Engine/Source/Programs/UnrealBuildTool/Platform/IOS/IOSToolChain.cs
@@ -1264,7 +1264,7 @@ namespace UnrealBuildTool
 			}
 
 			string FabricPath = Unreal.EngineDirectory + "/Intermediate/UnzippedFrameworks/Crashlytics/Fabric.embeddedframework";
-			if (Directory.Exists(FabricPath) && Environment.GetEnvironmentVariable("IsBuildMachine") == "1")
+			if (Directory.Exists(FabricPath) && Unreal.IsBuildMachine())
 			{
 				//string PlistFile = ProjectDir + "/Intermediate/IOS/" + ProjectName + "-Info.plist";
 				Process FabricProcess = new Process();
diff --git a/Engine/Source/Programs/UnrealBuildTool/Platform/IOS/UEBuildIOS.cs b/Engine/Source/Programs/UnrealBuildTool/Platform/IOS/UEBuildIOS.cs
--- a/Engine/Source/Programs/UnrealBuildTool/Platform/IOS/UEBuildIOS.cs
+++ b/Engine/Source/Programs/UnrealBuildTool/Platform/IOS/UEBuildIOS.cs
@@ -9,6 +9,7 @@ using System.Text;
 using System.Text.RegularExpressions;
 using EpicGames.Core;
 using Microsoft.Extensions.Logging;
+using UnrealBuildBase;
 
 namespace UnrealBuildTool
 {
@@ -699,7 +700,7 @@ namespace UnrealBuildTool
 				}
 
 				// if we are stripping the executable, or if the project requested it, or if it's a buildmachine, generate the dsym
-				if (Target.IOSPlatform.bStripSymbols || Target.IOSPlatform.ProjectSettings.bGeneratedSYMFile || Environment.GetEnvironmentVariable("IsBuildMachine") == "1")
+				if (Target.IOSPlatform.bStripSymbols || Target.IOSPlatform.ProjectSettings.bGeneratedSYMFile || Unreal.IsBuildMachine())
 				{
 					Target.IOSPlatform.bGeneratedSYM = true;
 				}
diff --git a/Engine/Source/Programs/UnrealBuildTool/Platform/IOS/UEDeployIOS.cs b/Engine/Source/Programs/UnrealBuildTool/Platform/IOS/UEDeployIOS.cs
--- a/Engine/Source/Programs/UnrealBuildTool/Platform/IOS/UEDeployIOS.cs
+++ b/Engine/Source/Programs/UnrealBuildTool/Platform/IOS/UEDeployIOS.cs
@@ -150,7 +150,7 @@ namespace UnrealBuildTool
 			public static string UpdateBundleVersion(string OldPList, string EngineDirectory)
 			{
 				string CFBundleVersion = "-1";
-				if (Environment.GetEnvironmentVariable("IsBuildMachine") != "1")
+				if (!Unreal.IsBuildMachine())
 				{
 					int Index = OldPList.IndexOf("CFBundleVersion");
 					if (Index != -1)
@@ -1203,7 +1203,7 @@ namespace UnrealBuildTool
 				DestFileInfo = new FileInfo(Environment.GetEnvironmentVariable("HOME") + "/Library/MobileDevice/Provisioning Profiles/" + ProjectProvision);
 				DestFileInfo.Attributes = DestFileInfo.Attributes & ~FileAttributes.ReadOnly;
 			}
-			if (!File.Exists(ProvisionWithPrefix) || Environment.GetEnvironmentVariable("IsBuildMachine") == "1")
+			if (!File.Exists(ProvisionWithPrefix) || Unreal.IsBuildMachine())
 			{
 				// copy all provisions from the game directory, the engine directory, notforlicensees directory, and, if defined, the ProvisionDirectory.
 				CopyAllProvisions(BuildDirectory, Logger);
diff --git a/Engine/Source/Programs/UnrealBuildTool/Platform/Linux/UEBuildLinux.cs b/Engine/Source/Programs/UnrealBuildTool/Platform/Linux/UEBuildLinux.cs
--- a/Engine/Source/Programs/UnrealBuildTool/Platform/Linux/UEBuildLinux.cs
+++ b/Engine/Source/Programs/UnrealBuildTool/Platform/Linux/UEBuildLinux.cs
@@ -669,7 +669,7 @@ namespace UnrealBuildTool
 			}
 
 			// Disable color logging if we are on a build machine
-			if (Environment.GetEnvironmentVariable("IsBuildMachine") == "1")
+			if (Unreal.IsBuildMachine())
 			{
 				Log.ColorConsoleOutput = false;
 			}
diff --git a/Engine/Source/Programs/UnrealBuildTool/Platform/Mac/ApplePlatformSDK.cs b/Engine/Source/Programs/UnrealBuildTool/Platform/Mac/ApplePlatformSDK.cs
--- a/Engine/Source/Programs/UnrealBuildTool/Platform/Mac/ApplePlatformSDK.cs
+++ b/Engine/Source/Programs/UnrealBuildTool/Platform/Mac/ApplePlatformSDK.cs
@@ -3,6 +3,7 @@
 using System;
 using EpicGames.Core;
 using Microsoft.Extensions.Logging;
+using UnrealBuildBase;
 
 ///////////////////////////////////////////////////////////////////
 // If you are looking for supported version numbers, look in the
@@ -35,7 +36,7 @@ namespace UnrealBuildTool
 
 			// iTunes is technically only need to deploy to and run on connected devices.
 			// This code removes requirement for Windows builders to have Xcode installed.
-			if (Status == SDKStatus.Invalid && !RuntimePlatform.IsMac && Environment.GetEnvironmentVariable("IsBuildMachine") == "1")
+			if (Status == SDKStatus.Invalid && !RuntimePlatform.IsMac && Unreal.IsBuildMachine())
 			{
 				Status = SDKStatus.Valid;
 			}
diff --git a/Engine/Source/Programs/UnrealBuildTool/Platform/Mac/UEBuildMac.cs b/Engine/Source/Programs/UnrealBuildTool/Platform/Mac/UEBuildMac.cs
--- a/Engine/Source/Programs/UnrealBuildTool/Platform/Mac/UEBuildMac.cs
+++ b/Engine/Source/Programs/UnrealBuildTool/Platform/Mac/UEBuildMac.cs
@@ -158,7 +158,7 @@ namespace UnrealBuildTool
 			}
 
 			bool bIsEditor = false;
-			bool bIsBuildMachine = Environment.GetEnvironmentVariable("IsBuildMachine") == "1";
+			bool bIsBuildMachine = Unreal.IsBuildMachine();
 
 			// get project ini from ProjetFile, or if null, then try to get it from the target rules
 			if (TargetName != null)
@@ -553,7 +553,7 @@ namespace UnrealBuildTool
 		public override bool ShouldCreateDebugInfo(ReadOnlyTargetRules Target)
 		{
 			// Always generate debug symbols on the build machines.
-			bool IsBuildMachine = Environment.GetEnvironmentVariable("IsBuildMachine") == "1";
+			bool IsBuildMachine = Unreal.IsBuildMachine();
 
 			switch (Target.Configuration)
 			{
