UnrealBuildTool: Respect warning as error setting for clang

#rnx

diff --git a/Engine/Source/Programs/UnrealBuildTool/Platform/Windows/VCToolChain.cs b/Engine/Source/Programs/UnrealBuildTool/Platform/Windows/VCToolChain.cs
--- a/Engine/Source/Programs/UnrealBuildTool/Platform/Windows/VCToolChain.cs
+++ b/Engine/Source/Programs/UnrealBuildTool/Platform/Windows/VCToolChain.cs
@@ -312,6 +312,12 @@ namespace UnrealBuildTool
 					// Enable the static analyzer with default checks.
 					Arguments.Add("--analyze");
 
+					// Emit analyzer warnings as errors
+					if (CompileEnvironment.DefaultWarningLevel == WarningLevel.Error)
+					{
+						Arguments.Add("-Xclang -analyzer-werror");
+					}
+
 					// Deprecated in LLVM 15
 					if (EnvVars.CompilerVersion <= new VersionNumber(14))
 					{
diff --git a/Engine/Source/Programs/UnrealBuildTool/ToolChain/ClangToolChain.cs b/Engine/Source/Programs/UnrealBuildTool/ToolChain/ClangToolChain.cs
--- a/Engine/Source/Programs/UnrealBuildTool/ToolChain/ClangToolChain.cs
+++ b/Engine/Source/Programs/UnrealBuildTool/ToolChain/ClangToolChain.cs
@@ -544,7 +544,12 @@ namespace UnrealBuildTool
 		protected virtual void GetCompileArguments_WarningsAndErrors(CppCompileEnvironment CompileEnvironment, List<string> Arguments)
 		{
 			Arguments.Add("-Wall");                                     // https://clang.llvm.org/docs/DiagnosticsReference.html#wall
-			Arguments.Add("-Werror");                                   // https://clang.llvm.org/docs/UsersManual.html#cmdoption-werror
+
+			// Emit warnings as errors
+			if (CompileEnvironment.DefaultWarningLevel == WarningLevel.Error)
+			{
+				Arguments.Add("-Werror");                               // https://clang.llvm.org/docs/UsersManual.html#cmdoption-werror
+			}
 
 			ClangWarnings.GetEnabledWarnings(Arguments);
 
@@ -715,6 +720,12 @@ namespace UnrealBuildTool
 			// Enable the static analyzer with default checks.
 			Arguments.Add("--analyze");
 
+			// Emit analyzer warnings as errors
+			if (CompileEnvironment.DefaultWarningLevel == WarningLevel.Error)
+			{
+				Arguments.Add("-Xclang -analyzer-werror");
+			}
+
 			// Deprecated in LLVM 15
 			if (CompilerVersionLessThan(15, 0, 0))
 			{
