Readded splitting up the architectures in VS. In the case of Win64, the name in VS stays the same(i.e. not Win64-x64) but the other architectures have the -<architecture> suffix. This gets around an issue where VS would reset the selected config/platform because the name changed.

diff --git a/Engine/Source/Programs/UnrealBuildTool/ProjectFiles/VisualStudio/VCProject.cs b/Engine/Source/Programs/UnrealBuildTool/ProjectFiles/VisualStudio/VCProject.cs
--- a/Engine/Source/Programs/UnrealBuildTool/ProjectFiles/VisualStudio/VCProject.cs
+++ b/Engine/Source/Programs/UnrealBuildTool/ProjectFiles/VisualStudio/VCProject.cs
@@ -682,14 +682,14 @@ namespace UnrealBuildTool
 										continue;
 									}
 
-									foreach (UnrealArch? Architecture in VCProjectFileGenerator.GetPlatformArchitectures(BuildPlatform))
+									var AddProjectAndTargetCombination = (UnrealArch? Arch) =>
 									{
 										PlatformProjectGenerator? PlatformProjectGenerator = PlatformProjectGenerators.GetPlatformProjectGenerator(Platform, bInAllowFailure: true);
 										string ProjectPlatformName;
 										string ProjectConfigurationName = Configuration.ToString();
 										bool CreateDistinctConfigName = false;
 
-										VSSettings VSSettings = new(Platform, Configuration, ProjectFileFormat, Architecture);
+										VSSettings VSSettings = new(Platform, Configuration, ProjectFileFormat, Arch);
 
 										// Check to see if this platform is supported directly by Visual Studio projects.
 										if (PlatformProjectGenerator != null && PlatformProjectGenerator.HasVisualStudioSupport(VSSettings))
@@ -716,7 +716,7 @@ namespace UnrealBuildTool
 
 										if (CreateDistinctConfigName)
 										{
-											ProjectConfigurationName = string.Format("{0}{1}_{2}", Platform.ToString(), Architecture != null ? "_" + Architecture.ToString() : string.Empty, Configuration.ToString());
+											ProjectConfigurationName = string.Format("{0}{1}_{2}", Platform.ToString(), Arch != null ? "_" + Arch.ToString() : string.Empty, Configuration.ToString());
 										}
 
 										TargetType TargetConfigurationType = ProjectTarget.TargetRules!.Type;
@@ -730,7 +730,20 @@ namespace UnrealBuildTool
 											throw new BuildException("'{0}' '{1} is already in the platform/config list. This means a platform generator is not marking that the config needs to be distinct.", ProjectPlatformName, ProjectConfigurationName);
 										}
 
-										ProjectConfigAndTargetCombinations.Add(new ProjectConfigAndTargetCombination(Platform, Configuration, ProjectPlatformName, ProjectConfigurationName, ProjectTarget, Architecture));
+										ProjectConfigAndTargetCombinations.Add(new ProjectConfigAndTargetCombination(Platform, Configuration, ProjectPlatformName, ProjectConfigurationName, ProjectTarget, Arch));
+									};
+
+									UnrealArchitectures? Architectures = VCProjectFileGenerator.GetPlatformArchitecturesToGenerate(BuildPlatform);
+									if (Architectures == null)
+									{
+										AddProjectAndTargetCombination(null);
+									}
+									else
+									{
+										foreach (UnrealArch Arch in Architectures.Architectures)
+										{
+											AddProjectAndTargetCombination(Arch);
+										}
 									}
 								}
 							}
diff --git a/Engine/Source/Programs/UnrealBuildTool/ProjectFiles/VisualStudio/VCProjectFileGenerator.cs b/Engine/Source/Programs/UnrealBuildTool/ProjectFiles/VisualStudio/VCProjectFileGenerator.cs
--- a/Engine/Source/Programs/UnrealBuildTool/ProjectFiles/VisualStudio/VCProjectFileGenerator.cs
+++ b/Engine/Source/Programs/UnrealBuildTool/ProjectFiles/VisualStudio/VCProjectFileGenerator.cs
@@ -265,18 +265,10 @@ namespace UnrealBuildTool
 		/// <summary>
 		/// Returns a list of architectures to generate unique VS platforms for.
 		/// </summary>
-		public static IList<UnrealArch?> GetPlatformArchitectures(UEBuildPlatform BuildPlatform)
+		public static UnrealArchitectures? GetPlatformArchitecturesToGenerate(UEBuildPlatform BuildPlatform)
 		{
-			List<UnrealArch?> Architectures = new();
-			//if (BuildPlatform.ArchitectureConfig.Mode == UnrealArchitectureMode.OneTargetPerArchitecture)
-			//{
-			//	Architectures.AddRange(BuildPlatform.ArchitectureConfig.AllSupportedArchitectures.Architectures.Cast<UnrealArch?>());
-			//}
-			//else
-			{
-				Architectures.Add(null);
-			}
-			return Architectures;
+			return BuildPlatform.ArchitectureConfig.Mode == UnrealArchitectureMode.OneTargetPerArchitecture ?
+				BuildPlatform.ArchitectureConfig.AllSupportedArchitectures : null;
 		}
 
 		/// <inheritdoc/>
@@ -957,6 +949,7 @@ namespace UnrealBuildTool
 					VCSolutionConfigCombination? DefaultConfig = SolutionConfigCombinations.Find(x =>
 						x.Configuration == UnrealTargetConfiguration.Development &&
 						x.Platform == UnrealTargetPlatform.Win64 &&
+						(x.Architecture == null || x.Architecture == UnrealArch.X64) &&
 						(bMakeProjectPerTarget || x.TargetConfigurationName == TargetType.Editor));
 					if (DefaultConfig != null)
 					{
@@ -1069,9 +1062,7 @@ namespace UnrealBuildTool
 					{
 						ProjectTarget ProjectTarget = SolutionConfigKeyValue.Value.Item2.Item1;
 
-						IList<UnrealArch?> Architectures = GetPlatformArchitectures(BuildPlatform);
-						bool SingleArchitecture = Architectures.Count == 1;
-						foreach (UnrealArch? Arch in Architectures)
+						var AddSolutionConfig = (UnrealArch? Arch, List<VCSolutionConfigCombination> OutSolutionConfigs) =>
 						{
 							// e.g.  "Development|Win64 = Development|Win64"
 							string SolutionConfigName = SolutionConfigKeyValue.Key;
@@ -1079,7 +1070,9 @@ namespace UnrealBuildTool
 							TargetType TargetType = SolutionConfigKeyValue.Value.Item2.Item2;
 
 							string SolutionPlatformName = CurPlatform.ToString();
-							if (!SingleArchitecture)
+							// We use RequiresArchitectureFilenames to determine whether the architecture suffix should be added.
+							// This is used to tell us what the "default" architecture is.
+							if (Arch != null && BuildPlatform.ArchitectureConfig.RequiresArchitectureFilenames(new UnrealArchitectures(Arch.Value)))
 							{
 								SolutionPlatformName += $"-{Arch}";
 							}
@@ -1094,6 +1087,19 @@ namespace UnrealBuildTool
 									Architecture = Arch
 								}
 							);
+						};
+
+						UnrealArchitectures? Architectures = GetPlatformArchitecturesToGenerate(BuildPlatform);
+						if (Architectures == null)
+						{
+							AddSolutionConfig(null, OutSolutionConfigs);
+						}
+						else
+						{
+							foreach (UnrealArch Arch in Architectures.Architectures)
+							{
+								AddSolutionConfig(Arch, OutSolutionConfigs);
+							}
 						}
 					}
 				}
