[UBT]
* Fixed so manual action runners can control scheduling of work instead of being forced into a task. Some runners need to return result straight away and can't be scheduled later

diff --git a/Engine/Source/Programs/UnrealBuildTool/Executors/ImmediateActionQueue.cs b/Engine/Source/Programs/UnrealBuildTool/Executors/ImmediateActionQueue.cs
--- a/Engine/Source/Programs/UnrealBuildTool/Executors/ImmediateActionQueue.cs
+++ b/Engine/Source/Programs/UnrealBuildTool/Executors/ImmediateActionQueue.cs
@@ -565,7 +565,7 @@ namespace UnrealBuildTool
 				// If we have an action, run it and account for any completed actions
 				if (runAction != null && action != null)
 				{
-					Task.Factory.StartNew(() =>
+					if (runner != null && runner.Type == ImmediateActionQueueRunnerType.Manual)
 					{
 						try
 						{
@@ -575,7 +575,21 @@ namespace UnrealBuildTool
 						{
 							HandleException(action, ex);
 						}
-					}, CancellationToken, TaskCreationOptions.LongRunning | TaskCreationOptions.PreferFairness, TaskScheduler.Default);
+					}
+					else
+					{
+						Task.Factory.StartNew(() =>
+						{
+							try
+							{
+								runAction().Wait();
+							}
+							catch (Exception ex)
+							{
+								HandleException(action, ex);
+							}
+						}, CancellationToken, TaskCreationOptions.LongRunning | TaskCreationOptions.PreferFairness, TaskScheduler.Default);
+					}
 					AddCompletedActions(completedActions);
 					return true;
 				}
