Add a way to soft-disable PGO during a build.

- Sometimes PGO data gets stale and can needs to be removed (until a new one is collected) because it can even break the build. Instead of changing the build scripts each time when such an update process happens, add a way for UBT to silently ignore non-existent PGO data (when the file exists but has 0 size).

#rnx

diff --git a/Engine/Source/Programs/UnrealBuildTool/Platform/Linux/UEBuildLinux.cs b/Engine/Source/Programs/UnrealBuildTool/Platform/Linux/UEBuildLinux.cs
--- a/Engine/Source/Programs/UnrealBuildTool/Platform/Linux/UEBuildLinux.cs
+++ b/Engine/Source/Programs/UnrealBuildTool/Platform/Linux/UEBuildLinux.cs
@@ -506,11 +506,20 @@ namespace UnrealBuildTool
 				CompileEnvironment.PGODirectory = CompileEnvironment.PGODirectory.Replace('\\', '/') + "/";
 				CompileEnvironment.PGOFilenamePrefix = string.Format("{0}-{1}-{2}.profdata", Target.Name, Target.Platform, Target.Configuration);
 
-				// check if the profdata file exists and disable if not
+				// Check if the profdata file exists and disable if not.
+				// If the file exists but has zero length, this is a "soft" disabling. E.g. PGO data has become stale and we want to temporarily compile without PGO - do not complain about it.
 				String PGOFilePath = Path.Combine(CompileEnvironment.PGODirectory, CompileEnvironment.PGOFilenamePrefix);
-				if (!File.Exists(PGOFilePath))
+				FileInfo Info = new FileInfo(PGOFilePath);
+				if (!Info.Exists || Info.Length == 0)
 				{
-					Logger.LogWarning("Warning: PGO file '{0}' does not exist, disabling optimization", PGOFilePath);
+					if (!Info.Exists)
+					{
+						Logger.LogWarning("Warning: PGO file '{0}' does not exist, disabling optimization", PGOFilePath);
+					}
+					else
+					{
+						Logger.LogInformation("PGO file '{0}' exists but has 0 length. Assuming that PGO data is temporarily missing, disabling optimization without a warning.", PGOFilePath);
+					}
 					CompileEnvironment.bPGOOptimize = false;
 					LinkEnvironment.bPGOOptimize = false;
 
