Updated the debugger and cmdline arguments to be local or remote based on the host architecture.

#rnx

diff --git a/Engine/Source/Programs/UnrealBuildTool/Platform/Windows/UEBuildWindows.cs b/Engine/Source/Programs/UnrealBuildTool/Platform/Windows/UEBuildWindows.cs
--- a/Engine/Source/Programs/UnrealBuildTool/Platform/Windows/UEBuildWindows.cs
+++ b/Engine/Source/Programs/UnrealBuildTool/Platform/Windows/UEBuildWindows.cs
@@ -846,6 +846,17 @@ namespace UnrealBuildTool
 		{
 			return Architectures.SingleArchitecture != UnrealArch.X64;
 		}
+
+		public override UnrealArch GetHostArchitecture()
+		{
+			switch (System.Runtime.InteropServices.RuntimeInformation.ProcessArchitecture)
+			{
+				case System.Runtime.InteropServices.Architecture.Arm64:
+					return UnrealArch.Arm64;
+				default: 
+					return UnrealArch.X64;
+			}
+		}
 	}
 
 	class WindowsPlatform : UEBuildPlatform
diff --git a/Engine/Source/Programs/UnrealBuildTool/Platform/Windows/WindowsProjectGenerator.cs b/Engine/Source/Programs/UnrealBuildTool/Platform/Windows/WindowsProjectGenerator.cs
--- a/Engine/Source/Programs/UnrealBuildTool/Platform/Windows/WindowsProjectGenerator.cs
+++ b/Engine/Source/Programs/UnrealBuildTool/Platform/Windows/WindowsProjectGenerator.cs
@@ -51,6 +51,9 @@ namespace UnrealBuildTool
 		{
 			StringBuilder VCUserFileContent = new StringBuilder();
 
+			string LocalOrRemoteString = InVSSettings.Architecture == null || InVSSettings.Architecture.Value == UnrealArch.Host.Value
+				? "Local" : "Remote";
+
 			VCUserFileContent.AppendLine("  <PropertyGroup {0}>", InConditionString);
 			if (InTargetRules.Type != TargetType.Game)
 			{
@@ -66,9 +69,10 @@ namespace UnrealBuildTool
 					DebugOptions += ProjectName;
 				}
 
-				VCUserFileContent.AppendLine("    <LocalDebuggerCommandArguments>{0}</LocalDebuggerCommandArguments>", DebugOptions);
+				VCUserFileContent.AppendLine($"    <{LocalOrRemoteString}DebuggerCommandArguments>{DebugOptions}</{LocalOrRemoteString}DebuggerCommandArguments>");
 			}
-			VCUserFileContent.AppendLine("    <DebuggerFlavor>WindowsLocalDebugger</DebuggerFlavor>");
+
+			VCUserFileContent.AppendLine($"    <DebuggerFlavor>Windows{LocalOrRemoteString}Debugger</DebuggerFlavor>");
 			VCUserFileContent.AppendLine("  </PropertyGroup>");
 
 			return VCUserFileContent.ToString();
diff --git a/Engine/Source/Programs/UnrealBuildTool/ProjectFiles/VisualStudio/VCProject.cs b/Engine/Source/Programs/UnrealBuildTool/ProjectFiles/VisualStudio/VCProject.cs
--- a/Engine/Source/Programs/UnrealBuildTool/ProjectFiles/VisualStudio/VCProject.cs
+++ b/Engine/Source/Programs/UnrealBuildTool/ProjectFiles/VisualStudio/VCProject.cs
@@ -2232,7 +2232,7 @@ namespace UnrealBuildTool
 					if (ProjGenerator != null)
 					{
 						string? ForeignUProjectPath = (IsForeignProject && !String.IsNullOrEmpty(UProjectPath)) ? UProjectPath : null;
-						VCUserFileContent.Append(ProjGenerator.GetVisualStudioUserFileStrings(new(Combination.Platform!.Value, Configuration, ProjectFileFormat, null), ConditionString, TargetRulesObject, Combination.ProjectTarget.TargetFilePath, ProjectFilePath, ProjectName, ForeignUProjectPath));
+						VCUserFileContent.Append(ProjGenerator.GetVisualStudioUserFileStrings(new(Combination.Platform!.Value, Configuration, ProjectFileFormat, Combination.Architecture), ConditionString, TargetRulesObject, Combination.ProjectTarget.TargetFilePath, ProjectFilePath, ProjectName, ForeignUProjectPath));
 					}
 				}
 			}
