Ignore architecture offset for bundles when copying symbols

#android

diff --git a/Engine/Source/Programs/UnrealBuildTool/Platform/Android/UEDeployAndroid.cs b/Engine/Source/Programs/UnrealBuildTool/Platform/Android/UEDeployAndroid.cs
--- a/Engine/Source/Programs/UnrealBuildTool/Platform/Android/UEDeployAndroid.cs
+++ b/Engine/Source/Programs/UnrealBuildTool/Platform/Android/UEDeployAndroid.cs
@@ -4507,7 +4507,7 @@ namespace UnrealBuildTool
 					if (bSaveSymbols || (Configuration == UnrealTargetConfiguration.Shipping && bBuildWithHiddenSymbolVisibility))
 					{
 						// Copy .so with symbols to 
-						int StoreVersion = GetStoreVersion(Arch);
+						int StoreVersion = GetStoreVersion(bEnableBundle ? null : Arch);
 						string SymbolSODirectory = Path.Combine(DestApkDirectory, ProjectName + "_Symbols_v" + StoreVersion + "/" + ProjectName + Arch);
 						string SymbolifiedSOPath = Path.Combine(SymbolSODirectory, Path.GetFileName(SourceSOName));
 						MakeDirectoryIfRequired(SymbolifiedSOPath);
@@ -4889,7 +4889,7 @@ namespace UnrealBuildTool
 				if (bSaveSymbols || (Configuration == UnrealTargetConfiguration.Shipping && bBuildWithHiddenSymbolVisibility))
 				{
 					// Copy .so with symbols to 
-					int StoreVersion = GetStoreVersion(Arch);
+					int StoreVersion = GetStoreVersion(bEnableBundle ? null : Arch);
 					string SymbolSODirectory = Path.Combine(DestApkDirectory, ProjectName + "_Symbols_v" + StoreVersion + "/" + ProjectName + Arch);
 					string SymbolifiedSOPath = Path.Combine(SymbolSODirectory, Path.GetFileName(FinalSOName));
 					MakeDirectoryIfRequired(SymbolifiedSOPath);
