UE-189611 ARImageDetection Fails to Launch on Android - Error: Unable to load plugin 'XRBase'
-The logic in NativeProjects.cs IsHybridContentOnlyProject which decides whether we need to build an executable for a project was always using the buildhost platform (usually win64) to decide which plugins were active.  This caused it to miss enabled code plugins that build for the actual target platform but not the host platform and potentially not do an executible build when one is required.
-Now we pass the TargetPlatforms for the current build into ConditionalMakeTempTargetForHybridProject so that we can test against them.  If any project in this list needs code built we will make a temp target.

diff --git a/Engine/Source/Programs/AutomationTool/AutomationUtils/ProjectUtils.cs b/Engine/Source/Programs/AutomationTool/AutomationUtils/ProjectUtils.cs
--- a/Engine/Source/Programs/AutomationTool/AutomationUtils/ProjectUtils.cs
+++ b/Engine/Source/Programs/AutomationTool/AutomationUtils/ProjectUtils.cs
@@ -253,7 +253,7 @@ namespace AutomationTool
 					}
 				}
 
-				if (NativeProjects.ConditionalMakeTempTargetForHybridProject(RawProjectPath, Logger))
+				if (NativeProjects.ConditionalMakeTempTargetForHybridProject(RawProjectPath, TargetPlatforms, Logger))
 				{
 					Properties.bWasGenerated = true;
 					string TempTargetDir = CommandUtils.CombinePaths(Path.GetDirectoryName(RawProjectPath.FullName), "Intermediate", "Source");
diff --git a/Engine/Source/Programs/UnrealBuildTool/Configuration/TargetDescriptor.cs b/Engine/Source/Programs/UnrealBuildTool/Configuration/TargetDescriptor.cs
--- a/Engine/Source/Programs/UnrealBuildTool/Configuration/TargetDescriptor.cs
+++ b/Engine/Source/Programs/UnrealBuildTool/Configuration/TargetDescriptor.cs
@@ -573,7 +573,7 @@ namespace UnrealBuildTool
 						// make a temp target for hybrid content-as-code projects
 						if (TargetProjectFile != null)
 						{
-							NativeProjects.ConditionalMakeTempTargetForHybridProject(TargetProjectFile, Logger);
+							NativeProjects.ConditionalMakeTempTargetForHybridProject(TargetProjectFile, new List<UnrealTargetPlatform>(1) { Platform }, Logger);
 						}
 
 						if (ParamArchitectures != null)
diff --git a/Engine/Source/Programs/UnrealBuildTool/ProjectFiles/ProjectFileGenerator.cs b/Engine/Source/Programs/UnrealBuildTool/ProjectFiles/ProjectFileGenerator.cs
--- a/Engine/Source/Programs/UnrealBuildTool/ProjectFiles/ProjectFileGenerator.cs
+++ b/Engine/Source/Programs/UnrealBuildTool/ProjectFiles/ProjectFileGenerator.cs
@@ -961,7 +961,7 @@ namespace UnrealBuildTool
 			{
 				foreach (FileReference GameUProjectFile in AllGameProjects)
 				{
-					NativeProjects.ConditionalMakeTempTargetForHybridProject(GameUProjectFile, Logger);
+					NativeProjects.ConditionalMakeTempTargetForHybridProject(GameUProjectFile, new List<UnrealTargetPlatform>(1) { BuildHostPlatform.Current.Platform }, Logger);
 				}
 
 				// they are created in a temp location, which we need to scan
diff --git a/Engine/Source/Programs/UnrealBuildTool/System/NativeProjects.cs b/Engine/Source/Programs/UnrealBuildTool/System/NativeProjects.cs
--- a/Engine/Source/Programs/UnrealBuildTool/System/NativeProjects.cs
+++ b/Engine/Source/Programs/UnrealBuildTool/System/NativeProjects.cs
@@ -99,14 +99,15 @@ namespace UnrealBuildTool
 		/// Returns true if this project is a Hybrid content only project that requires it to be built as code
 		/// </summary>
 		/// <param name="UProjectFile"></param>
+		/// <param name="TargetPlatforms">The target platforms we are asking about.</param>
 		/// <param name="Reason">Contains a description of the reason the project is hybrid</param>
 		/// <param name="Logger"></param>
 		/// <returns></returns>
-		public static bool IsHybridContentOnlyProject(FileReference UProjectFile, [NotNullWhen(true)] out string? Reason, ILogger Logger)
+		public static bool IsHybridContentOnlyProject(FileReference UProjectFile, List<UnrealTargetPlatform> TargetPlatforms, [NotNullWhen(true)] out string? Reason, ILogger Logger)
 		{
 			return RequiresTempTarget(
 				UProjectFile,
-				new List<UnrealTargetPlatform>() { BuildHostPlatform.Current.Platform },
+				TargetPlatforms,
 				new List<UnrealTargetConfiguration>() { UnrealTargetConfiguration.Development, UnrealTargetConfiguration.Shipping },
 				out Reason,
 				Logger);
@@ -116,12 +117,13 @@ namespace UnrealBuildTool
 		/// Creates temporary target files, if needed, for a hybrid content only project
 		/// </summary>
 		/// <param name="UProjectFile"></param>
+		/// <param name="TargetPlatforms"></param>
 		/// <param name="Logger"></param>
 		/// <returns>True if the project is hybrid</returns>
-		public static bool ConditionalMakeTempTargetForHybridProject(FileReference UProjectFile, ILogger Logger)
+		public static bool ConditionalMakeTempTargetForHybridProject(FileReference UProjectFile, List<UnrealTargetPlatform> TargetPlatforms, ILogger Logger)
 		{
 			string? Reason;
-			bool bIsHybrid = IsHybridContentOnlyProject(UProjectFile, out Reason, Logger);
+			bool bIsHybrid = IsHybridContentOnlyProject(UProjectFile, TargetPlatforms, out Reason, Logger);
 
 			DirectoryReference TempDir = DirectoryReference.Combine(UProjectFile.Directory, "Intermediate", "Source");
 
