Removed debug logging in release because it is noticeably slowing down the project generation.

#rnx

diff --git a/Engine/Source/Programs/UnrealBuildTool/Configuration/UEBuildModuleCPP.cs b/Engine/Source/Programs/UnrealBuildTool/Configuration/UEBuildModuleCPP.cs
--- a/Engine/Source/Programs/UnrealBuildTool/Configuration/UEBuildModuleCPP.cs
+++ b/Engine/Source/Programs/UnrealBuildTool/Configuration/UEBuildModuleCPP.cs
@@ -823,18 +823,18 @@ namespace UnrealBuildTool
 		}
 
 		// Cache of files found while searching the includes
-		static ConcurrentDictionary<DirectoryReference, HashSet<FileReference>> KnownIncludeFilesDict = new();
-		static HashSet<FileReference> GetIncludeFiles(DirectoryReference Directory)
+		static ConcurrentDictionary<DirectoryReference, HashSet<FileReference>?> KnownIncludeFilesDict = new();
+		static HashSet<FileReference>? GetIncludeFiles(DirectoryReference Directory)
 		{
 			return KnownIncludeFilesDict.GetOrAdd(Directory, (_) =>
 			{
-				if (DirectoryReference.Exists(Directory))
+				if (DirectoryLookupCache.DirectoryExists(Directory))
 				{
-					return DirectoryReference.EnumerateFiles(Directory, "*").ToHashSet();
+					return DirectoryLookupCache.EnumerateFiles(Directory).ToHashSet();
 				}
 				else
 				{
-					return new HashSet<FileReference>();
+					return null;
 				}
 			});
 		}
@@ -847,8 +847,7 @@ namespace UnrealBuildTool
 				if (TranformedHeaderInclude.Contains("COMPILED_PLATFORM_HEADER("))
 				{
 					string PlatformName = UEBuildPlatform.GetBuildPlatform(CompileEnvironment.Platform).GetPlatformName();
-					TranformedHeaderInclude = TranformedHeaderInclude.Replace("COMPILED_PLATFORM_HEADER(", PlatformName + "/" + PlatformName);
-					TranformedHeaderInclude = TranformedHeaderInclude.TrimEnd(')');
+					TranformedHeaderInclude = TranformedHeaderInclude.Replace("COMPILED_PLATFORM_HEADER(", PlatformName + "/" + PlatformName).TrimEnd(')');
 				}
 
 				if (VisitedIncludes.Add(TranformedHeaderInclude))
@@ -856,7 +855,8 @@ namespace UnrealBuildTool
 					var SearchForFileItem = (DirectoryReference dir) =>
 					{
 						FileReference FileRef = FileReference.Combine(dir, TranformedHeaderInclude);
-						if (GetIncludeFiles(FileRef.Directory).Contains(FileRef))
+						HashSet<FileReference>? Files = GetIncludeFiles(FileRef.Directory);
+						if (Files != null && Files.Contains(FileRef))
 						{
 							return FileItem.GetItemByFileReference(FileRef);
 						}
@@ -890,15 +890,21 @@ namespace UnrealBuildTool
 					{
 						if (!TranformedHeaderInclude.Contains('.'))
 						{
+#if DEBUG
 							Logger.LogDebug("{0} SharedPCH - Skipping '{1}' found in '{2}' because it doesn't appear to be a module header.", Name, TranformedHeaderInclude, FileToSearch.Location);
+#endif
 						}
 						else if (TranformedHeaderInclude.EndsWith(".generated.h"))
 						{
+#if DEBUG
 							Logger.LogDebug("{0} SharedPCH - Skipping '{1}' found in '{2}' because it appears to be a generated UHT header.", Name, TranformedHeaderInclude, FileToSearch.Location);
+#endif
 						}
 						else
 						{
+#if DEBUG
 							Logger.LogDebug("{0} SharedPCH - Could not find include directory for '{1}' found in '{2}'.", Name, TranformedHeaderInclude, FileToSearch.Location);
+#endif
 							IncludesNotFound.Add(TranformedHeaderInclude);
 						}
 					}
@@ -947,28 +953,36 @@ namespace UnrealBuildTool
 					UEBuildModule? FoundModule = ModuleDeps.FirstOrDefault(Module => Module.ContainsFile(IncludeFile.Location));
 					if (FoundModule != null)
 					{
+#if DEBUG
 						if (ModuleToIncludePathsOnlyFlag[FoundModule])
 						{
 							Logger.LogDebug("{0} SharedPCH - '{1}' is exporting types but '{2}' isn't declared as a public dependency.", Name, IncludeFile.Location, FoundModule);
 						}
+#endif
 
 						OptModules.Add(FoundModule);
 					}
 					else
 					{
+#if DEBUG
 						Logger.LogDebug("{0} SharedPCH - '{1}' is exporting types but the module this file belongs to isn't declared as a public dependency or include.", Name, IncludeFile.Location);
+#endif
 						FoundAllModules = false;
 					}
 				}
 				else
 				{
+#if DEBUG
 					Logger.LogDebug("{0} SharedPCH - '{1}' is not exporting types so we are ignoring the dependency", Name, IncludeFile.Location);
-				}	
+#endif
+				}
 			}
 
 			if (!FoundAllModules)
 			{
+#if DEBUG
 				Logger.LogDebug("{0} SharedPCH - Is missing public dependencies. To be safe, this shared PCH will fall back to use all the module dependencies. Note that this could affect compile times.", Name);
+#endif
 				return AllModuleDeps;
 			}
 			else
