tvOS app icon fix for modern xcode

diff --git a/Engine/Build/TVOS/Resources/Assets.xcassets/App Icon & Top Shelf Image.brandassets/App Icon - App Store.imagestack/Back.imagestacklayer/Content.imageset/Contents.json b/Engine/Build/TVOS/Resources/Assets.xcassets/App Icon & Top Shelf Image.brandassets/App Icon - App Store.imagestack/Back.imagestacklayer/Content.imageset/Contents.json
deleted file mode 100644
diff --git a/Engine/Build/TVOS/Resources/Assets.xcassets/App Icon & Top Shelf Image.brandassets/App Icon - App Store.imagestack/Back.imagestacklayer/Contents.json b/Engine/Build/TVOS/Resources/Assets.xcassets/App Icon & Top Shelf Image.brandassets/App Icon - App Store.imagestack/Back.imagestacklayer/Contents.json
deleted file mode 100644
diff --git a/Engine/Build/TVOS/Resources/Assets.xcassets/App Icon & Top Shelf Image.brandassets/App Icon - App Store.imagestack/Contents.json b/Engine/Build/TVOS/Resources/Assets.xcassets/App Icon & Top Shelf Image.brandassets/App Icon - App Store.imagestack/Contents.json
deleted file mode 100644
diff --git a/Engine/Build/TVOS/Resources/Assets.xcassets/App Icon & Top Shelf Image.brandassets/App Icon - App Store.imagestack/Front.imagestacklayer/Content.imageset/Contents.json b/Engine/Build/TVOS/Resources/Assets.xcassets/App Icon & Top Shelf Image.brandassets/App Icon - App Store.imagestack/Front.imagestacklayer/Content.imageset/Contents.json
deleted file mode 100644
diff --git a/Engine/Build/TVOS/Resources/Assets.xcassets/App Icon & Top Shelf Image.brandassets/App Icon - App Store.imagestack/Front.imagestacklayer/Contents.json b/Engine/Build/TVOS/Resources/Assets.xcassets/App Icon & Top Shelf Image.brandassets/App Icon - App Store.imagestack/Front.imagestacklayer/Contents.json
deleted file mode 100644
diff --git a/Engine/Build/TVOS/Resources/Assets.xcassets/App Icon & Top Shelf Image.brandassets/App Icon - App Store.imagestack/Middle.imagestacklayer/Content.imageset/Contents.json b/Engine/Build/TVOS/Resources/Assets.xcassets/App Icon & Top Shelf Image.brandassets/App Icon - App Store.imagestack/Middle.imagestacklayer/Content.imageset/Contents.json
deleted file mode 100644
diff --git a/Engine/Build/TVOS/Resources/Assets.xcassets/App Icon & Top Shelf Image.brandassets/App Icon - App Store.imagestack/Middle.imagestacklayer/Contents.json b/Engine/Build/TVOS/Resources/Assets.xcassets/App Icon & Top Shelf Image.brandassets/App Icon - App Store.imagestack/Middle.imagestacklayer/Contents.json
deleted file mode 100644
diff --git a/Engine/Build/TVOS/Resources/Assets.xcassets/App Icon & Top Shelf Image.brandassets/App Icon.imagestack/Back.imagestacklayer/Content.imageset/Contents.json b/Engine/Build/TVOS/Resources/Assets.xcassets/App Icon & Top Shelf Image.brandassets/App Icon.imagestack/Back.imagestacklayer/Content.imageset/Contents.json
deleted file mode 100644
diff --git a/Engine/Build/TVOS/Resources/Assets.xcassets/App Icon & Top Shelf Image.brandassets/App Icon.imagestack/Back.imagestacklayer/Contents.json b/Engine/Build/TVOS/Resources/Assets.xcassets/App Icon & Top Shelf Image.brandassets/App Icon.imagestack/Back.imagestacklayer/Contents.json
deleted file mode 100644
diff --git a/Engine/Build/TVOS/Resources/Assets.xcassets/App Icon & Top Shelf Image.brandassets/App Icon.imagestack/Contents.json b/Engine/Build/TVOS/Resources/Assets.xcassets/App Icon & Top Shelf Image.brandassets/App Icon.imagestack/Contents.json
deleted file mode 100644
diff --git a/Engine/Build/TVOS/Resources/Assets.xcassets/App Icon & Top Shelf Image.brandassets/App Icon.imagestack/Front.imagestacklayer/Content.imageset/Contents.json b/Engine/Build/TVOS/Resources/Assets.xcassets/App Icon & Top Shelf Image.brandassets/App Icon.imagestack/Front.imagestacklayer/Content.imageset/Contents.json
deleted file mode 100644
diff --git a/Engine/Build/TVOS/Resources/Assets.xcassets/App Icon & Top Shelf Image.brandassets/App Icon.imagestack/Front.imagestacklayer/Contents.json b/Engine/Build/TVOS/Resources/Assets.xcassets/App Icon & Top Shelf Image.brandassets/App Icon.imagestack/Front.imagestacklayer/Contents.json
deleted file mode 100644
diff --git a/Engine/Build/TVOS/Resources/Assets.xcassets/App Icon & Top Shelf Image.brandassets/App Icon.imagestack/Middle.imagestacklayer/Content.imageset/Contents.json b/Engine/Build/TVOS/Resources/Assets.xcassets/App Icon & Top Shelf Image.brandassets/App Icon.imagestack/Middle.imagestacklayer/Content.imageset/Contents.json
deleted file mode 100644
diff --git a/Engine/Build/TVOS/Resources/Assets.xcassets/App Icon & Top Shelf Image.brandassets/App Icon.imagestack/Middle.imagestacklayer/Contents.json b/Engine/Build/TVOS/Resources/Assets.xcassets/App Icon & Top Shelf Image.brandassets/App Icon.imagestack/Middle.imagestacklayer/Contents.json
deleted file mode 100644
diff --git a/Engine/Build/TVOS/Resources/Assets.xcassets/App Icon & Top Shelf Image.brandassets/Contents.json b/Engine/Build/TVOS/Resources/Assets.xcassets/App Icon & Top Shelf Image.brandassets/Contents.json
deleted file mode 100644
diff --git a/Engine/Build/TVOS/Resources/Assets.xcassets/App Icon & Top Shelf Image.brandassets/Top Shelf Image Wide.imageset/Contents.json b/Engine/Build/TVOS/Resources/Assets.xcassets/App Icon & Top Shelf Image.brandassets/Top Shelf Image Wide.imageset/Contents.json
deleted file mode 100644
diff --git a/Engine/Build/TVOS/Resources/Assets.xcassets/AppIcon.brandassets/App Icon - App Store.imagestack/Back.imagestacklayer/Content.imageset/Contents.json b/Engine/Build/TVOS/Resources/Assets.xcassets/AppIcon.brandassets/App Icon - App Store.imagestack/Back.imagestacklayer/Content.imageset/Contents.json
new file mode 100644
--- /dev/null
+++ b/Engine/Build/TVOS/Resources/Assets.xcassets/AppIcon.brandassets/App Icon - App Store.imagestack/Back.imagestacklayer/Content.imageset/Contents.json	
@@ -0,0 +1,17 @@
+{
+	"images" : [
+		{
+			"idiom" : "tv",
+			"filename" : "Icon_Large_Back-1280x768.png",
+			"scale" : "1x"
+    },
+    {
+      "idiom" : "tv",
+      "scale" : "2x"
+    }
+	],
+	"info" : {
+		"version" : 1,
+		"author" : "xcode"
+	}
+}
\ No newline at end of file
diff --git a/Engine/Build/TVOS/Resources/Assets.xcassets/AppIcon.brandassets/App Icon - App Store.imagestack/Back.imagestacklayer/Contents.json b/Engine/Build/TVOS/Resources/Assets.xcassets/AppIcon.brandassets/App Icon - App Store.imagestack/Back.imagestacklayer/Contents.json
new file mode 100644
--- /dev/null
+++ b/Engine/Build/TVOS/Resources/Assets.xcassets/AppIcon.brandassets/App Icon - App Store.imagestack/Back.imagestacklayer/Contents.json	
@@ -0,0 +1,6 @@
+{
+	"info" : {
+		"version" : 1,
+		"author" : "xcode"
+	}
+}
\ No newline at end of file
diff --git a/Engine/Build/TVOS/Resources/Assets.xcassets/AppIcon.brandassets/App Icon - App Store.imagestack/Contents.json b/Engine/Build/TVOS/Resources/Assets.xcassets/AppIcon.brandassets/App Icon - App Store.imagestack/Contents.json
new file mode 100644
--- /dev/null
+++ b/Engine/Build/TVOS/Resources/Assets.xcassets/AppIcon.brandassets/App Icon - App Store.imagestack/Contents.json	
@@ -0,0 +1,17 @@
+{
+	"layers" : [
+		{
+			"filename" : "Front.imagestacklayer"
+		},
+		{
+			"filename" : "Middle.imagestacklayer"
+		},
+		{
+			"filename" : "Back.imagestacklayer"
+		}
+	],
+	"info" : {
+		"version" : 1,
+		"author" : "xcode"
+	}
+}
\ No newline at end of file
diff --git a/Engine/Build/TVOS/Resources/Assets.xcassets/AppIcon.brandassets/App Icon - App Store.imagestack/Front.imagestacklayer/Content.imageset/Contents.json b/Engine/Build/TVOS/Resources/Assets.xcassets/AppIcon.brandassets/App Icon - App Store.imagestack/Front.imagestacklayer/Content.imageset/Contents.json
new file mode 100644
--- /dev/null
+++ b/Engine/Build/TVOS/Resources/Assets.xcassets/AppIcon.brandassets/App Icon - App Store.imagestack/Front.imagestacklayer/Content.imageset/Contents.json	
@@ -0,0 +1,17 @@
+{
+	"images" : [
+		{
+			"idiom" : "tv",
+			"filename" : "Icon_Large_Front-1280x768.png",
+			"scale" : "1x"
+    },
+    {
+      "idiom" : "tv",
+      "scale" : "2x"
+    }
+	],
+	"info" : {
+		"version" : 1,
+		"author" : "xcode"
+	}
+}
\ No newline at end of file
diff --git a/Engine/Build/TVOS/Resources/Assets.xcassets/AppIcon.brandassets/App Icon - App Store.imagestack/Front.imagestacklayer/Contents.json b/Engine/Build/TVOS/Resources/Assets.xcassets/AppIcon.brandassets/App Icon - App Store.imagestack/Front.imagestacklayer/Contents.json
new file mode 100644
--- /dev/null
+++ b/Engine/Build/TVOS/Resources/Assets.xcassets/AppIcon.brandassets/App Icon - App Store.imagestack/Front.imagestacklayer/Contents.json	
@@ -0,0 +1,6 @@
+{
+	"info" : {
+		"version" : 1,
+		"author" : "xcode"
+	}
+}
\ No newline at end of file
diff --git a/Engine/Build/TVOS/Resources/Assets.xcassets/AppIcon.brandassets/App Icon - App Store.imagestack/Middle.imagestacklayer/Content.imageset/Contents.json b/Engine/Build/TVOS/Resources/Assets.xcassets/AppIcon.brandassets/App Icon - App Store.imagestack/Middle.imagestacklayer/Content.imageset/Contents.json
new file mode 100644
--- /dev/null
+++ b/Engine/Build/TVOS/Resources/Assets.xcassets/AppIcon.brandassets/App Icon - App Store.imagestack/Middle.imagestacklayer/Content.imageset/Contents.json	
@@ -0,0 +1,17 @@
+{
+	"images" : [
+		{
+			"idiom" : "tv",
+			"filename" : "Icon_Large_Middle-1280x768.png",
+			"scale" : "1x"
+    },
+    {
+      "idiom" : "tv",
+      "scale" : "2x"
+    }
+	],
+	"info" : {
+		"version" : 1,
+		"author" : "xcode"
+	}
+}
\ No newline at end of file
diff --git a/Engine/Build/TVOS/Resources/Assets.xcassets/AppIcon.brandassets/App Icon - App Store.imagestack/Middle.imagestacklayer/Contents.json b/Engine/Build/TVOS/Resources/Assets.xcassets/AppIcon.brandassets/App Icon - App Store.imagestack/Middle.imagestacklayer/Contents.json
new file mode 100644
--- /dev/null
+++ b/Engine/Build/TVOS/Resources/Assets.xcassets/AppIcon.brandassets/App Icon - App Store.imagestack/Middle.imagestacklayer/Contents.json	
@@ -0,0 +1,6 @@
+{
+	"info" : {
+		"version" : 1,
+		"author" : "xcode"
+	}
+}
\ No newline at end of file
diff --git a/Engine/Build/TVOS/Resources/Assets.xcassets/AppIcon.brandassets/App Icon.imagestack/Back.imagestacklayer/Content.imageset/Contents.json b/Engine/Build/TVOS/Resources/Assets.xcassets/AppIcon.brandassets/App Icon.imagestack/Back.imagestacklayer/Content.imageset/Contents.json
new file mode 100644
--- /dev/null
+++ b/Engine/Build/TVOS/Resources/Assets.xcassets/AppIcon.brandassets/App Icon.imagestack/Back.imagestacklayer/Content.imageset/Contents.json	
@@ -0,0 +1,18 @@
+{
+	"images" : [
+		{
+			"idiom" : "tv",
+			"filename" : "Icon_Small_Back-400x240.png",
+			"scale" : "1x"
+		},
+		{
+			"idiom" : "tv",
+			"filename" : "Icon_Small_Back-800x480.png",
+			"scale" : "2x"
+		}
+	],
+	"info" : {
+		"version" : 1,
+		"author" : "xcode"
+	}
+}
\ No newline at end of file
diff --git a/Engine/Build/TVOS/Resources/Assets.xcassets/AppIcon.brandassets/App Icon.imagestack/Back.imagestacklayer/Contents.json b/Engine/Build/TVOS/Resources/Assets.xcassets/AppIcon.brandassets/App Icon.imagestack/Back.imagestacklayer/Contents.json
new file mode 100644
--- /dev/null
+++ b/Engine/Build/TVOS/Resources/Assets.xcassets/AppIcon.brandassets/App Icon.imagestack/Back.imagestacklayer/Contents.json	
@@ -0,0 +1,6 @@
+{
+	"info" : {
+		"version" : 1,
+		"author" : "xcode"
+	}
+}
\ No newline at end of file
diff --git a/Engine/Build/TVOS/Resources/Assets.xcassets/AppIcon.brandassets/App Icon.imagestack/Contents.json b/Engine/Build/TVOS/Resources/Assets.xcassets/AppIcon.brandassets/App Icon.imagestack/Contents.json
new file mode 100644
--- /dev/null
+++ b/Engine/Build/TVOS/Resources/Assets.xcassets/AppIcon.brandassets/App Icon.imagestack/Contents.json	
@@ -0,0 +1,17 @@
+{
+	"layers" : [
+		{
+			"filename" : "Front.imagestacklayer"
+		},
+		{
+			"filename" : "Middle.imagestacklayer"
+		},
+		{
+			"filename" : "Back.imagestacklayer"
+		}
+	],
+	"info" : {
+		"version" : 1,
+		"author" : "xcode"
+	}
+}
\ No newline at end of file
diff --git a/Engine/Build/TVOS/Resources/Assets.xcassets/AppIcon.brandassets/App Icon.imagestack/Front.imagestacklayer/Content.imageset/Contents.json b/Engine/Build/TVOS/Resources/Assets.xcassets/AppIcon.brandassets/App Icon.imagestack/Front.imagestacklayer/Content.imageset/Contents.json
new file mode 100644
--- /dev/null
+++ b/Engine/Build/TVOS/Resources/Assets.xcassets/AppIcon.brandassets/App Icon.imagestack/Front.imagestacklayer/Content.imageset/Contents.json	
@@ -0,0 +1,18 @@
+{
+	"images" : [
+		{
+			"idiom" : "tv",
+			"filename" : "Icon_Small_Front-400x240.png",
+			"scale" : "1x"
+		},
+		{
+			"idiom" : "tv",
+			"filename" : "Icon_Small_Front-800x480.png",
+			"scale" : "2x"
+		}
+	],
+	"info" : {
+		"version" : 1,
+		"author" : "xcode"
+	}
+}
\ No newline at end of file
diff --git a/Engine/Build/TVOS/Resources/Assets.xcassets/AppIcon.brandassets/App Icon.imagestack/Front.imagestacklayer/Contents.json b/Engine/Build/TVOS/Resources/Assets.xcassets/AppIcon.brandassets/App Icon.imagestack/Front.imagestacklayer/Contents.json
new file mode 100644
--- /dev/null
+++ b/Engine/Build/TVOS/Resources/Assets.xcassets/AppIcon.brandassets/App Icon.imagestack/Front.imagestacklayer/Contents.json	
@@ -0,0 +1,6 @@
+{
+	"info" : {
+		"version" : 1,
+		"author" : "xcode"
+	}
+}
\ No newline at end of file
diff --git a/Engine/Build/TVOS/Resources/Assets.xcassets/AppIcon.brandassets/App Icon.imagestack/Middle.imagestacklayer/Content.imageset/Contents.json b/Engine/Build/TVOS/Resources/Assets.xcassets/AppIcon.brandassets/App Icon.imagestack/Middle.imagestacklayer/Content.imageset/Contents.json
new file mode 100644
--- /dev/null
+++ b/Engine/Build/TVOS/Resources/Assets.xcassets/AppIcon.brandassets/App Icon.imagestack/Middle.imagestacklayer/Content.imageset/Contents.json	
@@ -0,0 +1,18 @@
+{
+	"images" : [
+		{
+			"idiom" : "tv",
+			"filename" : "Icon_Small_Middle-400x240.png",
+			"scale" : "1x"
+		},
+		{
+			"idiom" : "tv",
+			"filename" : "Icon_Small_Middle-800x480.png",
+			"scale" : "2x"
+		}
+	],
+	"info" : {
+		"version" : 1,
+		"author" : "xcode"
+	}
+}
\ No newline at end of file
diff --git a/Engine/Build/TVOS/Resources/Assets.xcassets/AppIcon.brandassets/App Icon.imagestack/Middle.imagestacklayer/Contents.json b/Engine/Build/TVOS/Resources/Assets.xcassets/AppIcon.brandassets/App Icon.imagestack/Middle.imagestacklayer/Contents.json
new file mode 100644
--- /dev/null
+++ b/Engine/Build/TVOS/Resources/Assets.xcassets/AppIcon.brandassets/App Icon.imagestack/Middle.imagestacklayer/Contents.json	
@@ -0,0 +1,6 @@
+{
+	"info" : {
+		"version" : 1,
+		"author" : "xcode"
+	}
+}
\ No newline at end of file
diff --git a/Engine/Build/TVOS/Resources/Assets.xcassets/AppIcon.brandassets/Contents.json b/Engine/Build/TVOS/Resources/Assets.xcassets/AppIcon.brandassets/Contents.json
new file mode 100644
--- /dev/null
+++ b/Engine/Build/TVOS/Resources/Assets.xcassets/AppIcon.brandassets/Contents.json
@@ -0,0 +1,26 @@
+{
+  "assets" : [
+    {
+      "size" : "1280x768",
+      "idiom" : "tv",
+      "filename" : "App Icon - App Store.imagestack",
+      "role" : "primary-app-icon"
+    },
+    {
+      "size" : "400x240",
+      "idiom" : "tv",
+      "filename" : "App Icon.imagestack",
+      "role" : "primary-app-icon"
+    },
+    {
+      "size" : "2320x720",
+      "idiom" : "tv",
+      "filename" : "Top Shelf Image Wide.imageset",
+      "role" : "top-shelf-image-wide"
+    }
+  ],
+  "info" : {
+    "version" : 1,
+    "author" : "xcode"
+  }
+}
diff --git a/Engine/Build/TVOS/Resources/Assets.xcassets/AppIcon.brandassets/Top Shelf Image Wide.imageset/Contents.json b/Engine/Build/TVOS/Resources/Assets.xcassets/AppIcon.brandassets/Top Shelf Image Wide.imageset/Contents.json
new file mode 100644
--- /dev/null
+++ b/Engine/Build/TVOS/Resources/Assets.xcassets/AppIcon.brandassets/Top Shelf Image Wide.imageset/Contents.json	
@@ -0,0 +1,18 @@
+{
+	"images" : [
+		{
+			"idiom" : "tv",
+			"filename" : "TopShelfWide-1920x720.png",
+			"scale" : "1x"
+		},
+		{
+			"idiom" : "tv",
+			"filename" : "[REDACTED_EMAIL]",
+			"scale" : "2x"
+		}
+	],
+	"info" : {
+		"version" : 1,
+		"author" : "xcode"
+	}
+}
\ No newline at end of file
diff --git a/Engine/Source/Programs/UnrealBuildTool/ProjectFiles/Xcode/XcodeProject.cs b/Engine/Source/Programs/UnrealBuildTool/ProjectFiles/Xcode/XcodeProject.cs
--- a/Engine/Source/Programs/UnrealBuildTool/ProjectFiles/Xcode/XcodeProject.cs
+++ b/Engine/Source/Programs/UnrealBuildTool/ProjectFiles/Xcode/XcodeProject.cs
@@ -1437,19 +1437,19 @@ namespace UnrealBuildTool.XcodeProjectXcconfig
 
 			// look for Assets
 			string? StoryboardPath;
-			if (Platform == UnrealTargetPlatform.Mac)
-			{
-				string AssetsPath = UnrealData.ProjectOrEnginePath("Build/Mac/Resources/Assets.xcassets", false);
-				ResourcesBuildPhase.AddResource(new FileReference(AssetsPath));
-				StoryboardPath = UnrealData.FindFile(StoryboardPaths, UnrealTargetPlatform.Mac, false);
-			}
-			else
+			string AssetsSubPath = $"Build/{Platform}/Resources/Assets.xcassets";
+
+			//default to IOS path for other platforms (eg VisionOS)
+			if (Platform != UnrealTargetPlatform.TVOS ||
+			Platform != UnrealTargetPlatform.IOS ||
+			Platform != UnrealTargetPlatform.Mac)
 			{
-				string AssetsPath = UnrealData.ProjectOrEnginePath("Build/IOS/Resources/Assets.xcassets", false);
-				ResourcesBuildPhase.AddResource(new FileReference(AssetsPath));
-				StoryboardPath = UnrealData.FindFile(StoryboardPaths, UnrealTargetPlatform.IOS, false);
+				AssetsSubPath = $"Platforms/IOS/Build/Resources/Assets.xcassets";
 			}
-
+			string AssetsPath = UnrealData.ProjectOrEnginePath(AssetsSubPath, false);
+			ResourcesBuildPhase.AddResource(new FileReference(AssetsPath));
+			StoryboardPath = UnrealData.FindFile(StoryboardPaths, Platform, false);
+			
 			if (StoryboardPath != null)
 			{
 				ResourcesBuildPhase.AddResource(new FileReference(StoryboardPath));
