- Fixed an issue with naming of schemes in stub xcodeprojects when packaging BP projects (the previous method of checking TargetRules name doesnt account for the rename of UnrealGame to MyProjectGame)

diff --git a/Engine/Source/Programs/UnrealBuildTool/ProjectFiles/Xcode/XcodeProjectFileGenerator.cs b/Engine/Source/Programs/UnrealBuildTool/ProjectFiles/Xcode/XcodeProjectFileGenerator.cs
--- a/Engine/Source/Programs/UnrealBuildTool/ProjectFiles/Xcode/XcodeProjectFileGenerator.cs
+++ b/Engine/Source/Programs/UnrealBuildTool/ProjectFiles/Xcode/XcodeProjectFileGenerator.cs
@@ -320,7 +320,7 @@ namespace UnrealBuildTool
 							foreach (XcodeProjectXcconfig.XcodeProjectFile XcodeProject in SupportedProjects)
 							{
 								// if we are only generating a single target project, skip any others now
-								if (!String.IsNullOrEmpty(SingleTargetName) && !XcodeProject.ProjectTargets.Any(x => x.TargetRules?.Name == SingleTargetName))
+								if (!String.IsNullOrEmpty(SingleTargetName) && XcodeProject.ProjectFilePath.GetFileNameWithoutAnyExtensions() != SingleTargetName)
 								{
 									continue;
 								}
