- Changed the logic for Programs to allow for Resources to exist in Engine/Source/Programs/Foo/Resources/Mac (which was an existing place to put Mac resources) instead of Engine/Programs/Foo/Build/Mac. WIll be used in an upcoming changelist for some icons and plist fixes for UnrealInsigits and UnrealMultiUserServer. This will also be helpful for licensees that have icons in the old location
- Allow extra .apps to be staged into the final .app (like CrashReportClient.app)

diff --git a/Engine/Source/Programs/UnrealBuildTool/ProjectFiles/Xcode/XcodeProject.cs b/Engine/Source/Programs/UnrealBuildTool/ProjectFiles/Xcode/XcodeProject.cs
--- a/Engine/Source/Programs/UnrealBuildTool/ProjectFiles/Xcode/XcodeProject.cs
+++ b/Engine/Source/Programs/UnrealBuildTool/ProjectFiles/Xcode/XcodeProject.cs
@@ -462,6 +462,16 @@ namespace UnrealBuildTool.XcodeProjectXcconfig
 			}
 			else if (ProjectFile.ProjectTargets[0].TargetRules!.Type == TargetType.Program)
 			{
+				// if a Programs directory under Source has a Resources directory, then use it as the Product directory - if it doesn't have 
+				// a Resources dir, then go up outside of Source and then into Programs (where .ini files are, etc)
+				bool bSetProductDirectory = true;
+				if (DirectoryReference.Exists(DirectoryReference.Combine(ProjectFile.BaseDir, "Resources")))
+				{
+					ProductDirectory = ProjectFile.BaseDir;
+					// don't set the ProductDirectory, only the ConfigDirectory, below
+					bSetProductDirectory = false;
+				}
+
 				DirectoryReference? ProgramFinder = DirectoryReference.Combine(ProjectFile.BaseDir);
 				while (ProgramFinder != null && String.Compare(ProgramFinder.GetDirectoryName(), "Source", true) != 0)
 				{
@@ -474,8 +484,11 @@ namespace UnrealBuildTool.XcodeProjectXcconfig
 					// if it exists, we have a ProductDir we can use for plists, icons, etc
 					if (DirectoryReference.Exists(ProgramFinder))
 					{
-						ProductDirectory = ProgramFinder;
-						ConfigDirectory = ProductDirectory;
+						if (bSetProductDirectory)
+						{
+							ProductDirectory = ProgramFinder;
+						}
+						ConfigDirectory = ProgramFinder;
 					}
 				}
 			}
@@ -571,18 +584,29 @@ namespace UnrealBuildTool.XcodeProjectXcconfig
 			return null;
 		}
 
-		public string ProjectOrEnginePath(string SubPath, bool bMakeRelative)
+		public string ProjectOrEnginePath(string SubPath, bool bMakeRelative, string? AltProjectSubPath=null)
 		{
 			string? FinalPath = null;
-			if (UProjectFileLocation != null)
+			if (ProductDirectory != Unreal.EngineDirectory)
 			{
 				string PathToCheck = Path.Combine(ProductDirectory.FullName, SubPath);
 				if (File.Exists(PathToCheck) || Directory.Exists(PathToCheck))
 				{
 					FinalPath = PathToCheck;
 				}
+				else if (AltProjectSubPath != null)
+				{
+					// allow for an alternate project sub path, for back compat. We wouldn't use the the alt path in Engine
+					// because we would have fixed it up
+					PathToCheck = Path.Combine(ProductDirectory.FullName, AltProjectSubPath);
+					if (File.Exists(PathToCheck) || Directory.Exists(PathToCheck))
+					{
+						FinalPath = PathToCheck;
+					}
+				}
 			}
 
+			// if the SubPath (or optional AlProjectsubPath) wasn't found, then fall back to engine location
 			if (FinalPath == null)
 			{
 				FinalPath = Path.Combine(Unreal.EngineDirectory.FullName, SubPath);
@@ -1390,7 +1414,7 @@ namespace UnrealBuildTool.XcodeProjectXcconfig
 				{
 					"",
 					$"echo \\\"Syncing ${{STAGED_DIR}}{SyncSourceSubdir} to ${{CONFIGURATION_BUILD_DIR}}/${{CONTENTS_FOLDER_PATH}}{SyncDestSubdir}\\\"",
-					$"rsync -a --delete --exclude=Info.plist --exclude=/Manifest_* --exclude=*.app \\\"${{STAGED_DIR}}{SyncSourceSubdir}/\\\" \\\"${{CONFIGURATION_BUILD_DIR}}/${{CONTENTS_FOLDER_PATH}}{SyncDestSubdir}\\\"",
+					$"rsync -a --delete --exclude=/Info.plist --exclude=/Manifest_* --exclude=${{UE_TARGET_NAME}}.app \\\"${{STAGED_DIR}}{SyncSourceSubdir}/\\\" \\\"${{CONFIGURATION_BUILD_DIR}}/${{CONTENTS_FOLDER_PATH}}{SyncDestSubdir}\\\"",
 				});
 			}
 
@@ -1435,9 +1459,10 @@ namespace UnrealBuildTool.XcodeProjectXcconfig
 					"$(Engine)/Build/Apple/Resources/Interface/LaunchScreen.storyboard",
 				};
 
-			// look for Assets
+			// look for Assets (in normal place, or an alternate for Programs)
 			string? StoryboardPath;
 			string AssetsSubPath = $"Build/{Platform}/Resources/Assets.xcassets";
+			string AssetsAltSubPath = $"Resources/{Platform}/Assets.xcassets";
 
 			//default to IOS path for other platforms (eg VisionOS)
 			if (Platform != UnrealTargetPlatform.TVOS ||
@@ -1446,7 +1471,7 @@ namespace UnrealBuildTool.XcodeProjectXcconfig
 			{
 				AssetsSubPath = $"Platforms/IOS/Build/Resources/Assets.xcassets";
 			}
-			string AssetsPath = UnrealData.ProjectOrEnginePath(AssetsSubPath, false);
+			string AssetsPath = UnrealData.ProjectOrEnginePath(AssetsSubPath, false, AssetsAltSubPath);
 			ResourcesBuildPhase.AddResource(new FileReference(AssetsPath));
 			StoryboardPath = UnrealData.FindFile(StoryboardPaths, Platform, false);
 			
