Fix missing launch image on iOS.  Also fix AppIcon inverted logic for iOS.

#rnx

diff --git a/Engine/Source/Programs/UnrealBuildTool/ProjectFiles/Xcode/XcodeProject.File.cs b/Engine/Source/Programs/UnrealBuildTool/ProjectFiles/Xcode/XcodeProject.File.cs
--- a/Engine/Source/Programs/UnrealBuildTool/ProjectFiles/Xcode/XcodeProject.File.cs
+++ b/Engine/Source/Programs/UnrealBuildTool/ProjectFiles/Xcode/XcodeProject.File.cs
@@ -123,6 +123,7 @@ namespace UnrealBuildTool.XcodeProjectXcconfig
 			".xcframework",
 			".framework",
 			".lproj",
+			".png",
 			// no extensions will be folders, like "cookeddata" that we want to copy as a resource
 			"",
 		};
diff --git a/Engine/Source/Programs/UnrealBuildTool/ProjectFiles/Xcode/XcodeProject.cs b/Engine/Source/Programs/UnrealBuildTool/ProjectFiles/Xcode/XcodeProject.cs
--- a/Engine/Source/Programs/UnrealBuildTool/ProjectFiles/Xcode/XcodeProject.cs
+++ b/Engine/Source/Programs/UnrealBuildTool/ProjectFiles/Xcode/XcodeProject.cs
@@ -1465,11 +1465,11 @@ namespace UnrealBuildTool.XcodeProjectXcconfig
 			string AssetsAltSubPath = $"Resources/{Platform}/Assets.xcassets";
 
 			//default to IOS path for other platforms (eg VisionOS)
-			if (Platform != UnrealTargetPlatform.TVOS ||
-			Platform != UnrealTargetPlatform.IOS ||
+			if (Platform != UnrealTargetPlatform.TVOS &&
+			Platform != UnrealTargetPlatform.IOS &&
 			Platform != UnrealTargetPlatform.Mac)
 			{
-				AssetsSubPath = $"Platforms/IOS/Build/Resources/Assets.xcassets";
+				AssetsSubPath = $"Build/IOS/Resources/Assets.xcassets";
 			}
 			string AssetsPath = UnrealData.ProjectOrEnginePath(AssetsSubPath, false, AssetsAltSubPath);
 			ResourcesBuildPhase.AddResource(new FileReference(AssetsPath));
@@ -1479,6 +1479,21 @@ namespace UnrealBuildTool.XcodeProjectXcconfig
 			{
 				ResourcesBuildPhase.AddResource(new FileReference(StoryboardPath));
 			}
+
+			if (Platform == UnrealTargetPlatform.IOS)
+			{
+				List<string> LaunchImagePaths = new List<string>()
+				{
+					"$(Project)/Build/$(Platform)/Resources/Graphics/LaunchScreenIOS.png",
+					"$(Engine)/Build/$(Platform)/Resources/Graphics/LaunchScreenIOS.png",
+				};
+
+				string? LaunchImagePath = UnrealData.FindFile(LaunchImagePaths, UnrealTargetPlatform.IOS, false);
+				if (LaunchImagePath != null)
+				{
+					ResourcesBuildPhase.AddResource(new FileReference(LaunchImagePath));
+				}
+			}
 		}
 		protected void ProcessFrameworks(XcodeResourcesBuildPhase ResourcesBuildPhase, XcodeProject Project, XcodeProjectFile ProjectFile, ILogger Logger)
 		{
