Fix tvOS missing LaunchScreen causing a crash on startup.

#rnx

diff --git a/Engine/Source/Programs/UnrealBuildTool/ProjectFiles/Xcode/XcodeProject.cs b/Engine/Source/Programs/UnrealBuildTool/ProjectFiles/Xcode/XcodeProject.cs
--- a/Engine/Source/Programs/UnrealBuildTool/ProjectFiles/Xcode/XcodeProject.cs
+++ b/Engine/Source/Programs/UnrealBuildTool/ProjectFiles/Xcode/XcodeProject.cs
@@ -1457,6 +1457,8 @@ namespace UnrealBuildTool.XcodeProjectXcconfig
 					"$(Engine)/Build/$(Platform)/Resources/Interface/LaunchScreen.storyboard",
 					"$(Engine)/Build/Apple/Resources/Interface/LaunchScreen.storyboardc",
 					"$(Engine)/Build/Apple/Resources/Interface/LaunchScreen.storyboard",
+					"$(Project)/Build/IOS/Resources/Interface/LaunchScreen.storyboard",
+					"$(Engine)/Build/IOS/Resources/Interface/LaunchScreen.storyboardc",
 				};
 
 			// look for Assets (in normal place, or an alternate for Programs)
@@ -1480,7 +1482,7 @@ namespace UnrealBuildTool.XcodeProjectXcconfig
 				ResourcesBuildPhase.AddResource(new FileReference(StoryboardPath));
 			}
 
-			if (Platform == UnrealTargetPlatform.IOS)
+			if (Platform == UnrealTargetPlatform.IOS || Platform == UnrealTargetPlatform.TVOS)
 			{
 				List<string> LaunchImagePaths = new List<string>()
 				{
