Add a 'general settings' platform ini section, so manifest-related config properties can be queried from a separate section to other platorm properties

#rnx

diff --git a/Engine/Source/Programs/UnrealBuildTool/Platform/Windows/AppXManifestGeneratorBase.cs b/Engine/Source/Programs/UnrealBuildTool/Platform/Windows/AppXManifestGeneratorBase.cs
--- a/Engine/Source/Programs/UnrealBuildTool/Platform/Windows/AppXManifestGeneratorBase.cs
+++ b/Engine/Source/Programs/UnrealBuildTool/Platform/Windows/AppXManifestGeneratorBase.cs
@@ -19,6 +19,9 @@ namespace UnrealBuildTool
 	{
 		/// config section for platform-specific target settings
 		protected virtual string IniSection_PlatformTargetSettings => String.Format("/Script/{0}PlatformEditor.{0}TargetSettings", Platform.ToString());
+		
+		/// config section for  platform-specific general target settings (i.e. settings with are unrelated to manifest generation)
+		protected virtual string? IniSection_GeneralPlatformSettings => null;
 
 		/// config section for general target settings
 		protected virtual string IniSection_GeneralProjectSettings => "/Script/EngineSettings.GeneralProjectSettings";
@@ -173,8 +176,9 @@ namespace UnrealBuildTool
 		[return: NotNullIfNotNull("DefaultValue")]
 		protected string? GetConfigString(string PlatformKey, string? GenericKey, string? DefaultValue = null)
 		{
+			string? GeneralPlatformValue = (IniSection_GeneralPlatformSettings != null) ? ReadIniString(PlatformKey, IniSection_GeneralPlatformSettings) : null;
 			string? GenericValue = ReadIniString(GenericKey, IniSection_GeneralProjectSettings, DefaultValue);
-			return ReadIniString(PlatformKey, IniSection_PlatformTargetSettings, GenericValue);
+			return GeneralPlatformValue ?? ReadIniString(PlatformKey, IniSection_PlatformTargetSettings, GenericValue);
 		}
 
 		/// <summary>
@@ -182,8 +186,9 @@ namespace UnrealBuildTool
 		/// </summary>
 		protected bool GetConfigBool(string PlatformKey, string? GenericKey, bool DefaultValue = false)
 		{
+			string? GeneralPlatformValue = (IniSection_GeneralPlatformSettings != null) ? ReadIniString(PlatformKey, IniSection_GeneralPlatformSettings) : null;
 			string? GenericValue = ReadIniString(GenericKey, IniSection_GeneralProjectSettings, null);
-			string? ResultStr = ReadIniString(PlatformKey, IniSection_PlatformTargetSettings, GenericValue);
+			string? ResultStr = GeneralPlatformValue ?? ReadIniString(PlatformKey, IniSection_PlatformTargetSettings, GenericValue);
 
 			if (ResultStr == null)
 			{
