Fixed typo.

diff --git a/Engine/Source/Programs/UnrealBuildTool/Platform/Apple/AppleExports.cs b/Engine/Source/Programs/UnrealBuildTool/Platform/Apple/AppleExports.cs
--- a/Engine/Source/Programs/UnrealBuildTool/Platform/Apple/AppleExports.cs
+++ b/Engine/Source/Programs/UnrealBuildTool/Platform/Apple/AppleExports.cs
@@ -28,7 +28,7 @@ namespace UnrealBuildTool
 				{
 					throw new BuildException("Both -modernxcode and -legacyxcode were specified, please use one or the other.");
 				}
-				Log.TraceInformationOnce("Forcing MDOERN XCODE with -modernxcode");
+				Log.TraceInformationOnce("Forcing MODERN XCODE with -modernxcode");
 				return true;
 			}
 			if (bForceLegacyXcode)
