fix UBT interpretation of ProvisioningProfile paths, so that it works with Editor's browse file button in settings

diff --git a/Engine/Source/Programs/UnrealBuildTool/Platform/Apple/AppleExports.cs b/Engine/Source/Programs/UnrealBuildTool/Platform/Apple/AppleExports.cs
--- a/Engine/Source/Programs/UnrealBuildTool/Platform/Apple/AppleExports.cs
+++ b/Engine/Source/Programs/UnrealBuildTool/Platform/Apple/AppleExports.cs
@@ -88,8 +88,16 @@ namespace UnrealBuildTool
 			{
 				return FileReference.Combine(ProductDirectory, FilePath.Substring(6));
 			}
-
-			return new FileReference(FilePath);
+			else if (FilePath.StartsWith("/", StringComparison.OrdinalIgnoreCase))
+			{
+				// Absolute path
+				return new FileReference(FilePath);
+			}
+			else
+			{
+				// UE-193103, using the file selector in UE will set the path relative to /Engine/Binaries/Mac
+				return FileReference.Combine(Unreal.EngineDirectory, "Binaries", "Mac", FilePath);
+			}
 		}
 
 		/// <summary>
