UnrealBuildTool: Fix generate clang database regression due to alternate intermediate directory

#rnx

diff --git a/Engine/Source/Programs/UnrealBuildTool/Configuration/UEBuildTarget.cs b/Engine/Source/Programs/UnrealBuildTool/Configuration/UEBuildTarget.cs
--- a/Engine/Source/Programs/UnrealBuildTool/Configuration/UEBuildTarget.cs
+++ b/Engine/Source/Programs/UnrealBuildTool/Configuration/UEBuildTarget.cs
@@ -1106,6 +1106,10 @@ namespace UnrealBuildTool
 		/// </summary>
 		Default,
 		/// <summary>
+		/// Generate clang database environment is non unity but uses the standard intermediate directory
+		/// </summary>
+		GenerateClangDatabase,
+		/// <summary>
 		/// Include what you use
 		/// </summary>
 		IWYU,
@@ -1134,6 +1138,7 @@ namespace UnrealBuildTool
 			switch (Environment)
 			{
 				case UnrealIntermediateEnvironment.IWYU:
+				case UnrealIntermediateEnvironment.GenerateClangDatabase:
 				case UnrealIntermediateEnvironment.NonUnity:
 					return false;
 			}
diff --git a/Engine/Source/Programs/UnrealBuildTool/Modes/GenerateClangDatabase.cs b/Engine/Source/Programs/UnrealBuildTool/Modes/GenerateClangDatabase.cs
--- a/Engine/Source/Programs/UnrealBuildTool/Modes/GenerateClangDatabase.cs
+++ b/Engine/Source/Programs/UnrealBuildTool/Modes/GenerateClangDatabase.cs
@@ -146,6 +146,7 @@ namespace UnrealBuildTool
 				{
 					// Disable PCHs and unity builds for the target
 					TargetDescriptor.bUseUnityBuild = false;
+					TargetDescriptor.IntermediateEnvironment = UnrealIntermediateEnvironment.GenerateClangDatabase;
 					TargetDescriptor.AdditionalArguments = TargetDescriptor.AdditionalArguments.Append(new string[] { "-NoPCH" });
 
 					// Create a makefile for the target
