Fix ClangToolchain passing -fautortfm to the IWYU compiler.

diff --git a/Engine/Source/Programs/UnrealBuildTool/ToolChain/ClangToolChain.cs b/Engine/Source/Programs/UnrealBuildTool/ToolChain/ClangToolChain.cs
--- a/Engine/Source/Programs/UnrealBuildTool/ToolChain/ClangToolChain.cs
+++ b/Engine/Source/Programs/UnrealBuildTool/ToolChain/ClangToolChain.cs
@@ -364,7 +364,7 @@ namespace UnrealBuildTool
 				Arguments.Add("-fpch-validate-input-files-content");
 			}
 
-			if (CompileEnvironment.bAllowAutoRTFMInstrumentation)
+			if (CompileEnvironment.bAllowAutoRTFMInstrumentation && CompileEnvironment.bUseAutoRTFMCompiler)
 			{
 				Arguments.Add("-fautortfm");
 			}
@@ -395,7 +395,7 @@ namespace UnrealBuildTool
 					throw new BuildException($"Unsupported C standard type set: {CompileEnvironment.CStandard}");
 			}
 
-			if (CompileEnvironment.bAllowAutoRTFMInstrumentation)
+			if (CompileEnvironment.bAllowAutoRTFMInstrumentation && CompileEnvironment.bUseAutoRTFMCompiler)
 			{
 				Arguments.Add("-fautortfm");
 			}
