UnrealBuildTool: Fix text case in condition

#rnx

diff --git a/Engine/Source/Programs/UnrealBuildTool/UnrealBuildTool.csproj b/Engine/Source/Programs/UnrealBuildTool/UnrealBuildTool.csproj
--- a/Engine/Source/Programs/UnrealBuildTool/UnrealBuildTool.csproj
+++ b/Engine/Source/Programs/UnrealBuildTool/UnrealBuildTool.csproj
@@ -75,7 +75,7 @@
   <PropertyGroup Condition="Exists('..\..\..\Restricted\NotForLicensees\Source\Programs\UnrealBuildTool\Executors\Experimental\BoxExecutor.cs')">
     <DefineConstants>$(DefineConstants);__BOXEXECUTOR_AVAILABLE__</DefineConstants>
   </PropertyGroup>
-  <PropertyGroup Condition="Exists('..\..\..\Restricted\NotForLicensees\Source\Programs\Shared\EpicGames.Box\Box.Props')">
+  <PropertyGroup Condition="Exists('..\..\..\Restricted\NotForLicensees\Source\Programs\Shared\EpicGames.Box\Box.props')">
     <DefineConstants>$(DefineConstants);__BOX_AVAILABLE__</DefineConstants>
   </PropertyGroup>
 
