Updated all of the intellisense data gathering to be done in parallel.
Updated a few project generators to be thread safe.
Fixed an issue where the vproject generation wasn't deterministic.
Updated the project generation to output the total time it took to execute.

Before: 21.62 seconds
After: 17.56 seconds

diff --git a/Engine/Source/Programs/UnrealBuildTool/Modes/GenerateProjectFilesMode.cs b/Engine/Source/Programs/UnrealBuildTool/Modes/GenerateProjectFilesMode.cs
--- a/Engine/Source/Programs/UnrealBuildTool/Modes/GenerateProjectFilesMode.cs
+++ b/Engine/Source/Programs/UnrealBuildTool/Modes/GenerateProjectFilesMode.cs
@@ -16,7 +16,7 @@ namespace UnrealBuildTool
 	/// <summary>
 	/// Generates project files for one or more projects
 	/// </summary>
-	[ToolMode("GenerateProjectFiles", ToolModeOptions.XmlConfig | ToolModeOptions.BuildPlatforms | ToolModeOptions.SingleInstance | ToolModeOptions.UseStartupTraceListener | ToolModeOptions.StartPrefetchingEngine)]
+	[ToolMode("GenerateProjectFiles", ToolModeOptions.XmlConfig | ToolModeOptions.BuildPlatforms | ToolModeOptions.SingleInstance | ToolModeOptions.UseStartupTraceListener | ToolModeOptions.StartPrefetchingEngine | ToolModeOptions.ShowExecutionTime)]
 	class GenerateProjectFilesMode : ToolMode
 	{
 		/// <summary>
diff --git a/Engine/Source/Programs/UnrealBuildTool/ProjectFiles/ProjectFileGenerator.cs b/Engine/Source/Programs/UnrealBuildTool/ProjectFiles/ProjectFileGenerator.cs
--- a/Engine/Source/Programs/UnrealBuildTool/ProjectFiles/ProjectFileGenerator.cs
+++ b/Engine/Source/Programs/UnrealBuildTool/ProjectFiles/ProjectFileGenerator.cs
@@ -6,6 +6,7 @@ using System.Diagnostics;
 using System.IO;
 using System.Linq;
 using System.Text;
+using System.Threading;
 using System.Xml.Linq;
 using EpicGames.Core;
 using Microsoft.Extensions.Logging;
@@ -2092,10 +2093,8 @@ namespace UnrealBuildTool
 				using (ProgressWriter Progress = new ProgressWriter(ProgressInfoText, true, Logger))
 				{
 					int NumTargets = Targets.Count;
-					int NumTasks = NumTargets * 2;
+					int NumTasks = NumTargets;
 					int TasksFinished = 0;
-					TargetDescriptor[] CreatedTargetDesc = new TargetDescriptor[NumTargets];
-					UEBuildTarget[] CreatedTargets = new UEBuildTarget[NumTargets];
 					System.Threading.Tasks.Parallel.For(0, NumTargets, TargetIndex =>
 					{
 						ProjectFile TargetProjectFile = Targets[TargetIndex].Item1;
@@ -2107,7 +2106,7 @@ namespace UnrealBuildTool
 						{
 							lock (Progress)
 							{
-								TasksFinished += 2;
+								Interlocked.Increment(ref TasksFinished);
 								Progress.Write(TasksFinished, NumTasks);
 							}
 							return;
@@ -2128,36 +2127,10 @@ namespace UnrealBuildTool
 							UnrealArchitectures DefaultArchitecture = UnrealArchitectureConfig.ForPlatform(IntellisensePlatform).ActiveArchitectures(CurTarget.UnrealProjectFilePath, CurTarget.Name);
 
 							// Create the target descriptor
-							CreatedTargetDesc[TargetIndex] = new TargetDescriptor(CurTarget.UnrealProjectFilePath, CurTarget.Name, IntellisensePlatform, UnrealTargetConfiguration.Development, DefaultArchitecture, new CommandLineArguments(NewArguments.ToArray()));
+							TargetDescriptor TargetDesc = new TargetDescriptor(CurTarget.UnrealProjectFilePath, CurTarget.Name, IntellisensePlatform, UnrealTargetConfiguration.Development, DefaultArchitecture, new CommandLineArguments(NewArguments.ToArray()));
 
 							// Create the target
-							CreatedTargets[TargetIndex] = UEBuildTarget.Create(CreatedTargetDesc[TargetIndex], false, false, bUsePrecompiled, Logger);
-						}
-						catch (Exception Ex)
-						{
-							Logger.LogWarning("Exception while generating include data for {Target}: {Ex}", CurTarget.Name, Ex.ToString());
-						}
-
-						lock (Progress)
-						{
-							Progress.Write(++TasksFinished, NumTasks);
-						}
-					});
-
-					for (int TargetIndex = 0; TargetIndex < Targets.Count; ++TargetIndex)
-					{
-						ProjectFile TargetProjectFile = Targets[TargetIndex].Item1;
-						ProjectTarget CurTarget = Targets[TargetIndex].Item2;
-
-						try
-						{
-							UEBuildTarget Target = CreatedTargets[TargetIndex];
-							TargetDescriptor TargetDesc = CreatedTargetDesc[TargetIndex];
-
-							if (TargetDesc == null || Target == null)
-							{
-								continue;
-							}
+							UEBuildTarget Target = UEBuildTarget.Create(TargetDesc, false, false, bUsePrecompiled, Logger);
 
 							AddTargetForIntellisense(Target, Logger);
 
@@ -2174,12 +2147,11 @@ namespace UnrealBuildTool
 										ProjectFile? ProjectFileForIDE;
 										if (ModuleToEditorProjectFileMap.TryGetValue(Module.RulesFile, out ProjectFileForIDE) && ProjectFileForIDE == TargetProjectFile)
 										{
-											Utils.WriteFileIfChangedContext = $"{Target.TargetName} {Target.Configuration} {Target.Platform} {Binary.OutputFilePaths[0].GetFileName()} {Module.Name}";
-
 											CppCompileEnvironment ModuleCompileEnvironment = Module.CreateCompileEnvironmentForIntellisense(Target.Rules, BinaryCompileEnvironment, Logger);
-											ProjectFileForIDE.AddModule(Module, ModuleCompileEnvironment);
-
-											Utils.WriteFileIfChangedContext = "";
+											lock (ProjectFileForIDE)
+											{
+												ProjectFileForIDE.AddModule(Module, ModuleCompileEnvironment);
+											}
 										}
 									}
 								}
@@ -2199,8 +2171,12 @@ namespace UnrealBuildTool
 							Logger.LogWarning("Exception while generating include data for {Target}: {Ex}", CurTarget.Name, Ex.ToString());
 						}
 
-						Progress.Write(++TasksFinished, NumTasks);
-					}
+						lock (Progress)
+						{
+							Interlocked.Increment(ref TasksFinished);
+							Progress.Write(TasksFinished, NumTasks);
+						}
+					});
 				}
 			}
 		}
diff --git a/Engine/Source/Programs/UnrealBuildTool/ProjectFiles/VisualStudioCode/VSCodeProjectFileGenerator.cs b/Engine/Source/Programs/UnrealBuildTool/ProjectFiles/VisualStudioCode/VSCodeProjectFileGenerator.cs
--- a/Engine/Source/Programs/UnrealBuildTool/ProjectFiles/VisualStudioCode/VSCodeProjectFileGenerator.cs
+++ b/Engine/Source/Programs/UnrealBuildTool/ProjectFiles/VisualStudioCode/VSCodeProjectFileGenerator.cs
@@ -1,6 +1,7 @@
 ﻿// Copyright Epic Games, Inc. All Rights Reserved.
 
 using System;
+using System.Collections.Concurrent;
 using System.Collections.Generic;
 using System.Diagnostics;
 using System.IO;
@@ -33,7 +34,7 @@ namespace UnrealBuildTool
 		private string FrameworkExecutableExtension = ".exe";
 		private string FrameworkLibraryExtension = ".dll";
 
-		private readonly List<BuildTarget> BuildTargets = new List<BuildTarget>();
+		private readonly ConcurrentBag<BuildTarget> BuildTargets = new();
 
 		/// <summary>
 		/// Includes all files in the generated workspace.
