LLT metadata generation fixes:
- make files readable when (re)generating
- make more platforms non-runnable by default
- adjust binaries path for plugins
- supported platforms are generated from TestMetadata instead of C# attribute

New tests brought to Horde: EventLoop and StateGraph

diff --git a/Engine/Build/LowLevelTests/EventLoop.xml b/Engine/Build/LowLevelTests/EventLoop.xml
new file mode 100644
--- /dev/null
+++ b/Engine/Build/LowLevelTests/EventLoop.xml
@@ -0,0 +1,41 @@
+﻿<?xml version="1.0" encoding="utf-8"?>
+<BuildGraph xmlns:xsi="http://www.w3.org/2001/XMLSchema-instance" xsi:schemaLocation="http://www.epicgames.com/BuildGraph ../../Build/Graph/Schema.xsd" xmlns="http://www.epicgames.com/BuildGraph">
+  <Property Name="TestNames" Value="$(TestNames);EventLoop" />
+  <Property Name="EventLoopDisabled" Value="False" />
+  <Property Name="EventLoopShort" Value="EventLoop" />
+  <Property Name="EventLoopStagesWithProjectFile" Value="False" />
+  <Property Name="EventLoopTarget" Value="EventLoopUnitTests" />
+  <Property Name="EventLoopBinariesRelative" Value="Engine\Binaries" />
+  <Property Name="EventLoopReportType" Value="console" />
+  <Property Name="EventLoopGauntletArgs" Value="" />
+  <Property Name="EventLoopExtraArgs" />
+  <Property Name="EventLoopInitialExtraArgs" />
+  <Property Name="EventLoopHasAfterSteps" Value="False" />
+  <Property Name="EventLoopUsesCatch2" Value="True" />
+  <Option Name="RunEventLoopTests" DefaultValue="False" Description="Run EventLoop Tests" />
+  <Property Name="EventLoopSupportedPlatforms" Value="Win64" />
+  <Property Name="Win64EventLoopTags" Value="" />
+  <Property Name="Win64EventLoopExtraCompilationArgs" Value="" />
+  <Property Name="Win64EventLoopRunSupported" Value="True" />
+  <Property Name="MacEventLoopTags" Value="" />
+  <Property Name="MacEventLoopExtraCompilationArgs" Value="" />
+  <Property Name="MacEventLoopRunSupported" Value="True" />
+  <Property Name="IOSEventLoopTags" Value="" />
+  <Property Name="IOSEventLoopExtraCompilationArgs" Value="" />
+  <Property Name="IOSEventLoopRunSupported" Value="False" />
+  <Property Name="AndroidEventLoopTags" Value="" />
+  <Property Name="AndroidEventLoopExtraCompilationArgs" Value="" />
+  <Property Name="AndroidEventLoopRunSupported" Value="False" />
+  <Property Name="LinuxEventLoopTags" Value="" />
+  <Property Name="LinuxEventLoopExtraCompilationArgs" Value="" />
+  <Property Name="LinuxEventLoopRunSupported" Value="True" />
+  <Property Name="LinuxArm64EventLoopTags" Value="" />
+  <Property Name="LinuxArm64EventLoopExtraCompilationArgs" Value="" />
+  <Property Name="LinuxArm64EventLoopRunSupported" Value="True" />
+  <Property Name="TVOSEventLoopTags" Value="" />
+  <Property Name="TVOSEventLoopExtraCompilationArgs" Value="" />
+  <Property Name="TVOSEventLoopRunSupported" Value="False" />
+  <Property Name="VisionOSEventLoopTags" Value="" />
+  <Property Name="VisionOSEventLoopExtraCompilationArgs" Value="" />
+  <Property Name="VisionOSEventLoopRunSupported" Value="False" />
+</BuildGraph>
\ No newline at end of file
diff --git a/Engine/Build/LowLevelTests/Foundation.xml b/Engine/Build/LowLevelTests/Foundation.xml
--- a/Engine/Build/LowLevelTests/Foundation.xml
+++ b/Engine/Build/LowLevelTests/Foundation.xml
@@ -14,7 +14,7 @@
   <Property Name="FoundationHasAfterSteps" Value="False" />
   <Property Name="FoundationUsesCatch2" Value="True" />
   <Option Name="RunFoundationTests" DefaultValue="False" Description="Run Foundation Tests" />
-  <Property Name="FoundationSupportedPlatforms" Value="Win64;Mac;Linux;Android;iOS" />
+  <Property Name="FoundationSupportedPlatforms" Value="Win64;Linux;Mac;Android;IOS" />
   <Property Name="Win64FoundationTags" Value="" />
   <Property Name="MacFoundationTags" Value="" />
   <Property Name="IOSFoundationTags" Value="" />
@@ -35,8 +35,8 @@
   <Property Name="AndroidFoundationRunSupported" Value="False" />
   <Property Name="LinuxFoundationRunSupported" Value="True" />
   <Property Name="LinuxArm64FoundationRunSupported" Value="True" />
-  <Property Name="TVOSFoundationRunSupported" Value="True" />
+  <Property Name="TVOSFoundationRunSupported" Value="False" />
   <Property Name="VisionOSFoundationTags" Value="" />
   <Property Name="VisionOSFoundationExtraCompilationArgs" Value="-allmodules" />
-  <Property Name="VisionOSFoundationRunSupported" Value="True" />
+  <Property Name="VisionOSFoundationRunSupported" Value="False" />
 </BuildGraph>
\ No newline at end of file
diff --git a/Engine/Build/LowLevelTests/HeadlessChaos.xml b/Engine/Build/LowLevelTests/HeadlessChaos.xml
--- a/Engine/Build/LowLevelTests/HeadlessChaos.xml
+++ b/Engine/Build/LowLevelTests/HeadlessChaos.xml
@@ -35,8 +35,8 @@
   <Property Name="AndroidHeadlessChaosRunSupported" Value="False" />
   <Property Name="LinuxHeadlessChaosRunSupported" Value="True" />
   <Property Name="LinuxArm64HeadlessChaosRunSupported" Value="True" />
-  <Property Name="TVOSHeadlessChaosRunSupported" Value="True" />
+  <Property Name="TVOSHeadlessChaosRunSupported" Value="False" />
   <Property Name="VisionOSHeadlessChaosTags" Value="" />
   <Property Name="VisionOSHeadlessChaosExtraCompilationArgs" Value="" />
-  <Property Name="VisionOSHeadlessChaosRunSupported" Value="True" />
+  <Property Name="VisionOSHeadlessChaosRunSupported" Value="False" />
 </BuildGraph>
\ No newline at end of file
diff --git a/Engine/Build/LowLevelTests/NetworkPredictionPlugin.xml b/Engine/Build/LowLevelTests/NetworkPredictionPlugin.xml
--- a/Engine/Build/LowLevelTests/NetworkPredictionPlugin.xml
+++ b/Engine/Build/LowLevelTests/NetworkPredictionPlugin.xml
@@ -35,8 +35,8 @@
   <Property Name="AndroidNetworkPredictionPluginRunSupported" Value="False" />
   <Property Name="LinuxNetworkPredictionPluginRunSupported" Value="True" />
   <Property Name="LinuxArm64NetworkPredictionPluginRunSupported" Value="True" />
-  <Property Name="TVOSNetworkPredictionPluginRunSupported" Value="True" />
+  <Property Name="TVOSNetworkPredictionPluginRunSupported" Value="False" />
   <Property Name="VisionOSNetworkPredictionPluginTags" Value="" />
   <Property Name="VisionOSNetworkPredictionPluginExtraCompilationArgs" Value="" />
-  <Property Name="VisionOSNetworkPredictionPluginRunSupported" Value="True" />
+  <Property Name="VisionOSNetworkPredictionPluginRunSupported" Value="False" />
 </BuildGraph>
\ No newline at end of file
diff --git a/Engine/Build/LowLevelTests/ReplicationSystem.xml b/Engine/Build/LowLevelTests/ReplicationSystem.xml
--- a/Engine/Build/LowLevelTests/ReplicationSystem.xml
+++ b/Engine/Build/LowLevelTests/ReplicationSystem.xml
@@ -35,8 +35,8 @@
   <Property Name="AndroidReplicationSystemRunSupported" Value="False" />
   <Property Name="LinuxReplicationSystemRunSupported" Value="True" />
   <Property Name="LinuxArm64ReplicationSystemRunSupported" Value="True" />
-  <Property Name="TVOSReplicationSystemRunSupported" Value="True" />
+  <Property Name="TVOSReplicationSystemRunSupported" Value="False" />
   <Property Name="VisionOSReplicationSystemTags" Value="" />
   <Property Name="VisionOSReplicationSystemExtraCompilationArgs" Value="" />
-  <Property Name="VisionOSReplicationSystemRunSupported" Value="True" />
+  <Property Name="VisionOSReplicationSystemRunSupported" Value="False" />
 </BuildGraph>
\ No newline at end of file
diff --git a/Engine/Build/LowLevelTests/StateGraph.xml b/Engine/Build/LowLevelTests/StateGraph.xml
new file mode 100644
--- /dev/null
+++ b/Engine/Build/LowLevelTests/StateGraph.xml
@@ -0,0 +1,40 @@
+﻿<?xml version="1.0" encoding="utf-8"?>
+<BuildGraph xmlns:xsi="http://www.w3.org/2001/XMLSchema-instance" xsi:schemaLocation="http://www.epicgames.com/BuildGraph ../../Build/Graph/Schema.xsd" xmlns="http://www.epicgames.com/BuildGraph">
+  <Property Name="TestNames" Value="$(TestNames);StateGraph" />
+  <Property Name="StateGraphDisabled" Value="False" />
+  <Property Name="StateGraphShort" Value="State Graph" />
+  <Property Name="StateGraphStagesWithProjectFile" Value="False" />
+  <Property Name="StateGraphTarget" Value="StateGraphTests" />
+  <Property Name="StateGraphBinariesRelative" Value="Engine\Binaries" />
+  <Property Name="StateGraphReportType" Value="console" />
+  <Property Name="StateGraphGauntletArgs" Value="" />
+  <Property Name="StateGraphExtraArgs" />
+  <Property Name="StateGraphHasAfterSteps" Value="False" />
+  <Property Name="StateGraphUsesCatch2" Value="True" />
+  <Option Name="RunStateGraphTests" DefaultValue="False" Description="Run State Graph Tests" />
+  <Property Name="StateGraphSupportedPlatforms" Value="Win64" />
+  <Property Name="Win64StateGraphTags" Value="" />
+  <Property Name="Win64StateGraphExtraCompilationArgs" Value="" />
+  <Property Name="Win64StateGraphRunSupported" Value="True" />
+  <Property Name="MacStateGraphTags" Value="" />
+  <Property Name="MacStateGraphExtraCompilationArgs" Value="" />
+  <Property Name="MacStateGraphRunSupported" Value="True" />
+  <Property Name="IOSStateGraphTags" Value="" />
+  <Property Name="IOSStateGraphExtraCompilationArgs" Value="" />
+  <Property Name="IOSStateGraphRunSupported" Value="False" />
+  <Property Name="AndroidStateGraphTags" Value="" />
+  <Property Name="AndroidStateGraphExtraCompilationArgs" Value="" />
+  <Property Name="AndroidStateGraphRunSupported" Value="False" />
+  <Property Name="LinuxStateGraphTags" Value="" />
+  <Property Name="LinuxStateGraphExtraCompilationArgs" Value="" />
+  <Property Name="LinuxStateGraphRunSupported" Value="True" />
+  <Property Name="LinuxArm64StateGraphTags" Value="" />
+  <Property Name="LinuxArm64StateGraphExtraCompilationArgs" Value="" />
+  <Property Name="LinuxArm64StateGraphRunSupported" Value="True" />
+  <Property Name="TVOSStateGraphTags" Value="" />
+  <Property Name="TVOSStateGraphExtraCompilationArgs" Value="" />
+  <Property Name="TVOSStateGraphRunSupported" Value="False" />
+  <Property Name="VisionOSStateGraphTags" Value="" />
+  <Property Name="VisionOSStateGraphExtraCompilationArgs" Value="" />
+  <Property Name="VisionOSStateGraphRunSupported" Value="False" />
+</BuildGraph>
\ No newline at end of file
diff --git a/Engine/Source/Programs/LowLevelTests/FoundationTests.Build.cs b/Engine/Source/Programs/LowLevelTests/FoundationTests.Build.cs
--- a/Engine/Source/Programs/LowLevelTests/FoundationTests.Build.cs
+++ b/Engine/Source/Programs/LowLevelTests/FoundationTests.Build.cs
@@ -3,7 +3,16 @@ using UnrealBuildTool;
 
 public class FoundationTests : TestModuleRules
 {
-	protected Metadata FoundationTestsMetadata = new Metadata() { TestName = "Foundation", TestShortName = "Foundation", ReportType = "xml" };
+	protected Metadata FoundationTestsMetadata = new Metadata() {
+		TestName = "Foundation",
+		TestShortName = "Foundation",
+		ReportType = "xml",
+		SupportedPlatforms = {
+			UnrealTargetPlatform.Win64,
+			UnrealTargetPlatform.Linux,
+			UnrealTargetPlatform.Mac,
+			UnrealTargetPlatform.Android,
+			UnrealTargetPlatform.IOS } };
 
 	/// <summary>
 	/// Test metadata to be used with BuildGraph
diff --git a/Engine/Source/Programs/NetworkPredictionTests/NetworkPredictionTests.Build.cs b/Engine/Source/Programs/NetworkPredictionTests/NetworkPredictionTests.Build.cs
--- a/Engine/Source/Programs/NetworkPredictionTests/NetworkPredictionTests.Build.cs
+++ b/Engine/Source/Programs/NetworkPredictionTests/NetworkPredictionTests.Build.cs
@@ -16,6 +16,10 @@ public class NetworkPredictionTests : TestModuleRules
 			}
 		);
 
-		UpdateBuildGraphPropertiesFile(new Metadata() { TestName = "NetworkPredictionPlugin", TestShortName = "Net Prediction" });
+		UpdateBuildGraphPropertiesFile(new Metadata() {
+			TestName = "NetworkPredictionPlugin",
+			TestShortName = "Net Prediction",
+			SupportedPlatforms = { UnrealTargetPlatform.Win64, UnrealTargetPlatform.Linux }
+		});
 	}
 }
diff --git a/Engine/Source/Programs/ReplicationSystemLowLevelTests/ReplicationSystemLowLevelTests.Build.cs b/Engine/Source/Programs/ReplicationSystemLowLevelTests/ReplicationSystemLowLevelTests.Build.cs
--- a/Engine/Source/Programs/ReplicationSystemLowLevelTests/ReplicationSystemLowLevelTests.Build.cs
+++ b/Engine/Source/Programs/ReplicationSystemLowLevelTests/ReplicationSystemLowLevelTests.Build.cs
@@ -27,6 +27,10 @@ public class ReplicationSystemLowLevelTests : TestModuleRules
 			}
 		);
 
-		UpdateBuildGraphPropertiesFile(new Metadata() { TestName = "ReplicationSystem", TestShortName = "Replication System" });
+		UpdateBuildGraphPropertiesFile(new Metadata() {
+			TestName = "ReplicationSystem",
+			TestShortName = "Replication System",
+			SupportedPlatforms = { UnrealTargetPlatform.Win64, UnrealTargetPlatform.Linux }
+		});
 	}
 }
diff --git a/Engine/Source/Programs/UnrealBuildTool/Configuration/TestModuleRules.cs b/Engine/Source/Programs/UnrealBuildTool/Configuration/TestModuleRules.cs
--- a/Engine/Source/Programs/UnrealBuildTool/Configuration/TestModuleRules.cs
+++ b/Engine/Source/Programs/UnrealBuildTool/Configuration/TestModuleRules.cs
@@ -176,9 +176,11 @@ namespace UnrealBuildTool
 				return;
 			}
 
+			string BaseFolder = GetBaseFolder();
+
 			string GeneratedPropertiesScriptFile;
 
-			string NonPublicPath = Path.Combine(Unreal.EngineDirectory.FullName, "Restricted", "NotForLicensees", "Build", "LowLevelTests", $"{TestMetadata.TestName}.xml");
+			string NonPublicPath = Path.Combine(BaseFolder, "Restricted", "NotForLicensees", "Build", "LowLevelTests", $"{TestMetadata.TestName}.xml");
 
 			bool ModuleInRestrictedPath = IsRestrictedPath(ModuleDirectory);
 
@@ -188,7 +190,7 @@ namespace UnrealBuildTool
 			}
 			else
 			{
-				GeneratedPropertiesScriptFile = Path.Combine(Unreal.EngineDirectory.FullName, "Build", "LowLevelTests", $"{TestMetadata.TestName}.xml");
+				GeneratedPropertiesScriptFile = Path.Combine(BaseFolder, "Build", "LowLevelTests", $"{TestMetadata.TestName}.xml");
 			}
 
 			if (!System.IO.File.Exists(GeneratedPropertiesScriptFile))
@@ -221,6 +223,7 @@ namespace UnrealBuildTool
 				TestBinariesPath = Path.GetRelativePath(Unreal.RootDirectory.FullName, TestBinariesPath);
 			}
 
+			MakeFileWriteable(GeneratedPropertiesScriptFile);
 			XDocument GenPropsDoc = XDocument.Load(GeneratedPropertiesScriptFile);
 			XElement? Root = GenPropsDoc.Root;
 			// First descendant must be TestNames
@@ -244,23 +247,7 @@ namespace UnrealBuildTool
 
 					InsertOrUpdateTestOption(ref lastUpdatedNode, TestMetadata.TestName, $"Run {TestMetadata.TestShortName} Tests", "Run", "Tests", false.ToString());
 
-					List<UnrealTargetPlatform> AllSupportedPlatforms = new List<UnrealTargetPlatform>();
-					object[] SupportedPlatforms = GetType().GetCustomAttributes(typeof(SupportedPlatformsAttribute), false);
-					// If none specified we assume Win64
-					if (SupportedPlatforms.Length == 0)
-					{
-						UnrealTargetPlatform[] SupportedByDefault = { UnrealTargetPlatform.Win64 };
-						AllSupportedPlatforms.AddRange(SupportedByDefault);
-					}
-					else
-					{
-						foreach (object Platform in SupportedPlatforms)
-						{
-							AllSupportedPlatforms.AddRange(((SupportedPlatformsAttribute)Platform).Platforms);
-						}
-					}
-
-					InsertOrUpdateTestFlagProperty(ref lastUpdatedNode, TestMetadata.TestName, "SupportedPlatforms", AllSupportedPlatforms.Aggregate("", (current, next) => (String.IsNullOrEmpty(current) ? next.ToString() : current + ";" + next.ToString())));
+					InsertOrUpdateTestFlagProperty(ref lastUpdatedNode, TestMetadata.TestName, "SupportedPlatforms", TestMetadata.SupportedPlatforms.Aggregate("", (current, next) => (String.IsNullOrEmpty(current) ? next.ToString() : current + ";" + next.ToString())));
 				}
 			}
 
@@ -277,11 +264,11 @@ namespace UnrealBuildTool
 				bool IsRestrictedPlatformName = IsPlatformRestricted(ValidPlatform);
 				if (IsRestrictedPlatformName)
 				{
-					NonPublicPathPlatform = Path.Combine(Unreal.EngineDirectory.FullName, "Restricted", "NotForLicensees", "Platforms", ValidPlatform.ToString(), "Build", "LowLevelTests", $"{TestMetadata.TestName}.xml");
+					NonPublicPathPlatform = Path.Combine(BaseFolder, "Restricted", "NotForLicensees", "Platforms", ValidPlatform.ToString(), "Build", "LowLevelTests", $"{TestMetadata.TestName}.xml");
 				}
 				else
 				{
-					NonPublicPathPlatform = Path.Combine(Unreal.EngineDirectory.FullName, "Restricted", "NotForLicensees", "Build", "LowLevelTests", $"{TestMetadata.TestName}.xml");
+					NonPublicPathPlatform = Path.Combine(BaseFolder, "Restricted", "NotForLicensees", "Build", "LowLevelTests", $"{TestMetadata.TestName}.xml");
 				}
 
 				if (ModuleInRestrictedPath)
@@ -292,11 +279,11 @@ namespace UnrealBuildTool
 				{
 					if (IsRestrictedPlatformName)
 					{
-						GeneratedPropertiesPlatformFile = Path.Combine(Unreal.EngineDirectory.FullName, "Platforms", ValidPlatform.ToString(), "Build", "LowLevelTests", $"{TestMetadata.TestName}.xml");
+						GeneratedPropertiesPlatformFile = Path.Combine(BaseFolder, "Platforms", ValidPlatform.ToString(), "Build", "LowLevelTests", $"{TestMetadata.TestName}.xml");
 					}
 					else
 					{
-						GeneratedPropertiesPlatformFile = Path.Combine(Unreal.EngineDirectory.FullName, "Build", "LowLevelTests", $"{TestMetadata.TestName}.xml");
+						GeneratedPropertiesPlatformFile = Path.Combine(BaseFolder, "Build", "LowLevelTests", $"{TestMetadata.TestName}.xml");
 					}
 				}
 
@@ -313,6 +300,7 @@ namespace UnrealBuildTool
 					}
 				}
 
+				MakeFileWriteable(GeneratedPropertiesPlatformFile);
 				XDocument XInitPlatformFile = XDocument.Load(GeneratedPropertiesPlatformFile);
 
 				// Adding per-test and per-platform tags
@@ -329,6 +317,17 @@ namespace UnrealBuildTool
 			}
 		}
 
+		private string GetBaseFolder()
+		{
+			string RelativeModulePath = Path.GetRelativePath(Unreal.RootDirectory.FullName, ModuleDirectory);
+			string[] BreadCrumbs = RelativeModulePath.Split(new char[] { '\\', '/' }, StringSplitOptions.RemoveEmptyEntries);
+			if (BreadCrumbs.Length > 0)
+			{
+				return Path.Combine(Unreal.RootDirectory.FullName, BreadCrumbs[0]);
+			}
+			return Unreal.EngineDirectory.FullName;
+		}
+
 		private bool IsPlatformRestricted(UnrealTargetPlatform Platform)
 		{
 			return RestrictedFolder.GetNames().Contains(Platform.ToString());
@@ -352,6 +351,11 @@ namespace UnrealBuildTool
 			int SourceFolderIndex = ModuleDirectory.IndexOf("Source");
 			if (SourceFolderIndex < 0)
 			{
+				int PluginFolderIndex = ModuleDirectory.IndexOf("Plugins");
+				if (PluginFolderIndex >= 0)
+				{
+					return ModuleDirectory.Substring(0, PluginFolderIndex) + "Binaries";
+				}
 				throw new Exception("Could not detect source folder path for module " + GetType());
 			}
 			return ModuleDirectory.Substring(0, SourceFolderIndex) + "Binaries";
@@ -418,6 +422,11 @@ namespace UnrealBuildTool
 #pragma warning restore 8604
 #pragma warning restore 8602
 
+		private void MakeFileWriteable(string InFilePath)
+		{
+			System.IO.File.SetAttributes(InFilePath, System.IO.File.GetAttributes(InFilePath) & ~FileAttributes.ReadOnly);
+		}
+
 #pragma warning disable 8618
 		/// <summary>
 		/// Test metadata class.
@@ -465,7 +474,7 @@ namespace UnrealBuildTool
 			}
 
 			/// <summary>
-			/// Any initial extra args to be passed to the test executable
+			/// Any initial Gauntlet args to be passed to the test executable
 			/// </summary>
 			public string GauntletArgs { get; set; }
 
@@ -490,6 +499,11 @@ namespace UnrealBuildTool
 				set => UsesCatch2Private = value;
 			}
 
+			/// <summary>
+			/// Set of supported platforms.
+			/// </summary>
+			public HashSet<UnrealTargetPlatform> SupportedPlatforms { get; set; } = new HashSet<UnrealTargetPlatform>() { UnrealTargetPlatform.Win64 };
+
 			private Dictionary<UnrealTargetPlatform, string> PlatformTagsPrivate = new Dictionary<UnrealTargetPlatform, string>();
 			/// <summary>
 			/// Per-platform tags.
@@ -510,7 +524,12 @@ namespace UnrealBuildTool
 				set => PlatformCompilationExtraArgsPrivate = value;
 			}
 
-			private List<UnrealTargetPlatform> PlatformsRunUnsupportedPrivate = new List<UnrealTargetPlatform>() { UnrealTargetPlatform.Android, UnrealTargetPlatform.IOS };
+			private List<UnrealTargetPlatform> PlatformsRunUnsupportedPrivate = new List<UnrealTargetPlatform>() {
+				UnrealTargetPlatform.Android,
+				UnrealTargetPlatform.IOS,
+				UnrealTargetPlatform.TVOS,
+				UnrealTargetPlatform.VisionOS };
+
 			/// <summary>
 			/// List of platforms that cannot run tests.
 			/// </summary>
