Kill java before clean step on Windows to fix locked files issue that sometimes occurs with minifyReleaseWithR8

#android

diff --git a/Engine/Source/Programs/UnrealBuildTool/Platform/Android/UEDeployAndroid.cs b/Engine/Source/Programs/UnrealBuildTool/Platform/Android/UEDeployAndroid.cs
--- a/Engine/Source/Programs/UnrealBuildTool/Platform/Android/UEDeployAndroid.cs
+++ b/Engine/Source/Programs/UnrealBuildTool/Platform/Android/UEDeployAndroid.cs
@@ -4866,6 +4866,12 @@ namespace UnrealBuildTool
 								//GradleOptions = "tasks --all";
 								//RunCommandLineProgramWithException(UnrealBuildGradlePath, ShellExecutable, ShellParametersBegin + "\"" + GradleScriptPath + "\" " + GradleOptions + ShellParametersEnd, "Listing all tasks...");
 
+								// on Windows sometimes minifyReleaseWithR8 is keeping a lock on classes.dex so kill java.exe
+								if (RuntimePlatform.IsWindows)
+								{
+									RunCommandLineProgramAndReturnResult(UnrealBuildGradlePath, ShellExecutable, "/c taskkill.exe /F /IM java.exe /T", Logger, "Terminate java.exe processes");
+								}
+
 								GradleOptions = "clean";
 								RunCommandLineProgramWithExceptionAndFiltering(UnrealBuildGradlePath, ShellExecutable, ShellParametersBegin + "\"" + GradleScriptPath + "\" " + GradleOptions + ShellParametersEnd, Logger, "Cleaning Gradle intermediates...");
 							}
