UnrealBuildTool: Add support for /PDBSTRIPPED msvc symbols with -StripPrivateSymbols

diff --git a/Engine/Source/Programs/UnrealBuildTool/Platform/Windows/UEBuildWindows.cs b/Engine/Source/Programs/UnrealBuildTool/Platform/Windows/UEBuildWindows.cs
--- a/Engine/Source/Programs/UnrealBuildTool/Platform/Windows/UEBuildWindows.cs
+++ b/Engine/Source/Programs/UnrealBuildTool/Platform/Windows/UEBuildWindows.cs
@@ -443,6 +443,14 @@ namespace UnrealBuildTool
 		[CommandLine("-StrictEnumTypes")]
 		public bool bStrictEnumTypesConformance = false;
 
+		/// <summary>
+		/// Whether to request the linker create a stripped pdb file as part of the build.
+		/// If enabled the full debug pdb will have the extension .full.pdb
+		/// </summary>
+		[XmlConfigFile(Category = "WindowsPlatform")]
+		[CommandLine("-StripPrivateSymbols")]
+		public bool bStripPrivateSymbols = false;
+
 		/// <summary>
 		/// Specify an alternate location for the PDB file. This option does not change the location of the generated PDB file,
 		/// it changes the name that is embedded into the executable. Path can contain %_PDB% which will be expanded to the original
@@ -743,6 +751,8 @@ namespace UnrealBuildTool
 
 		public bool bStrictEnumTypesConformance => Inner.bStrictEnumTypesConformance;
 
+		public bool bStripPrivateSymbols => Inner.bStripPrivateSymbols;
+
 		public string? PdbAlternatePath => Inner.PdbAlternatePath;
 
 		public int DefaultStackSize => Inner.DefaultStackSize;
@@ -1298,6 +1308,7 @@ namespace UnrealBuildTool
 				|| IsBuildProductName(FileName, NamePrefixes, NameSuffixes, ".dll.rsp")
 				|| IsBuildProductName(FileName, NamePrefixes, NameSuffixes, ".lib")
 				|| IsBuildProductName(FileName, NamePrefixes, NameSuffixes, ".pdb")
+				|| IsBuildProductName(FileName, NamePrefixes, NameSuffixes, ".full.pdb")
 				|| IsBuildProductName(FileName, NamePrefixes, NameSuffixes, ".exp")
 				|| IsBuildProductName(FileName, NamePrefixes, NameSuffixes, ".obj")
 				|| IsBuildProductName(FileName, NamePrefixes, NameSuffixes, ".map")
diff --git a/Engine/Source/Programs/UnrealBuildTool/Platform/Windows/VCToolChain.cs b/Engine/Source/Programs/UnrealBuildTool/Platform/Windows/VCToolChain.cs
--- a/Engine/Source/Programs/UnrealBuildTool/Platform/Windows/VCToolChain.cs
+++ b/Engine/Source/Programs/UnrealBuildTool/Platform/Windows/VCToolChain.cs
@@ -2456,7 +2456,18 @@ namespace UnrealBuildTool
 				{
 					// Write the PDB file to the output directory.
 					{
-						FileReference PDBFilePath = FileReference.Combine(LinkEnvironment.OutputDirectory!, Path.GetFileNameWithoutExtension(OutputFile.AbsolutePath) + ".pdb");
+						FileReference? PDBFilePath = FileReference.Combine(LinkEnvironment.OutputDirectory!, Path.GetFileNameWithoutExtension(OutputFile.AbsolutePath) + ".pdb");
+
+						// If stripping private symbols, write the full pdb as .full.pdb
+						if (Target.WindowsPlatform.bStripPrivateSymbols)
+						{
+							FileItem StrippedPDBFile = FileItem.GetItemByFileReference(PDBFilePath);
+							Arguments.Add($"/PDBSTRIPPED:\"{NormalizeCommandLinePath(StrippedPDBFile)}\"");
+							ProducedItems.Add(StrippedPDBFile);
+
+							PDBFilePath = FileReference.Combine(LinkEnvironment.OutputDirectory!, Path.GetFileNameWithoutExtension(OutputFile.AbsolutePath) + ".full.pdb");
+						}
+
 						FileItem PDBFile = FileItem.GetItemByFileReference(PDBFilePath);
 						Arguments.Add($"/PDB:\"{NormalizeCommandLinePath(PDBFilePath)}\"");
 						ProducedItems.Add(PDBFile);
@@ -2527,7 +2538,7 @@ namespace UnrealBuildTool
 			// Delete PDB files for all produced items, since incremental updates are slower than full ones.
 			if (!LinkEnvironment.bUseIncrementalLinking)
 			{
-				LinkAction.DeleteItems.UnionWith(LinkAction.ProducedItems.Where(x => x.Location.HasExtension(".pdb")));
+				LinkAction.DeleteItems.UnionWith(LinkAction.ProducedItems.Where(x => x.Location.HasExtension(".pdb") || x.Location.HasExtension(".full.pdb")));
 			}
 
 			// Delete any .sup.lib files before building, even if they're not tracked
@@ -2616,6 +2627,7 @@ namespace UnrealBuildTool
 			"/FASTGENPROFILE",
 			"/OUT:",
 			"/PDB:",
+			"/PDBSTRIPPED:",
 			"/NOLOGO",
 			"/LIBPATH:",
 			"/errorReport:",
