UBT: Generate RulesFileCache for temporary targets created for hybrid projects

When we have hybrid projects, we add temporary target and build files in Intermediate/Source. We need to regenerate the RulesFileCache for that temporary folder.
If we don't, our 'initial' attempt to BuildCookRun a project will fail because we won't discover any target files because of an invalid cache entry.

diff --git a/Engine/Source/Programs/UnrealBuildTool/System/NativeProjects.cs b/Engine/Source/Programs/UnrealBuildTool/System/NativeProjects.cs
--- a/Engine/Source/Programs/UnrealBuildTool/System/NativeProjects.cs
+++ b/Engine/Source/Programs/UnrealBuildTool/System/NativeProjects.cs
@@ -220,6 +220,7 @@ namespace UnrealBuildTool
 			// need to clear out some caches now that we've added a Target
 			ClearCache();
 			Rules.InvalidateRulesFileCache(UProjectFile.Directory.FullName);
+			Rules.InvalidateRulesFileCache(TempDir.FullName);
 			DirectoryItem.ResetCachedInfo(UProjectFile.Directory.FullName);
 			DirectoryItem.ResetCachedInfo(TempDir.FullName);
 
