UnrealBuildTool: Update BuildSettingVersion for temp generated projects

#rnx

diff --git a/Engine/Source/Programs/UnrealBuildTool/System/NativeProjects.cs b/Engine/Source/Programs/UnrealBuildTool/System/NativeProjects.cs
--- a/Engine/Source/Programs/UnrealBuildTool/System/NativeProjects.cs
+++ b/Engine/Source/Programs/UnrealBuildTool/System/NativeProjects.cs
@@ -177,7 +177,7 @@ namespace UnrealBuildTool
 					Writer.WriteLine("{");
 					Writer.WriteLine("\tpublic {0}{1}Target(TargetInfo Target) : base(Target)", ProjectName, TargetTypeString);
 					Writer.WriteLine("\t{");
-					Writer.WriteLine("\t\tDefaultBuildSettings = BuildSettingsVersion.V3;");
+					Writer.WriteLine("\t\tDefaultBuildSettings = BuildSettingsVersion.Latest;");
 					Writer.WriteLine("\t\tIncludeOrderVersion = EngineIncludeOrderVersion.Latest;");
 					Writer.WriteLine("\t\tType = TargetType.{0};", TargetType);
 					Writer.WriteLine("\t\tExtraModuleNames.Add(\"{0}\");", ProjectName);
