UBT: Add AVX arch support to Linux toolchains.

#rnx

diff --git a/Engine/Source/Programs/UnrealBuildTool/Platform/Linux/LinuxToolChain.cs b/Engine/Source/Programs/UnrealBuildTool/Platform/Linux/LinuxToolChain.cs
--- a/Engine/Source/Programs/UnrealBuildTool/Platform/Linux/LinuxToolChain.cs
+++ b/Engine/Source/Programs/UnrealBuildTool/Platform/Linux/LinuxToolChain.cs
@@ -341,6 +341,43 @@ namespace UnrealBuildTool
 
 			AddCompilerLTOFlags(Arguments);
 
+			// architecture (all but None are AVX)
+			if (CompileEnvironment.Architecture == UnrealArch.X64 && CompileEnvironment.MinCpuArchX64 != MinimumCpuArchitectureX64.None)
+			{
+				// The binary created will be targeting AVX instructions. Machines without AVX support will crash on any AVX instructions if they run this compilation unit.
+
+				// AVX available implies sse4 and sse2 available.
+				Arguments.Add("-DPLATFORM_ENABLE_VECTORINTRINSICS=1");
+
+				if (CompileEnvironment.MinCpuArchX64 >= MinimumCpuArchitectureX64.AVX)
+				{
+					// Apparently MSVC enables (a subset?) of BMI (bit manipulation instructions) when /arch:AVX is set. Some code relies on this, so mirror it by enabling BMI1
+					Arguments.Add("-mavx -mbmi");
+					// Inform Unreal code that we have sse2, sse4, and AVX, both available to compile and available to run
+					Arguments.Add("-DPLATFORM_MAYBE_HAS_AVX=1");
+					// By setting the ALWAYS_HAS defines, we we direct Unreal code to skip cpuid checks to verify that the running hardware supports sse/avx.
+					Arguments.Add("-DPLATFORM_ALWAYS_HAS_AVX=1");
+				}
+
+				if (CompileEnvironment.MinCpuArchX64 >= MinimumCpuArchitectureX64.AVX2)
+				{
+					Arguments.Add("-mavx2");
+					Arguments.Add("-DPLATFORM_ALWAYS_HAS_AVX_2=1");
+				}
+
+				if (CompileEnvironment.MinCpuArchX64 >= MinimumCpuArchitectureX64.AVX512)
+				{
+					// Match MSVC which says (https://learn.microsoft.com/en-us/cpp/build/reference/arch-x64?view=msvc-170):
+					// > The __AVX512F__, __AVX512CD__, __AVX512BW__, __AVX512DQ__ and __AVX512VL__ preprocessor symbols are defined when the /arch:AVX512 compiler option is specified
+					Arguments.Add("-mavx512f");
+					Arguments.Add("-mavx512cd");
+					Arguments.Add("-mavx512bw");
+					Arguments.Add("-mavx512dq");
+					Arguments.Add("-mavx512vl");
+					Arguments.Add("-DPLATFORM_ALWAYS_HAS_AVX_512=1");
+				}
+			}
+
 			if (CompileEnvironment.bCodeCoverage)
 			{
 				Arguments.Add("-O0");
@@ -853,6 +890,8 @@ namespace UnrealBuildTool
 				Logger.LogInformation("Using fast way to relink  circularly dependent libraries (no FixDeps).");
 			}
 
+			Logger.LogInformation("Targeted minimum CPU architecture: {0}", (CompileEnvironment.Architecture == UnrealArch.X64) ? CompileEnvironment.MinCpuArchX64 : "default");
+
 			if (CompileEnvironment.bPGOOptimize)
 			{
 				Logger.LogInformation("Using PGO (profile guided optimization).");
