- Brought back Legacy style temp keychain method because IPP cannot currently re-sign Frameworks embedded in a .app

diff --git a/Engine/Source/Programs/UnrealBuildTool/Platform/Apple/AppleExports.cs b/Engine/Source/Programs/UnrealBuildTool/Platform/Apple/AppleExports.cs
--- a/Engine/Source/Programs/UnrealBuildTool/Platform/Apple/AppleExports.cs
+++ b/Engine/Source/Programs/UnrealBuildTool/Platform/Apple/AppleExports.cs
@@ -186,32 +186,39 @@ namespace UnrealBuildTool
 			}
 			else
 			{
-			// look for the special app store connect key information
-			ConfigHierarchy SharedPlatformIni = ConfigCache.ReadHierarchy(ConfigHierarchyType.Engine, ProjectFile?.Directory, UnrealTargetPlatform.Mac);
-			bool bUseAutomaticCodeSigning;
-			SharedPlatformIni.TryGetValue("/Script/MacTargetPlatform.XcodeProjectSettings", "bUseAutomaticCodeSigning", out bUseAutomaticCodeSigning);
-			if (bUseAutomaticCodeSigning)
-			{
-				ExtraOptions += " -allowProvisioningUpdates";
+				// look for the special app store connect key information
+				ConfigHierarchy SharedPlatformIni = ConfigCache.ReadHierarchy(ConfigHierarchyType.Engine, ProjectFile?.Directory, UnrealTargetPlatform.Mac);
+				bool bUseAutomaticCodeSigning;
+				SharedPlatformIni.TryGetValue("/Script/MacTargetPlatform.XcodeProjectSettings", "bUseAutomaticCodeSigning", out bUseAutomaticCodeSigning);
 
-				// handle AppStore Connect settings
-				bool bUseAppStoreConnect;
-				SharedPlatformIni.TryGetValue("/Script/MacTargetPlatform.XcodeProjectSettings", "bUseAppStoreConnect", out bUseAppStoreConnect);
-				if (bUseAppStoreConnect)
+				// disable automatic signing, if some extra options imply manual
+				if (ExtraOptions.Contains("CODE_SIGN_IDENTITY"))
 				{
-					string? IssuerID, KeyID, KeyPath;
-					if (SharedPlatformIni.TryGetValue("/Script/MacTargetPlatform.XcodeProjectSettings", "AppStoreConnectIssuerID", out IssuerID) &&
-						SharedPlatformIni.TryGetValue("/Script/MacTargetPlatform.XcodeProjectSettings", "AppStoreConnectKeyID", out KeyID) &&
-						SharedPlatformIni.TryGetValue("/Script/MacTargetPlatform.XcodeProjectSettings", "AppStoreConnectKeyPath", out KeyPath))
+					bUseAutomaticCodeSigning = false;
+				}
+
+				if (bUseAutomaticCodeSigning)
+				{
+					ExtraOptions += " -allowProvisioningUpdates";
+	
+					// handle AppStore Connect settings
+					bool bUseAppStoreConnect;
+					SharedPlatformIni.TryGetValue("/Script/MacTargetPlatform.XcodeProjectSettings", "bUseAppStoreConnect", out bUseAppStoreConnect);
+					if (bUseAppStoreConnect)
 					{
-						FileReference KeyFile = ConvertFilePath(ProjectFile?.Directory, KeyPath);
-						ExtraOptions += $" -authenticationKeyIssuerID {IssuerID}";
-						ExtraOptions += $" -authenticationKeyID {KeyID}";
-						ExtraOptions += $" -authenticationKeyPath \"{KeyFile}\"";
+						string? IssuerID, KeyID, KeyPath;
+						if (SharedPlatformIni.TryGetValue("/Script/MacTargetPlatform.XcodeProjectSettings", "AppStoreConnectIssuerID", out IssuerID) &&
+							SharedPlatformIni.TryGetValue("/Script/MacTargetPlatform.XcodeProjectSettings", "AppStoreConnectKeyID", out KeyID) &&
+							SharedPlatformIni.TryGetValue("/Script/MacTargetPlatform.XcodeProjectSettings", "AppStoreConnectKeyPath", out KeyPath))
+						{
+							FileReference KeyFile = ConvertFilePath(ProjectFile?.Directory, KeyPath);
+							ExtraOptions += $" -authenticationKeyIssuerID {IssuerID}";
+							ExtraOptions += $" -authenticationKeyID {KeyID}";
+							ExtraOptions += $" -authenticationKeyPath \"{KeyFile}\"";
+						}
 					}
 				}
 			}
-			}
 
 			// run xcodebuild on the generated project to make the .app
 			string XcodeBuildAction = BuildMode == XcodeBuildMode.Distribute ? "archive" : "build";
diff --git a/Engine/Source/Programs/UnrealBuildTool/Platform/IOS/IOSToolChain.cs b/Engine/Source/Programs/UnrealBuildTool/Platform/IOS/IOSToolChain.cs
--- a/Engine/Source/Programs/UnrealBuildTool/Platform/IOS/IOSToolChain.cs
+++ b/Engine/Source/Programs/UnrealBuildTool/Platform/IOS/IOSToolChain.cs
@@ -724,7 +724,7 @@ namespace UnrealBuildTool
 			return OutputFile;
 		}
 
-		public static void PackageStub(string BinaryPath, string GameName, string ExeName, bool bUseModernXcode)
+		public static void PackageStub(string BinaryPath, string GameName, string ExeName, bool bUseModernXcode, bool bPerformDummySigning=false)
 		{
 			// create the ipa
 			string IPAName = BinaryPath + "/" + ExeName + ".stub";
@@ -745,7 +745,7 @@ namespace UnrealBuildTool
 			string ZipWorkingDir = bUseModernXcode ? $"Payload/{ExeName}.app/" : $"Payload/{GameName}.app/";
 			string ZipSourceDir = bUseModernXcode ? $"{BinaryPath}/{ExeName}.app" : $"{BinaryPath}/Payload/{GameName}.app";
 
-			if (bUseModernXcode)
+			if (bPerformDummySigning)
 			{
 				// with modern, we don't codesign when building the source .app due to keychain issues over ssh and Xcode wants provision setup, even if using the - identity
 				// so, now codesign right at the end with the dummy signature so that there is signature at all, which IPP on windows will replace
@@ -755,6 +755,15 @@ namespace UnrealBuildTool
 				{
 					throw new BuildException($"Failed to codesign with dummy identity, which should never happen [ExitCode = {ExitCode}, output = {Output}");
 				}
+				// also codesign the frameworks
+				foreach (string FrameworkDir in Directory.EnumerateDirectories(Path.Combine(ZipSourceDir, "Frameworks")))
+				{
+					Output = Utils.RunLocalProcessAndReturnStdOut("/bin/sh", $"-c 'codesign -f -s - \"{FrameworkDir}\"", null, out ExitCode);
+					if (ExitCode != 0)
+					{
+						throw new BuildException($"Failed to codesign with dummy identity, which should never happen [ExitCode = {ExitCode}, output = {Output}");
+					}
+				}
 			}
 
 			// create the file
diff --git a/Engine/Source/Programs/UnrealBuildTool/ToolChain/AppleToolChain.cs b/Engine/Source/Programs/UnrealBuildTool/ToolChain/AppleToolChain.cs
--- a/Engine/Source/Programs/UnrealBuildTool/ToolChain/AppleToolChain.cs
+++ b/Engine/Source/Programs/UnrealBuildTool/ToolChain/AppleToolChain.cs
@@ -5,6 +5,7 @@ using System.Collections.Generic;
 using System.Diagnostics;
 using System.IO;
 using System.Linq;
+using System.Security.Cryptography.X509Certificates;
 using System.Runtime.InteropServices;
 using System.Text;
 using System.Text.RegularExpressions;
@@ -63,7 +64,7 @@ namespace UnrealBuildTool
 		/// <param name="bVerbose"></param>
 		/// <param name="Logger"></param>
 		protected AppleToolChainSettings(string OSPrefix, string? SimulatorOSPrefix, string TargetOSName, bool bVerbose, ILogger Logger)
-		{
+			{   
 			this.TargetOSName = TargetOSName;
 
 			if (_XcodeDeveloperDir == null)
@@ -92,7 +93,7 @@ namespace UnrealBuildTool
 		/// <param name="Architecture"></param>
 		/// <returns></returns>
 		public DirectoryReference GetSDKPath(UnrealArch Architecture)
-		{
+			{
 			// note that VisionOS uses IOSSimulator (as TVOS should eventually do as well)
 			if (Architecture == UnrealArch.IOSSimulator || Architecture == UnrealArch.TVOSSimulator)
 		{
@@ -107,7 +108,7 @@ namespace UnrealBuildTool
 		/// <param name="Architecture"></param>
 		/// <returns></returns>
 		public string GetTargetTuple(UnrealArch Architecture)
-			{
+		{
 			string Prefix = Architecture.AppleName;
 			string Suffix = (Architecture == UnrealArch.IOSSimulator || Architecture == UnrealArch.TVOSSimulator) ? "-simulator" : "";
 
@@ -142,13 +143,13 @@ namespace UnrealBuildTool
 			if (XcodeDeveloperDir.ContainsName("CommandLineTools", 0))
 			{
 				throw new BuildException($"Your Mac is set to use CommandLineTools for its build tools ({XcodeDeveloperDir}). Unreal expects Xcode as the build tools. Please install Xcode if it's not already, then do one of the following:\n" +
-						"  - Run Xcode, go to Settings, and in the Locations tab, choose your Xcode in Command Line Tools dropdown.\n" +
-						"  - In Terminal, run 'sudo xcode-select -s /Applications/Xcode.app' (or an alternate location if you installed Xcode to a non-standard location)\n" +
-						"Either way, you will need to enter your Mac password.");
+					"  - Run Xcode, go to Settings, and in the Locations tab, choose your Xcode in Command Line Tools dropdown.\n" +
+					"  - In Terminal, run 'sudo xcode-select -s /Applications/Xcode.app' (or an alternate location if you installed Xcode to a non-standard location)\n" +
+					"Either way, you will need to enter your Mac password.");
 			}
 
 			if (bVerbose && !XcodeDeveloperDir.FullName.StartsWith("/Applications/Xcode.app"))
-			{
+				{
 				Log.TraceInformationOnce("Compiling with non-standard Xcode: {0}", XcodeDeveloperDir);
 			}
 
@@ -168,62 +169,62 @@ namespace UnrealBuildTool
 		}
 
 		private static string SelectSDK(DirectoryReference BaseSDKDir, string OSPrefix, bool bVerbose, ILogger Logger)
-		{
+			{
 			string PlatformSDKVersion = "";
-				try
-				{
-					// loop over the subdirs and parse out the version
-					int MaxSDKVersionMajor = 0;
-					int MaxSDKVersionMinor = 0;
-					string? MaxSDKVersionString = null;
+			try
+			{
+				// loop over the subdirs and parse out the version
+				int MaxSDKVersionMajor = 0;
+				int MaxSDKVersionMinor = 0;
+				string? MaxSDKVersionString = null;
 				foreach (DirectoryReference SubDir in DirectoryReference.EnumerateDirectories(BaseSDKDir))
-					{
+				{
 					string SubDirName = Path.GetFileNameWithoutExtension(SubDir.GetDirectoryName());
-						if (SubDirName.StartsWith(OSPrefix))
-						{
-							// get the SDK version from the directory name
-							string SDKString = SubDirName.Replace(OSPrefix, "");
-							int Major = 0;
-							int Minor = 0;
-
-							// parse it into whole and fractional parts (since 10.10 > 10.9 in versions, but not in math)
-							try
-							{
-								string[] Tokens = SDKString.Split(".".ToCharArray());
-								if (Tokens.Length == 2)
-								{
-									Major = Int32.Parse(Tokens[0]);
-									Minor = Int32.Parse(Tokens[1]);
-								}
-							}
-							catch (Exception)
-							{
-								// weirdly formatted SDKs
-								continue;
-							}
+					if (SubDirName.StartsWith(OSPrefix))
+					{
+						// get the SDK version from the directory name
+						string SDKString = SubDirName.Replace(OSPrefix, "");
+						int Major = 0;
+						int Minor = 0;
 
-							// update largest SDK version number
-							if (Major > MaxSDKVersionMajor || (Major == MaxSDKVersionMajor && Minor > MaxSDKVersionMinor))
+						// parse it into whole and fractional parts (since 10.10 > 10.9 in versions, but not in math)
+						try
+						{
+							string[] Tokens = SDKString.Split(".".ToCharArray());
+							if (Tokens.Length == 2)
 							{
-								MaxSDKVersionString = SDKString;
-								MaxSDKVersionMajor = Major;
-								MaxSDKVersionMinor = Minor;
+								Major = Int32.Parse(Tokens[0]);
+								Minor = Int32.Parse(Tokens[1]);
 							}
 						}
-					}
+						catch (Exception)
+						{
+							// weirdly formatted SDKs
+							continue;
+						}
 
-					// use the largest version
-					if (MaxSDKVersionString != null)
-					{
-						PlatformSDKVersion = MaxSDKVersionString;
+						// update largest SDK version number
+						if (Major > MaxSDKVersionMajor || (Major == MaxSDKVersionMajor && Minor > MaxSDKVersionMinor))
+						{
+							MaxSDKVersionString = SDKString;
+							MaxSDKVersionMajor = Major;
+							MaxSDKVersionMinor = Minor;
+						}
 					}
 				}
-				catch (Exception Ex)
+
+				// use the largest version
+				if (MaxSDKVersionString != null)
 				{
-					// on any exception, just use the backup version
-					Logger.LogInformation("Triggered an exception while looking for SDK directory in Xcode.app");
-					Logger.LogInformation("{Ex}", Ex.ToString());
+					PlatformSDKVersion = MaxSDKVersionString;
 				}
+			}
+			catch (Exception Ex)
+			{
+				// on any exception, just use the backup version
+				Logger.LogInformation("Triggered an exception while looking for SDK directory in Xcode.app");
+				Logger.LogInformation("{Ex}", Ex.ToString());
+			}
 
 			if (bVerbose && !ProjectFileGenerator.bGenerateProjectFiles)
 			{
@@ -905,7 +906,7 @@ namespace UnrealBuildTool
 					ExtraOptions,
 					//$"-sdk {SDKName}",
 				};
-
+	
 				Process LocalProcess = new Process();
 				LocalProcess.StartInfo = new ProcessStartInfo("/usr/bin/env", String.Join(" ", Arguments));
 				LocalProcess.OutputDataReceived += (Sender, Args) => { LocalProcessOutput(Args, false, Logger); };
@@ -955,9 +956,13 @@ namespace UnrealBuildTool
 
 		// For iOS/TVOS
 		public bool bCreateStubIPA;
+		public string? RemoteImportProvision;
+		public string? RemoteImportCertificate;
+		public string? RemoteImportCertificatePassword;
+		public DirectoryReference ProjectIntermediateDirectory;
 		public FileReference StubOutputPath;
 
-		public ApplePostBuildSyncTarget(ReadOnlyTargetRules Target, FileItem Executable)
+		public ApplePostBuildSyncTarget(ReadOnlyTargetRules Target, FileItem Executable, DirectoryReference IntermediateDir)
 		{
 			Platform = Target.Platform;
 			Configuration = Target.Configuration;
@@ -966,6 +971,11 @@ namespace UnrealBuildTool
 			TargetName = Target.Name;
 
 			bCreateStubIPA = Target.IOSPlatform.bCreateStubIPA;
+			RemoteImportProvision = Target.IOSPlatform.ImportProvision;
+			RemoteImportCertificate = Target.IOSPlatform.ImportCertificate;
+			RemoteImportCertificatePassword = Target.IOSPlatform.ImportCertificatePassword;
+			ProjectIntermediateDirectory = IntermediateDir;
+
 			StubOutputPath = Executable.Location;
 		}
 	}
@@ -1013,27 +1023,125 @@ namespace UnrealBuildTool
 				return 0;
 			}
 
+			string ExtraOptions = "";
 			// for mobile builds (which need real codesigning to be able to run), we use dummy codesigning when making a .stub (which will
 			// be sent to Windows an re-codesigned), or when making UnrealGame.app without a .uproject (we will always make a .app
 			// again with a .uproject on the commandline to be able to get the staged data that will be pulled in to the app - we can't run 
 			// without a Staged directory, so this dummy codesigned .app won't be used directly)
-			bool bUseDummySigning = Target.Platform != UnrealTargetPlatform.Mac && (Target.bCreateStubIPA || Target.ProjectFile == null);
+
+			// NOTE: Actually for _now_ we are using legacy-style signing with temp keychain because IPhonePackager cannot codesign
+			// Frameworks, so we have to do full non-dummy signing of stubs until we get IPP working
+			bool bUseDummySigning = Target.Platform != UnrealTargetPlatform.Mac && (/*Target.bCreateStubIPA ||*/ Target.ProjectFile == null);
+			bool bCreateStub = Target.Platform.IsInGroup(UnrealPlatformGroup.IOS) && Target.bCreateStubIPA;
+			bool bUseLegacyStubSigning = bCreateStub && true;
+
+			if (bUseLegacyStubSigning)
+			{
+				// create and run a script that will make a temp keychain, and return the options needed to pass to xcodebuild to use it
+				ExtraOptions += SetupRemoteCodesigning(Target);
+			}
 
 			int ExitCode = AppleExports.BuildWithStubXcodeProject(Target.ProjectFile, Target.Platform, Target.Architectures, Target.Configuration, Target.TargetName, 
-				AppleExports.XcodeBuildMode.PostBuildSync, Logger, "", bForceDummySigning: bUseDummySigning);
+				AppleExports.XcodeBuildMode.PostBuildSync, Logger, ExtraOptions, bForceDummySigning: bUseDummySigning);
+
+			// restore the keychain as soon as possible
+			if (bUseLegacyStubSigning)
+			{
+				// cleanup the 
+				CleanupRemoteCodesigning(Target);
+			}
+
 			if (ExitCode != 0)
 			{
 				Logger.LogError("ERROR: Failed to finalize the .app with Xcode. Check the log for more information");
 			}
 
-			if (Target.Platform.IsInGroup(UnrealPlatformGroup.IOS) && Target.bCreateStubIPA)
+			if (bCreateStub)
 			{
-				IOSToolChain.PackageStub(Target.StubOutputPath.Directory.FullName, Target.TargetName, Target.StubOutputPath.GetFileNameWithoutExtension(), true);
+				IOSToolChain.PackageStub(Target.StubOutputPath.Directory.FullName, Target.TargetName, Target.StubOutputPath.GetFileNameWithoutExtension(), true, !bUseLegacyStubSigning);
+
 			}
 
 			return ExitCode;
 		}
 
+		// This is hopefully until we can get codesigning of Frameworks working in IPhonePackager, then we can go back to dummy codesigning, without needing
+		// to mess with keychains and what not
+		private static string SetupRemoteCodesigning(ApplePostBuildSyncTarget Target)
+		{
+			FileReference TempKeychain = FileReference.Combine(Target.ProjectIntermediateDirectory!, "TempKeychain.keychain");
+			FileReference SignProjectScript = FileReference.Combine(Target.ProjectIntermediateDirectory!, "SignProject.sh");
+			string MobileProvisionUUID = "";
+			string SigningCertificate = "";
+
+			using (StreamWriter Writer = new StreamWriter(SignProjectScript.FullName))
+			{
+				// Boilerplate
+				Writer.WriteLine("#!/bin/sh");
+				Writer.WriteLine("set -e");
+				Writer.WriteLine("set -x");
+				// Copy the mobile provision into the system store
+				if (Target.RemoteImportProvision == null || Target.RemoteImportCertificate == null)
+				{
+					throw new BuildException("Expecting stub to be run with -ImportCertificate and -ImportProvision when using modern xcode");
+				}
+
+				// copy the provision into standard location
+				Writer.WriteLine("cp -f {0} ~/Library/MobileDevice/Provisioning\\ Profiles/", Utils.EscapeShellArgument(Target.RemoteImportProvision));
+				MobileProvisionContents MobileProvision = MobileProvisionContents.Read(new FileReference(Target.RemoteImportProvision));
+				MobileProvisionUUID = MobileProvision.GetUniqueId();
+
+				// Get the signing certificate to use
+				X509Certificate2 Certificate;
+				try
+				{
+					Certificate = new X509Certificate2(Target.RemoteImportCertificate, Target.RemoteImportCertificatePassword ?? "");
+				}
+				catch (Exception Ex)
+				{
+					throw new BuildException(Ex, "Unable to read certificate '{0}': {1}", Target.RemoteImportCertificate, Ex.Message);
+				}
+				// Read the name from the certificate
+				SigningCertificate = Certificate.GetNameInfo(X509NameType.SimpleName, false);
+
+				// Install a certificate given on the command line to a temporary keychain
+				Writer.WriteLine("security delete-keychain \"{0}\" || true", TempKeychain);
+				Writer.WriteLine("security create-keychain -p \"A\" \"{0}\"", TempKeychain);
+				Writer.WriteLine("security list-keychains -s \"{0}\"", TempKeychain);
+				Writer.WriteLine("security list-keychains");
+				Writer.WriteLine("security set-keychain-settings -t 3600 -l  \"{0}\"", TempKeychain);
+				Writer.WriteLine("security -v unlock-keychain -p \"A\" \"{0}\"", TempKeychain);
+				Writer.WriteLine("security import {0} -P {1} -k \"{2}\" -T /usr/bin/codesign -T /usr/bin/security -t agg", Utils.EscapeShellArgument(Target.RemoteImportCertificate), Utils.EscapeShellArgument(Target.RemoteImportCertificatePassword!), TempKeychain);
+				Writer.WriteLine("security set-key-partition-list -S apple-tool:,apple:,codesign: -s -k \"A\" -D '{0}' -t private {1}", SigningCertificate, TempKeychain);
+			}
+
+			// run the script
+			Utils.RunLocalProcessAndReturnStdOut("sh", $"\"{SignProjectScript.FullName}\"");
+
+			// Set parameters to make sure it uses the correct identity and keychain
+			// pass back the comandline arguments to xcodebuild to use these certicicates
+			return $" CODE_SIGN_STYLE=Manual CODE_SIGN_IDENTITY=\"{SigningCertificate}\" PROVISIONING_PROFILE_SPECIFIER={MobileProvisionUUID}";
+		}
+
+		private static void CleanupRemoteCodesigning(ApplePostBuildSyncTarget Target)
+		{
+			FileReference TempKeychain = FileReference.Combine(Target.ProjectIntermediateDirectory!, "TempKeychain.keychain");
+
+			FileReference CleanProjectScript = FileReference.Combine(Target.ProjectIntermediateDirectory!, "CleanProject.sh");
+			using (StreamWriter CleanWriter = new StreamWriter(CleanProjectScript.FullName))
+			{
+				CleanWriter.WriteLine("#!/bin/sh");
+				CleanWriter.WriteLine("set -e");
+				CleanWriter.WriteLine("set -x");
+				// Remove the temporary keychain from the search list
+				CleanWriter.WriteLine("security delete-keychain \"{0}\" || true", TempKeychain);
+				// Restore the login keychain as active
+				CleanWriter.WriteLine("security list-keychain -s login.keychain");
+			}
+
+			Utils.RunLocalProcessAndReturnStdOut("sh", $"\"{CleanProjectScript.FullName}\"");
+		}
+
 		private static FileItem GetPostBuildOutputFile(FileReference Executable, string TargetName, UnrealTargetPlatform Platform)
 		{
 			FileReference StagedExe;
@@ -1050,7 +1158,7 @@ namespace UnrealBuildTool
 
 		public static Action CreatePostBuildSyncAction(ReadOnlyTargetRules Target, FileItem Executable, DirectoryReference IntermediateDir, IActionGraphBuilder Graph)
 		{
-			ApplePostBuildSyncTarget PostBuildSync = new(Target, Executable);
+			ApplePostBuildSyncTarget PostBuildSync = new(Target, Executable, IntermediateDir);
 			FileReference PostBuildSyncFile = FileReference.Combine(IntermediateDir!, "PostBuildSync.dat");
 			BinaryFormatterUtils.Save(PostBuildSyncFile, PostBuildSync);
 
