UnrealBuildTool: Add WindowsTargetRule setting for PDB page size

#rnx

diff --git a/Engine/Source/Programs/UnrealBuildTool/Platform/Windows/UEBuildWindows.cs b/Engine/Source/Programs/UnrealBuildTool/Platform/Windows/UEBuildWindows.cs
--- a/Engine/Source/Programs/UnrealBuildTool/Platform/Windows/UEBuildWindows.cs
+++ b/Engine/Source/Programs/UnrealBuildTool/Platform/Windows/UEBuildWindows.cs
@@ -451,6 +451,13 @@ namespace UnrealBuildTool
 		[CommandLine("-StripPrivateSymbols")]
 		public bool bStripPrivateSymbols = false;
 
+		/// <summary>
+		/// Set page size to allow for larger than 4GB PDBs to be generated by the msvc linker.
+		/// Will default to 16384 for monolithic editor builds.
+		/// Values should be a power of two such as 4096, 8192, 16384, or 32768
+		/// </summary>
+		public uint? PdbPageSize = null;
+
 		/// <summary>
 		/// Specify an alternate location for the PDB file. This option does not change the location of the generated PDB file,
 		/// it changes the name that is embedded into the executable. Path can contain %_PDB% which will be expanded to the original
@@ -753,6 +760,8 @@ namespace UnrealBuildTool
 
 		public bool bStripPrivateSymbols => Inner.bStripPrivateSymbols;
 
+		public uint? PdbPageSize => Inner.PdbPageSize;
+
 		public string? PdbAlternatePath => Inner.PdbAlternatePath;
 
 		public int DefaultStackSize => Inner.DefaultStackSize;
@@ -1033,6 +1042,14 @@ namespace UnrealBuildTool
 				throw new BuildException("bAdaptiveUnityEnablesEditAndContinue requires bAdaptiveUnityDisablesPCH or bAdaptiveUnityCreatesDedicatedPCH");
 			}
 
+			// for monolithic editor builds, add the PDBPAGESIZE option, (VS 16.11, VC toolchain 14.29.30133), but the pdb will be too large without this
+			// some monolithic game builds could be too large as well, but they can be added in a .Target.cs with:
+			// TargetRules.WindowsPlatform.PdbPageSize = 8192;
+			if (!Target.WindowsPlatform.PdbPageSize.HasValue && Target.LinkType == TargetLinkType.Monolithic && Target.Type == TargetType.Editor)
+			{
+				Target.WindowsPlatform.PdbPageSize = 16384;
+			}
+
 			// If we're using PDB files and PCHs, the generated code needs to be compiled with the same options as the PCH.
 			if ((Target.bUsePDBFiles || Target.bSupportEditAndContinue) && Target.bUsePCHFiles)
 			{
diff --git a/Engine/Source/Programs/UnrealBuildTool/Platform/Windows/VCToolChain.cs b/Engine/Source/Programs/UnrealBuildTool/Platform/Windows/VCToolChain.cs
--- a/Engine/Source/Programs/UnrealBuildTool/Platform/Windows/VCToolChain.cs
+++ b/Engine/Source/Programs/UnrealBuildTool/Platform/Windows/VCToolChain.cs
@@ -1227,12 +1227,10 @@ namespace UnrealBuildTool
 				Arguments.Add("/fastfail");
 			}
 
-			// for monolithic editor builds, add the PDBPAGESIZE option, (VS 16.11, VC toolchain 14.29.30133), but the pdb will be too large without this
-			// some monolithic game builds could be too large as well, but they can be added in a .Target.cs with:
-			//   			WindowsPlatform.AdditionalLinkerOptions = "/PDBPAGESIZE:8192";
-			if (Target.LinkType == TargetLinkType.Monolithic && Target.Type == TargetType.Editor)
+			// Allow for PDBs larger than 4GB
+			if (Target.WindowsPlatform.PdbPageSize.HasValue)
 			{
-				Arguments.Add("/PDBPAGESIZE:16384");
+				Arguments.Add($"/PDBPAGESIZE:{System.Numerics.BitOperations.RoundUpToPowerOf2(Target.WindowsPlatform.PdbPageSize.Value)}");
 				//Arguments.Add("/PDBCompress"); // Do not turn this on, it makes link times almost 2x slower. This is _only_ to save local disk space. Will _not_ make actual file smaller for network transfer
 			}
 
