Make CPP files on windows depend on the AutoRTFM compiler version if we are using it.

diff --git a/Engine/Source/Programs/UnrealBuildTool/Platform/Windows/VCToolChain.cs b/Engine/Source/Programs/UnrealBuildTool/Platform/Windows/VCToolChain.cs
--- a/Engine/Source/Programs/UnrealBuildTool/Platform/Windows/VCToolChain.cs
+++ b/Engine/Source/Programs/UnrealBuildTool/Platform/Windows/VCToolChain.cs
@@ -1601,6 +1601,18 @@ namespace UnrealBuildTool
 						{
 							CompileAction.Arguments.Add("-fautortfm");
 						}
+
+						// If we are using the AutoRTFM compiler, we make the compile action depend on the version of the compiler itself.
+						// This lets us update the compiler (which might not cause a version update of the compiler, which instead tracks
+						// the LLVM versioning scheme that Clang uses), but ensure that we rebuild the source if the compiler has changed.
+						if (CompileEnvironment.bUseAutoRTFMCompiler)
+						{
+							FileReference? CompilerPath = GetCppCompilerPath();
+							if (null != CompilerPath)
+							{
+								BaseCompileAction.AdditionalPrerequisiteItems.Add(FileItem.GetItemByFileReference(CompilerPath));
+							}
+						}
 					}
 				}
 				else
diff --git a/Engine/Source/Programs/UnrealBuildTool/System/CppCompileEnvironment.cs b/Engine/Source/Programs/UnrealBuildTool/System/CppCompileEnvironment.cs
--- a/Engine/Source/Programs/UnrealBuildTool/System/CppCompileEnvironment.cs
+++ b/Engine/Source/Programs/UnrealBuildTool/System/CppCompileEnvironment.cs
@@ -728,6 +728,7 @@ namespace UnrealBuildTool
 			StaticAnalyzerDisabledCheckers = new HashSet<string>(Other.StaticAnalyzerDisabledCheckers);
 			StaticAnalyzerAdditionalCheckers = new HashSet<string>(Other.StaticAnalyzerAdditionalCheckers);
 			bOptimizeCode = Other.bOptimizeCode;
+			bUseAutoRTFMCompiler = Other.bUseAutoRTFMCompiler;
 			bAllowAutoRTFMInstrumentation = Other.bAllowAutoRTFMInstrumentation;
 			bCodeCoverage = Other.bCodeCoverage;
 			OptimizationLevel = Other.OptimizationLevel;
