Turning on FormatStringSan in UBT

diff --git a/Engine/Source/Editor/GameProjectGeneration/Private/GameProjectUtils.cpp b/Engine/Source/Editor/GameProjectGeneration/Private/GameProjectUtils.cpp
--- a/Engine/Source/Editor/GameProjectGeneration/Private/GameProjectUtils.cpp
+++ b/Engine/Source/Editor/GameProjectGeneration/Private/GameProjectUtils.cpp
@@ -2963,7 +2963,7 @@ bool GameProjectUtils::UpdateAdditionalPluginDirectory(const FString& InDir, con
 
 const TCHAR* GameProjectUtils::GetDefaultBuildSettingsVersion()
 {
-	return TEXT("BuildSettingsVersion.V4");
+	return TEXT("BuildSettingsVersion.V5");
 }
 
 bool GameProjectUtils::ReadTemplateFile(const FString& TemplateFileName, FString& OutFileContents, FText& OutFailReason)
diff --git a/Engine/Source/Programs/UnrealBuildTool/Configuration/TargetRules.cs b/Engine/Source/Programs/UnrealBuildTool/Configuration/TargetRules.cs
--- a/Engine/Source/Programs/UnrealBuildTool/Configuration/TargetRules.cs
+++ b/Engine/Source/Programs/UnrealBuildTool/Configuration/TargetRules.cs
@@ -139,12 +139,18 @@ namespace UnrealBuildTool
 		/// </summary>
 		V4,
 
+		/// <summary>
+		/// New defaults for 5.4:
+		/// * TargetRules.bValidateFormatStrings = true
+		/// </summary>
+		V5,
+
 		// *** When adding new entries here, be sure to update GameProjectUtils::GetDefaultBuildSettingsVersion() to ensure that new projects are created correctly. ***
 
 		/// <summary>
 		/// Always use the defaults for the current engine version. Note that this may cause compatibility issues when upgrading.
 		/// </summary>
-		Latest = V4,
+		Latest = V5,
 	}
 
 	/// <summary>
@@ -1628,6 +1634,18 @@ namespace UnrealBuildTool
 		}
 		private bool? bRequireObjectPtrForAddReferencedObjectsPrivate;		
 
+		/// <summary>
+		/// Emits compilation errors for incorrect UE_LOG format strings.
+		/// </summary>
+		[RequiresUniqueBuildEnvironment]
+		[XmlConfigFile(Category = "BuildConfiguration")]		
+		public bool bValidateFormatStrings
+		{
+				get { return bValidateFormatStringsPrivate ?? (DefaultBuildSettings >= BuildSettingsVersion.V5); }
+				set { bValidateFormatStringsPrivate = value; }
+		}
+		private bool? bValidateFormatStringsPrivate;
+			
 		/// <summary>
 		/// Level to report deprecation warnings as errors
 		/// </summary>
@@ -2910,6 +2928,11 @@ namespace UnrealBuildTool
 					ModifiedSettings.Add(Tuple.Create(String.Format("{0}.{1} = true", nameof(WindowsPlatform), nameof(WindowsPlatform.bStrictConformanceMode)), "Updates MSVC strict conformance mode to true (Previously: false)."));
 				}
 
+				if (BuildSettingsVersion.V5 <= LatestVersion && DefaultBuildSettings < BuildSettingsVersion.V5)
+				{
+					ModifiedSettings.Add(Tuple.Create(String.Format("{0} = true", nameof(bValidateFormatStrings)), "Enables compile-time validation of strings+args passed to UE_LOG. (Previously: false)."));
+				}
+
 				if (ModifiedSettings.Count > 0)
 				{
 					string FormatString = String.Format("[Upgrade]     {{0,-{0}}}   => {{1}}", ModifiedSettings.Max(x => x.Item1.Length));
diff --git a/Engine/Source/Programs/UnrealBuildTool/System/RulesAssembly.cs b/Engine/Source/Programs/UnrealBuildTool/System/RulesAssembly.cs
--- a/Engine/Source/Programs/UnrealBuildTool/System/RulesAssembly.cs
+++ b/Engine/Source/Programs/UnrealBuildTool/System/RulesAssembly.cs
@@ -683,6 +683,11 @@ namespace UnrealBuildTool
 				Rules.GlobalDefinitions.Add("UE_REFERENCE_COLLECTOR_REQUIRE_OBJECTPTR=1");
 			}
 
+			if (Rules.bValidateFormatStrings)
+			{
+				Rules.GlobalDefinitions.Add("UE_VALIDATE_FORMAT_STRINGS=1");
+			}
+			
 			// if the Target has opted in only some platforms, disable any plugins of other platforms (there may be editor, etc, modules that
 			// will just add themselves, with no other reference to be able to remove them, other than disabling them here)
 			if (Rules.OptedInModulePlatforms != null)
