- Fixed codesigning of UnrealGame made with -createstub (will use dummy signing, which doesn't need a cert/provision pair, so will be signed the same way UnrealGame is signed without making a stub)

diff --git a/Engine/Source/Programs/UnrealBuildTool/Platform/IOS/IOSToolChain.cs b/Engine/Source/Programs/UnrealBuildTool/Platform/IOS/IOSToolChain.cs
--- a/Engine/Source/Programs/UnrealBuildTool/Platform/IOS/IOSToolChain.cs
+++ b/Engine/Source/Programs/UnrealBuildTool/Platform/IOS/IOSToolChain.cs
@@ -756,12 +756,16 @@ namespace UnrealBuildTool
 					throw new BuildException($"Failed to codesign with dummy identity, which should never happen [ExitCode = {ExitCode}, output = {Output}");
 				}
 				// also codesign the frameworks
-				foreach (string FrameworkDir in Directory.EnumerateDirectories(Path.Combine(ZipSourceDir, "Frameworks")))
+				string FrameworksDir = Path.Combine(ZipSourceDir, "Frameworks");
+				if (Directory.Exists(FrameworksDir))
 				{
-					Output = Utils.RunLocalProcessAndReturnStdOut("/bin/sh", $"-c 'codesign -f -s - \"{FrameworkDir}\"", null, out ExitCode);
-					if (ExitCode != 0)
+					foreach (string FrameworkDir in Directory.EnumerateDirectories(FrameworksDir))
 					{
-						throw new BuildException($"Failed to codesign with dummy identity, which should never happen [ExitCode = {ExitCode}, output = {Output}");
+						Output = Utils.RunLocalProcessAndReturnStdOut("/bin/sh", $"-c 'codesign -f -s - \"{FrameworkDir}\"", null, out ExitCode);
+						if (ExitCode != 0)
+						{
+							throw new BuildException($"Failed to codesign with dummy identity, which should never happen [ExitCode = {ExitCode}, output = {Output}");
+						}
 					}
 				}
 			}
diff --git a/Engine/Source/Programs/UnrealBuildTool/ToolChain/AppleToolChain.cs b/Engine/Source/Programs/UnrealBuildTool/ToolChain/AppleToolChain.cs
--- a/Engine/Source/Programs/UnrealBuildTool/ToolChain/AppleToolChain.cs
+++ b/Engine/Source/Programs/UnrealBuildTool/ToolChain/AppleToolChain.cs
@@ -1031,9 +1031,10 @@ namespace UnrealBuildTool
 
 			// NOTE: Actually for _now_ we are using legacy-style signing with temp keychain because IPhonePackager cannot codesign
 			// Frameworks, so we have to do full non-dummy signing of stubs until we get IPP working
-			bool bUseDummySigning = Target.Platform != UnrealTargetPlatform.Mac && (/*Target.bCreateStubIPA ||*/ Target.ProjectFile == null);
+			bool bUseDummySigning = Target.Platform != UnrealTargetPlatform.Mac && (Target.ProjectFile == null);
 			bool bCreateStub = Target.Platform.IsInGroup(UnrealPlatformGroup.IOS) && Target.bCreateStubIPA;
-			bool bUseLegacyStubSigning = bCreateStub && true;
+			// if we want dummy signing (no project at all) then we don't want to use legacy signing - that is only used for remote builds from Windows
+			bool bUseLegacyStubSigning = bCreateStub && !bUseDummySigning;
 
 			if (bUseLegacyStubSigning)
 			{
