GetMainVersion now requires a registered platform, which only can be done from the Linux SDK that was first registered. Current StripSymbols was making a new LinuxPlatformSDK which wont be registered and fails getting the MainVersion

Now we will use the existing LinuxPlatfromSDK or if doesn't exist make a new one and correctly register it

diff --git a/Engine/Source/Programs/UnrealBuildTool/Platform/Linux/LinuxExports.cs b/Engine/Source/Programs/UnrealBuildTool/Platform/Linux/LinuxExports.cs
--- a/Engine/Source/Programs/UnrealBuildTool/Platform/Linux/LinuxExports.cs
+++ b/Engine/Source/Programs/UnrealBuildTool/Platform/Linux/LinuxExports.cs
@@ -18,7 +18,15 @@ namespace UnrealBuildTool
 		/// <param name="Logger">Logger for output</param>
 		public static void StripSymbols(FileReference SourceFile, FileReference TargetFile, ILogger Logger)
 		{
-			LinuxToolChain ToolChain = new LinuxToolChain(LinuxPlatform.DefaultHostArchitecture, (LinuxPlatformSDK)UEBuildPlatform.GetSDK(UnrealTargetPlatform.Linux)!, ClangToolChainOptions.None, Logger);
+			LinuxPlatformSDK? LinuxSDK = UEBuildPlatformSDK.GetSDKForPlatform("Linux") as LinuxPlatformSDK;
+
+			if (LinuxSDK == null)
+			{
+				LinuxSDK = new LinuxPlatformSDK(Logger);
+				UEBuildPlatformSDK.RegisterSDKForPlatform(LinuxSDK, "Linux", true);
+			}
+
+			LinuxToolChain ToolChain = new LinuxToolChain(LinuxPlatform.DefaultHostArchitecture, LinuxSDK, ClangToolChainOptions.None, Logger);
 			ToolChain.StripSymbols(SourceFile, TargetFile, Logger);
 		}
 	}
