Clang static analysis requires unity builds to be disabled

#rnx

diff --git a/Engine/Source/Programs/UnrealBuildTool/Platform/Android/UEBuildAndroid.cs b/Engine/Source/Programs/UnrealBuildTool/Platform/Android/UEBuildAndroid.cs
--- a/Engine/Source/Programs/UnrealBuildTool/Platform/Android/UEBuildAndroid.cs
+++ b/Engine/Source/Programs/UnrealBuildTool/Platform/Android/UEBuildAndroid.cs
@@ -186,7 +186,7 @@ namespace UnrealBuildTool
 			{
 				Target.StaticAnalyzer = StaticAnalyzer.Default;
 				Target.StaticAnalyzerOutputType = (Environment.GetEnvironmentVariable("CLANG_ANALYZER_OUTPUT")?.Contains("html", StringComparison.OrdinalIgnoreCase) == true) ? StaticAnalyzerOutputType.Html : StaticAnalyzerOutputType.Text;
-				Target.StaticAnalyzerMode = String.Equals(Environment.GetEnvironmentVariable("CLANG_STATIC_ANALYZER_MODE"), "shallow") ? StaticAnalyzerMode.Shallow : StaticAnalyzerMode.Deep;
+				Target.StaticAnalyzerMode = String.Equals(Environment.GetEnvironmentVariable("CLANG_STATIC_ANALYZER_MODE"), "shallow", StringComparison.OrdinalIgnoreCase) ? StaticAnalyzerMode.Shallow : StaticAnalyzerMode.Deep;
 			}
 			else if (Target.StaticAnalyzer == StaticAnalyzer.Clang)
 			{
@@ -198,6 +198,9 @@ namespace UnrealBuildTool
 			{
 				Target.bDisableLinking = true;
 				Target.bIgnoreBuildOutputs = true;
+
+				// Clang static analysis requires non unity builds
+				Target.bUseUnityBuild = false;
 			}
 
 			Target.bCompileRecast = true;
diff --git a/Engine/Source/Programs/UnrealBuildTool/Platform/IOS/UEBuildIOS.cs b/Engine/Source/Programs/UnrealBuildTool/Platform/IOS/UEBuildIOS.cs
--- a/Engine/Source/Programs/UnrealBuildTool/Platform/IOS/UEBuildIOS.cs
+++ b/Engine/Source/Programs/UnrealBuildTool/Platform/IOS/UEBuildIOS.cs
@@ -716,7 +716,7 @@ namespace UnrealBuildTool
 			{
 				Target.StaticAnalyzer = StaticAnalyzer.Default;
 				Target.StaticAnalyzerOutputType = (Environment.GetEnvironmentVariable("CLANG_ANALYZER_OUTPUT")?.Contains("html", StringComparison.OrdinalIgnoreCase) == true) ? StaticAnalyzerOutputType.Html : StaticAnalyzerOutputType.Text;
-				Target.StaticAnalyzerMode = String.Equals(Environment.GetEnvironmentVariable("CLANG_STATIC_ANALYZER_MODE"), "shallow") ? StaticAnalyzerMode.Shallow : StaticAnalyzerMode.Deep;
+				Target.StaticAnalyzerMode = String.Equals(Environment.GetEnvironmentVariable("CLANG_STATIC_ANALYZER_MODE"), "shallow", StringComparison.OrdinalIgnoreCase) ? StaticAnalyzerMode.Shallow : StaticAnalyzerMode.Deep;
 			}
 			else if (Target.StaticAnalyzer == StaticAnalyzer.Clang)
 			{
@@ -728,6 +728,9 @@ namespace UnrealBuildTool
 			{
 				Target.bDisableLinking = true;
 				Target.bIgnoreBuildOutputs = true;
+
+				// Clang static analysis requires non unity builds
+				Target.bUseUnityBuild = false;
 			}
 
 			// we assume now we are building with IOS8 or later
diff --git a/Engine/Source/Programs/UnrealBuildTool/Platform/Linux/UEBuildLinux.cs b/Engine/Source/Programs/UnrealBuildTool/Platform/Linux/UEBuildLinux.cs
--- a/Engine/Source/Programs/UnrealBuildTool/Platform/Linux/UEBuildLinux.cs
+++ b/Engine/Source/Programs/UnrealBuildTool/Platform/Linux/UEBuildLinux.cs
@@ -212,7 +212,7 @@ namespace UnrealBuildTool
 			{
 				Target.StaticAnalyzer = StaticAnalyzer.Default;
 				Target.StaticAnalyzerOutputType = (Environment.GetEnvironmentVariable("CLANG_ANALYZER_OUTPUT")?.Contains("html", StringComparison.OrdinalIgnoreCase) == true) ? StaticAnalyzerOutputType.Html : StaticAnalyzerOutputType.Text;
-				Target.StaticAnalyzerMode = String.Equals(Environment.GetEnvironmentVariable("CLANG_STATIC_ANALYZER_MODE"), "shallow") ? StaticAnalyzerMode.Shallow : StaticAnalyzerMode.Deep;
+				Target.StaticAnalyzerMode = String.Equals(Environment.GetEnvironmentVariable("CLANG_STATIC_ANALYZER_MODE"), "shallow", StringComparison.OrdinalIgnoreCase) ? StaticAnalyzerMode.Shallow : StaticAnalyzerMode.Deep;
 			}
 			else if (Target.StaticAnalyzer == StaticAnalyzer.Clang)
 			{
@@ -224,6 +224,9 @@ namespace UnrealBuildTool
 			{
 				Target.bDisableLinking = true;
 				Target.bIgnoreBuildOutputs = true;
+
+				// Clang static analysis requires non unity builds
+				Target.bUseUnityBuild = false;
 			}
 
 			// Editor target types get overwritten in UEBuildTarget.cs so lets avoid adding this here. ResetTarget is called with
diff --git a/Engine/Source/Programs/UnrealBuildTool/Platform/Mac/UEBuildMac.cs b/Engine/Source/Programs/UnrealBuildTool/Platform/Mac/UEBuildMac.cs
--- a/Engine/Source/Programs/UnrealBuildTool/Platform/Mac/UEBuildMac.cs
+++ b/Engine/Source/Programs/UnrealBuildTool/Platform/Mac/UEBuildMac.cs
@@ -299,7 +299,7 @@ namespace UnrealBuildTool
 			{
 				Target.StaticAnalyzer = StaticAnalyzer.Default;
 				Target.StaticAnalyzerOutputType = (Environment.GetEnvironmentVariable("CLANG_ANALYZER_OUTPUT")?.Contains("html", StringComparison.OrdinalIgnoreCase) == true) ? StaticAnalyzerOutputType.Html : StaticAnalyzerOutputType.Text;
-				Target.StaticAnalyzerMode = String.Equals(Environment.GetEnvironmentVariable("CLANG_STATIC_ANALYZER_MODE"), "shallow") ? StaticAnalyzerMode.Shallow : StaticAnalyzerMode.Deep;
+				Target.StaticAnalyzerMode = String.Equals(Environment.GetEnvironmentVariable("CLANG_STATIC_ANALYZER_MODE"), "shallow", StringComparison.OrdinalIgnoreCase) ? StaticAnalyzerMode.Shallow : StaticAnalyzerMode.Deep;
 			}
 			else if (Target.StaticAnalyzer == StaticAnalyzer.Clang)
 			{
@@ -311,6 +311,10 @@ namespace UnrealBuildTool
 			{
 				Target.bDisableLinking = true;
 				Target.bIgnoreBuildOutputs = true;
+
+				// Clang static analysis requires non unity builds
+				Target.bUseUnityBuild = false;
+
 				// Disable chaining PCHs for the moment because it is crashing clang
 				Target.bChainPCHs = false;
 			}
diff --git a/Engine/Source/Programs/UnrealBuildTool/Platform/Windows/UEBuildWindows.cs b/Engine/Source/Programs/UnrealBuildTool/Platform/Windows/UEBuildWindows.cs
--- a/Engine/Source/Programs/UnrealBuildTool/Platform/Windows/UEBuildWindows.cs
+++ b/Engine/Source/Programs/UnrealBuildTool/Platform/Windows/UEBuildWindows.cs
@@ -995,6 +995,9 @@ namespace UnrealBuildTool
 					Target.WindowsPlatform.Compiler = WindowsCompiler.Clang;
 				}
 				Target.StaticAnalyzer = StaticAnalyzer.Default;
+
+				// Clang static analysis requires non unity builds
+				Target.bUseUnityBuild = false;
 			}
 			else if (Target.StaticAnalyzer != StaticAnalyzer.None &&
 					 Target.StaticAnalyzerOutputType != StaticAnalyzerOutputType.Text)
