UnrealBuildTool: Always enable ValidateFormatStrings for engine modules, and create a pch variant for modules that opt out

#rnx

diff --git a/Engine/Source/Programs/UnrealBuildTool/Configuration/ModuleRules.cs b/Engine/Source/Programs/UnrealBuildTool/Configuration/ModuleRules.cs
--- a/Engine/Source/Programs/UnrealBuildTool/Configuration/ModuleRules.cs
+++ b/Engine/Source/Programs/UnrealBuildTool/Configuration/ModuleRules.cs
@@ -780,6 +780,16 @@ namespace UnrealBuildTool
 		/// </summary>
 		public bool bTreatAsEngineModule;
 
+		/// <summary>
+		/// Emits compilation errors for incorrect UE_LOG format strings.
+		/// </summary>
+		public bool bValidateFormatStrings
+		{
+			get => bValidateFormatStringsPrivate ?? (bTreatAsEngineModule || Target.bValidateFormatStrings);
+			set => bValidateFormatStringsPrivate = value;
+		}
+		private bool? bValidateFormatStringsPrivate;
+
 		/// <summary>
 		/// Which engine version's build settings to use by default. 
 		/// </summary>
diff --git a/Engine/Source/Programs/UnrealBuildTool/Configuration/TargetRules.cs b/Engine/Source/Programs/UnrealBuildTool/Configuration/TargetRules.cs
--- a/Engine/Source/Programs/UnrealBuildTool/Configuration/TargetRules.cs
+++ b/Engine/Source/Programs/UnrealBuildTool/Configuration/TargetRules.cs
@@ -1637,7 +1637,6 @@ namespace UnrealBuildTool
 		/// <summary>
 		/// Emits compilation errors for incorrect UE_LOG format strings.
 		/// </summary>
-		[RequiresUniqueBuildEnvironment]
 		[XmlConfigFile(Category = "BuildConfiguration")]		
 		public bool bValidateFormatStrings
 		{
@@ -3275,6 +3274,8 @@ namespace UnrealBuildTool
 
 		public int MinGameModuleSourceFilesForUnityBuild => Inner.MinGameModuleSourceFilesForUnityBuild;
 
+		public bool bValidateFormatStrings => Inner.bValidateFormatStrings;
+
 		public WarningLevel DefaultWarningLevel => Inner.DefaultWarningLevel;
 
 		public WarningLevel DeprecationWarningLevel => Inner.DeprecationWarningLevel;
diff --git a/Engine/Source/Programs/UnrealBuildTool/Configuration/UEBuildModuleCPP.cs b/Engine/Source/Programs/UnrealBuildTool/Configuration/UEBuildModuleCPP.cs
--- a/Engine/Source/Programs/UnrealBuildTool/Configuration/UEBuildModuleCPP.cs
+++ b/Engine/Source/Programs/UnrealBuildTool/Configuration/UEBuildModuleCPP.cs
@@ -1022,6 +1022,7 @@ namespace UnrealBuildTool
 			foreach (string Definition in Definitions)
 			{
 				if (Definition.Contains("UE_IS_ENGINE_MODULE") ||
+					Definition.Contains("UE_VALIDATE_FORMAT_STRINGS") ||
 					Definition.Contains("DEPRECATED_FORGAME") ||
 					Definition.Contains("UE_DEPRECATED_FORGAME") ||
 					Definition.Contains("UE_ENABLE_INCLUDE_ORDER_DEPRECATED_IN_"))
@@ -1096,6 +1097,13 @@ namespace UnrealBuildTool
 					Definitions.AddRange(EngineIncludeOrderHelper.GetDeprecationDefines(ModuleCompileEnvironment.IncludeOrderVersion));
 				}
 
+				// Modify definitions if we need to create a new shared pch for validating format strings
+				if (ModuleCompileEnvironment.bValidateFormatStrings != Template.BaseCompileEnvironment.bValidateFormatStrings)
+				{
+					Definitions.RemoveAll(x => x.Contains("UE_VALIDATE_FORMAT_STRINGS"));
+					Definitions.Add($"UE_VALIDATE_FORMAT_STRINGS={(ModuleCompileEnvironment.bValidateFormatStrings ? "1" : "0")}");
+				}
+
 				// Create a suffix to distinguish this shared PCH variant from any others. Currently only optimized and non-optimized shared PCHs are supported.
 				string Variant = GetSuffixForSharedPCH(ModuleCompileEnvironment, Template.BaseCompileEnvironment);
 
@@ -1219,6 +1227,11 @@ namespace UnrealBuildTool
 			{
 				return false;
 			}
+			
+			if (ModuleCompileEnvironment.bValidateFormatStrings != CompileEnvironment.bValidateFormatStrings)
+			{
+				return false;
+			}
 			return true;
 		}
 
@@ -1264,6 +1277,17 @@ namespace UnrealBuildTool
 					Variant += ".NoExceptions";
 				}
 			}
+			if (CompileEnvironment.bValidateFormatStrings != BaseCompileEnvironment.bValidateFormatStrings)
+			{
+				if (CompileEnvironment.bValidateFormatStrings)
+				{
+					Variant += ".ValFmtStr";
+				}
+				else
+				{
+					Variant += ".NoValFmtStr";
+				}
+			}
 
 			if (CompileEnvironment.bDeterministic != BaseCompileEnvironment.bDeterministic)
 			{
@@ -1325,6 +1349,7 @@ namespace UnrealBuildTool
 			CompileEnvironment.CppStandard = ModuleCompileEnvironment.CppStandard;
 			CompileEnvironment.CStandard = ModuleCompileEnvironment.CStandard;
 			CompileEnvironment.IncludeOrderVersion = ModuleCompileEnvironment.IncludeOrderVersion;
+			CompileEnvironment.bValidateFormatStrings = ModuleCompileEnvironment.bValidateFormatStrings;
 		}
 
 		/// <summary>
@@ -1561,6 +1586,7 @@ namespace UnrealBuildTool
 							if (!Rules.bTreatAsEngineModule)
 							{
 								Writer.AppendLine("#undef UE_IS_ENGINE_MODULE");
+								Writer.AppendLine("#undef UE_VALIDATE_FORMAT_STRINGS");
 								Writer.AppendLine("#undef DEPRECATED_FORGAME");
 								Writer.AppendLine("#define DEPRECATED_FORGAME DEPRECATED");
 								Writer.AppendLine("#undef UE_DEPRECATED_FORGAME");
@@ -1855,7 +1881,7 @@ namespace UnrealBuildTool
 			Result.bEnableUndefinedIdentifierWarnings = Rules.bEnableUndefinedIdentifierWarnings;
 			Result.IncludeOrderVersion = Rules.IncludeOrderVersion;
 			Result.DeterministicWarningLevel = Rules.DeterministicWarningLevel;
-
+			Result.bValidateFormatStrings = Rules.bValidateFormatStrings;
 			Result.bUseAutoRTFMCompiler = Target.bUseAutoRTFMCompiler;
 
 			// Only enable the AutoRTFM flag if we are using the AutoRTFM compiler
@@ -1945,6 +1971,8 @@ namespace UnrealBuildTool
 				Result.Definitions.Add("UE_IS_ENGINE_MODULE=0");
 			}
 
+			Result.Definitions.Add($"UE_VALIDATE_FORMAT_STRINGS={(Rules.bValidateFormatStrings ? "1" : "0")}");
+
 			Result.Definitions.AddRange(EngineIncludeOrderHelper.GetDeprecationDefines(Rules.IncludeOrderVersion));
 
 			// If module intermediate files are under the project folder we can set UE_PROJECT_NAME and friends.
@@ -1991,6 +2019,7 @@ namespace UnrealBuildTool
 			// Use the default optimization setting for
 			CompileEnvironment.bOptimizeCode = ShouldEnableOptimization(ModuleRules.CodeOptimization.Default, Target.Configuration, Rules.bTreatAsEngineModule, Rules.bCodeCoverage);
 			CompileEnvironment.bCodeCoverage = Rules.bCodeCoverage;
+			CompileEnvironment.bValidateFormatStrings = Rules.bValidateFormatStrings;
 
 			// Override compile environment
 			CompileEnvironment.bIsBuildingDLL = !Target.ShouldCompileMonolithic();
@@ -2006,6 +2035,8 @@ namespace UnrealBuildTool
 				CompileEnvironment.Definitions.Add("UE_IS_ENGINE_MODULE=0");
 			}
 
+			CompileEnvironment.Definitions.Add($"UE_VALIDATE_FORMAT_STRINGS={(Rules.bValidateFormatStrings ? "1" : "0")}");
+
 			CompileEnvironment.Definitions.AddRange(EngineIncludeOrderHelper.GetDeprecationDefines(Rules.IncludeOrderVersion));
 
 			// Add the module's private definitions.
diff --git a/Engine/Source/Programs/UnrealBuildTool/Configuration/UEBuildTarget.cs b/Engine/Source/Programs/UnrealBuildTool/Configuration/UEBuildTarget.cs
--- a/Engine/Source/Programs/UnrealBuildTool/Configuration/UEBuildTarget.cs
+++ b/Engine/Source/Programs/UnrealBuildTool/Configuration/UEBuildTarget.cs
@@ -4766,6 +4766,7 @@ namespace UnrealBuildTool
 			GlobalCompileEnvironment.bDeterministic = Rules.bDeterministic;
 			GlobalCompileEnvironment.CrashDiagnosticDirectory = Rules.CrashDiagnosticDirectory;
 			GlobalCompileEnvironment.bCodeCoverage = Rules.bCodeCoverage;
+			GlobalCompileEnvironment.bValidateFormatStrings = Rules.bValidateFormatStrings;
 
 			GlobalLinkEnvironment.bUseDebugCRT = bUseDebugCRT;
 			GlobalLinkEnvironment.bUseStaticCRT = Rules.bUseStaticCRT;
diff --git a/Engine/Source/Programs/UnrealBuildTool/System/CppCompileEnvironment.cs b/Engine/Source/Programs/UnrealBuildTool/System/CppCompileEnvironment.cs
--- a/Engine/Source/Programs/UnrealBuildTool/System/CppCompileEnvironment.cs
+++ b/Engine/Source/Programs/UnrealBuildTool/System/CppCompileEnvironment.cs
@@ -667,6 +667,11 @@ namespace UnrealBuildTool
 		/// </summary>
 		public WarningLevel DeterministicWarningLevel = WarningLevel.Off;
 
+		/// <summary>
+		/// Emits compilation errors for incorrect UE_LOG format strings.
+		/// </summary>
+		public bool bValidateFormatStrings = true;
+
 		/// <summary>
 		/// Directory where to put crash report files for platforms that support it
 		/// </summary>
@@ -779,6 +784,7 @@ namespace UnrealBuildTool
 			bDeterministic = Other.bDeterministic;
 			DeterministicWarningLevel = Other.DeterministicWarningLevel;
 			CrashDiagnosticDirectory = Other.CrashDiagnosticDirectory;
+			bValidateFormatStrings = Other.bValidateFormatStrings;
 		}
 
 		public CppCompileEnvironment(CppCompileEnvironment Other, UnrealArch OverrideArchitecture)
diff --git a/Engine/Source/Programs/UnrealBuildTool/System/RulesAssembly.cs b/Engine/Source/Programs/UnrealBuildTool/System/RulesAssembly.cs
--- a/Engine/Source/Programs/UnrealBuildTool/System/RulesAssembly.cs
+++ b/Engine/Source/Programs/UnrealBuildTool/System/RulesAssembly.cs
@@ -682,11 +682,6 @@ namespace UnrealBuildTool
 			{
 				Rules.GlobalDefinitions.Add("UE_REFERENCE_COLLECTOR_REQUIRE_OBJECTPTR=1");
 			}
-
-			if (Rules.bValidateFormatStrings)
-			{
-				Rules.GlobalDefinitions.Add("UE_VALIDATE_FORMAT_STRINGS=1");
-			}
 			
 			// if the Target has opted in only some platforms, disable any plugins of other platforms (there may be editor, etc, modules that
 			// will just add themselves, with no other reference to be able to remove them, other than disabling them here)
