- Building MacPlatformEditor and MacTargetPlatform so that the Xcode Projects settings show up in the Windows editor, for remote IOS building

diff --git a/Engine/Source/Developer/Mac/MacTargetPlatform/Private/MacTargetPlatformModule.cpp b/Engine/Source/Developer/Mac/MacTargetPlatform/Private/MacTargetPlatformModule.cpp
--- a/Engine/Source/Developer/Mac/MacTargetPlatform/Private/MacTargetPlatformModule.cpp
+++ b/Engine/Source/Developer/Mac/MacTargetPlatform/Private/MacTargetPlatformModule.cpp
@@ -31,6 +31,12 @@ public:
 		TargetPlatforms.Add(new TGenericMacTargetPlatform<false, true, false>());
 		// Client TP
 		TargetPlatforms.Add(new TGenericMacTargetPlatform<false, false, true>());
+
+#if PLATFORM_WINDOWS
+		// we added Mac to Windows so that the Xcode Project settings show up, but we don't
+		// want to see the Mac in the Platforms dropdown
+		FDataDrivenPlatformInfoRegistry::SetPlatformHiddenFromUI("Mac");
+#endif
 	}
 
 
diff --git a/Engine/Source/Programs/UnrealBuildTool/Platform/Mac/UEBuildMac.cs b/Engine/Source/Programs/UnrealBuildTool/Platform/Mac/UEBuildMac.cs
--- a/Engine/Source/Programs/UnrealBuildTool/Platform/Mac/UEBuildMac.cs
+++ b/Engine/Source/Programs/UnrealBuildTool/Platform/Mac/UEBuildMac.cs
@@ -455,6 +455,15 @@ namespace UnrealBuildTool
 			{
 				return;
 			}
+
+			if (Target.Platform == UnrealTargetPlatform.Win64 && Target.Type == TargetType.Editor)
+			{
+				// because remote IOS building needs the new XcodeProject Settings to show up in the editor, we bring in the Mac bits that expose it
+				if (ModuleName == "Engine")
+				{
+					Rules.DynamicallyLoadedModuleNames.AddAll("MacTargetPlatform", "MacPlatformEditor");
+				}
+			}
 		}
 
 		public override DirectoryReference? GetBundleDirectory(ReadOnlyTargetRules Rules, List<FileReference> OutputFiles)
