Fix the PCH rules when the AutoRTFM compiler is being used.

diff --git a/Engine/Source/Programs/UnrealBuildTool/Configuration/UEBuildModuleCPP.cs b/Engine/Source/Programs/UnrealBuildTool/Configuration/UEBuildModuleCPP.cs
--- a/Engine/Source/Programs/UnrealBuildTool/Configuration/UEBuildModuleCPP.cs
+++ b/Engine/Source/Programs/UnrealBuildTool/Configuration/UEBuildModuleCPP.cs
@@ -1325,6 +1325,8 @@ namespace UnrealBuildTool
 			CompileEnvironment.CppStandard = ModuleCompileEnvironment.CppStandard;
 			CompileEnvironment.CStandard = ModuleCompileEnvironment.CStandard;
 			CompileEnvironment.IncludeOrderVersion = ModuleCompileEnvironment.IncludeOrderVersion;
+			CompileEnvironment.bUseAutoRTFMCompiler = ModuleCompileEnvironment.bUseAutoRTFMCompiler;
+			CompileEnvironment.bAllowAutoRTFMInstrumentation = ModuleCompileEnvironment.bAllowAutoRTFMInstrumentation;
 		}
 
 		/// <summary>
