UnrealBuildTool: Denote default values for some BuildConfiguration properties

#rnx

diff --git a/Engine/Source/Programs/UnrealBuildTool/Configuration/BuildConfiguration.cs b/Engine/Source/Programs/UnrealBuildTool/Configuration/BuildConfiguration.cs
--- a/Engine/Source/Programs/UnrealBuildTool/Configuration/BuildConfiguration.cs
+++ b/Engine/Source/Programs/UnrealBuildTool/Configuration/BuildConfiguration.cs
@@ -49,21 +49,21 @@ namespace UnrealBuildTool
 #endif // #if __BOXEXECUTOR_AVAILABLE__
 
 		/// <summary>
-		/// Whether XGE may be used.
+		/// Whether XGE may be used if available, default is true.
 		/// </summary>
 		[XmlConfigFile]
 		[CommandLine("-NoXGE", Value = "false")]
 		public bool bAllowXGE = true;
 
 		/// <summary>
-		/// Whether FASTBuild may be used.
+		/// Whether FASTBuild may be used if availabe, default is true.
 		/// </summary>
 		[XmlConfigFile]
 		[CommandLine("-NoFASTBuild", Value = "false")]
 		public bool bAllowFASTBuild = true;
 
 		/// <summary>
-		/// Whether SN-DBS may be used.
+		/// Whether SN-DBS may be used if availabe, default is true.
 		/// </summary>
 		[XmlConfigFile]
 		[CommandLine("-NoSNDBS", Value = "false")]
