Fix iOS & tvOS not being able to set the BundleDisplayName (ie: the App name shown on the iPhone screen) when using Modern

diff --git a/Engine/Source/Programs/UnrealBuildTool/Platform/IOS/UEDeployIOS.cs b/Engine/Source/Programs/UnrealBuildTool/Platform/IOS/UEDeployIOS.cs
--- a/Engine/Source/Programs/UnrealBuildTool/Platform/IOS/UEDeployIOS.cs
+++ b/Engine/Source/Programs/UnrealBuildTool/Platform/IOS/UEDeployIOS.cs
@@ -216,7 +216,7 @@ namespace UnrealBuildTool
 			return MinVersionToReturn;
 		}
 
-		public static void WritePlistFile(FileReference PlistFile, DirectoryReference? ProjectLocation, UnrealPluginLanguage? UPL, string GameName, bool bIsUnrealGame, ILogger Logger)
+		public static void WritePlistFile(FileReference PlistFile, DirectoryReference? ProjectLocation, UnrealPluginLanguage? UPL, string GameName, bool bIsUnrealGame, string ProjectName, ILogger Logger)
 		{
 			ConfigHierarchy Ini = ConfigCache.ReadHierarchy(ConfigHierarchyType.Engine, ProjectLocation, UnrealTargetPlatform.IOS);
 			// required capabilities
@@ -276,6 +276,10 @@ namespace UnrealBuildTool
 			bool bDisableHTTPS;
 			Ini.GetBool("/Script/IOSRuntimeSettings.IOSRuntimeSettings", "bDisableHTTPS", out bDisableHTTPS);
 
+			// bundle display name
+			string BundleDisplayName;
+			Ini.GetString("/Script/IOSRuntimeSettings.IOSRuntimeSettings", "BundleDisplayName", out BundleDisplayName);
+
 			// short version string
 			string BundleShortVersion;
 			Ini.GetString("/Script/IOSRuntimeSettings.IOSRuntimeSettings", "VersionInfo", out BundleShortVersion);
@@ -345,6 +349,9 @@ namespace UnrealBuildTool
 				Text.AppendLine("\t<true/>");
 			}
 
+			Text.AppendLine("\t<key>CFBundleDisplayName</key>");
+			Text.AppendLine(String.Format("\t<string>{0}</string>", EncodeBundleName(BundleDisplayName, ProjectName)));
+
 			Text.AppendLine("\t<key>UIRequiresFullScreen</key>");
 			Text.AppendLine("\t<true/>");
 
@@ -514,7 +521,7 @@ namespace UnrealBuildTool
 			VersionUtilities.BuildDirectory = BuildDirectory;
 			VersionUtilities.GameName = GameName;
 
-			WritePlistFile(new FileReference(PListFile), DirRef, UPL, GameName, bIsUnrealGame, Logger);
+			WritePlistFile(new FileReference(PListFile), DirRef, UPL, GameName, bIsUnrealGame, ProjectName, Logger);
 
 			if (BuildHostPlatform.Current.Platform == UnrealTargetPlatform.Mac && !bBuildAsFramework)
 			{
