UnrealBuildTool: Don't remove old EngineIncludeOrderVersion enums to prevent breaking project generation

#rnx

diff --git a/Engine/Source/Programs/UnrealBuildTool/Configuration/TargetRules.cs b/Engine/Source/Programs/UnrealBuildTool/Configuration/TargetRules.cs
--- a/Engine/Source/Programs/UnrealBuildTool/Configuration/TargetRules.cs
+++ b/Engine/Source/Programs/UnrealBuildTool/Configuration/TargetRules.cs
@@ -161,10 +161,16 @@ namespace UnrealBuildTool
 	/// </summary>
 	public enum EngineIncludeOrderVersion
 	{
+		/// <summary>
+		/// Include order used in Unreal 5.1
+		/// </summary>
+		[Obsolete("The Unreal 5.1 include order is unsupported.")]
+		Unreal5_1,
+
 		/// <summary>
 		/// Include order used in Unreal 5.2
 		/// </summary>
-		[Obsolete("The Unreal 5.2 include order is deprecated and will be removed in 5.5.")]
+		[Obsolete("The Unreal 5.2 include order is deprecated and will be unsupported in 5.5.")]
 		Unreal5_2,
 
 		/// <summary>
@@ -173,7 +179,7 @@ namespace UnrealBuildTool
 		Unreal5_3,
 
 		/// <summary>
-		/// Include order used in Unreal 5.3
+		/// Include order used in Unreal 5.4
 		/// </summary>
 		Unreal5_4,
 
@@ -186,6 +192,7 @@ namespace UnrealBuildTool
 
 		/// <summary>
 		/// Contains the oldest version of include order that the engine supports.
+		/// Do not delete old enum entries to prevent breaking project generation
 		/// </summary>
 #pragma warning disable CS0618 // Type or member is obsolete
 		Oldest = Unreal5_2,
