UnrealBuildTool: Fix distribution not working for -SingleFile by moving several action properties to be set in the shared base action

#rnx

diff --git a/Engine/Source/Programs/UnrealBuildTool/Platform/Windows/VCCompileAction.cs b/Engine/Source/Programs/UnrealBuildTool/Platform/Windows/VCCompileAction.cs
--- a/Engine/Source/Programs/UnrealBuildTool/Platform/Windows/VCCompileAction.cs
+++ b/Engine/Source/Programs/UnrealBuildTool/Platform/Windows/VCCompileAction.cs
@@ -403,6 +403,7 @@ namespace UnrealBuildTool
 			bCanExecuteRemotely = Reader.ReadBool();
 			bCanExecuteRemotelyWithSNDBS = Reader.ReadBool();
 			bCanExecuteRemotelyWithXGE = Reader.ReadBool();
+			bCanExecuteInBox = Reader.ReadBool();
 			Architecture = UnrealArch.Parse(Reader.ReadString()!);
 			Weight = Reader.ReadDouble();
 
@@ -441,6 +442,7 @@ namespace UnrealBuildTool
 			Writer.WriteBool(bCanExecuteRemotely);
 			Writer.WriteBool(bCanExecuteRemotelyWithSNDBS);
 			Writer.WriteBool(bCanExecuteRemotelyWithXGE);
+			Writer.WriteBool(bCanExecuteInBox);
 			Writer.WriteString(Architecture.ToString());
 			Writer.WriteDouble(Weight);
 
diff --git a/Engine/Source/Programs/UnrealBuildTool/Platform/Windows/VCToolChain.cs b/Engine/Source/Programs/UnrealBuildTool/Platform/Windows/VCToolChain.cs
--- a/Engine/Source/Programs/UnrealBuildTool/Platform/Windows/VCToolChain.cs
+++ b/Engine/Source/Programs/UnrealBuildTool/Platform/Windows/VCToolChain.cs
@@ -1403,6 +1403,30 @@ namespace UnrealBuildTool
 					BaseCompileAction.Arguments.Add("/d1reportTime");
 				}
 			}
+
+			// MSVC uses multiple threads when compiling CPP files, so the "weight" is more than 1
+			if (Target.WindowsPlatform.Compiler.IsMSVC())
+			{
+				// If deterministic is enabled, MSVC does not use multiple threads
+				if (!CompileEnvironment.bDeterministic)
+				{
+					BaseCompileAction.Weight = Target.MSVCCompileActionWeight;
+				}
+			}
+			else if (Target.WindowsPlatform.Compiler.IsClang())
+			{
+				BaseCompileAction.Weight = Target.ClangCompileActionWeight;
+			}
+
+			// Don't farm out creation of precompiled headers as it is the critical path task.
+			BaseCompileAction.bCanExecuteRemotely =
+				CompileEnvironment.PrecompiledHeaderAction != PrecompiledHeaderAction.Create ||
+				CompileEnvironment.bAllowRemotelyCompiledPCHs
+				;
+
+			// When compiling with SN-DBS, modules that contain a #import must be built locally
+			BaseCompileAction.bCanExecuteRemotelyWithSNDBS = BaseCompileAction.bCanExecuteRemotely && !CompileEnvironment.bBuildLocallyWithSNDBS;
+
 			return BaseCompileAction;
 		}
 
@@ -1435,20 +1459,6 @@ namespace UnrealBuildTool
 		{
 			VCCompileAction BaseCompileAction = CreateBaseCompileAction(CompileEnvironment);
 
-			// MSVC uses multiple threads when compiling CPP files, so the "weight" is more than 1
-			if (Target.WindowsPlatform.Compiler.IsMSVC())
-			{
-				// If deterministic is enabled, MSVC does not use multiple threads
-				if (!CompileEnvironment.bDeterministic)
-				{
-					BaseCompileAction.Weight = Target.MSVCCompileActionWeight;
-				}
-			}
-			else if (Target.WindowsPlatform.Compiler.IsClang())
-			{
-				BaseCompileAction.Weight = Target.ClangCompileActionWeight;
-			}
-
 			// Create a compile action for each source file.
 			List<VCCompileAction> Actions = new List<VCCompileAction>();
 			foreach (FileItem SourceFile in InputFiles)
@@ -1544,12 +1554,6 @@ namespace UnrealBuildTool
 					}
 				}
 
-				// Don't farm out creation of precompiled headers as it is the critical path task.
-				CompileAction.bCanExecuteRemotely =
-					CompileEnvironment.PrecompiledHeaderAction != PrecompiledHeaderAction.Create ||
-					CompileEnvironment.bAllowRemotelyCompiledPCHs
-					;
-
 				// Create PDB files if we were configured to do that.
 				if (CompileEnvironment.bUsePDBFiles || CompileEnvironment.bSupportEditAndContinue)
 				{
@@ -1708,13 +1712,6 @@ namespace UnrealBuildTool
 
 				CompileAction.bIsAnalyzing = Target.StaticAnalyzer != StaticAnalyzer.None;
 
-				// When compiling with SN-DBS, modules that contain a #import must be built locally
-				CompileAction.bCanExecuteRemotelyWithSNDBS = CompileAction.bCanExecuteRemotely;
-				if (CompileEnvironment.bBuildLocallyWithSNDBS == true)
-				{
-					CompileAction.bCanExecuteRemotelyWithSNDBS = false;
-				}
-
 				// Update the output
 				Graph.AddAction(CompileAction);
 				Actions.Add(CompileAction);
diff --git a/Engine/Source/Programs/UnrealBuildTool/ToolChain/ClangToolChain.cs b/Engine/Source/Programs/UnrealBuildTool/ToolChain/ClangToolChain.cs
--- a/Engine/Source/Programs/UnrealBuildTool/ToolChain/ClangToolChain.cs
+++ b/Engine/Source/Programs/UnrealBuildTool/ToolChain/ClangToolChain.cs
@@ -1040,6 +1040,9 @@ namespace UnrealBuildTool
 			ClangSpecificFileActionGraphBuilder GraphBuilder = new(Logger);
 			Action Action = CompileCPPFile(CompileEnvironment, DummyFile, OutputDir, "<Unknown>", GraphBuilder, GlobalArguments, Result);
 			Action.PrerequisiteItems.RemoveWhere(File => File.Name.Contains(DummyName));
+			Action.bCanExecuteRemotely = true;
+			Action.bCanExecuteRemotelyWithSNDBS = Action.bCanExecuteRemotely && !CompileEnvironment.bBuildLocallyWithSNDBS;
+			Action.Weight = CompileActionWeight;
 
 			Graph.AddAction(new ClangSpecificFileAction(SourceDir, OutputDir, Action, GraphBuilder.ContentLines));
 		}
