UBT: Build fix.

diff --git a/Engine/Source/Programs/UnrealBuildTool/Artifacts/HordeStorageArtifactCache.cs b/Engine/Source/Programs/UnrealBuildTool/Artifacts/HordeStorageArtifactCache.cs
--- a/Engine/Source/Programs/UnrealBuildTool/Artifacts/HordeStorageArtifactCache.cs
+++ b/Engine/Source/Programs/UnrealBuildTool/Artifacts/HordeStorageArtifactCache.cs
@@ -525,7 +525,7 @@ namespace UnrealBuildTool.Artifacts
 				}
 				Directory.CreateDirectory(directory.FullName);
 
-				_store = new FileStorageClient(directory, logger);
+				_store = new FileStorageClient(directory, null, logger);
 
 				State = ArtifactCacheState.Available;
 				return State;
