UnrealBuildTool - Added a warning if a project module is found to have a dependency on a plugin module without listing that plugin as a dependency. The plugin has to be enabled by default or enabled in the .uproject or enabled in the target.cs

diff --git a/Engine/Source/Programs/UnrealBuildTool/Configuration/UEBuildTarget.cs b/Engine/Source/Programs/UnrealBuildTool/Configuration/UEBuildTarget.cs
--- a/Engine/Source/Programs/UnrealBuildTool/Configuration/UEBuildTarget.cs
+++ b/Engine/Source/Programs/UnrealBuildTool/Configuration/UEBuildTarget.cs
@@ -2402,6 +2402,63 @@ namespace UnrealBuildTool
 				}
 			}
 
+
+			// Check that each project module with a dependency on a plugin module, that the plugin is either enabled by default or it's enabled by the uproject file.
+			if (ProjectDescriptor?.Modules != null && ProjectDescriptor?.Plugins != null)
+			{
+				string ProjectName = ProjectFile?.GetFileName() ?? ".uproject";
+				bool bAllowEnginePluginsEnabledByDefault = !ProjectDescriptor.DisableEnginePluginsByDefault;
+
+				foreach (ModuleDescriptor Descriptor in ProjectDescriptor.Modules.Where(x => x.IsCompiledInConfiguration(Platform, Configuration, TargetName, TargetType, Rules.bBuildDeveloperTools, Rules.bBuildRequiresCookedData)))
+				{
+					if (Modules.TryGetValue(Descriptor.Name, out UEBuildModule? Module))
+					{
+						HashSet<UEBuildModule> DependencyModules = Module.GetDependencies(bWithIncludePathModules: true, bWithDynamicallyLoadedModules: true);
+						foreach (UEBuildModule DependencyModule in DependencyModules)
+						{
+							if (ModuleToPlugin.TryGetValue(DependencyModule, out UEBuildPlugin? DependencyPlugin))
+							{
+								PluginInfo? DependencyPluginInfo = Plugins.GetPlugin(DependencyPlugin.Name);
+
+								// Is the modules plugin enabled by default?
+								if (DependencyPluginInfo?.IsEnabledByDefault(bAllowEnginePluginsEnabledByDefault) == false)
+								{
+									// Try and find the project plugin reference in the plugins list.
+									PluginReferenceDescriptor? ProjectPluginReference = ProjectDescriptor.Plugins.FirstOrDefault(x => x.Name.Equals(DependencyPlugin.Name));
+									if (ProjectPluginReference != null)
+									{
+										if (!ProjectPluginReference.bEnabled)
+										{
+											Logger.LogWarning("Warning: {ProjectName} plugin dependency '{DependencyPluginName}' is not enabled, but module '{ModuleName}' depends on '{DependencyModuleName}'.", ProjectName, DependencyPlugin.Name, Module.Name, DependencyModule.Name);
+										}
+										else if (!ProjectPluginReference.IsEnabledForTarget(TargetType))
+										{
+											Logger.LogWarning("Warning: {ProjectName} plugin dependency '{DependencyPluginName}' is not enabled for target '{TargetType}', but module '{ModuleName}' depends on '{DependencyModuleName}'.", ProjectName, DependencyPlugin.Name, TargetType.ToString(), Module.Name, DependencyModule.Name);
+										}
+										else if (!ProjectPluginReference.IsEnabledForPlatform(Platform))
+										{
+											Logger.LogWarning("Warning: {ProjectName} plugin dependency '{DependencyPluginName}' is not enabled for platform '{Platform}', but module '{ModuleName}' depends on '{DependencyModuleName}'.", ProjectName, DependencyPlugin.Name, Platform.ToString(), Module.Name, DependencyModule.Name);
+										}
+										else if (!ProjectPluginReference.IsEnabledForTargetConfiguration(Configuration))
+										{
+											Logger.LogWarning("Warning: {ProjectName} plugin dependency '{DependencyPluginName}' is not enabled for target configuration '{Configuration}', but module '{ModuleName}' depends on '{DependencyModuleName}'.", ProjectName, DependencyPlugin.Name, Configuration.ToString(), Module.Name, DependencyModule.Name);
+										}
+									}
+									else
+									{
+										// No project plugin reference exists. Check inclusion of a plugin was not made as part of a Target.cs file or Build.cs file
+										if (!Rules.EnablePlugins.Contains(DependencyPlugin.Name) && !Rules.InternalPluginDependencies.Contains(DependencyPlugin.Name))
+										{
+											Logger.LogWarning("Warning: {ProjectName} does not list plugin '{DependencyPluginName}' as a dependency, but module '{ModuleName}' depends on '{DependencyModuleName}'.", ProjectName, DependencyPlugin.Name, Module.Name, DependencyModule.Name);
+										}
+									}
+								}
+							}
+						}
+					}
+				}
+			}
+
 			// Check that each plugin does not have a dependency on any sealed plugins
 			foreach (UEBuildPlugin Plugin in BuildPlugins)
 			{
