UBT: Clang: Added support for Clang's -fextend-lifetimes switch which improves debugging experience in optimized builds.

This switch makes the compiler extend lifetimes of local variables, function parameters and 'this' till the end of their scopes increasing the chance for the debugger being able to display their values when debugging.

It may slightly reduce performance so it's an opt-in feature controlled with bOptimizeForDebugging variable in xPlatform sections on supported platforms in BuildConfiguration.xml scripts. Even when set to true, it won't affect Test or Shipping builds.

[RN] minor Added support for -fextend-lifetimes compilation switch on selected Clang-based platforms which improves debugging experience of optimized builds by increasing the number of local variables for which the debugger can show values. The feature may slightly reduce performance and is intended to be used in development only. It has to be explicitly activated by setting bOptimizeForDebugging variable in <platform>Platform section of BuildConfiguration.xml. It won't affect Test or Shipping configurations.

diff --git a/Engine/Source/Programs/UnrealBuildTool/ToolChain/ClangToolChain.cs b/Engine/Source/Programs/UnrealBuildTool/ToolChain/ClangToolChain.cs
--- a/Engine/Source/Programs/UnrealBuildTool/ToolChain/ClangToolChain.cs
+++ b/Engine/Source/Programs/UnrealBuildTool/ToolChain/ClangToolChain.cs
@@ -112,6 +112,13 @@ namespace UnrealBuildTool
 		/// Enable LibFuzzer
 		/// </summary>
 		EnableLibFuzzer = 1 << 18,
+
+		/// <summary>
+		/// Modify code generation to help with debugging optimized builds e.g. by extending lifetimes of local variables.
+		/// It may slightly reduce performance. Thus, it's meant to be used during development only.
+		/// Supported only on some platforms.
+		/// </summary>
+		OptimizeForDebugging = 1 << 19,
 	}
 
 	abstract class ClangToolChain : ISPCToolChain
@@ -181,6 +188,8 @@ namespace UnrealBuildTool
 
 		protected ClangToolChainOptions Options;
 
+		protected bool bOptimizeForDebugging => Options.HasFlag(ClangToolChainOptions.OptimizeForDebugging);
+
 		// Dummy define to work around clang compilation related to the windows maximum path length limitation
 		protected static string ClangDummyDefine;
 		protected const int ClangCmdLineMaxSize = 32 * 1024;
@@ -650,6 +659,11 @@ namespace UnrealBuildTool
 				Arguments.Add("-fno-exceptions");
 				Arguments.Add("-DPLATFORM_EXCEPTIONS_DISABLED=1");
 			}
+
+			if (bOptimizeForDebugging)
+			{
+				Arguments.Add("-fextend-lifetimes");
+			}
 		}
 
 		/// <summary>
