UBT: change Linux PGO to use new code placement algo (extnended TSP).

- Expected to bring more icache-friendly code placement, especially for large binaries (https://reviews.llvm.org/D113424)
- Reportedly not default in lld because some synthetic benchmarks were regressed.
- Suggested by [REDACTED_PERSON]

diff --git a/Engine/Source/Programs/UnrealBuildTool/Platform/Linux/LinuxToolChain.cs b/Engine/Source/Programs/UnrealBuildTool/Platform/Linux/LinuxToolChain.cs
--- a/Engine/Source/Programs/UnrealBuildTool/Platform/Linux/LinuxToolChain.cs
+++ b/Engine/Source/Programs/UnrealBuildTool/Platform/Linux/LinuxToolChain.cs
@@ -310,7 +310,7 @@ namespace UnrealBuildTool
 				{
 					if (bForLinker)
 					{
-						Arguments.Add(String.Format("-flto=thin -Wl,--thinlto-jobs={0}, -Wl,-mllvm,-disable-auto-upgrade-debug-info", Utils.GetPhysicalProcessorCount()));
+						Arguments.Add(String.Format("-flto=thin -Wl,--thinlto-jobs={0}, -Wl,-mllvm,-disable-auto-upgrade-debug-info, -Wl,-mllvm,-enable-ext-tsp-block-placement=1", Utils.GetPhysicalProcessorCount()));
 					}
 					else
 					{
