Allow forcing `WITH_VERSE_VM=1` so that when generating Clang compilation databases we get useful ones for `clangd` to actually work properly.

#rnx

diff --git a/Engine/Source/Programs/UnrealBuildTool/Modes/GenerateClangDatabase.cs b/Engine/Source/Programs/UnrealBuildTool/Modes/GenerateClangDatabase.cs
--- a/Engine/Source/Programs/UnrealBuildTool/Modes/GenerateClangDatabase.cs
+++ b/Engine/Source/Programs/UnrealBuildTool/Modes/GenerateClangDatabase.cs
@@ -32,6 +32,13 @@ namespace UnrealBuildTool
 		[CommandLine("-ExecCodeGenActions")]
 		public bool bExecCodeGenActions = false;
 
+		/// <summary>
+		/// Used for generating a compilation database against the legacy/new VM for Verse. We default to `true`, since
+		/// in the general case the compilation database generated against the legacy VM is not very useful.
+		/// </summary>
+		[CommandLine("-NoUseVerseBPVM")]
+		public bool bNoUseVerseBPVM = true;
+
 		/// <summary>
 		/// This ActionGraphBuilder captures the build output from a UEBuildModuleCPP so it can be consumed later.
 		/// </summary>
@@ -148,6 +155,10 @@ namespace UnrealBuildTool
 					TargetDescriptor.bUseUnityBuild = false;
 					TargetDescriptor.IntermediateEnvironment = UnrealIntermediateEnvironment.GenerateClangDatabase;
 					TargetDescriptor.AdditionalArguments = TargetDescriptor.AdditionalArguments.Append(new string[] { "-NoPCH" });
+					if (bNoUseVerseBPVM)
+					{
+						TargetDescriptor.AdditionalArguments = TargetDescriptor.AdditionalArguments.Append(new string[] { "-NoUseVerseBPVM" });
+					}
 
 					// Create a makefile for the target
 					UEBuildTarget Target = UEBuildTarget.Create(TargetDescriptor, BuildConfiguration, Logger);
