Fix for TargetRules.ToolChainName instantiation failure

#rnx

diff --git a/Engine/Source/Programs/UnrealBuildTool/Configuration/UEBuildTarget.cs b/Engine/Source/Programs/UnrealBuildTool/Configuration/UEBuildTarget.cs
--- a/Engine/Source/Programs/UnrealBuildTool/Configuration/UEBuildTarget.cs
+++ b/Engine/Source/Programs/UnrealBuildTool/Configuration/UEBuildTarget.cs
@@ -2325,7 +2325,7 @@ namespace UnrealBuildTool
 			CppCompileEnvironment GlobalCompileEnvironment = new CppCompileEnvironment(Platform, CppConfiguration, Architectures, MetadataCache);
 			LinkEnvironment GlobalLinkEnvironment = new LinkEnvironment(GlobalCompileEnvironment.Platform, GlobalCompileEnvironment.Configuration, GlobalCompileEnvironment.Architectures);
 
-			UEToolChain TargetToolChain = CreateToolchain(Platform);
+			UEToolChain TargetToolChain = CreateToolchain(Platform, Logger);
 			SetupGlobalEnvironment(TargetToolChain, GlobalCompileEnvironment, GlobalLinkEnvironment);
 
 			FindSharedPCHs(Binaries, GlobalCompileEnvironment, Logger);
@@ -2345,7 +2345,7 @@ namespace UnrealBuildTool
 			CppCompileEnvironment GlobalCompileEnvironment = new CppCompileEnvironment(Platform, CppConfiguration, Architectures, MetadataCache);
 			LinkEnvironment GlobalLinkEnvironment = new LinkEnvironment(GlobalCompileEnvironment.Platform, GlobalCompileEnvironment.Configuration, GlobalCompileEnvironment.Architectures);
 
-			UEToolChain TargetToolChain = CreateToolchain(Platform);
+			UEToolChain TargetToolChain = CreateToolchain(Platform, Logger);
 			TargetToolChain.SetEnvironmentVariables();
 			SetupGlobalEnvironment(TargetToolChain, GlobalCompileEnvironment, GlobalLinkEnvironment);
 
@@ -3131,7 +3131,7 @@ namespace UnrealBuildTool
 		/// Creates a toolchain for the current target. May be overridden by the target rules.
 		/// </summary>
 		/// <returns>New toolchain instance</returns>
-		public UEToolChain CreateToolchain(UnrealTargetPlatform Platform)
+		public UEToolChain CreateToolchain(UnrealTargetPlatform Platform, ILogger Logger)
 		{
 			if (Rules.ToolChainName == null)
 			{
@@ -3144,7 +3144,7 @@ namespace UnrealBuildTool
 				{
 					throw new BuildException("Unable to create toolchain '{0}'. Check that the name is correct.", Rules.ToolChainName);
 				}
-				return (UEToolChain)Activator.CreateInstance(ToolchainType, Rules)!;
+				return (UEToolChain)Activator.CreateInstance(ToolchainType, Rules, Logger)!;
 			}
 		}
 
diff --git a/Engine/Source/Programs/UnrealBuildTool/Modes/GenerateClangDatabase.cs b/Engine/Source/Programs/UnrealBuildTool/Modes/GenerateClangDatabase.cs
--- a/Engine/Source/Programs/UnrealBuildTool/Modes/GenerateClangDatabase.cs
+++ b/Engine/Source/Programs/UnrealBuildTool/Modes/GenerateClangDatabase.cs
@@ -162,7 +162,7 @@ namespace UnrealBuildTool
 
 					// Create a makefile for the target
 					UEBuildTarget Target = UEBuildTarget.Create(TargetDescriptor, BuildConfiguration, Logger);
-					UEToolChain TargetToolChain = Target.CreateToolchain(Target.Platform);
+					UEToolChain TargetToolChain = Target.CreateToolchain(Target.Platform, Logger);
 
 					// Execute code generation actions
 					if (bExecCodeGenActions)
diff --git a/Engine/Source/Programs/UnrealBuildTool/Modes/ProfileUnitySizesMode.cs b/Engine/Source/Programs/UnrealBuildTool/Modes/ProfileUnitySizesMode.cs
--- a/Engine/Source/Programs/UnrealBuildTool/Modes/ProfileUnitySizesMode.cs
+++ b/Engine/Source/Programs/UnrealBuildTool/Modes/ProfileUnitySizesMode.cs
@@ -146,7 +146,7 @@ namespace UnrealBuildTool
 				// Create a makefile for the target
 				TimingLogger TimingLogger = new(Logger);
 				UEBuildTarget Target = UEBuildTarget.Create(TargetDescriptor, BuildConfiguration, TimingLogger);
-				UEToolChain TargetToolChain = Target.CreateToolchain(Target.Platform);
+				UEToolChain TargetToolChain = Target.CreateToolchain(Target.Platform, TimingLogger);
 
 				CppCompileEnvironment GlobalCompileEnvironment = Target.CreateCompileEnvironmentForProjectFiles(TimingLogger);
 				foreach (UEBuildBinary Binary in Target.Binaries)
diff --git a/Engine/Source/Programs/UnrealBuildTool/Modes/QueryMode.cs b/Engine/Source/Programs/UnrealBuildTool/Modes/QueryMode.cs
--- a/Engine/Source/Programs/UnrealBuildTool/Modes/QueryMode.cs
+++ b/Engine/Source/Programs/UnrealBuildTool/Modes/QueryMode.cs
@@ -264,7 +264,7 @@ namespace UnrealBuildTool
 				CppCompileEnvironment GlobalCompileEnvironment = new CppCompileEnvironment(CurrentTarget.Platform, CppConfiguration, CurrentTarget.Architectures, MetadataCache);
 				LinkEnvironment GlobalLinkEnvironment = new LinkEnvironment(GlobalCompileEnvironment.Platform, GlobalCompileEnvironment.Configuration, GlobalCompileEnvironment.Architectures);
 
-				UEToolChain TargetToolChain = CurrentTarget.CreateToolchain(CurrentTarget.Platform);
+				UEToolChain TargetToolChain = CurrentTarget.CreateToolchain(CurrentTarget.Platform, Logger);
 				TargetToolChain.SetEnvironmentVariables();
 				CurrentTarget.SetupGlobalEnvironment(TargetToolChain, GlobalCompileEnvironment, GlobalLinkEnvironment);
 
diff --git a/Engine/Source/Programs/UnrealBuildTool/Modes/ServerMode.cs b/Engine/Source/Programs/UnrealBuildTool/Modes/ServerMode.cs
--- a/Engine/Source/Programs/UnrealBuildTool/Modes/ServerMode.cs
+++ b/Engine/Source/Programs/UnrealBuildTool/Modes/ServerMode.cs
@@ -335,7 +335,7 @@ namespace UnrealBuildTool
 				CppCompileEnvironment GlobalCompileEnvironment = new CppCompileEnvironment(CurrentTarget.Platform, CppConfiguration, CurrentTarget.Architectures, MetadataCache);
 				LinkEnvironment GlobalLinkEnvironment = new LinkEnvironment(GlobalCompileEnvironment.Platform, GlobalCompileEnvironment.Configuration, GlobalCompileEnvironment.Architectures);
 
-				UEToolChain TargetToolChain = CurrentTarget.CreateToolchain(CurrentTarget.Platform);
+				UEToolChain TargetToolChain = CurrentTarget.CreateToolchain(CurrentTarget.Platform, Logger);
 				TargetToolChain.SetEnvironmentVariables();
 				CurrentTarget.SetupGlobalEnvironment(TargetToolChain, GlobalCompileEnvironment, GlobalLinkEnvironment);
 
