Add support for CustomConfig to AppXManifestGenerator

#rnx

diff --git a/Engine/Source/Programs/UnrealBuildTool/Platform/Windows/AppXManifestGeneratorBase.cs b/Engine/Source/Programs/UnrealBuildTool/Platform/Windows/AppXManifestGeneratorBase.cs
--- a/Engine/Source/Programs/UnrealBuildTool/Platform/Windows/AppXManifestGeneratorBase.cs
+++ b/Engine/Source/Programs/UnrealBuildTool/Platform/Windows/AppXManifestGeneratorBase.cs
@@ -65,6 +65,9 @@ namespace UnrealBuildTool
 		/// Whether we have logged the deprecation warning for PerCultureResources CultureId being replaced by StageIdOverrides
 		protected static bool bHasWarnedAboutDeprecatedCultureId = false;
 
+		/// CustomConfig to use when reading the ini files
+		public string CustomConfig = "";
+
 		/// <summary>
 		/// Create a manifest generator for the given platform variant.
 		/// </summary>
@@ -503,8 +506,8 @@ namespace UnrealBuildTool
 
 			// Load up INI settings. We'll use engine settings to retrieve the manifest configuration, but these may reference
 			// values in either game or engine settings, so we'll keep both.
-			GameIni = ConfigCache.ReadHierarchy(ConfigHierarchyType.Game, DirectoryReference.FromFile(InProjectFile), Platform);
-			EngineIni = ConfigCache.ReadHierarchy(ConfigHierarchyType.Engine, DirectoryReference.FromFile(InProjectFile), Platform);
+			GameIni = ConfigCache.ReadHierarchy(ConfigHierarchyType.Game, DirectoryReference.FromFile(InProjectFile), Platform, CustomConfig);
+			EngineIni = ConfigCache.ReadHierarchy(ConfigHierarchyType.Engine, DirectoryReference.FromFile(InProjectFile), Platform, CustomConfig);
 			PostConfigurationInit();
 
 			// Load and verify/clean culture list
