Add support to ConfigCache.cs to allow a hotfix directory to be specified which will allow additions and overrides to specified in a similar way to the runtime hotfixing system
#rnx

diff --git a/Engine/Source/Programs/UnrealBuildTool/System/ConfigCache.cs b/Engine/Source/Programs/UnrealBuildTool/System/ConfigCache.cs
--- a/Engine/Source/Programs/UnrealBuildTool/System/ConfigCache.cs
+++ b/Engine/Source/Programs/UnrealBuildTool/System/ConfigCache.cs
@@ -174,6 +174,11 @@ namespace UnrealBuildTool
 			/// </summary>
 			public List<string> OverrideStrings;
 
+			/// <summary>
+			/// A hotfix directory where modifications/additions to config files can be read from
+			/// </summary>
+			public DirectoryReference? HotfixDir;
+
 			/// <summary>
 			/// Constructor
 			/// </summary>
@@ -182,13 +187,15 @@ namespace UnrealBuildTool
 			/// <param name="Platform">Which platform-specific files to read</param>
 			/// <param name="CustomConfig">Custom config subdirectory to load</param>
 			/// <param name="OverrideStrings">Custom override strings</param>
-			public ConfigHierarchyKey(ConfigHierarchyType Type, DirectoryReference? ProjectDir, UnrealTargetPlatform Platform, string CustomConfig, List<string> OverrideStrings)
+			/// <param name="HotfixDir">A hotfix directory where modifications/additions to config files can be read from</param>
+			public ConfigHierarchyKey(ConfigHierarchyType Type, DirectoryReference? ProjectDir, UnrealTargetPlatform Platform, string CustomConfig, List<string> OverrideStrings, DirectoryReference? HotfixDir)
 			{
 				this.Type = Type;
 				this.ProjectDir = ProjectDir;
 				this.Platform = Platform;
 				this.CustomConfig = CustomConfig;
 				this.OverrideStrings = OverrideStrings;
+				this.HotfixDir = HotfixDir;
 			}
 
 			/// <summary>
@@ -204,7 +211,8 @@ namespace UnrealBuildTool
 					OtherKey.ProjectDir == ProjectDir &&
 					OtherKey.Platform == Platform &&
 					OtherKey.CustomConfig == CustomConfig &&
-					OtherKey.OverrideStrings.SequenceEqual(OverrideStrings);
+					OtherKey.OverrideStrings.SequenceEqual(OverrideStrings) &&
+					OtherKey.HotfixDir == HotfixDir;
 			}
 
 			/// <summary>
@@ -222,6 +230,8 @@ namespace UnrealBuildTool
 				{
 					Hash = (Hash * 31) + OverrideString.GetHashCode();
 				}
+				Hash = (Hash * 31) + ((HotfixDir == null) ? 0 : HotfixDir.GetHashCode());
+
 				return Hash;
 			}
 		}
@@ -270,8 +280,9 @@ namespace UnrealBuildTool
 		/// <param name="Platform">Which platform to read platform-specific config files for</param>
 		/// <param name="CustomConfig">Optional override config directory to search, for support of multiple target types</param>
 		/// <param name="CustomArgs">Optional list of command line arguments added to the existing command line arguments</param>
+		/// <param name="HotfixDir">A hotfix directory where modifications/additions to config files can be read from</param>
 		/// <returns>The requested config hierarchy</returns>
-		public static ConfigHierarchy ReadHierarchy(ConfigHierarchyType Type, DirectoryReference? ProjectDir, UnrealTargetPlatform Platform, string CustomConfig = "", string[]? CustomArgs = null)
+		public static ConfigHierarchy ReadHierarchy(ConfigHierarchyType Type, DirectoryReference? ProjectDir, UnrealTargetPlatform Platform, string CustomConfig = "", string[]? CustomArgs = null, DirectoryReference? HotfixDir = null)
 		{
 			CommandLineArguments CombinedArgs = new CommandLineArguments(Environment.GetCommandLineArgs());
 
@@ -290,7 +301,7 @@ namespace UnrealBuildTool
 				}
 			}
 
-			return ReadHierarchy(Type, ProjectDir, Platform, CustomConfig, CombinedArgs);
+			return ReadHierarchy(Type, ProjectDir, Platform, CustomConfig, CombinedArgs, HotfixDir);
 		}
 
 		/// <summary>
@@ -301,9 +312,10 @@ namespace UnrealBuildTool
 		/// <param name="Platform">Which platform to read platform-specific config files for</param>
 		/// <param name="CustomConfig">Optional override config directory to search, for support of multiple target types</param>
 		/// <param name="CmdLineArgs">The command line arguments to parse</param>
+		/// <param name="HotfixDir">A hotfix directory where modifications/additions to config files can be read from</param>
 		/// <returns>The requested config hierarchy</returns>
 		public static ConfigHierarchy ReadHierarchy(ConfigHierarchyType Type, DirectoryReference? ProjectDir, UnrealTargetPlatform Platform,
-			string CustomConfig, CommandLineArguments CmdLineArgs)
+			string CustomConfig, CommandLineArguments CmdLineArgs, DirectoryReference? HotfixDir = null)
 		{
 			// Handle command line overrides
 			List<String> OverrideStrings = new List<String>();
@@ -333,7 +345,7 @@ namespace UnrealBuildTool
 			}
 
 			// Get the key to use for the cache. It cannot be null, so we use the engine directory if a project directory is not given.
-			ConfigHierarchyKey Key = new ConfigHierarchyKey(Type, ProjectDir, Platform, CustomConfig, OverrideStrings);
+			ConfigHierarchyKey Key = new ConfigHierarchyKey(Type, ProjectDir, Platform, CustomConfig, OverrideStrings, HotfixDir);
 
 			ILogger Logger = NullLogger.Instance;
 
@@ -357,6 +369,17 @@ namespace UnrealBuildTool
 						Logger.LogInformation($"    Found!");
 						Files.Add(File);
 					}
+
+					if (HotfixDir != null)
+					{
+						FileReference HotfixFile = FileReference.Combine(HotfixDir, IniFileName.GetFileName());
+
+						if (TryReadFile(HotfixFile, out File))
+						{
+							Logger.LogInformation($"    Hotfixed with '{HotfixFile}'");
+							Files.Add(File);
+						}
+					}
 				}
 
 				foreach (string OverrideString in OverrideStrings)
