Added a command-line option to enable XCurl.

diff --git a/Engine/Source/Programs/UnrealBuildTool/Platform/Windows/UEBuildWindows.cs b/Engine/Source/Programs/UnrealBuildTool/Platform/Windows/UEBuildWindows.cs
--- a/Engine/Source/Programs/UnrealBuildTool/Platform/Windows/UEBuildWindows.cs
+++ b/Engine/Source/Programs/UnrealBuildTool/Platform/Windows/UEBuildWindows.cs
@@ -562,6 +562,7 @@ namespace UnrealBuildTool
 		/// Whether this build will use Microsoft's custom XCurl instead of libcurl
 		/// Note that XCurl is not part of the normal Windows SDK and will require additional downloads
 		/// </summary>
+		[CommandLine("-UseXCurl")]
 		public bool bUseXCurl = false;
 
 		/// <summary>
