CustomConfig can be specified on the UBT command line

#rnx

diff --git a/Engine/Source/Programs/UnrealBuildTool/Configuration/TargetRules.cs b/Engine/Source/Programs/UnrealBuildTool/Configuration/TargetRules.cs
--- a/Engine/Source/Programs/UnrealBuildTool/Configuration/TargetRules.cs
+++ b/Engine/Source/Programs/UnrealBuildTool/Configuration/TargetRules.cs
@@ -766,6 +766,7 @@ namespace UnrealBuildTool
 		/// This will get baked into the game executable as CUSTOM_CONFIG and used when staging to filter files and settings
 		/// </summary>
 		[RequiresUniqueBuildEnvironment]
+		[CommandLine("-CustomConfig")]
 		public string CustomConfig = String.Empty;
 
 		/// <summary>
