[UBT]
* Added text describing if error code was Access violation or Stack buffer overflow

diff --git a/Engine/Source/Programs/UnrealBuildTool/Executors/ImmediateActionQueue.cs b/Engine/Source/Programs/UnrealBuildTool/Executors/ImmediateActionQueue.cs
--- a/Engine/Source/Programs/UnrealBuildTool/Executors/ImmediateActionQueue.cs
+++ b/Engine/Source/Programs/UnrealBuildTool/Executors/ImmediateActionQueue.cs
@@ -1129,7 +1129,12 @@ namespace UnrealBuildTool
 					// user know something went wrong.
 					if (logLines == null || logLines.Count <= (action.bShouldOutputStatusDescription ? 0 : 1))
 					{
-						Logger.LogError("{TargetDetails} {Description}: Exited with error code {ExitCode}. The build will fail.", targetDetails, description, exitCode);
+						string exitCodeStr = string.Empty;
+						if ((uint)exitCode == 0xC0000005)
+							exitCodeStr = "(Access violation)";
+						else if ((uint)exitCode == 0xC0000409)
+							exitCodeStr = "(Stack buffer overflow)";
+						Logger.LogError("{TargetDetails} {Description}: Exited with error code {ExitCode} {ExitCodeStr}. The build will fail.", targetDetails, description, exitCode, exitCodeStr);
 						Logger.LogInformation("{TargetDetails} {Description}: WorkingDirectory {WorkingDirectory}", targetDetails, description, action.WorkingDirectory);
 						Logger.LogInformation("{TargetDetails} {Description}: {CommandPath} {CommandArguments}", targetDetails, description, action.CommandPath, action.CommandArguments);
 					}
