QueryMode: Add debug launch configuration to target info

diff --git a/Engine/Source/Programs/UnrealBuildTool/Modes/QueryMode.cs b/Engine/Source/Programs/UnrealBuildTool/Modes/QueryMode.cs
--- a/Engine/Source/Programs/UnrealBuildTool/Modes/QueryMode.cs
+++ b/Engine/Source/Programs/UnrealBuildTool/Modes/QueryMode.cs
@@ -6,6 +6,7 @@ using System.Linq;
 using System.Text.Json;
 using System.Threading.Tasks;
 using EpicGames.Core;
+using Microsoft.CodeAnalysis.CSharp.Syntax;
 using Microsoft.Extensions.Logging;
 using OpenTracing.Util;
 using UnrealBuildBase;
@@ -62,6 +63,13 @@ namespace UnrealBuildTool
 			return null;
 		}
 	}
+	
+	internal class LaunchSettings
+	{
+		public string? Description { get; set; }
+		public string? BinaryPath { get; set; }
+		public List<string> Arguments {get; set; } = new();
+	}
 
 	internal class TargetConfigs
 	{
@@ -219,14 +227,25 @@ namespace UnrealBuildTool
 			{
 				Logger.LogError("Failed to query available targets: {0}", e.Message);
 				return 1;
-			}
-		}
-		private int QueryTargetDetails(CommandLineArguments Arguments, ILogger Logger, JsonSerializerOptions JsonOptions)
-		{
-			if (TargetName == null || TargetConfiguration == null || TargetPlatform == null)
-			{
-				return 1;
-			}
+            }
+        }
+        private int QueryTargetDetails(CommandLineArguments Arguments, ILogger Logger, JsonSerializerOptions JsonOptions)
+        {
+            if (TargetName == null)
+            {
+                Logger.LogError("Missing argument Target");
+                return 1;
+            }
+            else if (TargetConfiguration == null)
+            {
+                Logger.LogError("Missing argument Configuration");
+                return 1;
+            }
+            else if (TargetPlatform == null)
+            {
+                Logger.LogError("Missing argument Platform");
+                return 1;
+            }
 
 			GenerateProjectFilesMode.TryParseProjectFileArgument(Arguments, Logger, out FileReference? ProjectFileArg);
 			List<string> RawArgs = new List<string> { TargetName, TargetConfiguration, TargetPlatform };
@@ -254,6 +273,7 @@ namespace UnrealBuildTool
 					CurrentTarget = UEBuildTarget.Create(TargetDescriptors[0], false, false, bUsePrecompiled, Logger);
 				}
 
+				LaunchSettings? CurrentLaunchSettings = null;
 				TargetIntellisenseInfo CurrentTargetIntellisenseInfo = new TargetIntellisenseInfo();
 				GetBrowseConfigurationResponse CurrentBrowseConfiguration = new GetBrowseConfigurationResponse { Success = true };
 
@@ -271,6 +291,18 @@ namespace UnrealBuildTool
 				// TODO: For installed builds, filter out all the binaries that aren't in mods
 				foreach (UEBuildBinary Binary in CurrentTarget.Binaries)
 				{
+					if (Binary.Type == UEBuildBinaryType.Executable && CurrentLaunchSettings == null && Binary.OutputFilePaths.Count == 1)
+					{
+						CurrentLaunchSettings = new LaunchSettings();
+						string ShortName = CurrentTarget.AppName;
+						CurrentLaunchSettings.Description = $"{ShortName} {CurrentTarget.Configuration} {CurrentTarget.Platform}";
+						CurrentLaunchSettings.BinaryPath = Binary.OutputFilePath.ToString();
+						if (CurrentTarget.ProjectFile != null && CurrentTarget.TargetType != TargetType.Program)
+						{
+							CurrentLaunchSettings.Arguments.Add(CurrentTarget.ProjectFile.ToString());
+						}
+					}
+					
 					HashSet<UEBuildModule> LinkEnvironmentVisitedModules = new HashSet<UEBuildModule>();
 					CppCompileEnvironment BinaryCompileEnvironment = Binary.CreateBinaryCompileEnvironment(GlobalCompileEnvironment);
 					CurrentBrowseConfiguration.Standard = BinaryCompileEnvironment.CppStandard.ToString();
@@ -329,6 +361,7 @@ namespace UnrealBuildTool
 				{
 					DirToModule = CurrentTargetIntellisenseInfo.DirToModule.ToDictionary(x => x.Key.ToString(), x => x.Value.Name),
 					ModuleToCompileSettings = CurrentTargetIntellisenseInfo.ModuleToCompileSettings.ToDictionary(x => x.Key.Name, x => x.Value),
+					LaunchSettings = CurrentLaunchSettings,
 				};
 
 				Console.WriteLine(JsonSerializer.Serialize(Result, JsonOptions));
