Restore missing using statement. UBT does not build with implicit usings.

diff --git a/Engine/Source/Programs/UnrealBuildTool/Matchers/MicrosoftEventMatcher.cs b/Engine/Source/Programs/UnrealBuildTool/Matchers/MicrosoftEventMatcher.cs
--- a/Engine/Source/Programs/UnrealBuildTool/Matchers/MicrosoftEventMatcher.cs
+++ b/Engine/Source/Programs/UnrealBuildTool/Matchers/MicrosoftEventMatcher.cs
@@ -1,5 +1,6 @@
 // Copyright Epic Games, Inc. All Rights Reserved.
 
+using System;
 using System.Text.RegularExpressions;
 using EpicGames.Core;
 using Microsoft.Extensions.Logging;
