UE-195321 InternalncludePaths typo (missing I on Include)

diff --git a/Engine/Source/Programs/UnrealBuildTool/Configuration/ModuleRules.cs b/Engine/Source/Programs/UnrealBuildTool/Configuration/ModuleRules.cs
--- a/Engine/Source/Programs/UnrealBuildTool/Configuration/ModuleRules.cs
+++ b/Engine/Source/Programs/UnrealBuildTool/Configuration/ModuleRules.cs
@@ -1166,7 +1166,7 @@ namespace UnrealBuildTool
 		/// <summary>
 		/// (This setting is currently not need as we discover all files from the 'Internal' folder) List of all paths to include files that are exposed to other internal modules
 		/// </summary>
-		public List<string> InternalncludePaths = new List<string>();
+		public List<string> InternalIncludePaths = new List<string>();
 
 		/// <summary>
 		/// List of all paths to this module's internal include files, not exposed to other modules (at least one include to the 'Private' path, more if we want to avoid relative paths)
diff --git a/Engine/Source/Programs/UnrealBuildTool/Configuration/UEBuildModule.cs b/Engine/Source/Programs/UnrealBuildTool/Configuration/UEBuildModule.cs
--- a/Engine/Source/Programs/UnrealBuildTool/Configuration/UEBuildModule.cs
+++ b/Engine/Source/Programs/UnrealBuildTool/Configuration/UEBuildModule.cs
@@ -212,7 +212,7 @@ namespace UnrealBuildTool
 			HashSet<string> PublicPreBuildLibraries = HashSetFromOptionalEnumerableStringParameter(Rules.PublicPreBuildLibraries);
 			PublicDefinitions = HashSetFromOptionalEnumerableStringParameter(Rules.PublicDefinitions);
 			PublicIncludePaths = CreateDirectoryHashSet(Rules.PublicIncludePaths);
-			InternalIncludePaths = CreateDirectoryHashSet(Rules.InternalncludePaths);
+			InternalIncludePaths = CreateDirectoryHashSet(Rules.InternalIncludePaths);
 			PublicSystemIncludePaths = CreateDirectoryHashSet(Rules.PublicSystemIncludePaths);
 			PublicSystemLibraryPaths = CreateDirectoryHashSet(Rules.PublicSystemLibraryPaths);
 			HashSet<string> PublicAdditionalLibraries = HashSetFromOptionalEnumerableStringParameter(Rules.PublicAdditionalLibraries.Union(PublicPreBuildLibraries));
diff --git a/Engine/Source/Programs/UnrealBuildTool/Configuration/UEBuildPlatform.cs b/Engine/Source/Programs/UnrealBuildTool/Configuration/UEBuildPlatform.cs
--- a/Engine/Source/Programs/UnrealBuildTool/Configuration/UEBuildPlatform.cs
+++ b/Engine/Source/Programs/UnrealBuildTool/Configuration/UEBuildPlatform.cs
@@ -856,7 +856,7 @@ namespace UnrealBuildTool
 
 			IOrderedEnumerable<DirectoryReference> PublicIncludePaths = Module.Rules.PublicIncludePaths.Select(x => DirectoryReference.FromString(x)!).OrderBy(x => x.FullName);
 			IOrderedEnumerable<DirectoryReference> PrivateIncludePaths = Module.Rules.PrivateIncludePaths.Select(x => DirectoryReference.FromString(x)!).OrderBy(x => x.FullName);
-			IOrderedEnumerable<DirectoryReference> InternalncludePaths = Module.Rules.InternalncludePaths.Select(x => DirectoryReference.FromString(x)!).OrderBy(x => x.FullName);
+			IOrderedEnumerable<DirectoryReference> InternalIncludePaths = Module.Rules.InternalIncludePaths.Select(x => DirectoryReference.FromString(x)!).OrderBy(x => x.FullName);
 			IOrderedEnumerable<DirectoryReference> PublicSystemIncludePaths = Module.Rules.PublicSystemIncludePaths.Select(x => DirectoryReference.FromString(x)!).OrderBy(x => x.FullName);
 			IOrderedEnumerable<UEBuildModule> OtherModules = AllModules.Where(x => x != Module).OrderBy(x => x.Name);
 
@@ -885,9 +885,9 @@ namespace UnrealBuildTool
 					LoggerFunc("External module '{Name}' is adding '{Path}' to PrivateIncludePaths. This path is unused.", Module.Name, Path);
 				}
 
-				foreach (DirectoryReference Path in InternalncludePaths)
+				foreach (DirectoryReference Path in InternalIncludePaths)
 				{
-					LoggerFunc("External module '{Name}' is adding '{Path}' to InternalncludePaths. This path is unused.", Module.Name, Path);
+					LoggerFunc("External module '{Name}' is adding '{Path}' to InternalIncludePaths. This path is unused.", Module.Name, Path);
 				}
 			}
 			else if (Module is UEBuildModuleCPP)
@@ -977,46 +977,46 @@ namespace UnrealBuildTool
 					}
 				}
 
-				foreach (DirectoryReference Path in InternalncludePaths)
+				foreach (DirectoryReference Path in InternalIncludePaths)
 				{
 					if (Path.IsUnderDirectory(Module.ModuleDirectory))
 					{
 						if (Path == Module.ModuleDirectory)
 						{
-							LoggerFunc("Module '{Name}' is adding root directory to InternalncludePaths. This is not allowed.", Module.Name);
+							LoggerFunc("Module '{Name}' is adding root directory to InternalIncludePaths. This is not allowed.", Module.Name);
 						}
 						else if (Path.IsUnderDirectory(DirectoryReference.Combine(Module.ModuleDirectory, "Private"))
 						|| Path.IsUnderDirectory(DirectoryReference.Combine(Module.ModuleDirectory, "Internal")))
 						{
-							LoggerFunc("Module '{Name}' is adding subdirectory '{Path}' to InternalncludePaths. This is not allowed.", Module.Name, Path.MakeRelativeTo(Module.ModuleDirectory));
+							LoggerFunc("Module '{Name}' is adding subdirectory '{Path}' to InternalIncludePaths. This is not allowed.", Module.Name, Path.MakeRelativeTo(Module.ModuleDirectory));
 						}
 						else if (Path.IsUnderDirectory(DirectoryReference.Combine(Module.ModuleDirectory, "Public"))
 							|| Path.IsUnderDirectory(DirectoryReference.Combine(Module.ModuleDirectory, "Classes")))
 						{
-							LoggerFuncSubDir("Module '{Name}' is adding subdirectory '{Path}' to InternalncludePaths. This is not necessary.", Module.Name, Path.MakeRelativeTo(Module.ModuleDirectory));
+							LoggerFuncSubDir("Module '{Name}' is adding subdirectory '{Path}' to InternalIncludePaths. This is not necessary.", Module.Name, Path.MakeRelativeTo(Module.ModuleDirectory));
 						}
 					}
 
 					UEBuildModule? OtherModule = OtherModules.FirstOrDefault(x => Path.IsUnderDirectory(x.ModuleDirectory));
 					if (OtherModule is UEBuildModuleExternal)
 					{
-						LoggerFunc("Module '{Name}' is adding '{Path}' from external module '{OtherModule}' to InternalncludePaths. Did you intend to add a public reference?", Module.Name, Path, OtherModule.Name);
+						LoggerFunc("Module '{Name}' is adding '{Path}' from external module '{OtherModule}' to InternalIncludePaths. Did you intend to add a public reference?", Module.Name, Path, OtherModule.Name);
 					}
 					else if (OtherModule is UEBuildModuleCPP)
 					{
 						if (Path == OtherModule.ModuleDirectory)
 						{
-							LoggerFunc("Module '{Name}' is adding root directory from '{OtherModule}' to InternalncludePaths. This is not allowed.", Module.Name, OtherModule.Name);
+							LoggerFunc("Module '{Name}' is adding root directory from '{OtherModule}' to InternalIncludePaths. This is not allowed.", Module.Name, OtherModule.Name);
 						}
 						else if (Path.IsUnderDirectory(DirectoryReference.Combine(OtherModule.ModuleDirectory, "Private"))
 							| Path.IsUnderDirectory(DirectoryReference.Combine(OtherModule.ModuleDirectory, "Internal")))
 						{
-							LoggerFunc("Module '{Name}' is adding '{Path}' from module '{OtherModule}' to InternalncludePaths. This is not allowed.", Module.Name, Path, OtherModule.Name);
+							LoggerFunc("Module '{Name}' is adding '{Path}' from module '{OtherModule}' to InternalIncludePaths. This is not allowed.", Module.Name, Path, OtherModule.Name);
 						}
 						else if (Path.IsUnderDirectory(DirectoryReference.Combine(OtherModule.ModuleDirectory, "Public"))
 							|| Path.IsUnderDirectory(DirectoryReference.Combine(OtherModule.ModuleDirectory, "Classes")))
 						{
-							LoggerFunc("Module '{Name}' is adding '{Path}' from module '{OtherModule}' to InternalncludePaths. Did you intend to add a public reference?", Module.Name, Path, OtherModule.Name);
+							LoggerFunc("Module '{Name}' is adding '{Path}' from module '{OtherModule}' to InternalIncludePaths. Did you intend to add a public reference?", Module.Name, Path, OtherModule.Name);
 						}
 					}
 				}
diff --git a/Engine/Source/Programs/UnrealBuildTool/Configuration/UEBuildTarget.cs b/Engine/Source/Programs/UnrealBuildTool/Configuration/UEBuildTarget.cs
--- a/Engine/Source/Programs/UnrealBuildTool/Configuration/UEBuildTarget.cs
+++ b/Engine/Source/Programs/UnrealBuildTool/Configuration/UEBuildTarget.cs
@@ -5311,7 +5311,7 @@ namespace UnrealBuildTool
 			// Make sure include paths don't end in trailing slashes. This can result in enclosing quotes being escaped when passed to command line tools.
 			RemoveTrailingSlashes(RulesObject.PublicIncludePaths);
 			RemoveTrailingSlashes(RulesObject.PublicSystemIncludePaths);
-			RemoveTrailingSlashes(RulesObject.InternalncludePaths);
+			RemoveTrailingSlashes(RulesObject.InternalIncludePaths);
 			RemoveTrailingSlashes(RulesObject.PrivateIncludePaths);
 			RemoveTrailingSlashes(RulesObject.PublicSystemLibraryPaths);
 
@@ -5431,7 +5431,7 @@ namespace UnrealBuildTool
 					if (RulesObject.File.IsUnderDirectory(ProjectSourceDirectoryName))
 					{
 						RulesObject.PublicIncludePaths = CombinePathList(ProjectSourceDirectoryName, RulesObject.PublicIncludePaths);
-						RulesObject.InternalncludePaths = CombinePathList(ProjectSourceDirectoryName, RulesObject.InternalncludePaths);
+						RulesObject.InternalIncludePaths = CombinePathList(ProjectSourceDirectoryName, RulesObject.InternalIncludePaths);
 						RulesObject.PrivateIncludePaths = CombinePathList(ProjectSourceDirectoryName, RulesObject.PrivateIncludePaths);
 						RulesObject.PublicSystemLibraryPaths = CombinePathList(ProjectSourceDirectoryName, RulesObject.PublicSystemLibraryPaths);
 					}
