Replace binary name with target name

diff --git a/Engine/Source/Programs/UnrealBuildTool/Modes/QueryMode.cs b/Engine/Source/Programs/UnrealBuildTool/Modes/QueryMode.cs
--- a/Engine/Source/Programs/UnrealBuildTool/Modes/QueryMode.cs
+++ b/Engine/Source/Programs/UnrealBuildTool/Modes/QueryMode.cs
@@ -294,8 +294,7 @@ namespace UnrealBuildTool
 					if (Binary.Type == UEBuildBinaryType.Executable && CurrentLaunchSettings == null && Binary.OutputFilePaths.Count == 1)
 					{
 						CurrentLaunchSettings = new LaunchSettings();
-						string ShortName = CurrentTarget.AppName;
-						CurrentLaunchSettings.Description = $"{ShortName} {CurrentTarget.Configuration} {CurrentTarget.Platform}";
+						CurrentLaunchSettings.Description = $"{CurrentTarget.TargetName} {CurrentTarget.Configuration} {CurrentTarget.Platform}";
 						CurrentLaunchSettings.BinaryPath = Binary.OutputFilePath.ToString();
 						if (CurrentTarget.ProjectFile != null && CurrentTarget.TargetType != TargetType.Program)
 						{
