- Fixed C# projects from being missing in VSCode projects on Mac/Linux

diff --git a/Engine/Source/Programs/UnrealBuildTool/ProjectFiles/VisualStudioCode/VSCodeProjectFileGenerator.cs b/Engine/Source/Programs/UnrealBuildTool/ProjectFiles/VisualStudioCode/VSCodeProjectFileGenerator.cs
--- a/Engine/Source/Programs/UnrealBuildTool/ProjectFiles/VisualStudioCode/VSCodeProjectFileGenerator.cs
+++ b/Engine/Source/Programs/UnrealBuildTool/ProjectFiles/VisualStudioCode/VSCodeProjectFileGenerator.cs
@@ -31,7 +31,7 @@ namespace UnrealBuildTool
 		private UnrealTargetPlatform HostPlatform = BuildHostPlatform.Current.Platform;
 		private bool bForeignProject;
 		private DirectoryReference ProjectRoot;
-		private string FrameworkExecutableExtension = ".exe";
+		private string FrameworkExecutableExtension = OperatingSystem.IsWindows() ? ".exe" : "";
 		private string FrameworkLibraryExtension = ".dll";
 
 		private readonly ConcurrentBag<BuildTarget> BuildTargets = new();
