UnrealBuildTool: Update preferred msvc compiler to 14.37.32822

diff --git a/Engine/Source/Programs/UnrealBuildTool/Platform/Windows/MicrosoftPlatformSDK.Versions.cs b/Engine/Source/Programs/UnrealBuildTool/Platform/Windows/MicrosoftPlatformSDK.Versions.cs
--- a/Engine/Source/Programs/UnrealBuildTool/Platform/Windows/MicrosoftPlatformSDK.Versions.cs
+++ b/Engine/Source/Programs/UnrealBuildTool/Platform/Windows/MicrosoftPlatformSDK.Versions.cs
@@ -43,6 +43,7 @@ namespace UnrealBuildTool
 		/// </summary>
 		static readonly VersionNumberRange[] PreferredVisualCppVersions = new VersionNumberRange[]
 		{
+			VersionNumberRange.Parse("14.37.32822", "14.37.99990"), // VS2022 17.7.x
 			VersionNumberRange.Parse("14.36.32532", "14.36.99999"), // VS2022 17.6.x
 			VersionNumberRange.Parse("14.35.32215", "14.35.99999"), // VS2022 17.5.x
 			VersionNumberRange.Parse("14.34.31933", "14.34.99999"), // VS2022 17.4.x
