- Fix a problem with BundleResources not being copied in to Mac and IOS .app bundles (plugin author reported - this is a feature Epic has not used)

diff --git a/Engine/Source/Programs/UnrealBuildTool/Platform/Mac/MacToolChain.cs b/Engine/Source/Programs/UnrealBuildTool/Platform/Mac/MacToolChain.cs
--- a/Engine/Source/Programs/UnrealBuildTool/Platform/Mac/MacToolChain.cs
+++ b/Engine/Source/Programs/UnrealBuildTool/Platform/Mac/MacToolChain.cs
@@ -1171,11 +1171,15 @@ namespace UnrealBuildTool
 				return OutputFiles;
 			}
 
-			if (BinaryLinkEnvironment.BundleDirectory != null)
+			// modern handle bundles via xcode
+			if (!bUseModernXcode)
 			{
-				foreach (UEBuildBundleResource Resource in BinaryLinkEnvironment.AdditionalBundleResources)
+				if (BinaryLinkEnvironment.BundleDirectory != null)
 				{
-					OutputFiles.Add(CopyBundleResource(Resource, Executable, BinaryLinkEnvironment.BundleDirectory, Graph));
+					foreach (UEBuildBundleResource Resource in BinaryLinkEnvironment.AdditionalBundleResources)
+					{
+						OutputFiles.Add(CopyBundleResource(Resource, Executable, BinaryLinkEnvironment.BundleDirectory, Graph));
+					}
 				}
 			}
 
diff --git a/Engine/Source/Programs/UnrealBuildTool/ProjectFiles/Xcode/XcodeProject.File.cs b/Engine/Source/Programs/UnrealBuildTool/ProjectFiles/Xcode/XcodeProject.File.cs
--- a/Engine/Source/Programs/UnrealBuildTool/ProjectFiles/Xcode/XcodeProject.File.cs
+++ b/Engine/Source/Programs/UnrealBuildTool/ProjectFiles/Xcode/XcodeProject.File.cs
@@ -122,6 +122,7 @@ namespace UnrealBuildTool.XcodeProjectXcconfig
 			".storyboard",
 			".xcframework",
 			".framework",
+			".bundle",
 			".lproj",
 			".png",
 			// no extensions will be folders, like "cookeddata" that we want to copy as a resource
diff --git a/Engine/Source/Programs/UnrealBuildTool/ProjectFiles/Xcode/XcodeProject.cs b/Engine/Source/Programs/UnrealBuildTool/ProjectFiles/Xcode/XcodeProject.cs
--- a/Engine/Source/Programs/UnrealBuildTool/ProjectFiles/Xcode/XcodeProject.cs
+++ b/Engine/Source/Programs/UnrealBuildTool/ProjectFiles/Xcode/XcodeProject.cs
@@ -1501,6 +1501,15 @@ namespace UnrealBuildTool.XcodeProjectXcconfig
 		{
 			// look up to see if we had cached any Frameworks
 			Tuple<ProjectFile, UnrealTargetPlatform> FrameworkKey = Tuple.Create((ProjectFile)ProjectFile, Platform);
+			IEnumerable<UEBuildBundleResource>? Bundles;
+			if (XcodeProjectFileGenerator.TargetBundles.TryGetValue(FrameworkKey, out Bundles))
+			{
+				foreach (UEBuildBundleResource Bundle in Bundles)
+				{
+					ResourcesBuildPhase.AddFolderResource(new DirectoryReference(Bundle.ResourcePath!), "Resources");
+				}
+			}
+
 			IEnumerable<UEBuildFramework>? Frameworks;
 			if (XcodeProjectFileGenerator.TargetFrameworks.TryGetValue(FrameworkKey, out Frameworks))
 			{
diff --git a/Engine/Source/Programs/UnrealBuildTool/ProjectFiles/Xcode/XcodeProjectFileGenerator.cs b/Engine/Source/Programs/UnrealBuildTool/ProjectFiles/Xcode/XcodeProjectFileGenerator.cs
--- a/Engine/Source/Programs/UnrealBuildTool/ProjectFiles/Xcode/XcodeProjectFileGenerator.cs
+++ b/Engine/Source/Programs/UnrealBuildTool/ProjectFiles/Xcode/XcodeProjectFileGenerator.cs
@@ -559,6 +559,7 @@ namespace UnrealBuildTool
 							foreach (UEBuildModule Module in Binary.Modules)
 							{
 								Dylibs.AddRange(Module.Rules.RuntimeDependencies.Inner.Where(x => x.Path.StartsWith("$(BinaryOutputDir)") && Path.GetExtension(x.SourcePath) == ".dylib"));
+								Bundles.AddRange(Module.Rules.AdditionalBundleResources.Select(x => new UEBuildBundleResource(x)));
 							}
 						}
 
