Disable no-deprecated-copy warning when building with clang on Windows

diff --git a/Engine/Source/Programs/UnrealBuildTool/ToolChain/ClangWarnings.cs b/Engine/Source/Programs/UnrealBuildTool/ToolChain/ClangWarnings.cs
--- a/Engine/Source/Programs/UnrealBuildTool/ToolChain/ClangWarnings.cs
+++ b/Engine/Source/Programs/UnrealBuildTool/ToolChain/ClangWarnings.cs
@@ -36,6 +36,11 @@ namespace UnrealBuildTool
 			{
 				Arguments.Add("-Wno-bitwise-instead-of-logical");       // https://clang.llvm.org/docs/DiagnosticsReference.html#wbitwise-instead-of-logical			// new warning for clang 14
 			}
+			if (ClangVersion >= new VersionNumber(16))
+			{
+				Arguments.Add("-Wno-deprecated-copy");
+				Arguments.Add("-Wno-deprecated-copy-with-user-provided-copy");
+			}
 
 			Arguments.Add("-Wno-gnu-string-literal-operator-template"); // https://clang.llvm.org/docs/DiagnosticsReference.html#wgnu-string-literal-operator-template	// We use this feature to allow static FNames.
 			Arguments.Add("-Wno-inconsistent-missing-override");        // https://clang.llvm.org/docs/DiagnosticsReference.html#winconsistent-missing-override			// ?? no reason given
