- Respect the flag for frameworks that don't want to link, to not actually link

diff --git a/Engine/Source/Programs/UnrealBuildTool/Configuration/UEBuildFramework.cs b/Engine/Source/Programs/UnrealBuildTool/Configuration/UEBuildFramework.cs
--- a/Engine/Source/Programs/UnrealBuildTool/Configuration/UEBuildFramework.cs
+++ b/Engine/Source/Programs/UnrealBuildTool/Configuration/UEBuildFramework.cs
@@ -110,6 +110,7 @@ namespace UnrealBuildTool
 			this.FrameworkDirectory = FrameworkDirectory;
 			this.CopyBundledAssets = CopyBundledAssets;
 			this.bCopyFramework = bCopyFramework;
+			this.bLinkFramework = bLinkFramework;
 			if (this.FrameworkDirectory.FullName.EndsWith(".xcframework"))
 			{
 				XCFrameworkVariants = LoadXCFrameworkVariants(Logger);
diff --git a/Engine/Source/Programs/UnrealBuildTool/Platform/IOS/IOSToolChain.cs b/Engine/Source/Programs/UnrealBuildTool/Platform/IOS/IOSToolChain.cs
--- a/Engine/Source/Programs/UnrealBuildTool/Platform/IOS/IOSToolChain.cs
+++ b/Engine/Source/Programs/UnrealBuildTool/Platform/IOS/IOSToolChain.cs
@@ -225,17 +225,20 @@ namespace UnrealBuildTool
 			// Add additional frameworks so that their headers can be found
 			foreach (UEBuildFramework Framework in CompileEnvironment.AdditionalFrameworks)
 			{
-				DirectoryReference? FrameworkDirectory = Framework.GetFrameworkDirectory(CompileEnvironment.Platform, CompileEnvironment.Architecture, Logger);
-				if (FrameworkDirectory != null)
+				if (Framework.bLinkFramework)
 				{
-					string FrameworkDir = FrameworkDirectory.FullName;
-					// embedded frameworks have a framework inside of this directory, so we use this directory. regular frameworks need to go one up to point to the 
-					// directory containing the framework. -F gives a path to look for the -framework
-					if (FrameworkDir.EndsWith(".framework"))
+					DirectoryReference? FrameworkDirectory = Framework.GetFrameworkDirectory(CompileEnvironment.Platform, CompileEnvironment.Architecture, Logger);
+					if (FrameworkDirectory != null)
 					{
-						FrameworkDir = Path.GetDirectoryName(FrameworkDir)!;
+						string FrameworkDir = FrameworkDirectory.FullName;
+						// embedded frameworks have a framework inside of this directory, so we use this directory. regular frameworks need to go one up to point to the 
+						// directory containing the framework. -F gives a path to look for the -framework
+						if (FrameworkDir.EndsWith(".framework"))
+						{
+							FrameworkDir = Path.GetDirectoryName(FrameworkDir)!;
+						}
+						Arguments.Add($"-F\"{FrameworkDir}\"");
 					}
-					Arguments.Add($"-F\"{FrameworkDir}\"");
 				}
 			}
 		}
