UnrealBuildTool: Support generating msvc link repros with -LinkRepro=<path>

#rnx

diff --git a/Engine/Source/Programs/UnrealBuildTool/Platform/Windows/UEBuildWindows.cs b/Engine/Source/Programs/UnrealBuildTool/Platform/Windows/UEBuildWindows.cs
--- a/Engine/Source/Programs/UnrealBuildTool/Platform/Windows/UEBuildWindows.cs
+++ b/Engine/Source/Programs/UnrealBuildTool/Platform/Windows/UEBuildWindows.cs
@@ -383,6 +383,13 @@ namespace UnrealBuildTool
 		[CommandLine("-ObjSrcMap")]
 		public string? ObjSrcMapFile = null;
 
+		/// <summary>
+		/// Whether to have the linker or library tool to generate a link repro in the specified directory
+		/// See https://learn.microsoft.com/en-us/cpp/build/reference/linkrepro for more information
+		/// </summary>
+		[CommandLine("-LinkRepro=")]
+		public string? LinkReproDir = null;
+
 		/// <summary>
 		/// Provides a Module Definition File (.def) to the linker to describe various attributes of a DLL.
 		/// Necessary when exporting functions by ordinal values instead of by name.
@@ -753,6 +760,8 @@ namespace UnrealBuildTool
 
 		public string? ObjSrcMapFile => Inner.ObjSrcMapFile;
 
+		public string? LinkReproDir => Inner.LinkReproDir;
+
 		public string? ModuleDefinitionFile => Inner.ModuleDefinitionFile;
 
 		public string? ManifestFile => Inner.ManifestFile;
diff --git a/Engine/Source/Programs/UnrealBuildTool/Platform/Windows/VCToolChain.cs b/Engine/Source/Programs/UnrealBuildTool/Platform/Windows/VCToolChain.cs
--- a/Engine/Source/Programs/UnrealBuildTool/Platform/Windows/VCToolChain.cs
+++ b/Engine/Source/Programs/UnrealBuildTool/Platform/Windows/VCToolChain.cs
@@ -2563,6 +2563,18 @@ namespace UnrealBuildTool
 			// Allow remote linking. Note that this may be overriden by the executor (eg. XGE.bAllowRemoteLinking)
 			LinkAction.bCanExecuteRemotely = true;
 
+			// Create link repro if requested, this argument is intentionally not added to the response file
+			if (Target.WindowsPlatform.LinkReproDir != null)
+			{
+				DirectoryReference LinkReproRoot = new DirectoryReference(Target.WindowsPlatform.LinkReproDir);
+				DirectoryReference LinkReproPath = DirectoryReference.Combine(LinkReproRoot, OutputFile.Name);
+				DirectoryReference.CreateDirectory(LinkReproPath);
+				LinkAction.CommandArguments = $"{LinkAction.CommandArguments} /LINKREPRO:{Utils.MakePathSafeToUseWithCommandLine(LinkReproPath.FullName)}";
+
+				LinkAction.bCanExecuteRemotely = false;
+				LinkAction.bCanExecuteInBox = false;
+			}
+
 			Logger.LogDebug("     Linking: {StatusDescription}", LinkAction.StatusDescription);
 			Logger.LogDebug("     Command: {CommandArguments}", LinkAction.CommandArguments);
 
@@ -2640,6 +2652,7 @@ namespace UnrealBuildTool
 			"/errorReport:",
 			"/d2:",
 			"/NATVIS:",
+			"/LINKREPRO:",
 			"/experimental:deterministic"
 		};
 		private IEnumerable<string> RemovePGOIgnoredSwitches(IEnumerable<string> SourceArguments)
