[UBT]
* Minor optimization in EnumerateReadyToCompileActions

diff --git a/Engine/Source/Programs/UnrealBuildTool/Executors/ImmediateActionQueue.cs b/Engine/Source/Programs/UnrealBuildTool/Executors/ImmediateActionQueue.cs
--- a/Engine/Source/Programs/UnrealBuildTool/Executors/ImmediateActionQueue.cs
+++ b/Engine/Source/Programs/UnrealBuildTool/Executors/ImmediateActionQueue.cs
@@ -508,8 +508,9 @@ namespace UnrealBuildTool
 		/// <returns>Enumerations of all ready to compile actions.</returns>
 		public IEnumerable<LinkedAction> EnumerateReadyToCompileActions()
 		{
-			foreach (ActionState actionState in Actions)
+			for (int actionIndex = _firstPendingAction; actionIndex != Actions.Length; ++actionIndex)
 			{
+				var actionState = Actions[actionIndex];
 				if (actionState.Status == ActionStatus.Queued &&
 					actionState.Phase == ActionPhase.Compile &&
 					GetActionReadyState(actionState.Action) == ActionReadyState.Ready)
