UnrealBuildTool: Fix path separator when generting type library header

#rnx

diff --git a/Engine/Source/Programs/UnrealBuildTool/Platform/Windows/VCToolChain.cs b/Engine/Source/Programs/UnrealBuildTool/Platform/Windows/VCToolChain.cs
--- a/Engine/Source/Programs/UnrealBuildTool/Platform/Windows/VCToolChain.cs
+++ b/Engine/Source/Programs/UnrealBuildTool/Platform/Windows/VCToolChain.cs
@@ -2078,7 +2078,7 @@ namespace UnrealBuildTool
 			Contents.AppendLine("#include <unknwn.h>");
 			Contents.AppendLine();
 
-			Contents.AppendFormat("#import \"{0}\"", TypeLibrary.FileName);
+			Contents.AppendFormat("#import \"{0}\"", TypeLibrary.FileName.Replace('\\', '/'));
 			if (!String.IsNullOrEmpty(TypeLibrary.Attributes))
 			{
 				Contents.Append(' ');
