[UBT]
* Made sure linking of executable is always running on local machine (not distributed)

diff --git a/Engine/Source/Programs/UnrealBuildTool/Platform/Windows/VCToolChain.cs b/Engine/Source/Programs/UnrealBuildTool/Platform/Windows/VCToolChain.cs
--- a/Engine/Source/Programs/UnrealBuildTool/Platform/Windows/VCToolChain.cs
+++ b/Engine/Source/Programs/UnrealBuildTool/Platform/Windows/VCToolChain.cs
@@ -2561,7 +2561,10 @@ namespace UnrealBuildTool
 			LinkAction.bProducesImportLibrary = bBuildImportLibraryOnly || LinkEnvironment.bIsBuildingDLL;
 
 			// Allow remote linking. Note that this may be overriden by the executor (eg. XGE.bAllowRemoteLinking)
-			LinkAction.bCanExecuteRemotely = true;
+			if (LinkAction.bProducesImportLibrary || LinkEnvironment.bIsBuildingDLL)
+			{
+				LinkAction.bCanExecuteRemotely = true;
+			}
 
 			// Create link repro if requested, this argument is intentionally not added to the response file
 			if (Target.WindowsPlatform.LinkReproDir != null)
