UnrealBuildTool: Add setting to control disabling some debug info optimizations when compiling msvc platforms with clang that can in some cases cause missing debug information

Can be enabled with
* TargetRules: WindowsPlatform.bClangStandaloneDebug
* INI: /Script/WindowsTargetPlatform.WindowsTargetSettings - bClangStandaloneDebug
* Cmd Line: -ClangStandaloneDebug

This will cause link time and pdb size to increase, expecially for monolithic targets

diff --git a/Engine/Source/Programs/UnrealBuildTool/Platform/Windows/UEBuildWindows.cs b/Engine/Source/Programs/UnrealBuildTool/Platform/Windows/UEBuildWindows.cs
--- a/Engine/Source/Programs/UnrealBuildTool/Platform/Windows/UEBuildWindows.cs
+++ b/Engine/Source/Programs/UnrealBuildTool/Platform/Windows/UEBuildWindows.cs
@@ -291,6 +291,15 @@ namespace UnrealBuildTool
 		[CommandLine("-VCFastFail")]
 		public bool bVCFastFail = false;
 
+		/// <summary>
+		/// True if optimizations to reduce the size of debug information should be disabled
+		/// See https://clang.llvm.org/docs/UsersManual.html#cmdoption-fstandalone-debug for more information
+		/// </summary>
+		[ConfigFile(ConfigHierarchyType.Engine, "/Script/WindowsTargetPlatform.WindowsTargetSettings", "bClangStandaloneDebug")]
+		[XmlConfigFile(Category = "WindowsPlatform")]
+		[CommandLine("-ClangStandaloneDebug")]
+		public bool bClangStandaloneDebug = false;
+
 		/// <summary>
 		/// True if we should use the Clang linker (LLD) when we are compiling with Clang, or Intel linker (xilink\xilib) when we are compiling with Intel oneAPI, otherwise we use the MSVC linker.
 		/// </summary>
@@ -742,6 +751,8 @@ namespace UnrealBuildTool
 
 		public bool bVCFastFail => Inner.bVCFastFail;
 
+		public bool bClangStandaloneDebug => Inner.bClangStandaloneDebug;
+
 		public bool bAllowClangLinker => Inner.bAllowClangLinker;
 
 		public bool bEnableRayTracing => Inner.bEnableRayTracing;
diff --git a/Engine/Source/Programs/UnrealBuildTool/Platform/Windows/VCToolChain.cs b/Engine/Source/Programs/UnrealBuildTool/Platform/Windows/VCToolChain.cs
--- a/Engine/Source/Programs/UnrealBuildTool/Platform/Windows/VCToolChain.cs
+++ b/Engine/Source/Programs/UnrealBuildTool/Platform/Windows/VCToolChain.cs
@@ -770,6 +770,12 @@ namespace UnrealBuildTool
 				{
 					Arguments.Add("/Z7");
 				}
+
+				// https://clang.llvm.org/docs/UsersManual.html#cmdoption-fstandalone-debug
+				if (Target.WindowsPlatform.Compiler.IsClang() && Target.WindowsPlatform.bClangStandaloneDebug)
+				{
+					Arguments.Add("-fstandalone-debug");
+				}
 			}
 
 			// Specify the appropriate runtime library based on the platform and config.
