UnrealBuildTool: Use Hostarm64 compiler if running on an arm64 cpu

#rnx

diff --git a/Engine/Source/Programs/UnrealBuildTool/Platform/Windows/MicrosoftPlatformSDK.cs b/Engine/Source/Programs/UnrealBuildTool/Platform/Windows/MicrosoftPlatformSDK.cs
--- a/Engine/Source/Programs/UnrealBuildTool/Platform/Windows/MicrosoftPlatformSDK.cs
+++ b/Engine/Source/Programs/UnrealBuildTool/Platform/Windows/MicrosoftPlatformSDK.cs
@@ -8,6 +8,7 @@ using System.Diagnostics;
 using System.Diagnostics.CodeAnalysis;
 using System.IO;
 using System.Linq;
+using System.Runtime.InteropServices;
 using System.Runtime.Versioning;
 using System.Text.RegularExpressions;
 using EpicGames.Core;
@@ -1155,6 +1156,11 @@ namespace UnrealBuildTool
 			}
 		}
 
+		/// <summary>
+		/// The compiler host directory name for the running process architecture.
+		/// </summary>
+		public static string MSVCHostDirectoryName => RuntimeInformation.ProcessArchitecture == Architecture.Arm64 ? "Hostarm64" : "Hostx64";
+
 		/// <summary>
 		/// Determines if the given path is a valid Visual C++ version number
 		/// </summary>
@@ -1163,7 +1169,7 @@ namespace UnrealBuildTool
 		/// <returns>True if the path is a valid version</returns>
 		static bool IsValidToolChainDirMSVC(DirectoryReference ToolChainDir, [NotNullWhen(true)] out VersionNumber? Version)
 		{
-			FileReference CompilerExe = FileReference.Combine(ToolChainDir, "bin", "Hostx64", "x64", "cl.exe");
+			FileReference CompilerExe = FileReference.Combine(ToolChainDir, "bin", MSVCHostDirectoryName, "x64", "cl.exe");
 			if (!FileReference.Exists(CompilerExe))
 			{
 				Version = null;
@@ -1187,7 +1193,7 @@ namespace UnrealBuildTool
 		/// <returns>True if the given directory contains a 64-bit toolchain</returns>
 		static bool Has64BitToolChain(DirectoryReference ToolChainDir)
 		{
-			return FileReference.Exists(FileReference.Combine(ToolChainDir, "bin", "Hostx64", "x64", "cl.exe"));
+			return FileReference.Exists(FileReference.Combine(ToolChainDir, "bin", MSVCHostDirectoryName, "x64", "cl.exe"));
 		}
 
 		/// <summary>
@@ -1197,7 +1203,7 @@ namespace UnrealBuildTool
 		/// <returns>True if the given directory contains the arm64 toolchain</returns>
 		static bool HasArm64ToolChain(DirectoryReference ToolChainDir)
 		{
-			return FileReference.Exists(FileReference.Combine(ToolChainDir, "bin", "Hostx64", "arm64", "cl.exe"));
+			return FileReference.Exists(FileReference.Combine(ToolChainDir, "bin", MSVCHostDirectoryName, "arm64", "cl.exe"));
 		}
 
 		/// <summary>
diff --git a/Engine/Source/Programs/UnrealBuildTool/Platform/Windows/VCEnvironment.cs b/Engine/Source/Programs/UnrealBuildTool/Platform/Windows/VCEnvironment.cs
--- a/Engine/Source/Programs/UnrealBuildTool/Platform/Windows/VCEnvironment.cs
+++ b/Engine/Source/Programs/UnrealBuildTool/Platform/Windows/VCEnvironment.cs
@@ -183,7 +183,7 @@ namespace UnrealBuildTool
 		/// <returns>Directory containing the 64-bit toolchain binaries</returns>
 		public static DirectoryReference GetVCToolPath(DirectoryReference VCToolChainDir, UnrealArch Architecture)
 		{
-			FileReference CompilerPath = FileReference.Combine(VCToolChainDir, "bin", "Hostx64", Architecture.WindowsToolChain, "cl.exe");
+			FileReference CompilerPath = FileReference.Combine(VCToolChainDir, "bin", MicrosoftPlatformSDK.MSVCHostDirectoryName, Architecture.WindowsToolChain, "cl.exe");
 			if (FileReference.Exists(CompilerPath))
 			{
 				return CompilerPath.Directory;
