[UBT]
* Added support for pch using header units.

diff --git a/Engine/Source/Programs/UnrealBuildTool/Configuration/TargetRules.cs b/Engine/Source/Programs/UnrealBuildTool/Configuration/TargetRules.cs
--- a/Engine/Source/Programs/UnrealBuildTool/Configuration/TargetRules.cs
+++ b/Engine/Source/Programs/UnrealBuildTool/Configuration/TargetRules.cs
@@ -2099,6 +2099,12 @@ namespace UnrealBuildTool
 		/// </summary>
 		public bool bAllowRemotelyCompiledPCHs = false;
 
+		/// <summary>
+		/// Will replace pch with ifc and use header units instead. This is an experimental and msvc-only feature
+		/// </summary>
+		[CommandLine("-HeaderUnits")]
+		public bool bUseHeaderUnitsForPch = false;
+
 		/// <summary>
 		/// Whether headers in system paths should be checked for modification when determining outdated actions.
 		/// </summary>
@@ -3402,6 +3408,8 @@ namespace UnrealBuildTool
 
 		public bool bAllowRemotelyCompiledPCHs => Inner.bAllowRemotelyCompiledPCHs;
 
+		public bool bUseHeaderUnitsForPch => Inner.bUseHeaderUnitsForPch;
+
 		public bool bCheckSystemHeadersForModification => Inner.bCheckSystemHeadersForModification;
 
 		public bool bDisableLinking => Inner.bDisableLinking;
diff --git a/Engine/Source/Programs/UnrealBuildTool/Configuration/UEBuildTarget.cs b/Engine/Source/Programs/UnrealBuildTool/Configuration/UEBuildTarget.cs
--- a/Engine/Source/Programs/UnrealBuildTool/Configuration/UEBuildTarget.cs
+++ b/Engine/Source/Programs/UnrealBuildTool/Configuration/UEBuildTarget.cs
@@ -4807,6 +4807,7 @@ namespace UnrealBuildTool
 			GlobalCompileEnvironment.bPGOOptimize = Rules.bPGOOptimize;
 			GlobalCompileEnvironment.bPGOProfile = Rules.bPGOProfile;
 			GlobalCompileEnvironment.bAllowRemotelyCompiledPCHs = Rules.bAllowRemotelyCompiledPCHs;
+			GlobalCompileEnvironment.bUseHeaderUnitsForPch = Rules.bUseHeaderUnitsForPch;
 			GlobalCompileEnvironment.bCheckSystemHeadersForModification = Rules.bCheckSystemHeadersForModification;
 			GlobalCompileEnvironment.bPrintTimingInfo = Rules.bPrintToolChainTimingInfo;
 			GlobalCompileEnvironment.bUseRTTI = Rules.bForceEnableRTTI;
diff --git a/Engine/Source/Programs/UnrealBuildTool/Platform/Windows/VCToolChain.cs b/Engine/Source/Programs/UnrealBuildTool/Platform/Windows/VCToolChain.cs
--- a/Engine/Source/Programs/UnrealBuildTool/Platform/Windows/VCToolChain.cs
+++ b/Engine/Source/Programs/UnrealBuildTool/Platform/Windows/VCToolChain.cs
@@ -228,16 +228,34 @@ namespace UnrealBuildTool
 		{
 			string PchThroughHeaderFilePath = NormalizeCommandLinePath(PchThroughHeaderFile);
 			string CreatePchFilePath = NormalizeCommandLinePath(CreatePchFile);
-			Arguments.Add($"/Yc\"{PchThroughHeaderFilePath}\"");
-			Arguments.Add($"/Fp\"{CreatePchFilePath}\"");
+			if (CreatePchFile.Name.EndsWith(".ifc"))
+			{
+				Arguments.Add($"/exportHeader");
+				Arguments.Add($"/translateInclude");
+				Arguments.Add($"/dxifcInlineFunctions");
+				Arguments.Add($"/ifcOutput {CreatePchFilePath}");
+			}
+			else
+			{
+				Arguments.Add($"/Yc\"{PchThroughHeaderFilePath}\"");
+				Arguments.Add($"/Fp\"{CreatePchFilePath}\"");
+			}
 		}
 
 		public static void AddUsingPchFile(List<string> Arguments, FileItem PchThroughHeaderFile, FileItem UsingPchFile)
 		{
 			string PchThroughHeaderFilePath = NormalizeCommandLinePath(PchThroughHeaderFile);
 			string UsingPchFilePath = NormalizeCommandLinePath(UsingPchFile);
-			Arguments.Add($"/Yu\"{PchThroughHeaderFilePath}\"");
-			Arguments.Add($"/Fp\"{UsingPchFilePath}\"");
+			if (UsingPchFile.Name.EndsWith(".ifc"))
+			{
+				Arguments.Add($"/translateInclude");
+				Arguments.Add($"/headerUnit:quote {PchThroughHeaderFilePath}={UsingPchFilePath}");
+			}
+			else
+			{
+				Arguments.Add($"/Yu\"{PchThroughHeaderFilePath}\"");
+				Arguments.Add($"/Fp\"{UsingPchFilePath}\"");
+			}
 		}
 
 		public static void AddPreprocessedFile(List<string> Arguments, FileItem PreprocessedFile, ILogger Logger)
@@ -906,6 +924,14 @@ namespace UnrealBuildTool
 				Arguments.Add("/wd4838");
 			}
 
+			if (CompileEnvironment.bUseHeaderUnitsForPch)
+			{
+				Arguments.Add("/wd4324"); // 'struct_name' : structure was padded due to __declspec(align())
+				Arguments.Add("/wd4201"); // nonstandard extension used: nameless struct/union
+				//Arguments.Add("/wd5106");
+				//Arguments.Add("/wd5260");
+			}
+
 			if (CompileEnvironment.Architecture == UnrealArch.Arm64ec)
 			{
 				Arguments.Add("/arm64EC");
@@ -1483,7 +1509,8 @@ namespace UnrealBuildTool
 					Graph.CreateIntermediateTextFile(CompileAction.SourceFile, PchCppFile);
 
 					// Add the precompiled header file to the produced items list.
-					CompileAction.CreatePchFile = FileItem.GetItemByFileReference(FileReference.Combine(OutputDir, SourceFile.Location.GetFileName() + ".pch"));
+					string PchExtension = CompileEnvironment.bUseHeaderUnitsForPch ? ".ifc" : ".pch";
+					CompileAction.CreatePchFile = FileItem.GetItemByFileReference(FileReference.Combine(OutputDir, SourceFile.Location.GetFileName() + PchExtension));
 					CompileAction.PchThroughHeaderFile = FileItem.GetItemByFileReference(CompileEnvironment.PrecompiledHeaderIncludeFilename);
 
 					// If we're creating a PCH that will be used to compile source files for a library, we need
diff --git a/Engine/Source/Programs/UnrealBuildTool/System/CppCompileEnvironment.cs b/Engine/Source/Programs/UnrealBuildTool/System/CppCompileEnvironment.cs
--- a/Engine/Source/Programs/UnrealBuildTool/System/CppCompileEnvironment.cs
+++ b/Engine/Source/Programs/UnrealBuildTool/System/CppCompileEnvironment.cs
@@ -229,6 +229,11 @@ namespace UnrealBuildTool
 		/// </summary>
 		public PrecompiledHeaderAction PrecompiledHeaderAction = PrecompiledHeaderAction.None;
 
+		/// <summary>
+		/// Will replace pch with ifc and use header units instead
+		/// </summary>
+		public bool bUseHeaderUnitsForPch = false;
+
 		/// <summary>
 		/// Whether artifacts from this compile are shared with other targets. If so, we should not apply any target-wide modifications to the compile environment.
 		/// </summary>
@@ -757,6 +762,7 @@ namespace UnrealBuildTool
 			PGODirectory = Other.PGODirectory;
 			bPrintTimingInfo = Other.bPrintTimingInfo;
 			bAllowRemotelyCompiledPCHs = Other.bAllowRemotelyCompiledPCHs;
+			bUseHeaderUnitsForPch = Other.bUseHeaderUnitsForPch;
 			UserIncludePaths = new HashSet<DirectoryReference>(Other.UserIncludePaths);
 			SystemIncludePaths = new HashSet<DirectoryReference>(Other.SystemIncludePaths);
 			SharedUserIncludePaths = new HashSet<DirectoryReference>(Other.SharedUserIncludePaths);
