UnrealBuildTool: Preferr 2022 msbuild

#rnx

diff --git a/Engine/Source/Programs/UnrealBuildTool/Platform/Windows/MicrosoftPlatformSDK.cs b/Engine/Source/Programs/UnrealBuildTool/Platform/Windows/MicrosoftPlatformSDK.cs
--- a/Engine/Source/Programs/UnrealBuildTool/Platform/Windows/MicrosoftPlatformSDK.cs
+++ b/Engine/Source/Programs/UnrealBuildTool/Platform/Windows/MicrosoftPlatformSDK.cs
@@ -1236,9 +1236,9 @@ namespace UnrealBuildTool
 		[SupportedOSPlatform("windows")]
 		public static bool TryGetMsBuildPath(ILogger Logger, [NotNullWhen(true)] out FileReference? OutLocation)
 		{
-			// Get the Visual Studio 2019 install directory
-			List<DirectoryReference> InstallDirs2019 = MicrosoftPlatformSDK.FindVisualStudioInstallations(WindowsCompiler.VisualStudio2019, Logger).ConvertAll(x => x.BaseDir);
-			foreach (DirectoryReference InstallDir in InstallDirs2019)
+			// Get the Visual Studio 2022 install directory
+			List<DirectoryReference> InstallDirs2022 = MicrosoftPlatformSDK.FindVisualStudioInstallations(WindowsCompiler.VisualStudio2022, Logger).ConvertAll(x => x.BaseDir);
+			foreach (DirectoryReference InstallDir in InstallDirs2022)
 			{
 				FileReference MsBuildLocation = FileReference.Combine(InstallDir, "MSBuild", "Current", "Bin", "MSBuild.exe");
 				if (FileReference.Exists(MsBuildLocation))
@@ -1248,9 +1248,9 @@ namespace UnrealBuildTool
 				}
 			}
 
-			// Get the Visual Studio 2022 install directory
-			List<DirectoryReference> InstallDirs2022 = MicrosoftPlatformSDK.FindVisualStudioInstallations(WindowsCompiler.VisualStudio2022, Logger).ConvertAll(x => x.BaseDir);
-			foreach (DirectoryReference InstallDir in InstallDirs2022)
+			// Get the Visual Studio 2019 install directory
+			List<DirectoryReference> InstallDirs2019 = MicrosoftPlatformSDK.FindVisualStudioInstallations(WindowsCompiler.VisualStudio2019, Logger).ConvertAll(x => x.BaseDir);
+			foreach (DirectoryReference InstallDir in InstallDirs2019)
 			{
 				FileReference MsBuildLocation = FileReference.Combine(InstallDir, "MSBuild", "Current", "Bin", "MSBuild.exe");
 				if (FileReference.Exists(MsBuildLocation))
