Android Studio build improvements:
- Allow pass through of BypassGradlePackaging argument in UBT
- Allow setting skippackage to actually skip packaging on Android in UBT
- Fix to Android Studio detection to skip build.gradle stomping when building in RunUAT or VisualStudio. Requires setting File > Settings > Build, Execution, Deployment > Build Tools > Gradle > "Generate *.iml files for modules imported from Gradle"
- Save RunConfiguration for app config in AndroidStudio so that we can load already having the debugger set to dual and the correct LLDB StartupCommand set.
#tests Android, Dev, Client, compile, mobile

diff --git a/Engine/Build/Android/Java/gradle/.idea/runConfigurations/app.xml b/Engine/Build/Android/Java/gradle/.idea/runConfigurations/app.xml
new file mode 100644
--- /dev/null
+++ b/Engine/Build/Android/Java/gradle/.idea/runConfigurations/app.xml
@@ -0,0 +1,9 @@
+<component name="ProjectRunConfigurationManager">
+	<configuration default="false" name="app" type="AndroidRunConfigurationType" factoryName="Android App">
+		<module name="app.app.main" />
+		<option name="DEBUGGER_TYPE" value="Hybrid" />
+		<Hybrid>
+			<startup_commands startup_commands_attr="command script import &quot;${PY_VISUALIZER_PATH}&quot;" />
+		</Hybrid>
+	</configuration>
+</component>
\ No newline at end of file
diff --git a/Engine/Source/Programs/UnrealBuildTool/Platform/Android/UEDeployAndroid.cs b/Engine/Source/Programs/UnrealBuildTool/Platform/Android/UEDeployAndroid.cs
--- a/Engine/Source/Programs/UnrealBuildTool/Platform/Android/UEDeployAndroid.cs
+++ b/Engine/Source/Programs/UnrealBuildTool/Platform/Android/UEDeployAndroid.cs
@@ -4040,11 +4040,11 @@ namespace UnrealBuildTool
 			string UnrealJavaFilePath = Path.Combine(ProjectDirectory, "Build", "Android", GetUnrealJavaSrcPath());
 			string UnrealBuildFilesPath = GetUnrealBuildFilePath(EngineDirectory);
 			string UnrealPreBuiltFilesPath = GetUnrealPreBuiltFilePath(EngineDirectory);
-			string UnrealBuildFilesPath_NFL = GetUnrealBuildFilePath(Path.Combine(EngineDirectory, "Restricted/NotForLicensees"));
-			string UnrealBuildFilesPath_NR = GetUnrealBuildFilePath(Path.Combine(EngineDirectory, "Restricted/NoRedist"));
+			string UnrealBuildFilesPath_NFL = GetUnrealBuildFilePath(Path.Combine(EngineDirectory, "Restricted", "NotForLicensees"));
+			string UnrealBuildFilesPath_NR = GetUnrealBuildFilePath(Path.Combine(EngineDirectory, "Restricted", "NoRedist"));
 			string GameBuildFilesPath = Path.Combine(ProjectDirectory, "Build", "Android");
-			string GameBuildFilesPath_NFL = Path.Combine(Path.Combine(ProjectDirectory, "Restricted/NotForLicensees"), "Build", "Android");
-			string GameBuildFilesPath_NR = Path.Combine(Path.Combine(ProjectDirectory, "Restricted/NoRedist"), "Build", "Android");
+			string GameBuildFilesPath_NFL = Path.Combine(Path.Combine(ProjectDirectory, "Restricted","NotForLicensees"), "Build", "Android");
+			string GameBuildFilesPath_NR = Path.Combine(Path.Combine(ProjectDirectory, "Restricted", "NoRedist"), "Build", "Android");
 
 			// get a list of unique NDK architectures enabled for build
 			List<string> NDKArches = new List<string>();
@@ -4555,6 +4555,7 @@ namespace UnrealBuildTool
 				Logger.LogInformation("\n===={Time}====PREPARING NATIVE CODE====={Arch}============================================================", DateTime.Now.ToString(), Arch);
 
 				string UnrealBuildPath = Path.Combine(IntermediateAndroidPath, Arch.ToString());
+				string UnrealBuildGradlePath = Path.Combine(UnrealBuildPath, "gradle");
 
 				// If we are packaging for Amazon then we need to copy the  file to the correct location
 				Logger.LogInformation("bPackageDataInsideApk = {bPackageDataInsideApk}", bPackageDataInsideApk);
@@ -4594,17 +4595,31 @@ namespace UnrealBuildTool
 					SafeDeleteFile(CommandLineDestFileName);
 				}
 
+				// check for Android Studio project being setup for packaging the apk/debug. If this exists then we have setup and
+				//Android studio project that we intend to use to finish making the APK or for debugging.
+				string GradleAppImlFilename = Path.Combine(UnrealBuildGradlePath, ".idea", "modules", "app.iml");
+				bool bHasAndroidStudioProject = File.Exists(GradleAppImlFilename);
+
+				// We need to filter out copying over our gradle files and stomping the existing gradle project if we plan to use 
+				// the existing AndroidStudio project to build our APK and/or Debug.
+				string[]? ExcludeFolders = null;
+				if (bHasAndroidStudioProject)
+				{
+					//Path.Combine(UnrealBuildFilesPath, "gradle")
+					ExcludeFolders = new string[]{ "gradle", ".gradle", ".idea", "app", "runConfigurations" };
+				}
+
 				//Copy build files to the intermediate folder in this order (later overrides earlier):
 				//	- Shared Engine
 				//  - Shared Engine NoRedist (for Epic secret files)
 				//  - Game
 				//  - Game NoRedist (for Epic secret files)
-				CopyFileDirectory(UnrealBuildFilesPath, UnrealBuildPath, Replacements);
-				CopyFileDirectory(UnrealBuildFilesPath_NFL, UnrealBuildPath, Replacements);
-				CopyFileDirectory(UnrealBuildFilesPath_NR, UnrealBuildPath, Replacements);
-				CopyFileDirectory(GameBuildFilesPath, UnrealBuildPath, Replacements);
-				CopyFileDirectory(GameBuildFilesPath_NFL, UnrealBuildPath, Replacements);
-				CopyFileDirectory(GameBuildFilesPath_NR, UnrealBuildPath, Replacements);
+				CopyFileDirectory(UnrealBuildFilesPath, UnrealBuildPath, Replacements, ExcludeFolders);
+				CopyFileDirectory(UnrealBuildFilesPath_NFL, UnrealBuildPath, Replacements, ExcludeFolders);
+				CopyFileDirectory(UnrealBuildFilesPath_NR, UnrealBuildPath, Replacements, ExcludeFolders);
+				CopyFileDirectory(GameBuildFilesPath, UnrealBuildPath, Replacements, ExcludeFolders);
+				CopyFileDirectory(GameBuildFilesPath_NFL, UnrealBuildPath, Replacements, ExcludeFolders);
+				CopyFileDirectory(GameBuildFilesPath_NR, UnrealBuildPath, Replacements, ExcludeFolders);
 
 				// Parse Gradle filters (may have been replaced by above copies)
 				ParseFilterFile(Path.Combine(UnrealBuildPath, "GradleFilter.txt"));
@@ -4627,7 +4642,6 @@ namespace UnrealBuildTool
 				PickDownloaderScreenOrientation(UnrealBuildPath, bNeedPortrait, bNeedLandscape);
 
 				// use Gradle for compile/package
-				string UnrealBuildGradlePath = Path.Combine(UnrealBuildPath, "gradle");
 				string UnrealBuildGradleAppPath = Path.Combine(UnrealBuildGradlePath, "app");
 				string UnrealBuildGradleMainPath = Path.Combine(UnrealBuildGradleAppPath, "src", "main");
 				string CompileSDKVersion = SDKAPILevel.Replace("android-", "");
@@ -4654,7 +4668,7 @@ namespace UnrealBuildTool
 
 				// we don't actually need the SO for the bSkipGradleBuild case
 				string? FinalSOName = null;
-				string DestApkDirectory = Path.Combine(ProjectDirectory, "Binaries/Android");
+				string DestApkDirectory = Path.Combine(ProjectDirectory, "Binaries","Android");
 				string? DestApkName = null;
 				if (bSkipGradleBuild)
 				{
@@ -4846,8 +4860,7 @@ namespace UnrealBuildTool
 						string GradleSecondCallOptions = UPL.ProcessPluginNode(NDKArch, "gradleSecondCallParameters", "");
 
 						// check for Android Studio project, call Gradle if doesn't exist (assume user will build with Android Studio)
-						string GradleAppImlFilename = Path.Combine(UnrealBuildGradlePath, "app.iml");
-						if (!File.Exists(GradleAppImlFilename))
+						if (!bHasAndroidStudioProject)
 						{
 							// make sure destination exists
 							Directory.CreateDirectory(Path.GetDirectoryName(DestApkName)!);
@@ -4899,7 +4912,7 @@ namespace UnrealBuildTool
 				{
 					// Copy .so with symbols to 
 					int StoreVersion = GetStoreVersion(bEnableBundle ? null : Arch);
-					string SymbolSODirectory = Path.Combine(DestApkDirectory, ProjectName + "_Symbols_v" + StoreVersion + "/" + ProjectName + Arch);
+					string SymbolSODirectory = Path.Combine(DestApkDirectory, ProjectName + "_Symbols_v" + StoreVersion, ProjectName + Arch);
 					string SymbolifiedSOPath = Path.Combine(SymbolSODirectory, Path.GetFileName(FinalSOName));
 					MakeDirectoryIfRequired(SymbolifiedSOPath);
 					Logger.LogInformation("Writing symbols to {SymbolifiedSOPath}", SymbolifiedSOPath);
