- Fixed rpath calculation of paths relative to the executable, so .apps under locations other than Binaries/Mac can find thirdparty dylibs

diff --git a/Engine/Source/Programs/UnrealBuildTool/Platform/Mac/MacToolChain.cs b/Engine/Source/Programs/UnrealBuildTool/Platform/Mac/MacToolChain.cs
--- a/Engine/Source/Programs/UnrealBuildTool/Platform/Mac/MacToolChain.cs
+++ b/Engine/Source/Programs/UnrealBuildTool/Platform/Mac/MacToolChain.cs
@@ -368,7 +368,14 @@ namespace UnrealBuildTool
 				bool bCanUseMultipleRPATHs = !ExeAbsolutePath.Contains("EpicGamesLauncher-Mac-Shipping") || !Library.Contains("CEF3");
 
 				// First, add a path relative to the executable.
-				string RelativePath = Utils.MakePathRelativeTo(LibraryDir, ExeDir).Replace("\\", "/");
+				string FinalExeDir = ExeDir;
+				if (bIsBuildingAppBundle)
+				{
+					FinalExeDir = ExeAbsolutePath + ".app/Contents/MacOS";
+				}
+				string RelativePath = Utils.MakePathRelativeTo(LibraryDir, FinalExeDir).Replace("\\", "/");
+				
+
 				if (bCanUseMultipleRPATHs)
 				{
 					LinkCommand += " -rpath \"@loader_path/" + RelativePath + "\"";
