Linux can now be cooked on windows platforms without the sdk installed.
Added override cook platform for use by content beacon to the datadrivenplatforminfo.

diff --git a/Engine/Config/LinuxArm64/DataDrivenPlatformInfo.ini b/Engine/Config/LinuxArm64/DataDrivenPlatformInfo.ini
--- a/Engine/Config/LinuxArm64/DataDrivenPlatformInfo.ini
+++ b/Engine/Config/LinuxArm64/DataDrivenPlatformInfo.ini
@@ -12,6 +12,8 @@ bCanChangeGamepadType=true
 bSupportsTouch=false
 GlobalIdentifier=C0B45A339DE340C7BCE24C4715C3BABF
 
+OverridePlatformName=Linux
+
 NormalIconPath=Launcher/Linux/Platform_Linux_24x
 LargeIconPath=Launcher/Linux/Platform_Linux_128x
 XLargeIconPath=
diff --git a/Engine/Source/Programs/UnrealBuildTool/Configuration/UEBuildPlatform.cs b/Engine/Source/Programs/UnrealBuildTool/Configuration/UEBuildPlatform.cs
--- a/Engine/Source/Programs/UnrealBuildTool/Configuration/UEBuildPlatform.cs
+++ b/Engine/Source/Programs/UnrealBuildTool/Configuration/UEBuildPlatform.cs
@@ -617,10 +617,11 @@ namespace UnrealBuildTool
 		/// Determines whether a given platform is available
 		/// </summary>
 		/// <param name="Platform">The platform to check for</param>
+		/// <param name="bIgnoreSDKCheck">Ignore the sdks presence when checking for platform availablity (many platforms can be cooked as a target platform without requiring sdk)</param>
 		/// <returns>True if it's available, false otherwise</returns>
-		public static bool IsPlatformAvailable(UnrealTargetPlatform Platform)
+		public static bool IsPlatformAvailable(UnrealTargetPlatform Platform, bool bIgnoreSDKCheck = false)
 		{
-			return BuildPlatformDictionary.ContainsKey(Platform) && BuildPlatformDictionary[Platform].HasRequiredSDKsInstalled() == SDKStatus.Valid;
+			return BuildPlatformDictionary.ContainsKey(Platform) && (bIgnoreSDKCheck || BuildPlatformDictionary[Platform].HasRequiredSDKsInstalled() == SDKStatus.Valid);
 		}
 
 		/// <summary>
@@ -628,10 +629,11 @@ namespace UnrealBuildTool
 		/// </summary>
 		/// <param name="Platform">The platform to check for</param>
 		/// <param name="Target">A Target object that may further restrict available platforms</param>
+		/// <param name="bIgnoreSDKCheck">Ignore the sdks presence when checking for platform availablity (many platforms can be cooked as a target platform without requiring sdk)</param>
 		/// <returns>True if it's available, false otherwise</returns>
-		public static bool IsPlatformAvailableForTarget(UnrealTargetPlatform Platform, ReadOnlyTargetRules Target)
+		public static bool IsPlatformAvailableForTarget(UnrealTargetPlatform Platform, ReadOnlyTargetRules Target, bool bIgnoreSDKCheck = false)
 		{
-			return IsPlatformAvailable(Platform) && Target.IsPlatformOptedIn(Platform);
+			return IsPlatformAvailable(Platform, bIgnoreSDKCheck) && Target.IsPlatformOptedIn(Platform);
 		}
 
 		/// <summary>
diff --git a/Engine/Source/Programs/UnrealBuildTool/Platform/Android/UEBuildAndroid.cs b/Engine/Source/Programs/UnrealBuildTool/Platform/Android/UEBuildAndroid.cs
--- a/Engine/Source/Programs/UnrealBuildTool/Platform/Android/UEBuildAndroid.cs
+++ b/Engine/Source/Programs/UnrealBuildTool/Platform/Android/UEBuildAndroid.cs
@@ -326,7 +326,7 @@ namespace UnrealBuildTool
 		{
 			// don't do any target platform stuff if not available for host and opted in
 			// do not require SDK to build it since we don't necessarily need it for editor building
-			if (!(TryGetBuildPlatform(Platform, out _) && Target.IsPlatformOptedIn(Platform)))
+			if (!UEBuildPlatform.IsPlatformAvailableForTarget(Platform, Target, bIgnoreSDKCheck: true))
 			{
 				return;
 			}
diff --git a/Engine/Source/Programs/UnrealBuildTool/Platform/Linux/UEBuildLinux.cs b/Engine/Source/Programs/UnrealBuildTool/Platform/Linux/UEBuildLinux.cs
--- a/Engine/Source/Programs/UnrealBuildTool/Platform/Linux/UEBuildLinux.cs
+++ b/Engine/Source/Programs/UnrealBuildTool/Platform/Linux/UEBuildLinux.cs
@@ -384,7 +384,7 @@ namespace UnrealBuildTool
 		public override void ModifyModuleRulesForOtherPlatform(string ModuleName, ModuleRules Rules, ReadOnlyTargetRules Target)
 		{
 			// don't do any target platform stuff if SDK is not available
-			if (!UEBuildPlatform.IsPlatformAvailableForTarget(Platform, Target))
+			if (!UEBuildPlatform.IsPlatformAvailableForTarget(Platform, Target, bIgnoreSDKCheck:true))
 			{
 				return;
 			}
diff --git a/Engine/Source/Runtime/Core/Private/Misc/DataDrivenPlatformInfoRegistry.cpp b/Engine/Source/Runtime/Core/Private/Misc/DataDrivenPlatformInfoRegistry.cpp
--- a/Engine/Source/Runtime/Core/Private/Misc/DataDrivenPlatformInfoRegistry.cpp
+++ b/Engine/Source/Runtime/Core/Private/Misc/DataDrivenPlatformInfoRegistry.cpp
@@ -333,6 +333,8 @@ static void LoadDDPIIniSettings(const FConfigFile& IniFile, FDataDrivenPlatformI
 	DDPIGetBool(IniFile, TEXT("bCanChangeGamepadType"), Info.bCanChangeGamepadType);
 	DDPIGetBool(IniFile, TEXT("bSupportsTouch"), Info.bSupportsTouch);
 
+	DDPIGetName(IniFile, TEXT("OverrideCookPlatformName"), Info.OverrideCookPlatformName);
+
 #if DDPI_HAS_EXTENDED_PLATFORMINFO_DATA
 
 	DDPIGetString(IniFile, TEXT("AutoSDKPath"), Info.AutoSDKPath);
diff --git a/Engine/Source/Runtime/Core/Public/Misc/DataDrivenPlatformInfoRegistry.h b/Engine/Source/Runtime/Core/Public/Misc/DataDrivenPlatformInfoRegistry.h
--- a/Engine/Source/Runtime/Core/Public/Misc/DataDrivenPlatformInfoRegistry.h
+++ b/Engine/Source/Runtime/Core/Public/Misc/DataDrivenPlatformInfoRegistry.h
@@ -161,6 +161,18 @@ struct FDataDrivenPlatformInfo
 	// the compression format that this platform wants; overrides game unless bForceUseProjectCompressionFormat
 	FString HardwareCompressionFormat;
 
+	/** Platform name to be used for cooking (note: DOES NOT include the configuration name, eg "Server", "Client", only the platform name) */
+	FName OverrideCookPlatformName;
+
+	const FName& GetCookPlatformName() const
+	{
+		if (!OverrideCookPlatformName.IsValid())
+		{
+			return OverrideCookPlatformName;
+		}
+		return IniPlatformName;
+	}
+
 	// NOTE: add more settings here (and read them in in the LoadDDPIIniSettings() function in the .cpp)
 
 
