Prevent unnecessary rebuilds of windows resource files for development builds to reduce how often binaries relink

Only embed info from Build.version in exes and dlls for precompiled or shipping targets. This information is always available at runtime via the BuildSettings module regardless of being embedded in the binaries
Set default windows resource ProductVersion to ENGINE_VERSION_STRING if BUILD_VERSION is not defined

diff --git a/Engine/Build/Windows/Resources/Default.rc2 b/Engine/Build/Windows/Resources/Default.rc2
--- a/Engine/Build/Windows/Resources/Default.rc2
+++ b/Engine/Build/Windows/Resources/Default.rc2
@@ -29,6 +29,12 @@
 	#define BUILD_PROJECT_PRODUCT_NAME PREPROCESSOR_TO_STRING(UE_APP_NAME)
 #endif
 
+#ifdef BUILD_VERSION
+	#define BUILD_PROJECT_PRODUCT_VERSION PREPROCESSOR_TO_STRING(BUILD_VERSION)
+#else
+	#define BUILD_PROJECT_PRODUCT_VERSION ENGINE_VERSION_STRING
+#endif
+
 #ifdef PROJECT_PRODUCT_IDENTIFIER
 	#define BUILD_PROJECT_PRODUCT_IDENTIFIER PREPROCESSOR_TO_STRING(PROJECT_PRODUCT_IDENTIFIER)
 #else
@@ -67,9 +73,7 @@ BEGIN
 			VALUE "CompanyName", BUILD_PROJECT_COMPANY_NAME
 			VALUE "LegalCopyright", BUILD_PROJECT_COPYRIGHT_STRING
 			VALUE "ProductName", BUILD_PROJECT_PRODUCT_NAME
-#ifdef BUILD_VERSION
-			VALUE "ProductVersion", PREPROCESSOR_TO_STRING(BUILD_VERSION)
-#endif
+			VALUE "ProductVersion", BUILD_PROJECT_PRODUCT_VERSION
 			VALUE "FileDescription", BUILD_PROJECT_PRODUCT_NAME
 			VALUE "InternalName", BUILD_PROJECT_PRODUCT_IDENTIFIER
 #ifdef ORIGINAL_FILE_NAME
diff --git a/Engine/Source/Programs/UnrealBuildTool/Configuration/TargetRules.cs b/Engine/Source/Programs/UnrealBuildTool/Configuration/TargetRules.cs
--- a/Engine/Source/Programs/UnrealBuildTool/Configuration/TargetRules.cs
+++ b/Engine/Source/Programs/UnrealBuildTool/Configuration/TargetRules.cs
@@ -2123,7 +2123,7 @@ namespace UnrealBuildTool
 		public bool bIgnoreBuildOutputs = false;
 
 		/// <summary>
-		/// Indicates that this is a formal build, intended for distribution. This flag is automatically set to true when Build.version has a changelist set.
+		/// Indicates that this is a formal build, intended for distribution. This flag is automatically set to true when Build.version has a changelist set and is a promoted build.
 		/// The only behavior currently bound to this flag is to compile the default resource file separately for each binary so that the OriginalFilename field is set correctly.
 		/// By default, we only compile the resource once to reduce build times.
 		/// </summary>
@@ -2627,7 +2627,7 @@ namespace UnrealBuildTool
 			ConfigValueTracker = new ConfigValueTracker(ConfigValues);
 
 			// If we've got a changelist set, set that we're making a formal build
-			bFormalBuild = (Version.Changelist != 0 && Version.IsPromotedBuild);
+			bFormalBuild = bFormalBuild || (Version.Changelist != 0 && Version.IsPromotedBuild);
 
 			// Allow the build platform to set defaults for this target
 			UEBuildPlatform.GetBuildPlatform(Platform).ResetTarget(this);
diff --git a/Engine/Source/Programs/UnrealBuildTool/Configuration/UEBuildBinary.cs b/Engine/Source/Programs/UnrealBuildTool/Configuration/UEBuildBinary.cs
--- a/Engine/Source/Programs/UnrealBuildTool/Configuration/UEBuildBinary.cs
+++ b/Engine/Source/Programs/UnrealBuildTool/Configuration/UEBuildBinary.cs
@@ -814,9 +814,22 @@ namespace UnrealBuildTool
 						// Set the original file name macro; used in Default.rc2 to set the binary metadata fields.
 						ResourceCompileEnvironment.Definitions.Add("ORIGINAL_FILE_NAME=\"" + OutputFilePaths[0].GetFileName() + "\"");
 
-						// Set the other version fields
-						ResourceCompileEnvironment.Definitions.Add(String.Format("BUILT_FROM_CHANGELIST={0}", Target.Version.Changelist));
-						ResourceCompileEnvironment.Definitions.Add(String.Format("BUILD_VERSION={0}", Target.BuildVersion));
+						// Set the other version fields if requested
+						if (Target.WindowsPlatform.bSetResourceVersions)
+						{
+							if (Target.Version.Changelist != 0)
+							{
+								ResourceCompileEnvironment.Definitions.Add(String.Format("BUILT_FROM_CHANGELIST={0}", Target.Version.Changelist));
+							}
+							if (!String.IsNullOrEmpty(Target.Version.BranchName))
+							{
+								ResourceCompileEnvironment.Definitions.Add(String.Format("BRANCH_NAME={0}", Target.Version.BranchName));
+							}
+							if (!String.IsNullOrEmpty(Target.BuildVersion))
+							{
+								ResourceCompileEnvironment.Definitions.Add(String.Format("BUILD_VERSION={0}", Target.BuildVersion));
+							}
+						}
 
 						// Otherwise compile the default resource file per-binary, so that it gets the correct ORIGINAL_FILE_NAME macro.
 						FileItem DefaultResourceFile = FileItem.GetItemByFileReference(FileReference.Combine(Unreal.EngineDirectory, "Build", "Windows", "Resources", "Default.rc2"));
diff --git a/Engine/Source/Programs/UnrealBuildTool/Configuration/UEBuildTarget.cs b/Engine/Source/Programs/UnrealBuildTool/Configuration/UEBuildTarget.cs
--- a/Engine/Source/Programs/UnrealBuildTool/Configuration/UEBuildTarget.cs
+++ b/Engine/Source/Programs/UnrealBuildTool/Configuration/UEBuildTarget.cs
@@ -2634,23 +2634,17 @@ namespace UnrealBuildTool
 			}
 
 			// Compile the resource files common to all DLLs on Windows
-			if (!ShouldCompileMonolithic())
+			if (!ShouldCompileMonolithic() & !Rules.bFormalBuild && Platform.IsInGroup(UnrealPlatformGroup.Windows))
 			{
-				if (Platform == UnrealTargetPlatform.Win64)
+				FileReference DefaultResourceLocation = FileReference.Combine(Unreal.EngineDirectory, "Build", "Windows", "Resources", "Default.rc2");
+				if (!UnrealBuildTool.IsFileInstalled(DefaultResourceLocation))
 				{
-					if (!Rules.bFormalBuild)
-					{
-						FileReference DefaultResourceLocation = FileReference.Combine(Unreal.EngineDirectory, "Build", "Windows", "Resources", "Default.rc2");
-						if (!UnrealBuildTool.IsFileInstalled(DefaultResourceLocation))
-						{
-							CppCompileEnvironment DefaultResourceCompileEnvironment = new CppCompileEnvironment(GlobalCompileEnvironment);
+					CppCompileEnvironment DefaultResourceCompileEnvironment = new CppCompileEnvironment(GlobalCompileEnvironment);
 
-							FileItem DefaultResourceFile = FileItem.GetItemByFileReference(DefaultResourceLocation);
+					FileItem DefaultResourceFile = FileItem.GetItemByFileReference(DefaultResourceLocation);
 
-							CPPOutput DefaultResourceOutput = TargetToolChain.CompileRCFiles(DefaultResourceCompileEnvironment, new List<FileItem> { DefaultResourceFile }, EngineIntermediateDirectory, MakefileBuilder);
-							GlobalLinkEnvironment.DefaultResourceFiles.AddRange(DefaultResourceOutput.ObjectFiles);
-						}
-					}
+					CPPOutput DefaultResourceOutput = TargetToolChain.CompileRCFiles(DefaultResourceCompileEnvironment, new List<FileItem> { DefaultResourceFile }, EngineIntermediateDirectory, MakefileBuilder);
+					GlobalLinkEnvironment.DefaultResourceFiles.AddRange(DefaultResourceOutput.ObjectFiles);
 				}
 			}
 
diff --git a/Engine/Source/Programs/UnrealBuildTool/Platform/Windows/UEBuildWindows.cs b/Engine/Source/Programs/UnrealBuildTool/Platform/Windows/UEBuildWindows.cs
--- a/Engine/Source/Programs/UnrealBuildTool/Platform/Windows/UEBuildWindows.cs
+++ b/Engine/Source/Programs/UnrealBuildTool/Platform/Windows/UEBuildWindows.cs
@@ -200,6 +200,22 @@ namespace UnrealBuildTool
 		/// </summary>
 		TargetRules Target;
 
+		/// <summary>
+		/// If enabled will set the ProductVersion embeded in windows executables and dlls to contain BUILT_FROM_CHANGELIST and BuildVersion
+		/// Enabled by default for all precompiled and Shipping configurations. Regardless of this setting, the versions from Build.version will be available via the BuildSettings module
+		/// Note: Embedding these versions will cause resource files to be recompiled whenever changelist is updated which will cause binaries to relink
+		/// </summary>
+		[ConfigFile(ConfigHierarchyType.Engine, "/Script/WindowsTargetPlatform.WindowsTargetSettings", "SetResourceVersions")]
+		[XmlConfigFile(Category = "WindowsPlatform")]
+		[CommandLine("-SetResourceVersions")]
+		[CommandLine("-NoSetResourceVersions", Value = "false")]
+		public bool bSetResourceVersions
+		{
+			get => bSetResourceVersionPrivate ?? Target.bPrecompile || Target.Configuration == UnrealTargetConfiguration.Shipping;
+			set => bSetResourceVersionPrivate = value;
+		}
+		private bool? bSetResourceVersionPrivate = null;
+
 		/// <summary>
 		/// If -PGOOptimize is specified but the linker flags have changed since the last -PGOProfile, this will emit a warning and build without PGO instead of failing during link with LNK1268. 
 		/// </summary>
@@ -721,6 +737,8 @@ namespace UnrealBuildTool
 		#region Read-only accessor properties 
 #pragma warning disable CS1591
 
+		public bool bSetResourceVersions => Inner.bSetResourceVersions;
+
 		public bool bIgnoreStalePGOData => Inner.bIgnoreStalePGOData;
 
 		public bool bUseFastGenProfile => Inner.bUseFastGenProfile;
