UnrealBuildTool: Invalidate makefile if a header only directory is added or removed

diff --git a/Engine/Source/Programs/UnrealBuildTool/Configuration/UEBuildModuleCPP.cs b/Engine/Source/Programs/UnrealBuildTool/Configuration/UEBuildModuleCPP.cs
--- a/Engine/Source/Programs/UnrealBuildTool/Configuration/UEBuildModuleCPP.cs
+++ b/Engine/Source/Programs/UnrealBuildTool/Configuration/UEBuildModuleCPP.cs
@@ -2216,6 +2216,25 @@ namespace UnrealBuildTool
 			return Files;
 		}
 
+		/// <summary>
+		/// Gets a set of source files and headers for the given directory. Used to detect when the makefile is out of date due to a new directory
+		/// </summary>
+		/// <param name="Directory"></param>
+		/// <param name="Logger">Logger for output</param>
+		/// <returns>Array of source files and headers</returns>
+		public static FileItem[] GetSourceFilesAndHeaders(DirectoryItem Directory, ILogger Logger)
+		{
+			bool bIgnoreFileFound;
+			InputFileCollection InputFiles = new InputFileCollection();
+			List<FileItem> Files = new List<FileItem>(FindInputFilesFromDirectory(Directory, InputFiles, out bIgnoreFileFound, Logger));
+			if (bIgnoreFileFound)
+			{
+				return Array.Empty<FileItem>();
+			}
+			Files.AddRange(InputFiles.HeaderFiles);
+			return Files.ToArray();
+		}
+
 		/// <summary>
 		/// Checks a given directory path whether it exists and if it contains any Verse source files
 		/// </summary>
diff --git a/Engine/Source/Programs/UnrealBuildTool/System/TargetMakefile.cs b/Engine/Source/Programs/UnrealBuildTool/System/TargetMakefile.cs
--- a/Engine/Source/Programs/UnrealBuildTool/System/TargetMakefile.cs
+++ b/Engine/Source/Programs/UnrealBuildTool/System/TargetMakefile.cs
@@ -805,12 +805,17 @@ namespace UnrealBuildTool
 							ReasonNotLoaded = "source file modified";
 							return false;
 						}
+						else if (!InputDirectory.Exists)
+						{
+							ReasonNotLoaded = "source directory removed";
+							return false;
+						}
 
 						foreach (DirectoryItem Directory in InputDirectory.EnumerateDirectories())
 						{
-							if (!Makefile.DirectoryToSourceFiles.ContainsKey(Directory) && ContainsSourceFiles(Directory, ExcludedFolderNames, Logger))
+							if (!Makefile.DirectoryToSourceFiles.ContainsKey(Directory) && ContainsSourceFilesOrHeaders(Directory, ExcludedFolderNames, Logger))
 							{
-								ReasonNotLoaded = "directory added";
+								ReasonNotLoaded = "source directory added";
 								return false;
 							}
 						}
@@ -927,19 +932,19 @@ namespace UnrealBuildTool
 		}
 
 		/// <summary>
-		/// Determines if a directory, or any subdirectory of it, contains new source files
+		/// Determines if a directory, or any subdirectory of it, contains new source files or headers
 		/// </summary>
 		/// <param name="Directory">Directory to search through</param>
 		/// <param name="ExcludedFolderNames">Set of directory names to exclude</param>
 		/// <param name="Logger">Logger for output diagnostics</param>
 		/// <returns>True if the directory contains any source files</returns>
-		static bool ContainsSourceFiles(DirectoryItem Directory, IReadOnlySet<string> ExcludedFolderNames, ILogger Logger)
+		static bool ContainsSourceFilesOrHeaders(DirectoryItem Directory, IReadOnlySet<string> ExcludedFolderNames, ILogger Logger)
 		{
 			// Check this directory isn't ignored
 			if (!ExcludedFolderNames.Contains(Directory.Name))
 			{
 				// Check for any source files in this actual directory
-				FileItem[] SourceFiles = UEBuildModuleCPP.GetSourceFiles(Directory, Logger);
+				FileItem[] SourceFiles = UEBuildModuleCPP.GetSourceFilesAndHeaders(Directory, Logger);
 				if (SourceFiles.Length > 0)
 				{
 					return true;
@@ -948,7 +953,7 @@ namespace UnrealBuildTool
 				// Check for any source files in a subdirectory
 				foreach (DirectoryItem SubDirectory in Directory.EnumerateDirectories())
 				{
-					if (ContainsSourceFiles(SubDirectory, ExcludedFolderNames, Logger))
+					if (ContainsSourceFilesOrHeaders(SubDirectory, ExcludedFolderNames, Logger))
 					{
 						return true;
 					}
