UnrealBuildTool: Additional extensions for clang singlefile compiles

#rnx

diff --git a/Engine/Source/Programs/UnrealBuildTool/ToolChain/ClangSpecificFileAction.cs b/Engine/Source/Programs/UnrealBuildTool/ToolChain/ClangSpecificFileAction.cs
--- a/Engine/Source/Programs/UnrealBuildTool/ToolChain/ClangSpecificFileAction.cs
+++ b/Engine/Source/Programs/UnrealBuildTool/ToolChain/ClangSpecificFileAction.cs
@@ -67,7 +67,11 @@ namespace UnrealBuildTool
 			foreach (string L in RspLines)
 			{
 				string Line = L;
-				if (Line.Contains(".cpp.o") || Line.Contains(".cpp.d"))
+				if (Line.Contains(".cpp.bc", System.StringComparison.Ordinal) ||
+					Line.Contains(".cpp.d", System.StringComparison.Ordinal) ||
+					Line.Contains(".cpp.i", System.StringComparison.Ordinal) || 
+					Line.Contains(".cpp.json", System.StringComparison.Ordinal) ||
+					Line.Contains(".cpp.o", System.StringComparison.Ordinal))
 				{
 					Line = Line.Replace("SingleFile.cpp", UniqueDummyName);
 				}
