[UBT]
* Disable box for pgo optimized msvc link actions

diff --git a/Engine/Source/Programs/UnrealBuildTool/Platform/Windows/VCToolChain.cs b/Engine/Source/Programs/UnrealBuildTool/Platform/Windows/VCToolChain.cs
--- a/Engine/Source/Programs/UnrealBuildTool/Platform/Windows/VCToolChain.cs
+++ b/Engine/Source/Programs/UnrealBuildTool/Platform/Windows/VCToolChain.cs
@@ -2601,6 +2601,11 @@ namespace UnrealBuildTool
 				LinkAction.bCanExecuteRemotely = true;
 			}
 
+			if (LinkEnvironment.bPGOOptimize)
+			{
+				LinkAction.bCanExecuteInBox = false; // Disabled for now. Should revisit to see why it is not working
+			}
+
 			// Create link repro if requested, this argument is intentionally not added to the response file
 			if (Target.WindowsPlatform.LinkReproDir != null)
 			{
