Disable passing `/Os` and pass `/Ob0` when using `clang-cl` in debug builds by default since otherwise it
interferes with debugging `clang-cl`-generated executables in Visual Studio.

diff --git a/Engine/Source/Programs/UnrealBuildTool/Platform/Windows/VCToolChain.cs b/Engine/Source/Programs/UnrealBuildTool/Platform/Windows/VCToolChain.cs
--- a/Engine/Source/Programs/UnrealBuildTool/Platform/Windows/VCToolChain.cs
+++ b/Engine/Source/Programs/UnrealBuildTool/Platform/Windows/VCToolChain.cs
@@ -550,8 +550,10 @@ namespace UnrealBuildTool
 			// Maintain the old std::aligned_storage behavior from VS from v15.8 onwards, in case of prebuilt third party libraries are reliant on it
 			AddDefinition(Arguments, "_DISABLE_EXTENDED_ALIGNED_STORAGE");
 
-			// Do not allow inline method expansion if E&C support is enabled or inline expansion has been disabled
-			if (!CompileEnvironment.bSupportEditAndContinue && CompileEnvironment.bUseInlining)
+			// Do not allow inline method expansion if E&C support is enabled or inline expansion has been disabled, 
+			// or if we are compiling in a debug build with `clang-cl`, since this will interfere with debugging capabilities.
+			if (!CompileEnvironment.bSupportEditAndContinue && CompileEnvironment.bUseInlining
+				|| !(Target.WindowsPlatform.Compiler.IsClang() && CompileEnvironment.Configuration == CppConfiguration.Debug))
 			{
 				Arguments.Add($"/Ob{Math.Clamp(Target.WindowsPlatform.InlineFunctionExpansionLevel, 1, 3)}");
 			}
@@ -638,8 +640,13 @@ namespace UnrealBuildTool
 				// Disable compiler optimization.
 				Arguments.Add("/Od");
 
-				// Favor code size (especially useful for embedded platforms).
-				Arguments.Add("/Os");
+				// `/Os` causes `clang-cl` to effectively compile with `-Os` which enables optimizations,
+				// rendering passing `/Od` to it useless.
+				if (!Target.WindowsPlatform.Compiler.IsClang())
+				{
+					// Favor code size (especially useful for embedded platforms).
+					Arguments.Add("/Os");
+				}
 
 				// Runtime checks and ASan are incompatible.
 				if (!Target.WindowsPlatform.bEnableAddressSanitizer)
