UnrealBuildTool: Scan AutoSDK for WindowsSDKs rather than hardcoding a single folder to use. This will allow separate versioned folders for Windows SDKs rather than putting all versions into Windows Kits/10

#rnx

diff --git a/Engine/Source/Programs/UnrealBuildTool/Platform/Windows/MicrosoftPlatformSDK.cs b/Engine/Source/Programs/UnrealBuildTool/Platform/Windows/MicrosoftPlatformSDK.cs
--- a/Engine/Source/Programs/UnrealBuildTool/Platform/Windows/MicrosoftPlatformSDK.cs
+++ b/Engine/Source/Programs/UnrealBuildTool/Platform/Windows/MicrosoftPlatformSDK.cs
@@ -320,13 +320,20 @@ namespace UnrealBuildTool
 			}
 
 			DirectoryReference? HostAutoSdkDir;
-			if (UEBuildPlatformSDK.TryGetHostPlatformAutoSDKDir(out HostAutoSdkDir))
+			if (TryGetHostPlatformAutoSDKDir(out HostAutoSdkDir))
 			{
-				DirectoryReference RootDirAutoSdk = DirectoryReference.Combine(HostAutoSdkDir, "Win64", "Windows Kits", "10");
-				if (DirectoryReference.Exists(RootDirAutoSdk))
+				DirectoryReference WindowsKitsDirAutoSdk = DirectoryReference.Combine(HostAutoSdkDir, "Win64", "Windows Kits");
+				if (DirectoryReference.Exists(WindowsKitsDirAutoSdk))
 				{
-					Logger.LogDebug("Found Windows 10 AutoSDK root at {RootDirAutoSdk}", RootDirAutoSdk);
-					RootDirs.Add(RootDirAutoSdk);
+					foreach (DirectoryReference RootDirAutoSdk in DirectoryReference.EnumerateDirectories(WindowsKitsDirAutoSdk))
+					{
+						DirectoryReference IncludeRootDir = DirectoryReference.Combine(RootDirAutoSdk, "Include");
+						if (DirectoryReference.Exists(IncludeRootDir))
+						{
+							Logger.LogDebug("Found Windows 10 AutoSDK root at {RootDirAutoSdk}", RootDirAutoSdk);
+							RootDirs.Add(RootDirAutoSdk);
+						}
+					}
 				}
 			}
 		}
