UnrealBuildTool: Remove defunct HybridExecutor

diff --git a/Engine/Source/Programs/UnrealBuildTool/Configuration/BuildConfiguration.cs b/Engine/Source/Programs/UnrealBuildTool/Configuration/BuildConfiguration.cs
--- a/Engine/Source/Programs/UnrealBuildTool/Configuration/BuildConfiguration.cs
+++ b/Engine/Source/Programs/UnrealBuildTool/Configuration/BuildConfiguration.cs
@@ -33,9 +33,10 @@ namespace UnrealBuildTool
 		public bool bPrintDebugInfo = false;
 
 		/// <summary>
-		/// Whether the hybrid executor will be used (a remote executor and local executor).
+		/// Whether the hybrid executor will be used (a remote executor and local executor). No longer supported.
 		/// </summary>
 		[XmlConfigFile]
+		[Obsolete("HybridExecutor is no longer supported")]
 		public bool bAllowHybridExecutor = false;
 
 #if __BOXEXECUTOR_AVAILABLE__
diff --git a/Engine/Source/Programs/UnrealBuildTool/Executors/HybridExecutor.cs b/Engine/Source/Programs/UnrealBuildTool/Executors/HybridExecutor.cs
deleted file mode 100644
diff --git a/Engine/Source/Programs/UnrealBuildTool/System/ActionGraph.cs b/Engine/Source/Programs/UnrealBuildTool/System/ActionGraph.cs
--- a/Engine/Source/Programs/UnrealBuildTool/System/ActionGraph.cs
+++ b/Engine/Source/Programs/UnrealBuildTool/System/ActionGraph.cs
@@ -372,11 +372,7 @@ namespace UnrealBuildTool
 
 			if (ActionCount > ParallelExecutor.GetDefaultNumParallelProcesses(BuildConfiguration.MaxParallelActions, BuildConfiguration.bAllCores, Logger))
 			{
-				if (BuildConfiguration.bAllowHybridExecutor && HybridExecutor.IsAvailable(Logger))
-				{
-					return new HybridExecutor(TargetDescriptors, BuildConfiguration.MaxParallelActions, BuildConfiguration.bAllCores, BuildConfiguration.bCompactOutput, Logger);
-				}
-				else if (BuildConfiguration.bAllowXGE && XGE.IsAvailable(Logger) && ActionCount >= XGE.MinActions)
+				if (BuildConfiguration.bAllowXGE && XGE.IsAvailable(Logger) && ActionCount >= XGE.MinActions)
 				{
 					return new XGE(Logger);
 				}
