LLT: Pre-analyse dependency graph of low level test targets before setting compilation flags for UE components
Use -Mode=Test by default for all LLT Project to enable pre-analysis of dependency graph.

diff --git a/Engine/Build/LowLevelTests.xml b/Engine/Build/LowLevelTests.xml
--- a/Engine/Build/LowLevelTests.xml
+++ b/Engine/Build/LowLevelTests.xml
@@ -123,7 +123,7 @@
 			<Command Name="BuildCookRun" Arguments="-Build -NoBootstrapExe -NoCodeSign -SkipCook -Stage -NoSubmit -Archive -ArchiveDirectory=$(LocalAchiveDirectory) -project=$($(TestName)Target) -configuration=$(Configuration) -platform=$(Platform) $($(Platform)$(TestName)ExtraCompilationArgs)"/>
 		</Do>
 		<Do If="'$($(TestName)StagesWithProjectFile)' != 'True'">
-			<Compile Target="$($(TestName)Target)" Configuration="$(Configuration)" Platform="$(Platform)" AllowXGE="false" Arguments="$($(Platform)$(TestName)ExtraCompilationArgs)"/>
+			<Compile Target="$($(TestName)Target)" Configuration="$(Configuration)" Platform="$(Platform)" AllowXGE="false" Arguments="-Mode=Test $($(Platform)$(TestName)ExtraCompilationArgs)"/>
 		</Do>
 	</Macro>
 	<Macro Name="DeployAndTest" Arguments="TestName;Platform">
diff --git a/Engine/Source/Developer/LowLevelTestsRunner/LowLevelTestsRunner.Build.cs b/Engine/Source/Developer/LowLevelTestsRunner/LowLevelTestsRunner.Build.cs
--- a/Engine/Source/Developer/LowLevelTestsRunner/LowLevelTestsRunner.Build.cs
+++ b/Engine/Source/Developer/LowLevelTestsRunner/LowLevelTestsRunner.Build.cs
@@ -24,7 +24,24 @@ namespace UnrealBuildTool.Rules
 				new string[] {
 					"Core"
 				}
-			);		
+			);
+
+			if (TestTargetRules.bTestsRequireEditor)
+			{
+				PrivateDependencyModuleNames.Add("UnrealEd");
+			}
+			if (TestTargetRules.bTestsRequireEngine)
+			{
+				PrivateDependencyModuleNames.Add("Engine");
+			}
+			if (TestTargetRules.bTestsRequireApplicationCore)
+			{
+				PrivateDependencyModuleNames.Add("ApplicationCore");
+			}
+			if (TestTargetRules.bTestsRequireCoreUObject)
+			{
+				PrivateDependencyModuleNames.Add("CoreUObject");
+			}
 		}
 	}
 }
diff --git a/Engine/Source/Programs/UnrealBuildTool/Configuration/ModuleRules.cs b/Engine/Source/Programs/UnrealBuildTool/Configuration/ModuleRules.cs
--- a/Engine/Source/Programs/UnrealBuildTool/Configuration/ModuleRules.cs
+++ b/Engine/Source/Programs/UnrealBuildTool/Configuration/ModuleRules.cs
@@ -1750,47 +1750,6 @@ namespace UnrealBuildTool
 						PrivateIncludePathModuleNames.Add("LowLevelTestsRunner");
 					}
 				}
-				else if (Name == "LowLevelTestsRunner")
-				{
-					TestTargetRules.LowLevelTestsRunnerModule = this;
-				}
-
-				if (Name == "Engine" && !TestTargetRules.bNeverCompileAgainstEngine)
-				{
-					TestTargetRules.bTestsRequireEngine = true;
-				}
-				if (Name == "UnrealEd" && !TestTargetRules.bNeverCompileAgainstEditor)
-				{
-					TestTargetRules.bTestsRequireEditor = true;
-				}
-				if (Name == "ApplicationCore" && !TestTargetRules.bNeverCompileAgainstApplicationCore)
-				{
-					TestTargetRules.bTestsRequireApplicationCore = true;
-				}
-				if (Name == "CoreUObject" && !TestTargetRules.bNeverCompileAgainstCoreUObject)
-				{
-					TestTargetRules.bTestsRequireCoreUObject = true;
-				}
-
-				if (TestTargetRules.LowLevelTestsRunnerModule != null)
-				{
-					if (TestTargetRules.bTestsRequireEditor && !TestTargetRules.LowLevelTestsRunnerModule.PrivateDependencyModuleNames.Contains("UnrealEd"))
-					{
-						TestTargetRules.LowLevelTestsRunnerModule.PrivateDependencyModuleNames.Add("UnrealEd");
-					}
-					if (TestTargetRules.bTestsRequireEngine && !TestTargetRules.LowLevelTestsRunnerModule.PrivateDependencyModuleNames.Contains("Engine"))
-					{
-						TestTargetRules.LowLevelTestsRunnerModule.PrivateDependencyModuleNames.Add("Engine");
-					}
-					if (TestTargetRules.bTestsRequireApplicationCore && !TestTargetRules.LowLevelTestsRunnerModule.PrivateDependencyModuleNames.Contains("ApplicationCore"))
-					{
-						TestTargetRules.LowLevelTestsRunnerModule.PrivateDependencyModuleNames.Add("ApplicationCore");
-					}
-					if (TestTargetRules.bTestsRequireCoreUObject && !TestTargetRules.LowLevelTestsRunnerModule.PrivateDependencyModuleNames.Contains("CoreUObject"))
-					{
-						TestTargetRules.LowLevelTestsRunnerModule.PrivateDependencyModuleNames.Add("CoreUObject");
-					}
-				}
 			}
 		}
 
diff --git a/Engine/Source/Programs/UnrealBuildTool/Configuration/TestTargetRules.cs b/Engine/Source/Programs/UnrealBuildTool/Configuration/TestTargetRules.cs
--- a/Engine/Source/Programs/UnrealBuildTool/Configuration/TestTargetRules.cs
+++ b/Engine/Source/Programs/UnrealBuildTool/Configuration/TestTargetRules.cs
@@ -23,32 +23,27 @@ namespace UnrealBuildTool
 		/// <summary>
 		/// Keeps track if low level tests executable must build with the Editor.
 		/// </summary>
-		internal bool bTestsRequireEditor = false;
+		public static bool bTestsRequireEditor = false;
 
 		/// <summary>
 		/// Keeps track if low level tests executable must build with the Engine.
 		/// </summary>
-		internal bool bTestsRequireEngine = false;
+		public static bool bTestsRequireEngine = false;
 
 		/// <summary>
 		/// Keeps track if low level tests executable must build with the ApplicationCore.
 		/// </summary>
-		internal bool bTestsRequireApplicationCore = false;
+		public static bool bTestsRequireApplicationCore = false;
 
 		/// <summary>
 		/// Keeps track if low level tests executable must build with the CoreUObject.
 		/// </summary>
-		internal bool bTestsRequireCoreUObject = false;
-
-		/// <summary>
-		/// Keep track of low level test runner module instance.
-		/// </summary>
-		internal ModuleRules? LowLevelTestsRunnerModule;
+		public static bool bTestsRequireCoreUObject = false;
 
 		/// <summary>
 		/// Associated tested target of this test target, if defined.
 		/// </summary>
-		protected TargetRules? TestedTarget { get; private set; }
+		public TargetRules? TestedTarget { get; private set; }
 
 		/// <summary>
 		/// Test target override for bCompileAgainstApplicationCore.
@@ -238,7 +233,7 @@ namespace UnrealBuildTool
 		}
 
 		/// <summary>
-		/// Constructor for TestTargetRules based on existing target.
+		/// Constructor for TestTargetRules based on existing target a.k.a Implicit test target.
 		/// TestTargetRules is setup as a program and is linked monolithically.
 		/// It removes a lot of default compilation behavior in order to produce a minimal test environment.
 		/// </summary>
@@ -265,6 +260,7 @@ namespace UnrealBuildTool
 			WindowsPlatform = TestedTarget.WindowsPlatform;
 
 			SetupCommonProperties(Target);
+			SetupImplicitTestProperties(TestedTarget);
 		}
 
 		private void SetupCommonProperties(TargetInfo Target)
@@ -297,10 +293,10 @@ namespace UnrealBuildTool
 			bForceBuildShaderFormats = false;
 
 			// Do not compile against the engine, editor etc
-			bCompileAgainstEngine = false;
-			bCompileAgainstEditor = false;
-			bCompileAgainstCoreUObject = false;
-			bCompileAgainstApplicationCore = false;
+			bCompileAgainstEngine = bTestsRequireEngine;
+			bCompileAgainstEditor = bTestsRequireEditor;
+			bCompileAgainstCoreUObject = bTestsRequireCoreUObject;
+			bCompileAgainstApplicationCore = bTestsRequireApplicationCore;
 			bCompileCEF3 = false;
 
 			// No mixing with Functional Test framework
@@ -351,5 +347,12 @@ namespace UnrealBuildTool
 				bIsBuildingConsoleApplication = false;
 			}
 		}
+
+		private void SetupImplicitTestProperties(TargetRules TestedTarget)
+		{
+			bool IsEditorTestedTarget = (TestedTarget.Type == TargetType.Editor);
+			bBuildWithEditorOnlyData = bCompileAgainstEditor = IsEditorTestedTarget;
+			bBuildDeveloperTools = bCompileAgainstEngine && IsEditorTestedTarget;
+		}
 	}
 }
diff --git a/Engine/Source/Programs/UnrealBuildTool/Modes/CleanMode.cs b/Engine/Source/Programs/UnrealBuildTool/Modes/CleanMode.cs
--- a/Engine/Source/Programs/UnrealBuildTool/Modes/CleanMode.cs
+++ b/Engine/Source/Programs/UnrealBuildTool/Modes/CleanMode.cs
@@ -91,7 +91,7 @@ namespace UnrealBuildTool
 					RulesAssembly RulesAssembly = RulesCompiler.CreateTargetRulesAssembly(TargetDescriptor.ProjectFile, TargetDescriptor.Name, bSkipRulesCompile, BuildConfiguration.bForceRulesCompile, BuildConfiguration.bUsePrecompiled, TargetDescriptor.ForeignPlugin, TargetDescriptor.bBuildPluginAsLocal, Logger);
 
 					// Create the rules object
-					ReadOnlyTargetRules Target = new ReadOnlyTargetRules(RulesAssembly.CreateTargetRules(TargetDescriptor.Name, TargetDescriptor.Platform, TargetDescriptor.Configuration, TargetDescriptor.Architectures, TargetDescriptor.ProjectFile, TargetDescriptor.AdditionalArguments, Logger));
+					ReadOnlyTargetRules Target = new ReadOnlyTargetRules(RulesAssembly.CreateTargetRules(TargetDescriptor.Name, TargetDescriptor.Platform, TargetDescriptor.Configuration, TargetDescriptor.Architectures, TargetDescriptor.ProjectFile, TargetDescriptor.AdditionalArguments, Logger, TargetDescriptor.IsTestsTarget));
 
 					// Get the intermediate environment for this target
 					UnrealIntermediateEnvironment IntermediateEnvironment = TargetDescriptor.IntermediateEnvironment != UnrealIntermediateEnvironment.Default ? TargetDescriptor.IntermediateEnvironment : Target.IntermediateEnvironment;
diff --git a/Engine/Source/Programs/UnrealBuildTool/Modes/TestMode.cs b/Engine/Source/Programs/UnrealBuildTool/Modes/TestMode.cs
--- a/Engine/Source/Programs/UnrealBuildTool/Modes/TestMode.cs
+++ b/Engine/Source/Programs/UnrealBuildTool/Modes/TestMode.cs
@@ -1,6 +1,8 @@
 // Copyright Epic Games, Inc. All Rights Reserved.
 
+using System;
 using System.Collections.Generic;
+using System.Linq;
 using System.Threading.Tasks;
 using EpicGames.Core;
 using Microsoft.Extensions.Logging;
@@ -8,12 +10,37 @@ using UnrealBuildBase;
 
 namespace UnrealBuildTool
 {
+	using IncludeChainMap = Dictionary<UEBuildModule, List<string>>;
+
 	/// <summary>
 	/// Builds low level tests on one or more targets.
 	/// </summary>
 	[ToolMode("Test", ToolModeOptions.XmlConfig | ToolModeOptions.BuildPlatforms | ToolModeOptions.SingleInstance | ToolModeOptions.ShowExecutionTime)]
 	class TestMode : ToolMode
 	{
+		private readonly string ENGINE_MODULE = "Engine";
+		private readonly string APPLICATION_CORE_MODULE = "ApplicationCore";
+		private readonly string COREUOBJECT_MODULE = "CoreUObject";
+		private readonly string EDITOR_MODULE = "UnrealEd";
+
+		/// <summary>
+		/// Whether we're building implicit tests (default is explicit)
+		/// </summary>
+		[CommandLine("-Implicit")]
+		bool bImplicit = false;
+
+		/// <summary>
+		/// Whether we're cleaning tests instead of just building
+		/// </summary>
+		[CommandLine("-CleanTests")]
+		bool bCleanTests = false;
+
+		/// <summary>
+		/// Whether we're rebuilding tests instead of just building
+		/// </summary>
+		[CommandLine("-RebuildTests")]
+		bool bRebuildTests = false;
+
 		/// <summary>
 		/// Main entry point
 		/// </summary>
@@ -31,37 +58,157 @@ namespace UnrealBuildTool
 			// Parse all the targets being built
 			List<TargetDescriptor> TargetDescriptors = TargetDescriptor.ParseCommandLine(Arguments, BuildConfiguration, Logger);
 
-			await BuildTestsAsync(TargetDescriptors, BuildConfiguration, Logger);
+			if (TargetDescriptors.Count > 1)
+			{
+				throw new BuildException($"Only one target must be specified in test mode but we got {string.Join(", ", TargetDescriptors.Select(T => T.Name))} instead.");
+			}
+
+			TargetDescriptor TestTargetDescriptor = TargetDescriptors.First();
+
+			if (bImplicit)
+			{
+				TestTargetDescriptor.Name += "Tests";
+				TestTargetDescriptor.IsTestsTarget = true;
+			}
+
+			// Analyzing dependency graph before any build commands to automatically set compilation flags.
+			// Extracted from AnalyzeMode for simplicity.
+			AnalyzeDependencyGraph(TargetDescriptors.First(), BuildConfiguration, Logger);
 
+			using (ISourceFileWorkingSet WorkingSet = SourceFileWorkingSet.Create(Unreal.RootDirectory, new HashSet<DirectoryReference>(), Logger))
+			{
+				if (bCleanTests)
+				{
+					new CleanMode().Clean(TargetDescriptors, BuildConfiguration, Logger);
+				}
+				else if (bRebuildTests)
+				{
+					new CleanMode().Clean(TargetDescriptors, BuildConfiguration, Logger);
+					await BuildMode.BuildAsync(TargetDescriptors, BuildConfiguration, WorkingSet, BuildOptions.None, null, Logger);
+				}
+				else
+				{
+					await BuildMode.BuildAsync(TargetDescriptors, BuildConfiguration, WorkingSet, BuildOptions.None, null, Logger);
+				}
+			}
 			return 0;
 		}
 
-		/// <summary>
-		/// Build tests for a list of targets.
-		/// It generates artificial test target descriptors to build a target AND its dependencies's tests in one monolithic executable.
-		/// Each module containing a "Tests" folder is included.
-		/// The target that we generate the tests executable for must include a main.cpp file and Setup and Teardown methods.
-		/// The generated executable is the target name + "Tests", e.g. UnrealEditorTests.exe
-		/// Passing UnrealEditor here would build all the tests in all the modules.
-		/// </summary>
-		/// <param name="TargetDescriptors">Target descriptors</param>
-		/// <param name="BuildConfiguration">Current build configuration</param>
-		/// <param name="Logger"></param>
-		public static async Task BuildTestsAsync(List<TargetDescriptor> TargetDescriptors, BuildConfiguration BuildConfiguration, ILogger Logger)
+		private void AnalyzeDependencyGraph(TargetDescriptor TargetDescriptor, BuildConfiguration BuildConfiguration, ILogger Logger)
 		{
-			List<TargetDescriptor> TestTargetDescriptors = new List<TargetDescriptor>();
+			UEBuildTarget Target = UEBuildTarget.Create(TargetDescriptor, BuildConfiguration, Logger);
+			DirectoryReference.CreateDirectory(Target.ReceiptFileName.Directory);
 
-			for (int Idx = 0; Idx < TargetDescriptors.Count; ++Idx)
+			HashSet<UEBuildModule> Visited = new();
+			IncludeChainMap Chains = new();
+
+			if (Target.Rules.LaunchModuleName != null)
 			{
-				TargetDescriptor TestsTargetDescriptor = TargetDescriptors[Idx].Copy();
-				TestsTargetDescriptor.Name = TargetDescriptors[Idx].Name + "Tests";
-				TestsTargetDescriptor.IsTestsTarget = true;
-				TestTargetDescriptors.Add(TestsTargetDescriptor);
+				AnalyzeModuleDependencies(Target.Rules.LaunchModuleName, new List<string>() { "target" }, Target, Chains, Visited, Logger);
 			}
 
-			using (ISourceFileWorkingSet WorkingSet = SourceFileWorkingSet.Create(Unreal.RootDirectory, new HashSet<DirectoryReference>(), Logger))
+			foreach (string RootModuleName in Target.Rules.ExtraModuleNames)
+			{
+				AnalyzeModuleDependencies(RootModuleName, new List<string>() { "target" }, Target, Chains, Visited, Logger);
+			}
+
+			foreach (UEBuildPlugin Plugin in Target.BuildPlugins!)
+			{
+				foreach (UEBuildModule Module in Plugin.Modules)
+				{
+					if (!Chains.ContainsKey(Module))
+					{
+						AnalyzeModuleDependencies(Module.Name, new List<string>() { "target", Plugin.ReferenceChain, Plugin.File.GetFileName() }, Target, Chains, Visited, Logger);
+					}
+				}
+			}
+
+			if (Target.Rules.bBuildAllModules)
+			{
+				foreach (UEBuildBinary Binary in Target.Binaries)
+				{
+					foreach (UEBuildModule Module in Binary.Modules)
+					{
+						if (!Chains.ContainsKey(Module))
+						{
+							List<string> IncludeChain = Enumerable.Repeat(String.Empty, 1000).ToList();
+							IncludeChain.Add("allmodules option");
+							if (Module.Rules.Plugin != null)
+							{
+								IncludeChain.Add(Module.Rules.Plugin.File.GetFileName());
+								AnalyzeModuleDependencies(Module.Name, IncludeChain, Target, Chains, Visited, Logger);
+							}
+							else
+							{
+								AnalyzeModuleDependencies(Module.Name, IncludeChain, Target, Chains, Visited, Logger);
+							}
+						}
+					}
+				}
+			}
+
+			foreach(var Module in Visited)
+			{
+				SetupTestsTargetForModule(Module, Logger);
+			}
+		}
+
+		private void AnalyzeModuleDependencies(string ModuleName, List<string> ParentChain, UEBuildTarget Target, IncludeChainMap Chains, HashSet<UEBuildModule> Visited, ILogger Logger)
+		{
+			if (ParentChain.Contains(ModuleName))
+			{
+				return;
+			}
+
+			UEBuildModule Module = Target.GetModuleByName(ModuleName);
+
+			List<string> CurrentChain = new(ParentChain) { Module.Name };
+
+			if (!Chains.ContainsKey(Module))
+			{
+				Chains.Add(Module, CurrentChain);
+			}
+
+			if (Chains[Module].Count > CurrentChain.Count)
+			{
+				List<UEBuildModule> RecheckModules = new();
+				Module.GetAllDependencyModules(RecheckModules, new(), true, false, true);
+				Visited.ExceptWith(RecheckModules);
+				Chains[Module] = CurrentChain;
+			}
+			else if (Visited.Contains(Module))
+			{
+				return;
+			}
+
+			Visited.Add(Module);
+
+			List<UEBuildModule> TargetModules = new();
+			Module.GetAllDependencyModules(TargetModules, new(), true, false, true);
+			TargetModules.ForEach(x => AnalyzeModuleDependencies(x.Name, CurrentChain, Target, Chains, Visited, Logger));
+		}
+
+		private void SetupTestsTargetForModule(UEBuildModule Module, ILogger Logger)
+		{
+			if (Module.Name == ENGINE_MODULE)
+			{
+				TestTargetRules.bTestsRequireEngine = true;
+				Logger.LogInformation("{Module} module found in dependency graph, will compile against it.", ENGINE_MODULE);
+			}
+			else if (Module.Name == APPLICATION_CORE_MODULE)
+			{
+				TestTargetRules.bTestsRequireApplicationCore = true;
+				Logger.LogInformation("{Module} module found in dependency graph, will compile against it.", APPLICATION_CORE_MODULE);
+			}
+			else if (Module.Name == COREUOBJECT_MODULE)
+			{
+				TestTargetRules.bTestsRequireCoreUObject = true;
+				Logger.LogInformation("{Module} module found in dependency graph, will compile against it.", COREUOBJECT_MODULE);
+			}
+			else if (Module.Name == EDITOR_MODULE)
 			{
-				await BuildMode.BuildAsync(TestTargetDescriptors, BuildConfiguration, WorkingSet, BuildOptions.None, null, Logger);
+				TestTargetRules.bTestsRequireEditor = true;
+				Logger.LogInformation("{Module} module found in dependency graph, will compile against it.", EDITOR_MODULE);
 			}
 		}
 	}
diff --git a/Engine/Source/Programs/UnrealBuildTool/ProjectFiles/VisualStudio/VCProject.cs b/Engine/Source/Programs/UnrealBuildTool/ProjectFiles/VisualStudio/VCProject.cs
--- a/Engine/Source/Programs/UnrealBuildTool/ProjectFiles/VisualStudio/VCProject.cs
+++ b/Engine/Source/Programs/UnrealBuildTool/ProjectFiles/VisualStudio/VCProject.cs
@@ -2200,9 +2200,18 @@ namespace UnrealBuildTool
 					string BuildArguments = Builder.GetBuildArguments();
 
 					// NMake Build command line
-					VCProjectFileContent.AppendLine("    <NMakeBuildCommandLine>$(BuildBatchScript) {0}</NMakeBuildCommandLine>", BuildArguments);
-					VCProjectFileContent.AppendLine("    <NMakeReBuildCommandLine>$(RebuildBatchScript) {0}</NMakeReBuildCommandLine>", BuildArguments);
-					VCProjectFileContent.AppendLine("    <NMakeCleanCommandLine>$(CleanBatchScript) {0}</NMakeCleanCommandLine>", BuildArguments);
+					if (TargetRulesObject.IsTestTarget)
+					{
+						VCProjectFileContent.AppendLine("    <NMakeBuildCommandLine>$(BuildBatchScript) {0} -Mode=Test</NMakeBuildCommandLine>", BuildArguments);
+						VCProjectFileContent.AppendLine("    <NMakeReBuildCommandLine>$(BuildBatchScript) {0} -Mode=Test -RebuildTests</NMakeReBuildCommandLine>", BuildArguments);
+						VCProjectFileContent.AppendLine("    <NMakeCleanCommandLine>$(BuildBatchScript) {0} -Mode=Test -CleanTests</NMakeCleanCommandLine>", BuildArguments);
+					}
+					else
+					{
+						VCProjectFileContent.AppendLine("    <NMakeBuildCommandLine>$(BuildBatchScript) {0}</NMakeBuildCommandLine>", BuildArguments);
+						VCProjectFileContent.AppendLine("    <NMakeReBuildCommandLine>$(RebuildBatchScript) {0}</NMakeReBuildCommandLine>", BuildArguments);
+						VCProjectFileContent.AppendLine("    <NMakeCleanCommandLine>$(CleanBatchScript) {0}</NMakeCleanCommandLine>", BuildArguments);
+					}
 					VCProjectFileContent.AppendLine("    <NMakeOutput>{0}</NMakeOutput>", NormalizeProjectPath(NMakePath.FullName));
 
 					if (TargetRulesObject.Type == TargetType.Game || TargetRulesObject.Type == TargetType.Client || TargetRulesObject.Type == TargetType.Server)
