[UBT]
* Disabled box for pgooptimized link actions for now just to rule out if long link time is because of box

diff --git a/Engine/Source/Programs/UnrealBuildTool/Platform/Linux/LinuxToolChain.cs b/Engine/Source/Programs/UnrealBuildTool/Platform/Linux/LinuxToolChain.cs
--- a/Engine/Source/Programs/UnrealBuildTool/Platform/Linux/LinuxToolChain.cs
+++ b/Engine/Source/Programs/UnrealBuildTool/Platform/Linux/LinuxToolChain.cs
@@ -1117,6 +1117,10 @@ namespace UnrealBuildTool
 			{
 				LinkAction.CommandDescription = "Link";
 			}
+
+			// Saw a 6 hour link time potentially caused by box. Will disable for now and revisit later
+			LinkAction.bCanExecuteInBox = !LinkEnvironment.bPGOProfile && !LinkEnvironment.bPGOOptimize && !LinkEnvironment.bAllowLTCG;
+
 			// because the logic choosing between lld and ld is somewhat messy atm (lld fails to link .DSO due to bugs), make the name of the linker clear
 			LinkAction.CommandDescription += (LinkCommandString.Contains("-fuse-ld=lld")) ? " (lld)" : " (ld)";
 			LinkAction.CommandVersion = Info.ClangVersionString;
