- Improved the path generated for UE_ENGINE_DIRECTORY define by UBT for programs underneath Engine, it will now use ../../../Engine instead of ../.., or more important for Restricted programs, it will use ../../../../../Engine instead of ../../../..
- Some code expects there to be "Engine" in the EngineDir, so it will not use absolute paths for RootDir, etc for Restricted programs

diff --git a/Engine/Source/Programs/UnrealBuildTool/Configuration/UEBuildTarget.cs b/Engine/Source/Programs/UnrealBuildTool/Configuration/UEBuildTarget.cs
--- a/Engine/Source/Programs/UnrealBuildTool/Configuration/UEBuildTarget.cs
+++ b/Engine/Source/Programs/UnrealBuildTool/Configuration/UEBuildTarget.cs
@@ -5252,11 +5252,11 @@ namespace UnrealBuildTool
 						OutputFilePath = OutputFilePath.Substring(0, OutputFilePath.LastIndexOf(".app/Contents/MacOS") + 4);
 					}
 
-					string EnginePath = Utils.CleanDirectorySeparators(Unreal.EngineDirectory.MakeRelativeTo(new DirectoryReference(OutputFilePath).ParentDirectory!), '/');
-					if (EnginePath.EndsWith("/") == false)
-					{
-						EnginePath += "/";
-					}
+					DirectoryReference OutputDir = new DirectoryReference(OutputFilePath).ParentDirectory!;
+					DirectoryReference RootDir = Unreal.EngineDirectory.ParentDirectory!;
+					string RootPath = RootDir.MakeRelativeTo(OutputDir);
+					// now get the cleaned engine path (Root/Engine), and make sure to end with /
+					string EnginePath = Utils.CleanDirectorySeparators(Path.Combine(RootPath, "Engine/"), '/');
 					GlobalCompileEnvironment.Definitions.Add(String.Format("UE_ENGINE_DIRECTORY=\"{0}\"", EnginePath));
 				}
 			}
