early out if the PGO data directory doesn't exist, instead of causing an exception from directory enumeration.

#rnx

diff --git a/Engine/Source/Programs/UnrealBuildTool/Platform/Windows/VCToolChain.cs b/Engine/Source/Programs/UnrealBuildTool/Platform/Windows/VCToolChain.cs
--- a/Engine/Source/Programs/UnrealBuildTool/Platform/Windows/VCToolChain.cs
+++ b/Engine/Source/Programs/UnrealBuildTool/Platform/Windows/VCToolChain.cs
@@ -2635,6 +2635,12 @@ namespace UnrealBuildTool
 		{
 			if (LinkEnvironment.bPGOOptimize && LinkEnvironment.OutputFilePath.FullName.EndsWith(".exe"))
 			{
+				if (!Directory.Exists(LinkEnvironment.PGODirectory))
+				{
+					Logger.LogWarning("\"{PGODir}\" does not exist", LinkEnvironment.PGODirectory);
+					return false;
+				}
+
 				// The linker expects the .pgd and any .pgc files to be in the output directory.
 				// Copy the files there and make them writable...
 				Logger.LogInformation("...copying the profile guided optimization files to output directory...");
