UnrealBuildTool: Fix QueryMode to allow running in parallel

#rnx

diff --git a/Engine/Source/Programs/UnrealBuildTool/Modes/QueryMode.cs b/Engine/Source/Programs/UnrealBuildTool/Modes/QueryMode.cs
--- a/Engine/Source/Programs/UnrealBuildTool/Modes/QueryMode.cs
+++ b/Engine/Source/Programs/UnrealBuildTool/Modes/QueryMode.cs
@@ -2,6 +2,7 @@
 
 using System;
 using System.Collections.Generic;
+using System.IO;
 using System.Linq;
 using System.Text.Json;
 using System.Threading.Tasks;
@@ -63,12 +64,12 @@ namespace UnrealBuildTool
 			return null;
 		}
 	}
-	
+
 	internal class LaunchSettings
 	{
 		public string? Description { get; set; }
 		public string? BinaryPath { get; set; }
-		public List<string> Arguments {get; set; } = new();
+		public List<string> Arguments { get; set; } = new();
 	}
 
 	internal class TargetConfigs
@@ -84,6 +85,9 @@ namespace UnrealBuildTool
 		[CommandLine("-LogDirectory=")]
 		public DirectoryReference? LogDirectory = null;
 
+		[CommandLine("-OutputPath=")]
+		public FileReference? OutputPath = null;
+
 		[CommandLine("-Query=")]
 		public QueryType? Query = null;
 
@@ -110,12 +114,24 @@ namespace UnrealBuildTool
 
 			if (LogDirectory == null)
 			{
-				LogDirectory = DirectoryReference.Combine(Unreal.EngineProgramSavedDirectory, "UnrealBuildTool");
+				LogDirectory = DirectoryReference.Combine(Unreal.EngineProgramSavedDirectory, "UnrealBuildTool", "QueryMode");
 			}
 
 			DirectoryReference.CreateDirectory(LogDirectory);
 
-			FileReference LogFile = FileReference.Combine(LogDirectory, "Log_Query.txt"); // TODO: More history? Pass a log file path on the cmd line from extension and prune from there?
+			string LogFileName;
+			switch (Query)
+			{
+				case QueryType.TargetDetails:
+					LogFileName = $"Log_{TargetName}_{TargetConfiguration}_{TargetPlatform}.txt";
+					break;
+				default:
+					LogFileName = $"Log_{Query}.txt";
+					break;
+
+			}
+
+			FileReference LogFile = FileReference.Combine(LogDirectory, LogFileName);
 			Log.AddFileWriter("DefaultLogTraceListener", LogFile);
 
 			XmlConfig.ApplyTo(BuildConfiguration);
@@ -150,13 +166,25 @@ namespace UnrealBuildTool
 			return Task.FromResult(0);
 		}
 
+		private void WriteResults(object? Value, JsonSerializerOptions JsonOptions, ILogger Logger)
+		{
+			if (OutputPath == null)
+			{
+				Console.WriteLine(JsonSerializer.Serialize(Value, JsonOptions));
+				return;
+			}
+			using FileStream Stream = new FileStream(OutputPath.FullName, FileMode.Create, FileAccess.Write);
+			Logger.LogInformation("Writing {File}...", OutputPath.FullName);
+			JsonSerializer.Serialize(Stream, Value, JsonOptions);
+		}
+
 		private int QueryCapabilities(CommandLineArguments Arguments, ILogger Logger, JsonSerializerOptions JsonOptions)
 		{
 			var Reply = new
 			{
 				Queries = new List<string> { QueryType.Capabilities.ToString(), QueryType.AvailableTargets.ToString(), QueryType.TargetDetails.ToString() }
 			};
-			Console.WriteLine(JsonSerializer.Serialize(Reply, JsonOptions));
+			WriteResults(Reply, JsonOptions, Logger);
 			return 0;
 		}
 
@@ -220,32 +248,33 @@ namespace UnrealBuildTool
 					DefaultPlatform = Platforms[0].ToString(),
 					DefaultConfiguration = UnrealTargetConfiguration.Development.ToString(),
 				};
-				Console.WriteLine(JsonSerializer.Serialize(Reply, JsonOptions));
+
+				WriteResults(Reply, JsonOptions, Logger);
 				return 0;
 			}
 			catch (Exception e)
 			{
 				Logger.LogError("Failed to query available targets: {0}", e.Message);
 				return 1;
-            }
-        }
-        private int QueryTargetDetails(CommandLineArguments Arguments, ILogger Logger, JsonSerializerOptions JsonOptions)
-        {
-            if (TargetName == null)
-            {
-                Logger.LogError("Missing argument Target");
-                return 1;
-            }
-            else if (TargetConfiguration == null)
-            {
-                Logger.LogError("Missing argument Configuration");
-                return 1;
-            }
-            else if (TargetPlatform == null)
-            {
-                Logger.LogError("Missing argument Platform");
-                return 1;
-            }
+			}
+		}
+		private int QueryTargetDetails(CommandLineArguments Arguments, ILogger Logger, JsonSerializerOptions JsonOptions)
+		{
+			if (TargetName == null)
+			{
+				Logger.LogError("Missing argument Target");
+				return 1;
+			}
+			else if (TargetConfiguration == null)
+			{
+				Logger.LogError("Missing argument Configuration");
+				return 1;
+			}
+			else if (TargetPlatform == null)
+			{
+				Logger.LogError("Missing argument Platform");
+				return 1;
+			}
 
 			GenerateProjectFilesMode.TryParseProjectFileArgument(Arguments, Logger, out FileReference? ProjectFileArg);
 			List<string> RawArgs = new List<string> { TargetName, TargetConfiguration, TargetPlatform };
@@ -253,6 +282,7 @@ namespace UnrealBuildTool
 			{
 				RawArgs.Add(ProjectFileArg.ToString());
 			}
+			RawArgs.AddRange(Arguments.GetUnusedArguments());
 			CommandLineArguments Args = new CommandLineArguments(RawArgs.ToArray());
 			List<TargetDescriptor> TargetDescriptors = new();
 
@@ -270,7 +300,13 @@ namespace UnrealBuildTool
 				using (GlobalTracer.Instance.BuildSpan("UEBuildTarget.Create()").StartActive())
 				{
 					bool bUsePrecompiled = false;
-					CurrentTarget = UEBuildTarget.Create(TargetDescriptors[0], false, false, bUsePrecompiled, Logger);
+
+					// Prevent multiple conflicting processes building TargetRules at the same time
+					string MutexName = SingleInstanceMutex.GetUniqueMutexForPath("UnrealBuildTool_QueryMode_UEBuildTarget-Create", Unreal.RootDirectory.FullName);
+					using (new SingleInstanceMutex(MutexName, true))
+					{
+						CurrentTarget = UEBuildTarget.Create(TargetDescriptors[0], false, false, bUsePrecompiled, Logger);
+					}
 				}
 
 				LaunchSettings? CurrentLaunchSettings = null;
@@ -288,6 +324,8 @@ namespace UnrealBuildTool
 				TargetToolChain.SetEnvironmentVariables();
 				CurrentTarget.SetupGlobalEnvironment(TargetToolChain, GlobalCompileEnvironment, GlobalLinkEnvironment);
 
+				CurrentBrowseConfiguration.WindowsSdkVersion = CurrentTarget.Rules.WindowsPlatform.WindowsSdkVersion;
+
 				// TODO: For installed builds, filter out all the binaries that aren't in mods
 				foreach (UEBuildBinary Binary in CurrentTarget.Binaries)
 				{
@@ -301,11 +339,10 @@ namespace UnrealBuildTool
 							CurrentLaunchSettings.Arguments.Add(CurrentTarget.ProjectFile.ToString());
 						}
 					}
-					
+
 					HashSet<UEBuildModule> LinkEnvironmentVisitedModules = new HashSet<UEBuildModule>();
 					CppCompileEnvironment BinaryCompileEnvironment = Binary.CreateBinaryCompileEnvironment(GlobalCompileEnvironment);
 					CurrentBrowseConfiguration.Standard = BinaryCompileEnvironment.CppStandard.ToString();
-					CurrentBrowseConfiguration.WindowsSdkVersion = CurrentTarget.Rules.WindowsPlatform.WindowsSdkVersion;
 
 					foreach (UEBuildModuleCPP Module in Binary.Modules.OfType<UEBuildModuleCPP>())
 					{
@@ -335,13 +372,9 @@ namespace UnrealBuildTool
 						}
 
 						TargetIntellisenseInfo.CompileSettings Settings = new TargetIntellisenseInfo.CompileSettings();
-						if (OperatingSystem.IsWindows())
+						if (TargetToolChain is VCToolChain TargetVCToolChain)
 						{
-							if (CurrentTarget.Platform == UnrealTargetPlatform.Win64)
-							{
-								// TODO: Correct compiler
-								Settings.IncludePaths.AddRange(VCToolChain.GetVCIncludePaths(UnrealTargetPlatform.Win64, WindowsCompiler.VisualStudio2022, null, null, Logger).Split(";"));
-							}
+							Settings.IncludePaths.AddRange(TargetVCToolChain.GetVCIncludePaths().Select(x => x.ToString()));
 						}
 						Settings.IncludePaths.AddRange(ModuleCompileEnvironment.SystemIncludePaths.Select(x => x.ToString()));
 						Settings.IncludePaths.AddRange(ModuleCompileEnvironment.UserIncludePaths.Select(x => x.ToString()));
@@ -349,6 +382,7 @@ namespace UnrealBuildTool
 						Settings.Standard = ModuleCompileEnvironment.CppStandard.ToString();
 						Settings.ForcedIncludes = ModuleCompileEnvironment.ForceIncludeFiles.Select(x => x.ToString()).ToList();
 						Settings.CompilerPath = TargetToolChain.GetCppCompilerPath()?.ToString();
+						Settings.CompilerArgs = TargetToolChain.GetGlobalCommandLineArgs(ModuleCompileEnvironment).ToList();
 						Settings.WindowsSdkVersion = CurrentTarget.Rules.WindowsPlatform.WindowsSdkVersion;
 						CurrentTargetIntellisenseInfo.ModuleToCompileSettings.Add(Module, Settings);
 					}
@@ -363,7 +397,7 @@ namespace UnrealBuildTool
 					LaunchSettings = CurrentLaunchSettings,
 				};
 
-				Console.WriteLine(JsonSerializer.Serialize(Result, JsonOptions));
+				WriteResults(Result, JsonOptions, Logger);
 				return 0;
 			}
 			catch (Exception e)
diff --git a/Engine/Source/Programs/UnrealBuildTool/Platform/Windows/VCToolChain.cs b/Engine/Source/Programs/UnrealBuildTool/Platform/Windows/VCToolChain.cs
--- a/Engine/Source/Programs/UnrealBuildTool/Platform/Windows/VCToolChain.cs
+++ b/Engine/Source/Programs/UnrealBuildTool/Platform/Windows/VCToolChain.cs
@@ -58,6 +58,11 @@ namespace UnrealBuildTool
 			return EnvVars.CompilerPath;
 		}
 
+		public IEnumerable<DirectoryReference> GetVCIncludePaths()
+		{
+			return EnvVars.IncludePaths;
+		}
+
 		/// <summary>
 		/// Prepares the environment for building
 		/// </summary>
@@ -2196,6 +2201,13 @@ namespace UnrealBuildTool
 			CompileAction.bCanExecuteRemotely = false; // Incompatible with remote distribution
 		}
 
+		public override IEnumerable<string> GetGlobalCommandLineArgs(CppCompileEnvironment CompileEnvironment)
+		{
+			List<string> Arguments = new();
+			AppendCLArguments_Global(new(CompileEnvironment), Arguments);
+			return Arguments;
+		}
+
 		public override CppCompileEnvironment CreateSharedResponseFile(CppCompileEnvironment CompileEnvironment, FileReference OutResponseFile, IActionGraphBuilder Graph)
 		{
 			CppCompileEnvironment NewCompileEnvironment = new CppCompileEnvironment(CompileEnvironment);
diff --git a/Engine/Source/Programs/UnrealBuildTool/ToolChain/ClangToolChain.cs b/Engine/Source/Programs/UnrealBuildTool/ToolChain/ClangToolChain.cs
--- a/Engine/Source/Programs/UnrealBuildTool/ToolChain/ClangToolChain.cs
+++ b/Engine/Source/Programs/UnrealBuildTool/ToolChain/ClangToolChain.cs
@@ -1008,6 +1008,13 @@ namespace UnrealBuildTool
 			return NewList;
 		}
 
+		public override IEnumerable<string> GetGlobalCommandLineArgs(CppCompileEnvironment CompileEnvironment)
+		{
+			List<string> Arguments = new();
+			GetCompileArguments_Global(new CppCompileEnvironment(CompileEnvironment), Arguments);
+			return Arguments;
+		}
+
 		public override CppCompileEnvironment CreateSharedResponseFile(CppCompileEnvironment CompileEnvironment, FileReference OutResponseFile, IActionGraphBuilder Graph)
 		{
 			CppCompileEnvironment NewCompileEnvironment = new CppCompileEnvironment(CompileEnvironment);
diff --git a/Engine/Source/Programs/UnrealBuildTool/ToolChain/UEToolChain.cs b/Engine/Source/Programs/UnrealBuildTool/ToolChain/UEToolChain.cs
--- a/Engine/Source/Programs/UnrealBuildTool/ToolChain/UEToolChain.cs
+++ b/Engine/Source/Programs/UnrealBuildTool/ToolChain/UEToolChain.cs
@@ -162,6 +162,11 @@ namespace UnrealBuildTool
 			return Result.ToArray();
 		}
 
+		public virtual IEnumerable<string> GetGlobalCommandLineArgs(CppCompileEnvironment CompileEnvironment)
+		{
+			return Array.Empty<string>();
+		}
+
 		public virtual CppCompileEnvironment CreateSharedResponseFile(CppCompileEnvironment CompileEnvironment, FileReference OutResponseFile, IActionGraphBuilder Graph)
 		{
 			return CompileEnvironment;
