[UBT]
* Disable box for msvc pgoprofile link steps

diff --git a/Engine/Source/Programs/UnrealBuildTool/Platform/Windows/VCToolChain.cs b/Engine/Source/Programs/UnrealBuildTool/Platform/Windows/VCToolChain.cs
--- a/Engine/Source/Programs/UnrealBuildTool/Platform/Windows/VCToolChain.cs
+++ b/Engine/Source/Programs/UnrealBuildTool/Platform/Windows/VCToolChain.cs
@@ -2620,7 +2620,7 @@ namespace UnrealBuildTool
 				LinkAction.bCanExecuteRemotely = true;
 			}
 
-			if (LinkEnvironment.bPGOOptimize)
+			if (LinkEnvironment.bPGOOptimize || LinkEnvironment.bPGOProfile)
 			{
 				LinkAction.bCanExecuteInBox = false; // Disabled for now. Should revisit to see why it is not working
 			}
