UnrealBuildTool: Remove support for VS2019 msvc toolchain

diff --git a/Engine/Source/Programs/UnrealBuildTool/Platform/Windows/MicrosoftPlatformSDK.Versions.cs b/Engine/Source/Programs/UnrealBuildTool/Platform/Windows/MicrosoftPlatformSDK.Versions.cs
--- a/Engine/Source/Programs/UnrealBuildTool/Platform/Windows/MicrosoftPlatformSDK.Versions.cs
+++ b/Engine/Source/Programs/UnrealBuildTool/Platform/Windows/MicrosoftPlatformSDK.Versions.cs
@@ -47,7 +47,6 @@ namespace UnrealBuildTool
 			VersionNumberRange.Parse("14.36.32532", "14.36.99999"), // VS2022 17.6.x
 			VersionNumberRange.Parse("14.35.32215", "14.35.99999"), // VS2022 17.5.x
 			VersionNumberRange.Parse("14.34.31933", "14.34.99999"), // VS2022 17.4.x
-			VersionNumberRange.Parse("14.29.30133", "14.29.99999"), // VS2019 16.11.x
 		};
 
 		/// <summary>
@@ -58,22 +57,17 @@ namespace UnrealBuildTool
 			{ new VersionNumber(14, 37), new VersionNumber(16) }, // VS2022 17.7.x
 			{ new VersionNumber(14, 35), new VersionNumber(15) }, // VS2022 17.5.x - 17.6.x
 			{ new VersionNumber(14, 34), new VersionNumber(14) }, // VS2022 17.4.x
-			{ new VersionNumber(14, 29), new VersionNumber(13) }, // VS2019 16.11.x
-
 		};
 
 		/// <summary>
 		/// Tested compiler toolchains that should not be allowed.
 		/// </summary>
-		static readonly VersionNumberRange[] BannedVisualCppVersions = new VersionNumberRange[]
-		{
-			VersionNumberRange.Parse("14.30.0", "14.33.99999"), // VS2022 17.0.x - 17.3.x
-		};
+		static readonly VersionNumberRange[] BannedVisualCppVersions = System.Array.Empty<VersionNumberRange>();
 
 		/// <summary>
 		/// The minimum supported MSVC compiler
 		/// </summary>
-		static readonly VersionNumber MinimumVisualCppVersion = new VersionNumber(14, 29, 30133);
+		static readonly VersionNumber MinimumVisualCppVersion = new VersionNumber(14, 34, 31933);
 
 		/// <summary>
 		/// The default compiler version to be used, if installed. 
