UnrealBuildTool: ISPC pass --cpu when compiling not just generating headers

#rnx

diff --git a/Engine/Source/Programs/UnrealBuildTool/ToolChain/ISPCToolChain.cs b/Engine/Source/Programs/UnrealBuildTool/ToolChain/ISPCToolChain.cs
--- a/Engine/Source/Programs/UnrealBuildTool/ToolChain/ISPCToolChain.cs
+++ b/Engine/Source/Programs/UnrealBuildTool/ToolChain/ISPCToolChain.cs
@@ -565,6 +565,12 @@ namespace UnrealBuildTool
 			GlobalArguments.Add($"--target={TargetString}");
 			GlobalArguments.Add($"--emit-{PlatformObjectFileFormat}");
 
+			string? CpuTarget = GetISPCCpuTarget(CompileEnvironment.Platform);
+			if (!String.IsNullOrEmpty(CpuTarget))
+			{
+				GlobalArguments.Add($"--cpu={CpuTarget}");
+			}
+
 			bool bByteCodeOutput = (PlatformObjectFileFormat == "llvm");
 
 			List<string> CommonArgs = new List<string>();
