Fixes our Clang database generation to produce a compilation database that doesn't trip up against our checks in platform headers, so that it can actually be used properly.

diff --git a/Engine/Source/Programs/UnrealBuildTool/Modes/GenerateClangDatabase.cs b/Engine/Source/Programs/UnrealBuildTool/Modes/GenerateClangDatabase.cs
--- a/Engine/Source/Programs/UnrealBuildTool/Modes/GenerateClangDatabase.cs
+++ b/Engine/Source/Programs/UnrealBuildTool/Modes/GenerateClangDatabase.cs
@@ -186,6 +186,18 @@ namespace UnrealBuildTool
 
 					// Create all the binaries and modules
 					CppCompileEnvironment GlobalCompileEnvironment = Target.CreateCompileEnvironmentForProjectFiles(Logger);
+					// Some of our platform-specific wrappers check against this preprocessor macro, which `clang`/`clang-cl` will not set by default.
+					if (!GlobalCompileEnvironment.Architectures.bIsMultiArch)  // For now, this will only handle single-architecture builds since setting both macros at once doesn't really make sense.
+					{
+						if (GlobalCompileEnvironment.Architecture.bIsX64)
+						{
+							GlobalCompileEnvironment.AdditionalArguments = "-DPLATFORM_CPU_X86_FAMILY";
+						}
+						else  // Currently, this assumes that non-x64 arch means we're running on ARM.
+						{
+							GlobalCompileEnvironment.AdditionalArguments = "-DPLATFORM_CPU_ARM_FAMILY";
+						}
+					}
 					foreach (UEBuildBinary Binary in Target.Binaries)
 					{
 						CppCompileEnvironment BinaryCompileEnvironment = Binary.CreateBinaryCompileEnvironment(GlobalCompileEnvironment);
diff --git a/Engine/Source/Runtime/Core/Public/GenericPlatform/GenericPlatformProcess.h b/Engine/Source/Runtime/Core/Public/GenericPlatform/GenericPlatformProcess.h
--- a/Engine/Source/Runtime/Core/Public/GenericPlatform/GenericPlatformProcess.h
+++ b/Engine/Source/Runtime/Core/Public/GenericPlatform/GenericPlatformProcess.h
@@ -819,7 +819,7 @@ struct FGenericPlatformProcess
 		if (FPlatformMisc::HasTimedPauseCPUFeature())
 		{
 			uint64 PauseCycles = ReadCycleCounter() + Cycles;
-#if defined(_MSC_VER) && !defined(__clang__)
+#if defined(_MSC_VER)
 			_tpause(0, PauseCycles);
 #elif __has_builtin(__builtin_ia32_tpause)
 			__builtin_ia32_tpause(0, (uint32)(PauseCycles >> 32), (uint32)PauseCycles);
