UnrealBuildTool: WindowsPlatform.HasValidCompiler

#rnx

diff --git a/Engine/Source/Programs/UnrealBuildTool/Platform/Windows/MicrosoftPlatformSDK.cs b/Engine/Source/Programs/UnrealBuildTool/Platform/Windows/MicrosoftPlatformSDK.cs
--- a/Engine/Source/Programs/UnrealBuildTool/Platform/Windows/MicrosoftPlatformSDK.cs
+++ b/Engine/Source/Programs/UnrealBuildTool/Platform/Windows/MicrosoftPlatformSDK.cs
@@ -523,11 +523,18 @@ namespace UnrealBuildTool
 			return FindToolChainInstallations(Compiler, Logger).Where(x => (x.Architecture == Architecture)).Any();
 		}
 
+		/// <summary>
+		/// Determines if a given compiler is installed and valid
+		/// </summary>
+		/// <param name="Compiler">Compiler to check for</param>
+		/// <param name="Architecture">Architecture the compiler must support</param>
+		/// <param name="Logger">Logger for output</param>
+		/// <returns>True if the given compiler is installed and valid</returns>
 		public static bool HasValidCompiler(WindowsCompiler Compiler, UnrealArch Architecture, ILogger Logger)
 		{
 			// since this is static, we get the instance for it
 			MicrosoftPlatformSDK SDK = (MicrosoftPlatformSDK)GetSDKForPlatform("Win64")!;
-			return FindToolChainInstallations(Compiler, Logger).Where(x => (x.Version >= MinimumVisualCppVersion && x.Architecture == Architecture)).Any();
+			return FindToolChainInstallations(Compiler, Logger).Where(x => (x.Error == null && x.Architecture == Architecture)).Any();
 		}
 
 		/// <summary>
diff --git a/Engine/Source/Programs/UnrealBuildTool/Platform/Windows/UEBuildWindows.cs b/Engine/Source/Programs/UnrealBuildTool/Platform/Windows/UEBuildWindows.cs
--- a/Engine/Source/Programs/UnrealBuildTool/Platform/Windows/UEBuildWindows.cs
+++ b/Engine/Source/Programs/UnrealBuildTool/Platform/Windows/UEBuildWindows.cs
@@ -694,6 +694,18 @@ namespace UnrealBuildTool
 			}
 		}
 
+		/// <summary>
+		/// Determines if a given compiler is installed and valid
+		/// </summary>
+		/// <param name="Compiler">Compiler to check for</param>
+		/// <param name="Architecture">Architecture the compiler must support</param>
+		/// <param name="Logger">Logger for output</param>
+		/// <returns>True if the given compiler is installed and valid</returns>
+		public static bool HasValidCompiler(WindowsCompiler Compiler, UnrealArch Architecture, ILogger Logger)
+		{
+			return MicrosoftPlatformSDK.HasValidCompiler(Compiler, Architecture, Logger);
+		}
+
 		/// <summary>
 		/// Constructor
 		/// </summary>
@@ -1297,12 +1309,12 @@ namespace UnrealBuildTool
 		}
 
 		/// <summary>
-		/// Determines if a given compiler is installed
+		/// Determines if a given compiler is installed and valid
 		/// </summary>
 		/// <param name="Compiler">Compiler to check for</param>
 		/// <param name="Architecture">Architecture the compiler must support</param>
 		/// <param name="Logger">Logger for output</param>
-		/// <returns>True if the given compiler is installed</returns>
+		/// <returns>True if the given compiler is installed and valid</returns>
 		public static bool HasCompiler(WindowsCompiler Compiler, UnrealArch Architecture, ILogger Logger)
 		{
 			return MicrosoftPlatformSDK.HasCompiler(Compiler, Architecture, Logger);
