- Fixed various issues with making Mac/IOS builds that can be uploaded to AppStore in Xcode
  * Missing ThinApp.sh
  * Missing some quotes around a variable in GenerateUniversalDSYM.sh
  * Correctly filter out the stub .app's and such when copying the Staged data directory into the .app
  * Disabled the Mac's Sign To Run Locally option by default, as that causes Xcode's Validation/Upload to not have the team name embedded in it, causing a hassle while pushing up to AppStore/TestFlight
  * Fixed the PRODUCT_NAME for BP projects
  * Made CrashReportClient be a sandboxed app that inherits from parent
  * Fix Hybrid apps to check all platforms before in the project generator, so that if a project needs a temp target for IOS (via a plugin)
  * For hybrid projects, don't append Game to the end of the Game target names, as that is the default, and it makes for better named products
  * Disabled the LoginFlow module from OnlineFramework plugin - this was causing issues with having CEF and EpicWebHelper embedded into a sandboxed .app, and LoginFlow isn't seemingly actually used by anything

diff --git a/Engine/Build/BatchFiles/Mac/GenerateUniversalDSYM.sh b/Engine/Build/BatchFiles/Mac/GenerateUniversalDSYM.sh
--- a/Engine/Build/BatchFiles/Mac/GenerateUniversalDSYM.sh
+++ b/Engine/Build/BatchFiles/Mac/GenerateUniversalDSYM.sh
@@ -17,7 +17,7 @@ if [[ "${extension}" != "dSYM" &&  "${extension}" != "dsym" ]]; then
 fi
 
 # Cleanup any old dSYM
-if [ -e $2 ]; then
+if [ -e "$2" ]; then
 	echo "Deleting existing .dSYM"
 	rm -r "$2"
 fi
diff --git a/Engine/Build/InstalledEngineFilters.xml b/Engine/Build/InstalledEngineFilters.xml
--- a/Engine/Build/InstalledEngineFilters.xml
+++ b/Engine/Build/InstalledEngineFilters.xml
@@ -340,6 +340,7 @@
 		Engine/Build/BatchFiles/Mac/XcodeBuild.sh
 		Engine/Build/BatchFiles/Mac/GenerateUniversalDSYM.sh
 		Engine/Build/BatchFiles/Mac/UpdateVersionAfterBuild.sh
+		Engine/Build/BatchFiles/Mac/ThinApp.sh
 		Engine/Build/BatchFiles/RunDotnet.sh
 		Engine/Build/BatchFiles/RunDotnetMSBuild.sh
 		Engine/Build/BatchFiles/RunUAT.command
diff --git a/Engine/Config/BaseEngine.ini b/Engine/Config/BaseEngine.ini
--- a/Engine/Config/BaseEngine.ini
+++ b/Engine/Config/BaseEngine.ini
@@ -3059,7 +3059,7 @@ bEnableRayTracing=false
 [/Script/MacTargetPlatform.XcodeProjectSettings]
 bUseModernXcode=true
 bUseAutomaticCodeSigning=true
-bMacSignToRunLocally=true
+bMacSignToRunLocally=false
 CodeSigningTeam=
 CodeSigningPrefix=com.YourCompany
 BundleIdentifier=$(UE_SIGNING_PREFIX).$(UE_PRODUCT_NAME_STRIPPED)
diff --git a/Engine/Plugins/Online/OnlineFramework/OnlineFramework.uplugin b/Engine/Plugins/Online/OnlineFramework/OnlineFramework.uplugin
--- a/Engine/Plugins/Online/OnlineFramework/OnlineFramework.uplugin
+++ b/Engine/Plugins/Online/OnlineFramework/OnlineFramework.uplugin
@@ -41,7 +41,10 @@
 		},
 		{
 			"Name": "LoginFlow",
-			"Type":  "ClientOnly"
+			"Type":  "ClientOnly",
+			"PlatformDenyList": [
+				"Mac"
+			]
 		}
 	],
 	"LocalizationTargets": [
@@ -60,4 +63,4 @@
 			"Enabled": true
 		}
 	]
-}
\ No newline at end of file
+}
diff --git a/Engine/Programs/CrashReportClient/Config/DefaultEngine.ini b/Engine/Programs/CrashReportClient/Config/DefaultEngine.ini
--- a/Engine/Programs/CrashReportClient/Config/DefaultEngine.ini
+++ b/Engine/Programs/CrashReportClient/Config/DefaultEngine.ini
@@ -13,3 +13,8 @@ n.VerifyPeer=false
 [/Script/UdpMessaging.UdpMessagingSettings]
 EnableTransport=True
 MulticastTimeToLive=0
+
+[/Script/MacTargetPlatform.XcodeProjectSettings]
+PremadeMacEntitlements=(FilePath="/Engine/Build/Mac/Resources/Sandbox.EmbeddedApp.entitlements")
+ShippingSpecificMacEntitlements=(FilePath="/Engine/Build/Mac/Resources/Sandbox.EmbeddedApp.entitlements")
+bUseEntitlementsForPrograms=true
diff --git a/Engine/Source/Programs/AutomationTool/Apple/ApplePlatform.Automation.cs b/Engine/Source/Programs/AutomationTool/Apple/ApplePlatform.Automation.cs
--- a/Engine/Source/Programs/AutomationTool/Apple/ApplePlatform.Automation.cs
+++ b/Engine/Source/Programs/AutomationTool/Apple/ApplePlatform.Automation.cs
@@ -109,10 +109,6 @@ public abstract class ApplePlatform : Platform
 				string TargetName = Target.TargetName;
 				if (!Params.IsCodeBasedProject)
 				{
-					// instead of staging an UnrealGame.app, stage something with the project name, like say MyProjectClient-IOS-Shipping.app
-					string ProductName = AppleExports.MakeBinaryFileName(SC.ShortProjectName, Target.Platform, Target.Configuration, Target.Architectures, UnrealTargetConfiguration.Development, null);
-					ExtraOptions += $" PRODUCT_NAME={ProductName}";
-
 					if (Params.RawProjectPath != null)
 					{
 						TargetName = MakeContentOnlyTargetName(Target, Params.ShortProjectName);
@@ -205,11 +201,12 @@ public abstract class ApplePlatform : Platform
 					$"UE_OVERRIDE_STAGE_DIR=\"{SC.StageDirectory}\"";
 				if (!Params.IsCodeBasedProject)
 				{
-					// override where the .app will be located and named
-					ExtraOptions += $" SYMROOT=\"{SC.ProjectRoot}/Binaries\"";
-					string ProductName = AppleExports.MakeBinaryFileName(SC.ShortProjectName, Receipt.Platform, Receipt.Configuration, Receipt.Architectures, UnrealTargetConfiguration.Development, null);
-					ExtraOptions += $" PRODUCT_NAME={ProductName}";
-
+					if (!Params.Distribution)
+					{
+						// override where the .app will be located and named
+						ExtraOptions += $" SYMROOT=\"{SC.ProjectRoot}/Binaries\"";
+					}
+					
 					if (Params.RawProjectPath != null)
 					{
 						TargetName = MakeContentOnlyTargetName(Receipt, Params.ShortProjectName);
diff --git a/Engine/Source/Programs/UnrealBuildTool/ProjectFiles/PlatformProjectGeneratorCollection.cs b/Engine/Source/Programs/UnrealBuildTool/ProjectFiles/PlatformProjectGeneratorCollection.cs
--- a/Engine/Source/Programs/UnrealBuildTool/ProjectFiles/PlatformProjectGeneratorCollection.cs
+++ b/Engine/Source/Programs/UnrealBuildTool/ProjectFiles/PlatformProjectGeneratorCollection.cs
@@ -2,6 +2,7 @@
 
 using System.Collections.Generic;
 using System.Text;
+using System.Linq;
 using EpicGames.Core;
 using Microsoft.Extensions.Logging;
 
@@ -14,6 +15,15 @@ namespace UnrealBuildTool
 	{
 		Dictionary<UnrealTargetPlatform, PlatformProjectGenerator> ProjectGeneratorDictionary = new Dictionary<UnrealTargetPlatform, PlatformProjectGenerator>();
 
+		/// <summary>
+		/// Returns the list of platforms that have been registered in this collection
+		/// </summary>
+		/// <returns>Registered platforms</returns>
+		public List<UnrealTargetPlatform> GetRegisteredPlatforms()
+		{
+			return ProjectGeneratorDictionary.Keys.ToList();
+		}
+
 		/// <summary>
 		/// Register the given platforms UEPlatformProjectGenerator instance
 		/// </summary>
diff --git a/Engine/Source/Programs/UnrealBuildTool/ProjectFiles/ProjectFileGenerator.cs b/Engine/Source/Programs/UnrealBuildTool/ProjectFiles/ProjectFileGenerator.cs
--- a/Engine/Source/Programs/UnrealBuildTool/ProjectFiles/ProjectFileGenerator.cs
+++ b/Engine/Source/Programs/UnrealBuildTool/ProjectFiles/ProjectFileGenerator.cs
@@ -962,7 +962,7 @@ namespace UnrealBuildTool
 			{
 				foreach (FileReference GameUProjectFile in AllGameProjects)
 				{
-					NativeProjects.ConditionalMakeTempTargetForHybridProject(GameUProjectFile, new List<UnrealTargetPlatform>(1) { BuildHostPlatform.Current.Platform }, Logger);
+					NativeProjects.ConditionalMakeTempTargetForHybridProject(GameUProjectFile, PlatformProjectGenerators.GetRegisteredPlatforms(), Logger);
 				}
 
 				// they are created in a temp location, which we need to scan
@@ -2910,7 +2910,10 @@ namespace UnrealBuildTool
 						{
 							if (bMakeProjectPerTarget)
 							{
-								ProjectFile = FindOrAddProject(GetProjectLocation($"{ProjectName}{EngineTarget.TargetRules!.Type}"), ContentOnlyGameProject.Directory, IncludeInGeneratedProjects: true, bAlreadyExisted: out _);
+								// don't append Game suffix to the ContentOnly game target names, so the default case has the target/product named after the BP project name
+								// @note if this changes, see ApplePlatform.MakeContentOnlyTargetName
+								string TargetTypeSuffix = /*EngineTarget.TargetRules!.Type == TargetType.Game ? "" :*/ EngineTarget.TargetRules!.Type.ToString();
+								ProjectFile = FindOrAddProject(GetProjectLocation($"{ProjectName}{TargetTypeSuffix}"), ContentOnlyGameProject.Directory, IncludeInGeneratedProjects: true, bAlreadyExisted: out _);
 							}
 							else if (ProjectFile == null)
 							{
diff --git a/Engine/Source/Programs/UnrealBuildTool/ProjectFiles/Xcode/XcodeProject.Utils.cs b/Engine/Source/Programs/UnrealBuildTool/ProjectFiles/Xcode/XcodeProject.Utils.cs
--- a/Engine/Source/Programs/UnrealBuildTool/ProjectFiles/Xcode/XcodeProject.Utils.cs
+++ b/Engine/Source/Programs/UnrealBuildTool/ProjectFiles/Xcode/XcodeProject.Utils.cs
@@ -76,25 +76,6 @@ namespace UnrealBuildTool.XcodeProjectXcconfig
 			return false;
 		}
 
-		public static string MakeExecutableFileName(string BinaryName, UnrealTargetPlatform Platform, UnrealTargetConfiguration Configuration, UnrealArchitectures Architectures, UnrealTargetConfiguration UndecoratedConfiguration)
-		{
-			StringBuilder Result = new StringBuilder();
-
-			Result.Append(BinaryName);
-
-			if (Configuration != UndecoratedConfiguration)
-			{
-				Result.AppendFormat("-{0}-{1}", Platform.ToString(), Configuration.ToString());
-			}
-
-			if (UnrealArchitectureConfig.ForPlatform(Platform).RequiresArchitectureFilenames(Architectures))
-			{
-				Result.Append(Architectures.ToString());
-			}
-
-			return Result.ToString();
-		}
-
 		// cache for the below function
 		static Dictionary<string, UnrealArchitectures> CachedMacProjectArchitectures = new();
 
diff --git a/Engine/Source/Programs/UnrealBuildTool/ProjectFiles/Xcode/XcodeProject.cs b/Engine/Source/Programs/UnrealBuildTool/ProjectFiles/Xcode/XcodeProject.cs
--- a/Engine/Source/Programs/UnrealBuildTool/ProjectFiles/Xcode/XcodeProject.cs
+++ b/Engine/Source/Programs/UnrealBuildTool/ProjectFiles/Xcode/XcodeProject.cs
@@ -1370,6 +1370,8 @@ namespace UnrealBuildTool.XcodeProjectXcconfig
 				string DefaultStageDir = bIsEngineBuild ? "" : "${UE_PROJECT_DIR}/Saved/StagedBuilds/${UE_TARGET_PLATFORM_NAME}";
 				string SyncSourceSubdir = (Platform == UnrealTargetPlatform.Mac) ? "" : "/cookeddata";
 				string SyncDestSubdir = (Platform == UnrealTargetPlatform.Mac) ? "/UE" : "/cookeddata";
+				string ExecutableKey = $"UE_{Platform.ToString().ToUpper()}_EXECUTABLE_NAME";
+
 				CopyScript.AddRange(new string[]
 				{
 					"# Skip syncing if desired",
@@ -1414,7 +1416,7 @@ namespace UnrealBuildTool.XcodeProjectXcconfig
 				{
 					"",
 					$"echo \\\"Syncing ${{STAGED_DIR}}{SyncSourceSubdir} to ${{CONFIGURATION_BUILD_DIR}}/${{CONTENTS_FOLDER_PATH}}{SyncDestSubdir}\\\"",
-					$"rsync -a --delete --exclude=/Info.plist --exclude=/Manifest_* --exclude=${{UE_TARGET_NAME}}.app \\\"${{STAGED_DIR}}{SyncSourceSubdir}/\\\" \\\"${{CONFIGURATION_BUILD_DIR}}/${{CONTENTS_FOLDER_PATH}}{SyncDestSubdir}\\\"",
+					$"rsync -a --delete --exclude=/Info.plist --exclude=/Manifest_* --exclude=${{{ExecutableKey}}}.app --exclude=${{PRODUCT_NAME}} --exclude=${{{ExecutableKey}}} --exclude=${{PRODUCT_NAME}}.app \\\"${{STAGED_DIR}}{SyncSourceSubdir}/\\\" \\\"${{CONFIGURATION_BUILD_DIR}}/${{CONTENTS_FOLDER_PATH}}{SyncDestSubdir}\\\"",
 				});
 			}
 
@@ -1940,11 +1942,23 @@ namespace UnrealBuildTool.XcodeProjectXcconfig
 				// hook up the Buildconfig that matches this info to this xcconfig file
 				XcconfigFile ConfigXcconfig = MatchedConfig.Xcconfig!;
 
-				string ExecutableName = XcodeUtils.MakeExecutableFileName(Config.ExeName, Platform, Config.BuildConfig, TargetRules.Architectures, TargetRules.UndecoratedConfiguration);
+				string ExecutableName = AppleExports.MakeBinaryFileName(Config.ExeName, Platform, Config.BuildConfig, TargetRules.Architectures, TargetRules.UndecoratedConfiguration, null);
 				string ExetuableSubPath = FileReference.Combine(ConfigBuildDir, ExecutableName).MakeRelativeTo(ConfigBuildDir);
-				string ProductName = ExecutableName;
 				string ExecutableKey = $"UE_{Platform.ToString().ToUpper()}_EXECUTABLE_NAME";
 
+				// we want to make Foo.app, not FooGame.app, since we added on the target type when making targets
+				string PerTargetTypeProductName = UnrealData.ProductName;
+				if (UnrealData.bIsContentOnlyProject && TargetRules.Type == TargetType.Game)
+				{
+					PerTargetTypeProductName = UnrealData.UProjectFileLocation!.GetFileNameWithoutAnyExtensions();
+				}
+				string ProductName = ExecutableName;
+				// content only projects don't want UnrealGame, etc as the ProductName
+				if (UnrealData.bIsContentOnlyProject && TargetRules.Type != TargetType.Editor)
+				{
+					ProductName = AppleExports.MakeBinaryFileName(PerTargetTypeProductName, Platform, Config.BuildConfig, TargetRules.Architectures, TargetRules.UndecoratedConfiguration, null);
+				}
+
 				MetadataItem? EntitlementsMetadata = UnrealData.Metadata!.EntitlementsFiles[MetadataPlatform];
 
 				// if we have shipping specified, use it instead
