UnrealBuildTool: Don't always assume a working set is P4 unless there is a read-only file in the workspace root

#rnx

diff --git a/Engine/Source/Programs/UnrealBuildTool/System/SourceFileWorkingSet.cs b/Engine/Source/Programs/UnrealBuildTool/System/SourceFileWorkingSet.cs
--- a/Engine/Source/Programs/UnrealBuildTool/System/SourceFileWorkingSet.cs
+++ b/Engine/Source/Programs/UnrealBuildTool/System/SourceFileWorkingSet.cs
@@ -22,7 +22,7 @@ namespace UnrealBuildTool
 		/// </summary>
 		/// <param name="File">File to check</param>
 		/// <returns>True if the file is part of the working set, false otherwise</returns>
-		bool Contains(FileItem File);
+		abstract bool Contains(FileItem File);
 	}
 
 	/// <summary>
@@ -327,7 +327,7 @@ namespace UnrealBuildTool
 			}
 			else if (Provider == ProviderType.Git)
 			{
-				GitSourceFileWorkingSet? WorkingSet;
+				ISourceFileWorkingSet? WorkingSet;
 				if (!String.IsNullOrEmpty(RepositoryPath))
 				{
 					WorkingSet = new GitSourceFileWorkingSet(GitPath, DirectoryReference.Combine(RootDir, RepositoryPath), null, Logger);
@@ -344,19 +344,20 @@ namespace UnrealBuildTool
 			}
 			else
 			{
-				GitSourceFileWorkingSet? WorkingSet;
+				ISourceFileWorkingSet? WorkingSet;
 				if (TryCreateGitWorkingSet(RootDir, ProjectDirs, Logger, out WorkingSet))
 				{
 					return WorkingSet;
 				}
-				else
+				else if (TryCreatePerforceWorkingSet(RootDir, ProjectDirs, Logger, out WorkingSet))
 				{
-					return new PerforceSourceFileWorkingSet();
+					return WorkingSet;
 				}
 			}
+			return new EmptySourceFileWorkingSet();
 		}
 
-		static bool TryCreateGitWorkingSet(DirectoryReference RootDir, IEnumerable<DirectoryReference> ProjectDirs, ILogger Logger, [NotNullWhen(true)] out GitSourceFileWorkingSet? OutWorkingSet)
+		static bool TryCreateGitWorkingSet(DirectoryReference RootDir, IEnumerable<DirectoryReference> ProjectDirs, ILogger Logger, [NotNullWhen(true)] out ISourceFileWorkingSet? OutWorkingSet)
 		{
 			GitSourceFileWorkingSet? WorkingSet = null;
 
@@ -386,5 +387,19 @@ namespace UnrealBuildTool
 			OutWorkingSet = WorkingSet;
 			return OutWorkingSet != null;
 		}
+
+		static bool TryCreatePerforceWorkingSet(DirectoryReference RootDir, IEnumerable<DirectoryReference> ProjectDirs, ILogger Logger, [NotNullWhen(true)] out ISourceFileWorkingSet? OutWorkingSet)
+		{
+			PerforceSourceFileWorkingSet? WorkingSet = null;
+			// If an installed engine, or the root directory contains any read-only files assume this is a perforce working set
+			if (Unreal.IsEngineInstalled() || DirectoryReference.EnumerateFiles(RootDir).Any(x => x.ToFileInfo().Attributes.HasFlag(FileAttributes.ReadOnly)))
+			{
+				WorkingSet = new PerforceSourceFileWorkingSet();
+			}
+
+			// Set the output value
+			OutWorkingSet = WorkingSet;
+			return OutWorkingSet != null;
+		}
 	}
 }
