UnrealBuildTool: Simplify the UHT #include root path

#rnx

diff --git a/Engine/Source/Programs/UnrealBuildTool/System/ExternalExecution.cs b/Engine/Source/Programs/UnrealBuildTool/System/ExternalExecution.cs
--- a/Engine/Source/Programs/UnrealBuildTool/System/ExternalExecution.cs
+++ b/Engine/Source/Programs/UnrealBuildTool/System/ExternalExecution.cs
@@ -372,39 +372,18 @@ namespace UnrealBuildTool
 
 		static DirectoryReference FindIncludeBase(UEBuildModuleCPP Module, ILogger Logger)
 		{
-			DirectoryReference? ModuleIncludeBase = Module.PublicIncludePaths.FirstOrDefault() ??
-				Module.PrivateIncludePaths.FirstOrDefault() ??
-				Module.InternalIncludePaths.FirstOrDefault() ??
-				(Module.Rules.bLegacyParentIncludePaths ? Module.LegacyParentIncludePaths.FirstOrDefault() : null) ??
-				(Module.Rules.bLegacyPublicIncludePaths ? Module.LegacyPublicIncludePaths.FirstOrDefault() : null);
-
-			if (ModuleIncludeBase == null)
+			// Project Source directory is always an available include path for modules under a UProject (see UEBuildTarget.FindOrCreateModuleByName)
+			if (Module.Rules.Target.ProjectFile != null)
 			{
-				// Project Source directory is also always for modules under a UProject (see UEBuildTarget.FindOrCreateModuleByName)
-				if (Module.Rules.Target.ProjectFile != null)
-				{
-					DirectoryReference ProjectSourceDirectoryName = DirectoryReference.Combine(Module.Rules.Target.ProjectFile.Directory, "Source");
-					if (Module.Rules.File.IsUnderDirectory(ProjectSourceDirectoryName))
-					{
-						ModuleIncludeBase = ProjectSourceDirectoryName;
-					}
-				}
-
-				// If we still cannot find an include base, this most likely means the module is set up incorrectly when using BuildSettingsVersion.V3.=
-				// Fall back to Engine/Source or the UHT generated directory, however this should be resolved by fixing the ModuleRules.
-				if (ModuleIncludeBase == null && Module.ModuleDirectory.IsUnderDirectory(Unreal.EngineSourceDirectory))
-				{
-					ModuleIncludeBase = Unreal.EngineSourceDirectory;
-					Logger.LogDebug("Unable to find a module include path for {Module}, using engine source path '{Path}' because no standard Public/Private/Internal include paths were added. Please resolve by updating the module's .Build.cs", Module.Name, ModuleIncludeBase);
-				}
-				if (ModuleIncludeBase == null)
+				DirectoryReference ProjectSourceDirectoryName = DirectoryReference.Combine(Module.Rules.Target.ProjectFile.Directory, "Source");
+				if (Module.Rules.File.IsUnderDirectory(ProjectSourceDirectoryName))
 				{
-					ModuleIncludeBase = Module.GeneratedCodeDirectoryUHT!;
-					Logger.LogDebug("Unable to find a module include path for {Module}, using generated path '{Path}' because no standard Public/Private/Internal include paths were added. Please resolve by updating the module's .Build.cs", Module.Name, ModuleIncludeBase);
+					return ProjectSourceDirectoryName;
 				}
 			}
 
-			return ModuleIncludeBase;
+			// Otherwise use Engine/Source
+			return Unreal.EngineSourceDirectory;
 		}
 
 		public static void SetupUObjectModules(IEnumerable<UEBuildModuleCPP> ModulesToGenerateHeadersFor, UnrealTargetPlatform Platform, ProjectDescriptor? ProjectDescriptor, List<UHTModuleInfo> UObjectModules, List<UHTModuleHeaderInfo> UObjectModuleHeaders, EGeneratedCodeVersion GeneratedCodeVersion, SourceFileMetadataCache MetadataCache, ILogger Logger)
