UnrealBuildTool: Use IEnumerable for list of toolchain input lists that should be immutable

#rnx

diff --git a/Engine/Source/Programs/UnrealBuildTool/Platform/Android/AndroidToolChain.cs b/Engine/Source/Programs/UnrealBuildTool/Platform/Android/AndroidToolChain.cs
--- a/Engine/Source/Programs/UnrealBuildTool/Platform/Android/AndroidToolChain.cs
+++ b/Engine/Source/Programs/UnrealBuildTool/Platform/Android/AndroidToolChain.cs
@@ -1013,20 +1013,22 @@ namespace UnrealBuildTool
 		private bool bHasHandledLaunchModule = false;
 		private bool bHasHandledCoreModule = false;
 
-		protected override CPPOutput CompileCPPFiles(CppCompileEnvironment CompileEnvironment, List<FileItem> InputFiles, DirectoryReference OutputDir, string ModuleName, IActionGraphBuilder Graph)
+		protected override CPPOutput CompileCPPFiles(CppCompileEnvironment CompileEnvironment, IEnumerable<FileItem> InputFiles, DirectoryReference OutputDir, string ModuleName, IActionGraphBuilder Graph)
 		{
 			if (ShouldSkipModule(ModuleName, CompileEnvironment.Architecture))
 			{
 				return new CPPOutput();
 			}
 
+			List<FileItem> ModifiedInputFiles = new(InputFiles);
+
 			// Deal with Launch module special if first time seen
 			if (!bHasHandledLaunchModule && (ModuleName.Equals("Launch") || ModuleName.Equals("AndroidLauncher")))
 			{
 				// Directly added NDK files for NDK extensions
-				InputFiles.Add(FileItem.GetItemByPath(GetNativeGluePath()));
+				ModifiedInputFiles.Add(FileItem.GetItemByPath(GetNativeGluePath()));
 				// Deal with dynamic modules removed by architecture
-				GenerateEmptyLinkFunctionsForRemovedModules(InputFiles, CompileEnvironment.Architecture, ModuleName, OutputDir, Graph, Logger);
+				GenerateEmptyLinkFunctionsForRemovedModules(ModifiedInputFiles, CompileEnvironment.Architecture, ModuleName, OutputDir, Graph, Logger);
 
 				bHasHandledLaunchModule = true;
 			}
@@ -1034,11 +1036,11 @@ namespace UnrealBuildTool
 			if (!bHasHandledCoreModule && ModuleName.Equals("Core") && (CompileEnvironment.PrecompiledHeaderAction == PrecompiledHeaderAction.None))
 			{
 				// This is used by Crypto code in Core
-				InputFiles.Add(FileItem.GetItemByPath(GetCpuFeaturesPath()));
+				ModifiedInputFiles.Add(FileItem.GetItemByPath(GetCpuFeaturesPath()));
 				bHasHandledCoreModule = true;
 			}
 
-			return base.CompileCPPFiles(CompileEnvironment, InputFiles, OutputDir, ModuleName, Graph);
+			return base.CompileCPPFiles(CompileEnvironment, ModifiedInputFiles, OutputDir, ModuleName, Graph);
 		}
 
 		public static string InlineArchName(string Pathname, UnrealArch Arch, bool bUseShortNames = false)
@@ -1066,7 +1068,7 @@ namespace UnrealBuildTool
 			return DirectoryReference.Combine(PathRef, "include", Arch.ToString());
 		}
 
-		public override CPPOutput GenerateISPCHeaders(CppCompileEnvironment CompileEnvironment, List<FileItem> InputFiles, DirectoryReference OutputDir, IActionGraphBuilder Graph)
+		public override CPPOutput GenerateISPCHeaders(CppCompileEnvironment CompileEnvironment, IEnumerable<FileItem> InputFiles, DirectoryReference OutputDir, IActionGraphBuilder Graph)
 		{
 			CPPOutput Result = new CPPOutput();
 
@@ -1194,7 +1196,7 @@ namespace UnrealBuildTool
 			return Result;
 		}
 
-		public override CPPOutput CompileISPCFiles(CppCompileEnvironment CompileEnvironment, List<FileItem> InputFiles, DirectoryReference OutputDir, IActionGraphBuilder Graph)
+		public override CPPOutput CompileISPCFiles(CppCompileEnvironment CompileEnvironment, IEnumerable<FileItem> InputFiles, DirectoryReference OutputDir, IActionGraphBuilder Graph)
 		{
 			CPPOutput Result = new CPPOutput();
 
@@ -1695,7 +1697,7 @@ namespace UnrealBuildTool
 			File.WriteAllLines(FileName, ObjectFileDirectories.Select(x => x.FullName).OrderBy(x => x).ToArray());
 		}
 
-		public override void ModifyBuildProducts(ReadOnlyTargetRules Target, UEBuildBinary Binary, List<string> Libraries, List<UEBuildBundleResource> BundleResources, Dictionary<FileReference, BuildProductType> BuildProducts)
+		public override void ModifyBuildProducts(ReadOnlyTargetRules Target, UEBuildBinary Binary, IEnumerable<string> Libraries, IEnumerable<UEBuildBundleResource> BundleResources, Dictionary<FileReference, BuildProductType> BuildProducts)
 		{
 			// only the .so needs to be in the manifest; we always have to build the apk since its contents depend on the project
 
diff --git a/Engine/Source/Programs/UnrealBuildTool/Platform/IOS/IOSToolChain.cs b/Engine/Source/Programs/UnrealBuildTool/Platform/IOS/IOSToolChain.cs
--- a/Engine/Source/Programs/UnrealBuildTool/Platform/IOS/IOSToolChain.cs
+++ b/Engine/Source/Programs/UnrealBuildTool/Platform/IOS/IOSToolChain.cs
@@ -84,7 +84,7 @@ namespace UnrealBuildTool
 			return Settings.SDKVersionFloat.ToString();
 		}
 
-		public override void ModifyBuildProducts(ReadOnlyTargetRules Target, UEBuildBinary Binary, List<string> Libraries, List<UEBuildBundleResource> BundleResources, Dictionary<FileReference, BuildProductType> BuildProducts)
+		public override void ModifyBuildProducts(ReadOnlyTargetRules Target, UEBuildBinary Binary, IEnumerable<string> Libraries, IEnumerable<UEBuildBundleResource> BundleResources, Dictionary<FileReference, BuildProductType> BuildProducts)
 		{
 			if (Target.IOSPlatform.bCreateStubIPA && Binary.Type != UEBuildBinaryType.StaticLibrary)
 			{
@@ -1806,7 +1806,7 @@ namespace UnrealBuildTool
 			return TargetItem;
 		}
 
-		public override void SetupBundleDependencies(ReadOnlyTargetRules Target, List<UEBuildBinary> Binaries, string GameName)
+		public override void SetupBundleDependencies(ReadOnlyTargetRules Target, IEnumerable<UEBuildBinary> Binaries, string GameName)
 		{
 			base.SetupBundleDependencies(Target, Binaries, GameName);
 
diff --git a/Engine/Source/Programs/UnrealBuildTool/Platform/Linux/LinuxToolChain.cs b/Engine/Source/Programs/UnrealBuildTool/Platform/Linux/LinuxToolChain.cs
--- a/Engine/Source/Programs/UnrealBuildTool/Platform/Linux/LinuxToolChain.cs
+++ b/Engine/Source/Programs/UnrealBuildTool/Platform/Linux/LinuxToolChain.cs
@@ -925,7 +925,7 @@ namespace UnrealBuildTool
 			Logger.LogInformation("------------------------------");
 		}
 
-		protected override CPPOutput CompileCPPFiles(CppCompileEnvironment CompileEnvironment, List<FileItem> InputFiles, DirectoryReference OutputDir, string ModuleName, IActionGraphBuilder Graph)
+		protected override CPPOutput CompileCPPFiles(CppCompileEnvironment CompileEnvironment, IEnumerable<FileItem> InputFiles, DirectoryReference OutputDir, string ModuleName, IActionGraphBuilder Graph)
 		{
 			List<string> GlobalArguments = new();
 			GetCompileArguments_Global(CompileEnvironment, GlobalArguments);
@@ -1616,7 +1616,7 @@ namespace UnrealBuildTool
 			return OutputFile;
 		}
 
-		public override void SetupBundleDependencies(ReadOnlyTargetRules Target, List<UEBuildBinary> Binaries, string GameName)
+		public override void SetupBundleDependencies(ReadOnlyTargetRules Target, IEnumerable<UEBuildBinary> Binaries, string GameName)
 		{
 			if (bUseFixdeps)
 			{
diff --git a/Engine/Source/Programs/UnrealBuildTool/Platform/Mac/MacToolChain.cs b/Engine/Source/Programs/UnrealBuildTool/Platform/Mac/MacToolChain.cs
--- a/Engine/Source/Programs/UnrealBuildTool/Platform/Mac/MacToolChain.cs
+++ b/Engine/Source/Programs/UnrealBuildTool/Platform/Mac/MacToolChain.cs
@@ -1035,7 +1035,7 @@ namespace UnrealBuildTool
 
 		private static Dictionary<ReadOnlyTargetRules, DirectoryReference> BundleContentsDirectories = new();
 
-		public override void ModifyBuildProducts(ReadOnlyTargetRules Target, UEBuildBinary Binary, List<string> Libraries, List<UEBuildBundleResource> BundleResources, Dictionary<FileReference, BuildProductType> BuildProducts)
+		public override void ModifyBuildProducts(ReadOnlyTargetRules Target, UEBuildBinary Binary, IEnumerable<string> Libraries, IEnumerable<UEBuildBundleResource> BundleResources, Dictionary<FileReference, BuildProductType> BuildProducts)
 		{
 			if (Target.bUsePDBFiles == true)
 			{
diff --git a/Engine/Source/Programs/UnrealBuildTool/Platform/Windows/PVSToolChain.cs b/Engine/Source/Programs/UnrealBuildTool/Platform/Windows/PVSToolChain.cs
--- a/Engine/Source/Programs/UnrealBuildTool/Platform/Windows/PVSToolChain.cs
+++ b/Engine/Source/Programs/UnrealBuildTool/Platform/Windows/PVSToolChain.cs
@@ -647,7 +647,7 @@ namespace UnrealBuildTool
 			return CompileEnvironment;
 		}
 
-		protected override CPPOutput CompileCPPFiles(CppCompileEnvironment CompileEnvironment, List<FileItem> InputFiles, DirectoryReference OutputDir, string ModuleName, IActionGraphBuilder Graph)
+		protected override CPPOutput CompileCPPFiles(CppCompileEnvironment CompileEnvironment, IEnumerable<FileItem> InputFiles, DirectoryReference OutputDir, string ModuleName, IActionGraphBuilder Graph)
 		{
 			if (CompileEnvironment.bDisableStaticAnalysis)
 			{
diff --git a/Engine/Source/Programs/UnrealBuildTool/Platform/Windows/VCToolChain.cs b/Engine/Source/Programs/UnrealBuildTool/Platform/Windows/VCToolChain.cs
--- a/Engine/Source/Programs/UnrealBuildTool/Platform/Windows/VCToolChain.cs
+++ b/Engine/Source/Programs/UnrealBuildTool/Platform/Windows/VCToolChain.cs
@@ -1501,7 +1501,7 @@ namespace UnrealBuildTool
 			return null;
 		}
 
-		protected override CPPOutput CompileCPPFiles(CppCompileEnvironment CompileEnvironment, List<FileItem> InputFiles, DirectoryReference OutputDir, string ModuleName, IActionGraphBuilder Graph)
+		protected override CPPOutput CompileCPPFiles(CppCompileEnvironment CompileEnvironment, IEnumerable<FileItem> InputFiles, DirectoryReference OutputDir, string ModuleName, IActionGraphBuilder Graph)
 		{
 			VCCompileAction BaseCompileAction = CreateBaseCompileAction(CompileEnvironment);
 
@@ -1984,7 +1984,7 @@ namespace UnrealBuildTool
 			return Platform.IsInGroup(UnrealPlatformGroup.Windows);
 		}
 
-		public override CPPOutput CompileRCFiles(CppCompileEnvironment CompileEnvironment, List<FileItem> InputFiles, DirectoryReference OutputDir, IActionGraphBuilder Graph)
+		public override CPPOutput CompileRCFiles(CppCompileEnvironment CompileEnvironment, IEnumerable<FileItem> InputFiles, DirectoryReference OutputDir, IActionGraphBuilder Graph)
 		{
 			CPPOutput Result = new CPPOutput();
 
@@ -2894,7 +2894,7 @@ namespace UnrealBuildTool
 			return IncludePaths.ToString();
 		}
 
-		public override void ModifyBuildProducts(ReadOnlyTargetRules Target, UEBuildBinary Binary, List<string> Libraries, List<UEBuildBundleResource> BundleResources, Dictionary<FileReference, BuildProductType> BuildProducts)
+		public override void ModifyBuildProducts(ReadOnlyTargetRules Target, UEBuildBinary Binary, IEnumerable<string> Libraries, IEnumerable<UEBuildBundleResource> BundleResources, Dictionary<FileReference, BuildProductType> BuildProducts)
 		{
 			if (Binary.Type == UEBuildBinaryType.DynamicLinkLibrary)
 			{
diff --git a/Engine/Source/Programs/UnrealBuildTool/ToolChain/AppleToolChain.cs b/Engine/Source/Programs/UnrealBuildTool/ToolChain/AppleToolChain.cs
--- a/Engine/Source/Programs/UnrealBuildTool/ToolChain/AppleToolChain.cs
+++ b/Engine/Source/Programs/UnrealBuildTool/ToolChain/AppleToolChain.cs
@@ -285,7 +285,7 @@ namespace UnrealBuildTool
 			return Unreal.EngineSourceDirectory;
 		}
 
-		protected override CPPOutput CompileCPPFiles(CppCompileEnvironment CompileEnvironment, List<FileItem> InputFiles, DirectoryReference OutputDir, string ModuleName, IActionGraphBuilder Graph)
+		protected override CPPOutput CompileCPPFiles(CppCompileEnvironment CompileEnvironment, IEnumerable<FileItem> InputFiles, DirectoryReference OutputDir, string ModuleName, IActionGraphBuilder Graph)
 		{
 			List<string> GlobalArguments = new();
 
diff --git a/Engine/Source/Programs/UnrealBuildTool/ToolChain/ClangToolChain.cs b/Engine/Source/Programs/UnrealBuildTool/ToolChain/ClangToolChain.cs
--- a/Engine/Source/Programs/UnrealBuildTool/ToolChain/ClangToolChain.cs
+++ b/Engine/Source/Programs/UnrealBuildTool/ToolChain/ClangToolChain.cs
@@ -1197,7 +1197,7 @@ namespace UnrealBuildTool
 			return CompileAction;
 		}
 
-		protected override CPPOutput CompileCPPFiles(CppCompileEnvironment CompileEnvironment, List<FileItem> InputFiles, DirectoryReference OutputDir, string ModuleName, IActionGraphBuilder Graph)
+		protected override CPPOutput CompileCPPFiles(CppCompileEnvironment CompileEnvironment, IEnumerable<FileItem> InputFiles, DirectoryReference OutputDir, string ModuleName, IActionGraphBuilder Graph)
 		{
 			List<string> GlobalArguments = new();
 
diff --git a/Engine/Source/Programs/UnrealBuildTool/ToolChain/ISPCToolChain.cs b/Engine/Source/Programs/UnrealBuildTool/ToolChain/ISPCToolChain.cs
--- a/Engine/Source/Programs/UnrealBuildTool/ToolChain/ISPCToolChain.cs
+++ b/Engine/Source/Programs/UnrealBuildTool/ToolChain/ISPCToolChain.cs
@@ -384,7 +384,7 @@ namespace UnrealBuildTool
 			return NormalizeCommandLinePath(Item.Location);
 		}
 
-		public override CPPOutput GenerateISPCHeaders(CppCompileEnvironment CompileEnvironment, List<FileItem> InputFiles, DirectoryReference OutputDir, IActionGraphBuilder Graph)
+		public override CPPOutput GenerateISPCHeaders(CppCompileEnvironment CompileEnvironment, IEnumerable<FileItem> InputFiles, DirectoryReference OutputDir, IActionGraphBuilder Graph)
 		{
 			CPPOutput Result = new CPPOutput();
 
@@ -529,7 +529,7 @@ namespace UnrealBuildTool
 			return Result;
 		}
 
-		public override CPPOutput CompileISPCFiles(CppCompileEnvironment CompileEnvironment, List<FileItem> InputFiles, DirectoryReference OutputDir, IActionGraphBuilder Graph)
+		public override CPPOutput CompileISPCFiles(CppCompileEnvironment CompileEnvironment, IEnumerable<FileItem> InputFiles, DirectoryReference OutputDir, IActionGraphBuilder Graph)
 		{
 			CPPOutput Result = new CPPOutput();
 
diff --git a/Engine/Source/Programs/UnrealBuildTool/ToolChain/IWYUToolChain.cs b/Engine/Source/Programs/UnrealBuildTool/ToolChain/IWYUToolChain.cs
--- a/Engine/Source/Programs/UnrealBuildTool/ToolChain/IWYUToolChain.cs
+++ b/Engine/Source/Programs/UnrealBuildTool/ToolChain/IWYUToolChain.cs
@@ -71,9 +71,9 @@ namespace UnrealBuildTool
 			return FileName + Extension;
 		}
 
-		protected override CPPOutput CompileCPPFiles(CppCompileEnvironment CompileEnvironment, List<FileItem> InputFiles, DirectoryReference OutputDir, string ModuleName, IActionGraphBuilder Graph)
+		protected override CPPOutput CompileCPPFiles(CppCompileEnvironment CompileEnvironment, IEnumerable<FileItem> InputFiles, DirectoryReference OutputDir, string ModuleName, IActionGraphBuilder Graph)
 		{
-			if (InputFiles.Count == 0)
+			if (!InputFiles.Any())
 			{
 				return new CPPOutput();
 			}
@@ -199,12 +199,12 @@ namespace UnrealBuildTool
 		}
 
 		// Skip ISPC headers
-		public override CPPOutput GenerateISPCHeaders(CppCompileEnvironment CompileEnvironment, List<FileItem> InputFiles, DirectoryReference OutputDir, IActionGraphBuilder Graph)
+		public override CPPOutput GenerateISPCHeaders(CppCompileEnvironment CompileEnvironment, IEnumerable<FileItem> InputFiles, DirectoryReference OutputDir, IActionGraphBuilder Graph)
 		{
 			return new CPPOutput();
 		}
 
-		public override CPPOutput CompileISPCFiles(CppCompileEnvironment CompileEnvironment, List<FileItem> InputFiles, DirectoryReference OutputDir, IActionGraphBuilder Graph)
+		public override CPPOutput CompileISPCFiles(CppCompileEnvironment CompileEnvironment, IEnumerable<FileItem> InputFiles, DirectoryReference OutputDir, IActionGraphBuilder Graph)
 		{
 			return new CPPOutput();
 		}
diff --git a/Engine/Source/Programs/UnrealBuildTool/ToolChain/UEToolChain.cs b/Engine/Source/Programs/UnrealBuildTool/ToolChain/UEToolChain.cs
--- a/Engine/Source/Programs/UnrealBuildTool/ToolChain/UEToolChain.cs
+++ b/Engine/Source/Programs/UnrealBuildTool/ToolChain/UEToolChain.cs
@@ -55,9 +55,9 @@ namespace UnrealBuildTool
 			return null;
 		}
 
-		protected abstract CPPOutput CompileCPPFiles(CppCompileEnvironment CompileEnvironment, List<FileItem> InputFiles, DirectoryReference OutputDir, string ModuleName, IActionGraphBuilder Graph);
+		protected abstract CPPOutput CompileCPPFiles(CppCompileEnvironment CompileEnvironment, IEnumerable<FileItem> InputFiles, DirectoryReference OutputDir, string ModuleName, IActionGraphBuilder Graph);
 
-		public CPPOutput CompileAllCPPFiles(CppCompileEnvironment CompileEnvironment, List<FileItem> InputFiles, DirectoryReference OutputDir, string ModuleName, IActionGraphBuilder Graph)
+		public CPPOutput CompileAllCPPFiles(CppCompileEnvironment CompileEnvironment, IEnumerable<FileItem> InputFiles, DirectoryReference OutputDir, string ModuleName, IActionGraphBuilder Graph)
 		{
 			CPPOutput Result;
 
@@ -89,19 +89,19 @@ namespace UnrealBuildTool
 			return Result;
 		}
 
-		public virtual CPPOutput CompileRCFiles(CppCompileEnvironment Environment, List<FileItem> InputFiles, DirectoryReference OutputDir, IActionGraphBuilder Graph)
+		public virtual CPPOutput CompileRCFiles(CppCompileEnvironment CompileEnvironment, IEnumerable<FileItem> InputFiles, DirectoryReference OutputDir, IActionGraphBuilder Graph)
 		{
 			CPPOutput Result = new CPPOutput();
 			return Result;
 		}
 
-		public virtual CPPOutput CompileISPCFiles(CppCompileEnvironment Environment, List<FileItem> InputFiles, DirectoryReference OutputDir, IActionGraphBuilder Graph)
+		public virtual CPPOutput CompileISPCFiles(CppCompileEnvironment CompileEnvironment, IEnumerable<FileItem> InputFiles, DirectoryReference OutputDir, IActionGraphBuilder Graph)
 		{
 			CPPOutput Result = new CPPOutput();
 			return Result;
 		}
 
-		public virtual CPPOutput GenerateISPCHeaders(CppCompileEnvironment Environment, List<FileItem> InputFiles, DirectoryReference OutputDir, IActionGraphBuilder Graph)
+		public virtual CPPOutput GenerateISPCHeaders(CppCompileEnvironment CompileEnvironment, IEnumerable<FileItem> InputFiles, DirectoryReference OutputDir, IActionGraphBuilder Graph)
 		{
 			CPPOutput Result = new CPPOutput();
 			return Result;
@@ -121,7 +121,7 @@ namespace UnrealBuildTool
 		public virtual FileItem[] LinkImportLibrary(LinkEnvironment LinkEnvironment, IActionGraphBuilder Graph)
 		{
 			// by default doing nothing
-			return new FileItem[] { };
+			return Array.Empty<FileItem>();
 		}
 
 		public abstract FileItem? LinkFiles(LinkEnvironment LinkEnvironment, bool bBuildImportLibraryOnly, IActionGraphBuilder Graph);
@@ -182,7 +182,7 @@ namespace UnrealBuildTool
 		/// <param name="CompileEnvironment"></param>
 		/// <param name="OutputFile"></param>
 		/// <returns></returns>
-		public static FileReference GetResponseFileName(CppCompileEnvironment CompileEnvironment, FileItem OutputFile)
+		public virtual FileReference GetResponseFileName(CppCompileEnvironment CompileEnvironment, FileItem OutputFile)
 		{
 			// Construct a relative path for the intermediate response file
 			return OutputFile.Location.ChangeExtension(OutputFile.Location.GetExtension() + ResponseExt);
@@ -194,7 +194,7 @@ namespace UnrealBuildTool
 		/// <param name="LinkEnvironment"></param>
 		/// <param name="OutputFile"></param>
 		/// <returns></returns>
-		public static FileReference GetResponseFileName(LinkEnvironment LinkEnvironment, FileItem OutputFile)
+		public virtual FileReference GetResponseFileName(LinkEnvironment LinkEnvironment, FileItem OutputFile)
 		{
 			// Construct a relative path for the intermediate response file
 			return FileReference.Combine(LinkEnvironment.IntermediateDirectory!, OutputFile.Location.GetFileName() + ResponseExt);
@@ -205,7 +205,7 @@ namespace UnrealBuildTool
 			return new List<FileItem>();
 		}
 
-		public virtual ICollection<FileItem> PostBuild(ReadOnlyTargetRules Target, FileItem[] Executables, LinkEnvironment ExecutableLinkEnvironment, IActionGraphBuilder Graph)
+		public virtual ICollection<FileItem> PostBuild(ReadOnlyTargetRules Target, IEnumerable<FileItem> Executables, LinkEnvironment ExecutableLinkEnvironment, IActionGraphBuilder Graph)
 		{
 			// by default, run PostBuild for exe Exe and merge results
 			return Executables.SelectMany(x => PostBuild(Target, x, ExecutableLinkEnvironment, Graph)).ToList();
@@ -215,7 +215,7 @@ namespace UnrealBuildTool
 		{
 		}
 
-		public virtual void ModifyBuildProducts(ReadOnlyTargetRules Target, UEBuildBinary Binary, List<string> Libraries, List<UEBuildBundleResource> BundleResources, Dictionary<FileReference, BuildProductType> BuildProducts)
+		public virtual void ModifyBuildProducts(ReadOnlyTargetRules Target, UEBuildBinary Binary, IEnumerable<string> Libraries, IEnumerable<UEBuildBundleResource> BundleResources, Dictionary<FileReference, BuildProductType> BuildProducts)
 		{
 		}
 
@@ -247,7 +247,7 @@ namespace UnrealBuildTool
 			return OutputFile.ChangeExtension(DebugExtension);
 		}
 
-		public virtual void SetupBundleDependencies(ReadOnlyTargetRules Target, List<UEBuildBinary> Binaries, string GameName)
+		public virtual void SetupBundleDependencies(ReadOnlyTargetRules Target, IEnumerable<UEBuildBinary> Binaries, string GameName)
 		{
 		}
 
