UnrealBuildTool: Add the ability to disable debugging info for select modules and plugins via TargetRules configuration

Obsolete bDisableDebugInfo and replace with TargetRules.DebugInfo which is a flags enum DebugInfoMode. TargetRules.DebugInfo can be set to any or all of [Engine, EnginePlugins, Project, ProjectPlugins], whichever flags are set will have debug info enabled for the modules in those groups if the Target would otherwise has debug info enabled
* DebugInfoMode.None is equalivant to bDisableDebugInfo=true (-NoDebugInfo will still set this)
* DebugInfoMode.ProjectOnly = Project | ProjectPlugins
* DebugInfoMode.EngineOnly = Engine | EnginePlugins
* DebugInfoMode.Full = Engine | EnginePlugins | Project | ProjectPlugins
* DebugInfoMode.Full is equalivant to bDisableDebugInfo=false (the current default)
Add TargetRules.DisableDebugInfoModules, this is a list of module names that should disable debug info regardless of TargetRules.DebugInfo
Add TargetRules.DisableDebugInfoPlugins, this is a list of plugin names that should disable debug info for all of their modules regardless of TargetRules.DebugInfo

diff --git a/Engine/Source/Programs/Enterprise/Datasmith/DatasmithSDK/DatasmithSDK.Target.cs b/Engine/Source/Programs/Enterprise/Datasmith/DatasmithSDK/DatasmithSDK.Target.cs
--- a/Engine/Source/Programs/Enterprise/Datasmith/DatasmithSDK/DatasmithSDK.Target.cs
+++ b/Engine/Source/Programs/Enterprise/Datasmith/DatasmithSDK/DatasmithSDK.Target.cs
@@ -35,7 +35,7 @@ public class DatasmithSDKTarget : TargetRules
 		bCompileAgainstCoreUObject = true;
 		bCompileICU = false;
 		bUsesSlate = false;
-		bDisableDebugInfo = false;
+		DebugInfo = DebugInfoMode.Full;
 		bUsePDBFiles = true;
 		bHasExports = true;
 		bIsBuildingConsoleApplication = true;
diff --git a/Engine/Source/Programs/UnrealBuildTool/Configuration/TargetRules.cs b/Engine/Source/Programs/UnrealBuildTool/Configuration/TargetRules.cs
--- a/Engine/Source/Programs/UnrealBuildTool/Configuration/TargetRules.cs
+++ b/Engine/Source/Programs/UnrealBuildTool/Configuration/TargetRules.cs
@@ -291,6 +291,55 @@ namespace UnrealBuildTool
 		SizeAndSpeed
 	}
 
+	/// <summary>
+	/// Debug info mode for compiler settings to determine how much debug info is available
+	/// </summary>
+	[Flags]
+	public enum DebugInfoMode
+	{
+		/// <summary>
+		/// Disable all debugging info.
+		/// MSVC: object files will be compiled without /Z7 or /Zi but pdbs will still be created
+		///       callstacks should be available in this mode but there have been reports with them being incorrect
+		/// </summary>
+		None = 0,
+
+		/// <summary>
+		/// Enable debug info for engine modules
+		/// </summary>
+		Engine = 1 << 0,
+
+		/// <summary>
+		/// Enable debug info for engine plugins
+		/// </summary>
+		EnginePlugins = 1 << 1,
+
+		/// <summary>
+		/// Enable debug info for project modules
+		/// </summary>
+		Project = 1 << 2,
+
+		/// <summary>
+		/// Enable debug info for project plugins
+		/// </summary>
+		ProjectPlugins = 1 << 3,
+
+		/// <summary>
+		/// Only include debug info for engine modules and plugins
+		/// </summary>
+		EngineOnly = Engine | EnginePlugins,
+
+		/// <summary>
+		/// Only include debug info for project modules and project plugins
+		/// </summary>
+		ProjectOnly = Project | ProjectPlugins,
+
+		/// <summary>
+		/// Include full debugging information for all modules
+		/// </summary>
+		Full = Engine | EnginePlugins | Project | ProjectPlugins,
+	}
+
 	/// <summary>
 	/// Floating point math semantics
 	/// </summary>
@@ -1786,14 +1835,38 @@ namespace UnrealBuildTool
 		public bool bForceDebugInfo = false;
 
 		/// <summary>
-		/// Whether to globally disable debug info generation; see DebugInfoHeuristics.cs for per-config and per-platform options.
+		/// Whether to globally disable debug info generation; Obsolete, please use TargetRules.DebugInfoMode instead
 		/// </summary>
-		[CommandLine("-NoDebugInfo")]
 		[XmlConfigFile(Category = "BuildConfiguration")]
-		public bool bDisableDebugInfo = false;
+		[Obsolete("Deprecated in UE5.4 - Replace with TargetRules.DebugInfo")]
+		public bool bDisableDebugInfo
+		{
+			get => DebugInfo == DebugInfoMode.None;
+			set => DebugInfo = value ? DebugInfo = DebugInfoMode.None : DebugInfo = DebugInfoMode.Full;
+		}
 
 		/// <summary>
-		/// Whether to disable debug info generation for generated files. This improves link times for modules that have a lot of generated glue code.
+		/// How much debug info should be generated. See DebugInfoMode enum for more details
+		/// </summary>
+		[CommandLine("-NoDebugInfo", Value = "None")]
+		[CommandLine("-DebugInfo=")]
+		[XmlConfigFile(Category = "BuildConfiguration")]
+		public DebugInfoMode DebugInfo { get; set; } = DebugInfoMode.Full;
+
+		/// <summary>
+		/// Modules that should have debug info disabled
+		/// </summary>
+		[CommandLine("-DisableDebugInfoModules=", ListSeparator = '+')]
+		public HashSet<string> DisableDebugInfoModules { get; } = new();
+
+		/// <summary>
+		/// Plugins that should have debug info disabled
+		/// </summary>
+		[CommandLine("-DisableDebugInfoPlugins=", ListSeparator = '+')]
+		public HashSet<string> DisableDebugInfoPlugins { get; } = new();
+
+		/// <summary>
+		/// Whether to disable debug info generation for generated files. This improves link times and reduces pdb size for modules that have a lot of generated glue code.
 		/// </summary>
 		[XmlConfigFile(Category = "BuildConfiguration")]
 		public bool bDisableDebugInfoForGeneratedCode = false;
@@ -3320,8 +3393,15 @@ namespace UnrealBuildTool
 
 		public bool bDetailedUnityFiles => Inner.bDetailedUnityFiles;
 
+		[Obsolete("Deprecated in UE5.4 - Replace with ReadOnlyTargetRules.DebugInfo")]
 		public bool bDisableDebugInfo => Inner.bDisableDebugInfo;
 
+		public DebugInfoMode DebugInfo => Inner.DebugInfo;
+
+		public IReadOnlySet<string> DisableDebugInfoModules => Inner.DisableDebugInfoModules;
+
+		public IReadOnlySet<string> DisableDebugInfoPlugins => Inner.DisableDebugInfoPlugins;
+
 		public bool bDisableDebugInfoForGeneratedCode => Inner.bDisableDebugInfoForGeneratedCode;
 
 		public bool bOmitPCDebugInfoInDevelopment => Inner.bOmitPCDebugInfoInDevelopment;
diff --git a/Engine/Source/Programs/UnrealBuildTool/Configuration/UEBuildModuleCPP.cs b/Engine/Source/Programs/UnrealBuildTool/Configuration/UEBuildModuleCPP.cs
--- a/Engine/Source/Programs/UnrealBuildTool/Configuration/UEBuildModuleCPP.cs
+++ b/Engine/Source/Programs/UnrealBuildTool/Configuration/UEBuildModuleCPP.cs
@@ -1896,6 +1896,42 @@ namespace UnrealBuildTool
 			Result.bValidateFormatStrings = Rules.bValidateFormatStrings;
 			Result.bUseAutoRTFMCompiler = Target.bUseAutoRTFMCompiler;
 
+			// Disable debug info for modules if requested
+			if (!Target.bUsePDBFiles || !Target.Platform.IsInGroup("Microsoft"))
+			{
+				if (Rules.Target.ProjectFile != null && Rules.File.IsUnderDirectory(Rules.Target.ProjectFile.Directory))
+				{
+					if (!Target.DebugInfo.HasFlag(DebugInfoMode.ProjectPlugins) && Rules.Plugin != null)
+					{
+						Result.bCreateDebugInfo = false;
+					}
+					else if (!Target.DebugInfo.HasFlag(DebugInfoMode.Project) && Rules.Plugin == null)
+					{
+						Result.bCreateDebugInfo = false;
+					}
+				}
+				else
+				{
+					if (!Target.DebugInfo.HasFlag(DebugInfoMode.EnginePlugins) && Rules.Plugin != null)
+					{
+						Result.bCreateDebugInfo = false;
+					}
+					else if (!Target.DebugInfo.HasFlag(DebugInfoMode.Engine) && Rules.Plugin == null)
+					{
+						Result.bCreateDebugInfo = false;
+					}
+				}
+
+				if (Rules.Plugin != null && Rules.Target.DisableDebugInfoPlugins.Contains(Rules.Plugin.Name))
+				{
+					Result.bCreateDebugInfo = false;
+				}
+				else if (Rules.Target.DisableDebugInfoModules.Contains(Name))
+				{
+					Result.bCreateDebugInfo = false;
+				}
+			}
+
 			// Only enable the AutoRTFM flag if we are using the AutoRTFM compiler
 			if (Target.bUseAutoRTFMCompiler)
 			{
diff --git a/Engine/Source/Programs/UnrealBuildTool/Configuration/UEBuildPlatform.cs b/Engine/Source/Programs/UnrealBuildTool/Configuration/UEBuildPlatform.cs
--- a/Engine/Source/Programs/UnrealBuildTool/Configuration/UEBuildPlatform.cs
+++ b/Engine/Source/Programs/UnrealBuildTool/Configuration/UEBuildPlatform.cs
@@ -1265,7 +1265,7 @@ namespace UnrealBuildTool
 
 			// Create debug info based on the heuristics specified by the user.
 			GlobalCompileEnvironment.bCreateDebugInfo =
-				!Target.bDisableDebugInfo && ShouldCreateDebugInfo(Target);
+				Target.DebugInfo != DebugInfoMode.None && ShouldCreateDebugInfo(Target);
 			GlobalLinkEnvironment.bCreateDebugInfo = GlobalCompileEnvironment.bCreateDebugInfo;
 		}
 
diff --git a/Engine/Source/Programs/UnrealBuildTool/Platform/Mac/UEBuildMac.cs b/Engine/Source/Programs/UnrealBuildTool/Platform/Mac/UEBuildMac.cs
--- a/Engine/Source/Programs/UnrealBuildTool/Platform/Mac/UEBuildMac.cs
+++ b/Engine/Source/Programs/UnrealBuildTool/Platform/Mac/UEBuildMac.cs
@@ -355,7 +355,7 @@ namespace UnrealBuildTool
 
 			Target.GlobalDefinitions.Add("GL_SILENCE_DEPRECATION=1");
 
-			Target.bUsePDBFiles = !Target.bDisableDebugInfo && ShouldCreateDebugInfo(new ReadOnlyTargetRules(Target));
+			Target.bUsePDBFiles = Target.DebugInfo != DebugInfoMode.None && ShouldCreateDebugInfo(new ReadOnlyTargetRules(Target));
 			Target.bUsePDBFiles &= Target.MacPlatform.bUseDSYMFiles;
 
 			// we always deploy - the build machines need to be able to copy the files back, which needs the full bundle
diff --git a/Engine/Source/Programs/UnrealBuildTool/System/RulesAssembly.cs b/Engine/Source/Programs/UnrealBuildTool/System/RulesAssembly.cs
--- a/Engine/Source/Programs/UnrealBuildTool/System/RulesAssembly.cs
+++ b/Engine/Source/Programs/UnrealBuildTool/System/RulesAssembly.cs
@@ -648,7 +648,7 @@ namespace UnrealBuildTool
 			// Apply the override to force debug info to be enabled
 			if (Rules.bForceDebugInfo)
 			{
-				Rules.bDisableDebugInfo = false;
+				Rules.DebugInfo = DebugInfoMode.Full;
 				Rules.bOmitPCDebugInfoInDevelopment = false;
 			}
 
diff --git a/Engine/Source/Programs/UnrealBuildTool/ToolChain/IWYUToolChain.cs b/Engine/Source/Programs/UnrealBuildTool/ToolChain/IWYUToolChain.cs
--- a/Engine/Source/Programs/UnrealBuildTool/ToolChain/IWYUToolChain.cs
+++ b/Engine/Source/Programs/UnrealBuildTool/ToolChain/IWYUToolChain.cs
@@ -18,7 +18,7 @@ namespace UnrealBuildTool
 
 		public static void ValidateTarget(TargetRules Target)
 		{
-			Target.bDisableDebugInfo = true;
+			Target.DebugInfo = DebugInfoMode.None;
 			Target.bDisableLinking = true;
 			Target.bUsePCHFiles = false;
 			Target.bUseSharedPCHs = false;
