UnrealBuildTool: Support emiting line tables only for debugging rather than disabling debug info entirely

Configuration is the same as TargetRules.DebugInfo and will override that setting except when debug info is disabled for the entire target

diff --git a/Engine/Source/Programs/UnrealBuildTool/Configuration/TargetRules.cs b/Engine/Source/Programs/UnrealBuildTool/Configuration/TargetRules.cs
--- a/Engine/Source/Programs/UnrealBuildTool/Configuration/TargetRules.cs
+++ b/Engine/Source/Programs/UnrealBuildTool/Configuration/TargetRules.cs
@@ -1865,6 +1865,26 @@ namespace UnrealBuildTool
 		[CommandLine("-DisableDebugInfoPlugins=", ListSeparator = '+')]
 		public HashSet<string> DisableDebugInfoPlugins { get; } = new();
 
+		/// <summary>
+		/// True if only debug line number tables should be emitted in debug information for compilers that support doing so. Overrides TargetRules.DebugInfo
+		/// See https://clang.llvm.org/docs/UsersManual.html#cmdoption-gline-tables-only for more information
+		/// </summary>
+		[XmlConfigFile(Category = "BuildConfiguration")]
+		[CommandLine("-DebugInfoLineTablesOnly=")]
+		public DebugInfoMode DebugInfoLineTablesOnly { get; set; } = DebugInfoMode.None;
+
+		/// <summary>
+		/// Modules that should emit line number tables instead of full debug information for compilers that support doing so. Overrides DisableDebugInfoModules
+		/// </summary>
+		[CommandLine("-DebugInfoLineTablesOnlyModules=", ListSeparator = '+')]
+		public HashSet<string> DebugInfoLineTablesOnlyModules { get; } = new();
+
+		/// <summary>
+		/// Plugins that should emit line number tables instead of full debug information for compilers that support doing so. Overrides DisableDebugInfoPlugins
+		/// </summary>
+		[CommandLine("-DebugInfoLineTablesOnlyPlugins=", ListSeparator = '+')]
+		public HashSet<string> DebugInfoLineTablesOnlyPlugins { get; } = new();
+
 		/// <summary>
 		/// Whether to disable debug info generation for generated files. This improves link times and reduces pdb size for modules that have a lot of generated glue code.
 		/// </summary>
@@ -3402,6 +3422,12 @@ namespace UnrealBuildTool
 
 		public IReadOnlySet<string> DisableDebugInfoPlugins => Inner.DisableDebugInfoPlugins;
 
+		public DebugInfoMode DebugInfoLineTablesOnly => Inner.DebugInfoLineTablesOnly;
+
+		public IReadOnlySet<string> DebugInfoLineTablesOnlyModules => Inner.DebugInfoLineTablesOnlyModules;
+
+		public IReadOnlySet<string> DebugInfoLineTablesOnlyPlugins => Inner.DebugInfoLineTablesOnlyPlugins;
+
 		public bool bDisableDebugInfoForGeneratedCode => Inner.bDisableDebugInfoForGeneratedCode;
 
 		public bool bOmitPCDebugInfoInDevelopment => Inner.bOmitPCDebugInfoInDevelopment;
diff --git a/Engine/Source/Programs/UnrealBuildTool/Configuration/UEBuildModuleCPP.cs b/Engine/Source/Programs/UnrealBuildTool/Configuration/UEBuildModuleCPP.cs
--- a/Engine/Source/Programs/UnrealBuildTool/Configuration/UEBuildModuleCPP.cs
+++ b/Engine/Source/Programs/UnrealBuildTool/Configuration/UEBuildModuleCPP.cs
@@ -1823,6 +1823,89 @@ namespace UnrealBuildTool
 			return false;
 		}
 
+		private void CompileEnvironmentDebugInfoSettings(ReadOnlyTargetRules Target, CppCompileEnvironment Result)
+		{
+			// If bCreateDebugInfo is disabled for the whole target do not adjust any settings
+			if (!Result.bCreateDebugInfo)
+			{
+				return;
+			}
+
+			if (Rules.Target.ProjectFile != null && Rules.File.IsUnderDirectory(Rules.Target.ProjectFile.Directory))
+			{
+				if (Target.DebugInfoLineTablesOnly.HasFlag(DebugInfoMode.ProjectPlugins) && Rules.Plugin != null)
+				{
+					Result.bDebugLineTablesOnly = true;
+				}
+				else if (Target.DebugInfoLineTablesOnly.HasFlag(DebugInfoMode.Project) && Rules.Plugin == null)
+				{
+					Result.bDebugLineTablesOnly = true;
+				}
+			}
+			else
+			{
+				if (Target.DebugInfoLineTablesOnly.HasFlag(DebugInfoMode.EnginePlugins) && Rules.Plugin != null)
+				{
+					Result.bDebugLineTablesOnly = true;
+				}
+				else if (Target.DebugInfoLineTablesOnly.HasFlag(DebugInfoMode.Engine) && Rules.Plugin == null)
+				{
+					Result.bDebugLineTablesOnly = true;
+				}
+			}
+
+			if (Rules.Plugin != null && Rules.Target.DebugInfoLineTablesOnlyPlugins.Contains(Rules.Plugin.Name))
+			{
+				Result.bDebugLineTablesOnly = true;
+			}
+			else if (Rules.Target.DebugInfoLineTablesOnlyModules.Contains(Name))
+			{
+				Result.bDebugLineTablesOnly = true;
+			}
+
+			if (Result.bDebugLineTablesOnly)
+			{
+				// Don't disable debug info if only line tables are requested
+				return;
+			}
+
+			// Disable debug info for modules if requested
+			if (!Target.bUsePDBFiles || !Target.Platform.IsInGroup("Microsoft"))
+			{
+				if (Rules.Target.ProjectFile != null && Rules.File.IsUnderDirectory(Rules.Target.ProjectFile.Directory))
+				{
+					if (!Target.DebugInfo.HasFlag(DebugInfoMode.ProjectPlugins) && Rules.Plugin != null)
+					{
+						Result.bCreateDebugInfo = false;
+					}
+					else if (!Target.DebugInfo.HasFlag(DebugInfoMode.Project) && Rules.Plugin == null)
+					{
+						Result.bCreateDebugInfo = false;
+					}
+				}
+				else
+				{
+					if (!Target.DebugInfo.HasFlag(DebugInfoMode.EnginePlugins) && Rules.Plugin != null)
+					{
+						Result.bCreateDebugInfo = false;
+					}
+					else if (!Target.DebugInfo.HasFlag(DebugInfoMode.Engine) && Rules.Plugin == null)
+					{
+						Result.bCreateDebugInfo = false;
+					}
+				}
+
+				if (Rules.Plugin != null && Rules.Target.DisableDebugInfoPlugins.Contains(Rules.Plugin.Name))
+				{
+					Result.bCreateDebugInfo = false;
+				}
+				else if (Rules.Target.DisableDebugInfoModules.Contains(Name))
+				{
+					Result.bCreateDebugInfo = false;
+				}
+			}
+		}
+
 		/// <summary>
 		/// Determine whether optimization should be enabled for a given target
 		/// </summary>
@@ -1896,41 +1979,7 @@ namespace UnrealBuildTool
 			Result.bValidateFormatStrings = Rules.bValidateFormatStrings;
 			Result.bUseAutoRTFMCompiler = Target.bUseAutoRTFMCompiler;
 
-			// Disable debug info for modules if requested
-			if (!Target.bUsePDBFiles || !Target.Platform.IsInGroup("Microsoft"))
-			{
-				if (Rules.Target.ProjectFile != null && Rules.File.IsUnderDirectory(Rules.Target.ProjectFile.Directory))
-				{
-					if (!Target.DebugInfo.HasFlag(DebugInfoMode.ProjectPlugins) && Rules.Plugin != null)
-					{
-						Result.bCreateDebugInfo = false;
-					}
-					else if (!Target.DebugInfo.HasFlag(DebugInfoMode.Project) && Rules.Plugin == null)
-					{
-						Result.bCreateDebugInfo = false;
-					}
-				}
-				else
-				{
-					if (!Target.DebugInfo.HasFlag(DebugInfoMode.EnginePlugins) && Rules.Plugin != null)
-					{
-						Result.bCreateDebugInfo = false;
-					}
-					else if (!Target.DebugInfo.HasFlag(DebugInfoMode.Engine) && Rules.Plugin == null)
-					{
-						Result.bCreateDebugInfo = false;
-					}
-				}
-
-				if (Rules.Plugin != null && Rules.Target.DisableDebugInfoPlugins.Contains(Rules.Plugin.Name))
-				{
-					Result.bCreateDebugInfo = false;
-				}
-				else if (Rules.Target.DisableDebugInfoModules.Contains(Name))
-				{
-					Result.bCreateDebugInfo = false;
-				}
-			}
+			CompileEnvironmentDebugInfoSettings(Target, Result);
 
 			// Only enable the AutoRTFM flag if we are using the AutoRTFM compiler
 			if (Target.bUseAutoRTFMCompiler)
diff --git a/Engine/Source/Programs/UnrealBuildTool/Platform/Android/AndroidToolChain.cs b/Engine/Source/Programs/UnrealBuildTool/Platform/Android/AndroidToolChain.cs
--- a/Engine/Source/Programs/UnrealBuildTool/Platform/Android/AndroidToolChain.cs
+++ b/Engine/Source/Programs/UnrealBuildTool/Platform/Android/AndroidToolChain.cs
@@ -558,6 +558,11 @@ namespace UnrealBuildTool
 			{
 				Arguments.Add("-g2");
 				Arguments.Add("-gdwarf-4");
+
+				if (CompileEnvironment.bDebugLineTablesOnly)
+				{
+					Arguments.Add("-gline-tables-only");
+				}
 			}
 
 			if (!DisableStackProtector())
diff --git a/Engine/Source/Programs/UnrealBuildTool/Platform/Linux/LinuxToolChain.cs b/Engine/Source/Programs/UnrealBuildTool/Platform/Linux/LinuxToolChain.cs
--- a/Engine/Source/Programs/UnrealBuildTool/Platform/Linux/LinuxToolChain.cs
+++ b/Engine/Source/Programs/UnrealBuildTool/Platform/Linux/LinuxToolChain.cs
@@ -476,6 +476,11 @@ namespace UnrealBuildTool
 				{
 					Arguments.Add("-glldb");
 				}
+
+				if (CompileEnvironment.bDebugLineTablesOnly)
+				{
+					Arguments.Add("-gline-tables-only");
+				}
 			}
 
 			if (CompileEnvironment.bHideSymbolsByDefault)
diff --git a/Engine/Source/Programs/UnrealBuildTool/Platform/Windows/VCToolChain.cs b/Engine/Source/Programs/UnrealBuildTool/Platform/Windows/VCToolChain.cs
--- a/Engine/Source/Programs/UnrealBuildTool/Platform/Windows/VCToolChain.cs
+++ b/Engine/Source/Programs/UnrealBuildTool/Platform/Windows/VCToolChain.cs
@@ -801,6 +801,12 @@ namespace UnrealBuildTool
 					Arguments.Add("/Z7");
 				}
 
+				// https://clang.llvm.org/docs/UsersManual.html#cmdoption-gline-tables-only
+				if (Target.WindowsPlatform.Compiler.IsClang() && CompileEnvironment.bDebugLineTablesOnly)
+				{
+					Arguments.Add("-gline-tables-only");
+				}
+
 				// https://clang.llvm.org/docs/UsersManual.html#cmdoption-fstandalone-debug
 				if (Target.WindowsPlatform.Compiler.IsClang() && Target.WindowsPlatform.bClangStandaloneDebug)
 				{
diff --git a/Engine/Source/Programs/UnrealBuildTool/System/CppCompileEnvironment.cs b/Engine/Source/Programs/UnrealBuildTool/System/CppCompileEnvironment.cs
--- a/Engine/Source/Programs/UnrealBuildTool/System/CppCompileEnvironment.cs
+++ b/Engine/Source/Programs/UnrealBuildTool/System/CppCompileEnvironment.cs
@@ -416,6 +416,11 @@ namespace UnrealBuildTool
 		/// </summary>
 		public bool bCreateDebugInfo = true;
 
+		/// <summary>
+		/// True if debug info should only generate line number tables (clang)
+		/// </summary>
+		public bool bDebugLineTablesOnly = false;
+
 		/// <summary>
 		/// True if we're compiling .cpp files that will go into a library (.lib file)
 		/// </summary>
@@ -744,6 +749,7 @@ namespace UnrealBuildTool
 			OptimizationLevel = Other.OptimizationLevel;
 			FPSemantics = Other.FPSemantics;
 			bCreateDebugInfo = Other.bCreateDebugInfo;
+			bDebugLineTablesOnly = Other.bDebugLineTablesOnly;
 			bIsBuildingLibrary = Other.bIsBuildingLibrary;
 			bIsBuildingDLL = Other.bIsBuildingDLL;
 			bUseStaticCRT = Other.bUseStaticCRT;
diff --git a/Engine/Source/Programs/UnrealBuildTool/ToolChain/AppleToolChain.cs b/Engine/Source/Programs/UnrealBuildTool/ToolChain/AppleToolChain.cs
--- a/Engine/Source/Programs/UnrealBuildTool/ToolChain/AppleToolChain.cs
+++ b/Engine/Source/Programs/UnrealBuildTool/ToolChain/AppleToolChain.cs
@@ -529,6 +529,11 @@ namespace UnrealBuildTool
 			if (CompileEnvironment.bCreateDebugInfo)
 			{
 				Arguments.Add("-gdwarf-4");
+
+				if (CompileEnvironment.bDebugLineTablesOnly)
+				{
+					Arguments.Add("-gline-tables-only");
+				}
 			}
 		}
 
