UnrealBuildTool: Add a few additional CompilationResult so .cs issues can be detected by horde more easily

diff --git a/Engine/Source/Programs/Shared/EpicGames.Core/UnrealEngineTypes.cs b/Engine/Source/Programs/Shared/EpicGames.Core/UnrealEngineTypes.cs
--- a/Engine/Source/Programs/Shared/EpicGames.Core/UnrealEngineTypes.cs
+++ b/Engine/Source/Programs/Shared/EpicGames.Core/UnrealEngineTypes.cs
@@ -17,7 +17,7 @@ namespace EpicGames.Core
 	/// Standard return codes used by UBT and UHT
 	/// 
 	/// This MUST be kept in sync with EGeneratedBodyVersion defined in 
-	/// Engine\Source\Runtime\Core\Public\Misc\ComplilationResult.h.
+	/// Engine\Source\Runtime\Core\Public\Misc\CompilationResult.h
 	/// </summary>
 	public enum CompilationResult
 	{
@@ -57,10 +57,20 @@ namespace EpicGames.Core
 		OtherCompilationError = 6,
 
 		/// <summary>
-		/// Compilation failed due to live coding action limit being exceeded.
+		/// Compilation failed due to live coding action limit being exceeded
 		/// </summary>
 		LiveCodingLimitError = 7,
 
+		/// <summary>
+		/// Compilation failed due to TargetRules or ModuleRules errors
+		/// </summary>
+		RulesError = 8,
+
+		/// <summary>
+		/// Compilation failed due to invalid action graph
+		/// </summary>
+		ActionGraphInvalid = 9,
+
 		/// <summary>
 		/// Compilation is not supported in the current build
 		/// </summary>
diff --git a/Engine/Source/Programs/UnrealBuildTool/Configuration/ModuleRules.cs b/Engine/Source/Programs/UnrealBuildTool/Configuration/ModuleRules.cs
--- a/Engine/Source/Programs/UnrealBuildTool/Configuration/ModuleRules.cs
+++ b/Engine/Source/Programs/UnrealBuildTool/Configuration/ModuleRules.cs
@@ -625,7 +625,7 @@ namespace UnrealBuildTool
 				}
 				else
 				{
-					throw new BuildException("Module '{0}' cannot override package type because it is part of a plugin!", Name);
+					throw new CompilationResultException(CompilationResult.RulesError, "Module '{ModuleName}' cannot override package type because it is part of a plugin!", Name);
 				}
 			}
 		}
@@ -1208,7 +1208,7 @@ namespace UnrealBuildTool
 			FileReference? ModuleFileReference = RulesAssembly.GetModuleFileName(ModuleName);
 			if (ModuleFileReference == null)
 			{
-				throw new BuildException("Could not find a module named '{0}'.", ModuleName);
+				throw new CompilationResultException(CompilationResult.RulesError, "Could not find a module named '{ModuleName}'.", ModuleName);
 			}
 			return ModuleFileReference.Directory.FullName;
 		}
@@ -1407,7 +1407,7 @@ namespace UnrealBuildTool
 			{
 				if (Plugin == null)
 				{
-					throw new BuildException("Module '{0}' does not belong to a plugin; PluginDirectory property is invalid.", Name);
+					throw new CompilationResultException(CompilationResult.RulesError, "Module '{ModuleName}' does not belong to a plugin; PluginDirectory property is invalid.", Name);
 				}
 				else
 				{
diff --git a/Engine/Source/Programs/UnrealBuildTool/Configuration/TargetRules.cs b/Engine/Source/Programs/UnrealBuildTool/Configuration/TargetRules.cs
--- a/Engine/Source/Programs/UnrealBuildTool/Configuration/TargetRules.cs
+++ b/Engine/Source/Programs/UnrealBuildTool/Configuration/TargetRules.cs
@@ -2669,7 +2669,7 @@ namespace UnrealBuildTool
 			ConstructorInfo? Constructor = RulesType.GetConstructor(new Type[] { typeof(TargetInfo) });
 			if (Constructor == null)
 			{
-				throw new BuildException("No constructor found on {0} which takes an argument of type TargetInfo.", RulesType.Name);
+				throw new CompilationResultException(CompilationResult.RulesError, "No constructor found on {TargetName} which takes an argument of type TargetInfo.", RulesType.Name);
 			}
 
 			// Invoke the regular constructor
@@ -2679,7 +2679,8 @@ namespace UnrealBuildTool
 			}
 			catch (Exception Ex)
 			{
-				throw new BuildException(Ex, "Unable to instantiate instance of '{0}' object type from compiled assembly '{1}'.  Unreal Build Tool creates an instance of your module's 'Rules' object in order to find out about your module's requirements.  The CLR exception details may provide more information:  {2}", RulesType.Name, Path.GetFileNameWithoutExtension(RulesType.Assembly?.Location), Ex.ToString());
+				throw new CompilationResultException(CompilationResult.RulesError, Ex, "Unable to instantiate instance of '{TargetName}' object type from compiled assembly '{AssemblyPath}'.  Unreal Build Tool creates an instance of your module's 'Rules' object in order to find out about your module's requirements.  The CLR exception details may provide more information: {ExceptionMessage}",
+                    RulesType.Name, Path.GetFileNameWithoutExtension(RulesType.Assembly?.Location) ?? "Unknown Assembly", Ex.ToString());
 			}
 
 			return Rules;
diff --git a/Engine/Source/Programs/UnrealBuildTool/System/ActionGraph.cs b/Engine/Source/Programs/UnrealBuildTool/System/ActionGraph.cs
--- a/Engine/Source/Programs/UnrealBuildTool/System/ActionGraph.cs
+++ b/Engine/Source/Programs/UnrealBuildTool/System/ActionGraph.cs
@@ -111,7 +111,7 @@ namespace UnrealBuildTool
 
 			if (!bResult)
 			{
-				throw new BuildException("Action graph is invalid; unable to continue. See log for additional details.");
+				throw new CompilationResultException(CompilationResult.ActionGraphInvalid, "Action graph is invalid; unable to continue. See log for additional details.");
 			}
 		}
 
@@ -609,7 +609,7 @@ namespace UnrealBuildTool
 					}
 				}
 
-				throw new BuildException($"Action graph contains cycle!");
+				throw new CompilationResultException(CompilationResult.ActionGraphInvalid, "Action graph contains cycle!");
 			}
 		}
 
diff --git a/Engine/Source/Programs/UnrealBuildTool/System/BuildException.cs b/Engine/Source/Programs/UnrealBuildTool/System/BuildException.cs
--- a/Engine/Source/Programs/UnrealBuildTool/System/BuildException.cs
+++ b/Engine/Source/Programs/UnrealBuildTool/System/BuildException.cs
@@ -21,7 +21,7 @@ namespace UnrealBuildTool
 		}
 
 		/// <summary>
-		/// Constructor
+		/// Constructor which wraps another exception
 		/// </summary>
 		/// <param name="InnerException">An inner exception to wrap</param>
 		/// <param name="Message">The error message to display.</param>
@@ -51,6 +51,16 @@ namespace UnrealBuildTool
 		{
 		}
 
+		/// <summary>
+		/// Log BuildException with a provided ILogger
+		/// </summary>
+		/// <param name="Logger">The ILogger to use to log this exception</param>
+		public virtual void LogException(ILogger Logger)
+		{
+			Logger.LogError(this, "{Ex}", ExceptionUtils.FormatException(this));
+			Logger.LogDebug(this, "{Ex}", ExceptionUtils.FormatExceptionDetails(this));
+		}
+
 		/// <summary>
 		/// Returns the string representing the exception. Our build exceptions do not show the callstack since they are used to report known error conditions.
 		/// </summary>
@@ -81,7 +91,7 @@ namespace UnrealBuildTool
 		}
 
 		/// <summary>
-		/// Constructor
+		/// Constructor which wraps another exception
 		/// </summary>
 		/// <param name="InnerException">The inner exception</param>
 		/// <param name="Event">Event to construct from</param>
@@ -91,37 +101,118 @@ namespace UnrealBuildTool
 			this.Event = Event;
 		}
 
+		/// <inheritdoc/>
+		public BuildLogEventException(string Format, params object[] Arguments)
+			: this(LogEvent.Create(LogLevel.Error, Format, Arguments))
+		{
+		}
+
+		/// <inheritdoc/>
+		public BuildLogEventException(Exception? InnerException, string Format, params object[] Arguments)
+			: this(InnerException, LogEvent.Create(LogLevel.Error, default, InnerException, Format, Arguments))
+		{
+		}
+
 		/// <summary>
-		/// Constructor
+		/// Constructor which wraps another exception
 		/// </summary>
+		/// <param name="EventId">Event id for the error</param>
+		/// <param name="InnerException">Inner exception to wrap</param>
 		/// <param name="Format">Structured logging format string</param>
-		/// <param name="Args">Argument objects</param>
-		public BuildLogEventException(string Format, params object[] Args)
-			: this(LogEvent.Create(LogLevel.Error, Format, Args))
+		/// <param name="Arguments">Argument objects</param>
+		public BuildLogEventException(Exception? InnerException, EventId EventId, string Format, params object[] Arguments)
+			: this(InnerException, LogEvent.Create(LogLevel.Error, EventId, InnerException, Format, Arguments))
+		{
+		}
+
+		/// <inheritdoc/>
+		public override void LogException(ILogger Logger)
 		{
+			Logger.Log(Event.Level, Event.Id, Event, this, (s, e) => s.ToString());
+			Logger.LogDebug(this, "{Ex}", ExceptionUtils.FormatExceptionDetails(this));
 		}
+	}
+
+	/// <summary>
+	/// Implementation of <see cref="BuildLogEventException"/> that will return a unique exit code.
+	/// </summary>
+	class CompilationResultException : BuildLogEventException
+	{
+		/// <summary>
+		/// The exit code associated with this exception
+		/// </summary>
+		public CompilationResult Result { get; }
 
 		/// <summary>
 		/// Constructor
 		/// </summary>
-		/// <param name="InnerException">Inner exception to wrap</param>
-		/// <param name="Format">Structured logging format string</param>
-		/// <param name="Args">Argument objects</param>
-		public BuildLogEventException(Exception? InnerException, string Format, params object[] Args)
-			: this(InnerException, LogEvent.Create(LogLevel.Error, default, InnerException, Format, Args))
+		/// <param name="Result">The resulting exit code</param>
+		public CompilationResultException(CompilationResult Result)
+			: base(LogEvent.Create(LogLevel.Error, "{CompilationResult}", Result))
 		{
+			this.Result = Result;
 		}
 
 		/// <summary>
 		/// Constructor
 		/// </summary>
+		/// <param name="Result">The resulting exit code</param>
+		/// <param name="Event">Event to construct from</param>
+		public CompilationResultException(CompilationResult Result, LogEvent Event)
+			: base(Event)
+		{
+			this.Result = Result;
+		}
+
+		/// <summary>
+		/// Constructor which wraps another exception
+		/// </summary>
+		/// <param name="Result">The resulting exit code</param>
+		/// <param name="InnerException">The inner exception</param>
+		/// <param name="Event">Event to construct from</param>
+		public CompilationResultException(CompilationResult Result, Exception? InnerException, LogEvent Event)
+			: base(InnerException, Event)
+		{
+			this.Result = Result;
+		}
+
+		/// <summary>
+		/// Constructor
+		/// </summary>
+		/// <param name="Result">The resulting exit code</param>
+		/// <param name="Format">Formatting string for the error message</param>
+		/// <param name="Arguments">Arguments for the formatting string</param>
+		public CompilationResultException(CompilationResult Result, string Format, params object[] Arguments)
+			: base(LogEvent.Create(LogLevel.Error, Format, Arguments))
+		{
+			this.Result = Result;
+		}
+
+		/// <summary>
+		/// Constructor which wraps another exception
+		/// </summary>
+		/// <param name="Result">The resulting exit code</param>
+		/// <param name="InnerException">The inner exception being wrapped</param>
+		/// <param name="Format">Format for the message string</param>
+		/// <param name="Arguments">Format arguments</param>
+		public CompilationResultException(CompilationResult Result, Exception? InnerException, string Format, params object[] Arguments)
+			: base(InnerException, LogEvent.Create(LogLevel.Error, default, InnerException, Format, Arguments))
+		{
+			this.Result = Result;
+		}
+
+		/// <summary>
+		/// Constructor which wraps another exception
+		/// </summary>
+		/// <param name="Result">The resulting exit code</param>
 		/// <param name="EventId">Event id for the error</param>
 		/// <param name="InnerException">Inner exception to wrap</param>
 		/// <param name="Format">Structured logging format string</param>
-		/// <param name="Args">Argument objects</param>
-		public BuildLogEventException(Exception? InnerException, EventId EventId, string Format, params object[] Args)
-			: this(InnerException, LogEvent.Create(LogLevel.Error, EventId, InnerException, Format, Args))
+		/// <param name="Arguments">Argument objects</param>
+		public CompilationResultException(CompilationResult Result, Exception? InnerException, EventId EventId, string Format, params object[] Arguments)
+			: base(InnerException, LogEvent.Create(LogLevel.Error, EventId, InnerException, Format, Arguments))
 		{
+			this.Result = Result;
 		}
 	}
 }
diff --git a/Engine/Source/Programs/UnrealBuildTool/System/ExternalExecution.cs b/Engine/Source/Programs/UnrealBuildTool/System/ExternalExecution.cs
--- a/Engine/Source/Programs/UnrealBuildTool/System/ExternalExecution.cs
+++ b/Engine/Source/Programs/UnrealBuildTool/System/ExternalExecution.cs
@@ -20,18 +20,6 @@ using UnrealBuildTool.Modes;
 
 namespace UnrealBuildTool
 {
-
-	class CompilationResultException : BuildException
-	{
-		public readonly CompilationResult Result;
-
-		public CompilationResultException(CompilationResult Result)
-			: base("Error: {0}", Result)
-		{
-			this.Result = Result;
-		}
-	}
-
 	static class UHTModuleTypeExtensions
 	{
 		public static UHTModuleType? EngineModuleTypeFromHostType(ModuleHostType ModuleType)
diff --git a/Engine/Source/Programs/UnrealBuildTool/System/RulesAssembly.cs b/Engine/Source/Programs/UnrealBuildTool/System/RulesAssembly.cs
--- a/Engine/Source/Programs/UnrealBuildTool/System/RulesAssembly.cs
+++ b/Engine/Source/Programs/UnrealBuildTool/System/RulesAssembly.cs
@@ -424,7 +424,7 @@ namespace UnrealBuildTool
 			{
 				if (Parent == null)
 				{
-					throw new BuildException("Could not find definition for module '{0}', (referenced via {1})", ModuleTypeName, ReferenceChain);
+					throw new CompilationResultException(CompilationResult.RulesError, "Could not find definition for module '{ModuleTypeName}', (referenced via {ReferenceChain})", ModuleTypeName, ReferenceChain);
 				}
 				else
 				{
@@ -447,7 +447,7 @@ namespace UnrealBuildTool
 					// we expect only one platform group to be found in the extensions
 					if (GroupRulesObjectType != null && PlatformRulesObjectType != null)
 					{
-						throw new BuildException("Found multiple platform group overrides ({0} and {1}) for module {2} without a platform specific override. Create a platform override with the class hierarchy as needed.",
+						throw new CompilationResultException(CompilationResult.RulesError, "Found multiple platform group overrides ({GroupRulesName} and {PlatformRulesName}) for module {ModuleName} without a platform specific override. Create a platform override with the class hierarchy as needed.",
 							GroupRulesObjectType.Name, PlatformRulesObjectType.Name, ModuleName);
 					}
 					// remember the platform group if we found it, but keep searching to verify there isn't more than one
@@ -468,7 +468,7 @@ namespace UnrealBuildTool
 				if (ModuleFileName.ContainsAnyNames(DisallowedPlatformsAndGroups, Unreal.EngineDirectory) ||
 					(Target.ProjectFile != null && ModuleFileName.ContainsAnyNames(DisallowedPlatformsAndGroups, Target.ProjectFile.Directory)))
 				{
-					throw new BuildException("Platform module file {0} is not allowed (only platforms '{1}', and their groups, are allowed. This indicates a module reference not being checked with something like IsPlatformAvailableForTarget()).",
+					throw new CompilationResultException(CompilationResult.RulesError, "Platform module file {ModuleFileName} is not allowed (only platforms '{Platforms}', and their groups, are allowed. This indicates a module reference not being checked with something like IsPlatformAvailableForTarget()).",
 						ModuleFileName, String.Join(",", Target.OptedInModulePlatforms));
 				}
 			}
@@ -477,7 +477,8 @@ namespace UnrealBuildTool
 			Type? RulesObjectType = PlatformRulesObjectType != null ? PlatformRulesObjectType : BaseRulesObjectType;
 			if (RulesObjectType == null)
 			{
-				throw new BuildException("Expecting to find a type to be declared in a module rules named '{0}' in {1}.  This type must derive from the 'ModuleRules' type defined by Unreal Build Tool.", ModuleTypeName, CompiledAssembly?.FullName);
+				throw new CompilationResultException(CompilationResult.RulesError, "Expecting to find a type to be declared in a module rules named '{ModuleTypeName}' in '{AssemblyName}'.  This type must derive from the 'ModuleRules' type defined by UnrealBuildTool.",
+					ModuleTypeName, CompiledAssembly?.FullName ?? "Unknown Assembly");
 			}
 
 			// Create an instance of the module's rules object
@@ -519,7 +520,7 @@ namespace UnrealBuildTool
 						}
 						if (SubType.BaseType == null)
 						{
-							throw new BuildException("{0} is not derived from {1}", RulesObjectType.Name, BaseRulesObjectType.Name);
+							throw new CompilationResultException(CompilationResult.RulesError, "{TypeName} is not derived from {BaseTypeName}", RulesObjectType.Name, BaseRulesObjectType.Name);
 						}
 						SubType = SubType.BaseType;
 					}
@@ -529,7 +530,7 @@ namespace UnrealBuildTool
 				ConstructorInfo? Constructor = RulesObjectType.GetConstructor(new Type[] { typeof(ReadOnlyTargetRules) });
 				if (Constructor == null)
 				{
-					throw new BuildException("No valid constructor found for {0}.", ModuleName);
+					throw new CompilationResultException(CompilationResult.RulesError, "No valid constructor found for {ModuleName}.", ModuleName);
 				}
 
 				// Add the parent assemblies to the assembly cache so the types in them can be used when the constructor is called
@@ -559,7 +560,7 @@ namespace UnrealBuildTool
 			catch (Exception Ex)
 			{
 				Exception MessageEx = (Ex is TargetInvocationException && Ex.InnerException != null) ? Ex.InnerException : Ex;
-				throw new BuildException(Ex, "Unable to instantiate module '{0}': {1}\n(referenced via {2})", ModuleName, MessageEx.ToString(), ReferenceChain);
+				throw new CompilationResultException(CompilationResult.RulesError, Ex, "Unable to instantiate module '{ModuleName}': {ExceptionMessage}\n(referenced via {ReferenceChain})", ModuleName, MessageEx.ToString(), ReferenceChain);
 			}
 		}
 
@@ -596,7 +597,7 @@ namespace UnrealBuildTool
 					// we expect only one platform group to be found in the extensions
 					if (GroupRulesObjectType != null && PlatformRulesType != null)
 					{
-						throw new BuildException("Found multiple platform group overrides ({0} and {1}) for rules {2} without a platform specific override. Create a platform override with the class hierarchy as needed.",
+						throw new CompilationResultException(CompilationResult.RulesError, "Found multiple platform group overrides ({GroupRulesName} and {PlatformRulesName}) for rules {TypeName} without a platform specific override. Create a platform override with the class hierarchy as needed.",
 							GroupRulesObjectType.Name, PlatformRulesType.Name, TypeName);
 					}
 					// remember the platform group if we found it, but keep searching to verify there isn't more than one
@@ -609,7 +610,7 @@ namespace UnrealBuildTool
 			}
 			if (PlatformRulesType != null && !PlatformRulesType.IsSubclassOf(BaseRulesType))
 			{
-				throw new BuildException("Expecting {0} to be a specialization of {1}", PlatformRulesType, BaseRulesType);
+				throw new CompilationResultException(CompilationResult.RulesError, "Expecting {PlatformRulesType} to be a specialization of {BaseRulesType}", PlatformRulesType, BaseRulesType);
 			}
 
 			// Create an instance of the module's rules object, and set some defaults before calling the constructor.
@@ -624,13 +625,13 @@ namespace UnrealBuildTool
 			// Set the final value for the link type in the target rules
 			if (Rules.LinkType == TargetLinkType.Default)
 			{
-				throw new BuildException("TargetRules.LinkType should be inferred from TargetType");
+				throw new CompilationResultException(CompilationResult.RulesError, "TargetRules.LinkType should be inferred from TargetType");
 			}
 
 			// Set the default value for whether to use the shared build environment
 			if (Rules.BuildEnvironment == TargetBuildEnvironment.Unique && Unreal.IsEngineInstalled())
 			{
-				throw new BuildException("Targets with a unique build environment cannot be built with an installed engine.");
+				throw new CompilationResultException(CompilationResult.RulesError, "Targets with a unique build environment cannot be built with an installed engine.");
 			}
 
 			// Automatically include CoreUObject
@@ -720,7 +721,7 @@ namespace UnrealBuildTool
 			// Some platforms may *require* monolithic compilation...
 			if (Rules.LinkType != TargetLinkType.Monolithic && UEBuildPlatform.PlatformRequiresMonolithicBuilds(Rules.Platform, Rules.Configuration))
 			{
-				throw new BuildException(String.Format("{0}: {1} does not support modular builds", Rules.Name, Rules.Platform));
+				throw new CompilationResultException(CompilationResult.RulesError, String.Format("{RulesName}: {RulesPlatform} does not support modular builds", Rules.Name, Rules.Platform));
 			}
 
 			if (IsTestTarget)
@@ -753,14 +754,16 @@ namespace UnrealBuildTool
 					// if the target doesn't allow for an override at all, error
 					if (!Rules.AllowsPerProjectSDKVersion())
 					{
-						throw new BuildException($"Target {Rules.Name} is being built with a overridden {Platform} SDK version to '{SDK.GetMainVersion()}', but this target is not allowed - likely due to a modular build using a Shared BuildEnvironment");
+						throw new CompilationResultException(CompilationResult.RulesError, "Target {RulesName} is being built with a overridden {Platform} SDK version to '{SdkVersion}', but this target is not allowed - likely due to a modular build using a Shared BuildEnvironment",
+							Rules.Name, Platform, SDK.GetMainVersion());
 					}
 
 					// check to see if the project _doesn't_ override the SDK version (if it did, this was already handled during early processing)
 					if (Rules.ProjectFile == null || !SDK.ProjectsThatOverrodeSDK.Contains(Rules.ProjectFile))
 					{
 						string OverrideProject = SDK.ProjectsThatOverrodeSDK[0].GetFileNameWithoutAnyExtensions();
-						throw new BuildException($"Target {Rules.Name} is using default {Platform} SDK version, but another target (probably {OverrideProject}) has overridden the SDK version to '{SDK.GetMainVersion()}'. If this target doesn't care about SDK versions, set 'bAreTargetSDKVersionsRelevantOverride = false' in your Target.cs file");
+						throw new CompilationResultException(CompilationResult.RulesError, "Target {RulesName} is using default {Platform} SDK version, but another target (probably {OverrideProject}) has overridden the SDK version to '{SdkVersion}'. If this target doesn't care about SDK versions, set 'bAreTargetSDKVersionsRelevantOverride = false' in your Target.cs file",
+							Rules.Name, Platform, OverrideProject, SDK.GetMainVersion());
 					}
 				}
 			}
@@ -798,22 +801,19 @@ namespace UnrealBuildTool
 			{
 				if (Parent == null)
 				{
-					//				throw new BuildException("Couldn't find target rules file for target '{0}' in rules assembly '{1}'.", TargetName, RulesAssembly.FullName);
-					string ExceptionMessage = "Couldn't find target rules file for target '";
-					ExceptionMessage += TargetName;
-					ExceptionMessage += "' in rules assembly '";
-					ExceptionMessage += CompiledAssembly?.FullName;
-					ExceptionMessage += "'." + Environment.NewLine;
-
-					ExceptionMessage += "Location: " + CompiledAssembly?.Location + Environment.NewLine;
-
-					ExceptionMessage += "Target rules found:" + Environment.NewLine;
-					foreach (KeyValuePair<string, FileReference> entry in TargetNameToTargetFile)
+					string ExceptionMessage = "Couldn't find target rules file for target '{TargetName}' in rules assembly '{RulesAssembly}'" + Environment.NewLine;
+					ExceptionMessage += "Location: {AssemblyLocation}" + Environment.NewLine;
+					if (TargetNameToTargetFile.Any())
 					{
-						ExceptionMessage += "\t" + entry.Key + " - " + entry.Value + Environment.NewLine;
+						ExceptionMessage += "Target rules found:" + Environment.NewLine;
+						foreach (KeyValuePair<string, FileReference> entry in TargetNameToTargetFile)
+						{
+							ExceptionMessage += "\t" + entry.Key + " - " + entry.Value + Environment.NewLine;
+						}
 					}
 
-					throw new BuildException(ExceptionMessage);
+					throw new CompilationResultException(CompilationResult.RulesError, ExceptionMessage,
+						TargetName, CompiledAssembly?.FullName ?? "Unknown Assembly", CompiledAssembly?.Location ?? "Unknown Location");
 				}
 				else
 				{
@@ -829,7 +829,7 @@ namespace UnrealBuildTool
 
 			if (TargetRules == null)
 			{
-				throw new BuildException("Expecting to find a type to be declared in a target rules named '{0}'.  This type must derive from the 'TargetRules' type defined by Unreal Build Tool.", TargetTypeName);
+				throw new CompilationResultException(CompilationResult.RulesError, "Expecting to find a type to be declared in a target rules named '{TargetTypeName}'.  This type must derive from the 'TargetRules' type defined by UnrealBuildTool.", TargetTypeName);
 			}
 
 			return TargetRules;
@@ -862,7 +862,7 @@ namespace UnrealBuildTool
 			{
 				if (Parent == null)
 				{
-					throw new BuildException("Unable to find target of type '{0}' for project '{1}'", Type, ProjectFile);
+					throw new CompilationResultException(CompilationResult.RulesError, "Unable to find target of type '{Type}' for project '{1}'", Type, ProjectFile?.FullName ?? "NoProject");
 				}
 				else
 				{
@@ -890,7 +890,7 @@ namespace UnrealBuildTool
 					}
 				}
 
-				throw new BuildException("Found multiple targets with TargetType={0}: {1}.\nSpecify a default with a {2} entry in [/Script/BuildSettings.BuildSettings] section of your DefaultEngine.ini", Type, String.Join(", ", Matches), KeyName);
+				throw new CompilationResultException(CompilationResult.RulesError, "Found multiple targets with TargetType={Type}: {Matches}.\nSpecify a default with a {KeyName} entry in [/Script/BuildSettings.BuildSettings] section of your DefaultEngine.ini", Type, String.Join(", ", Matches), KeyName);
 			}
 		}
 
diff --git a/Engine/Source/Programs/UnrealBuildTool/UnrealBuildTool.cs b/Engine/Source/Programs/UnrealBuildTool/UnrealBuildTool.cs
--- a/Engine/Source/Programs/UnrealBuildTool/UnrealBuildTool.cs
+++ b/Engine/Source/Programs/UnrealBuildTool/UnrealBuildTool.cs
@@ -674,22 +674,20 @@ namespace UnrealBuildTool
 			}
 			catch (CompilationResultException Ex)
 			{
-				// Used to return a propagate a specific exit code after an error has occurred. Does not log any message.
-				Logger.LogDebug(Ex, "{Ex}", ExceptionUtils.FormatExceptionDetails(Ex));
+				// Used to return a propagate a specific exit code after an error has occurred.
+				Ex.LogException(Logger);
 				return (int)Ex.Result;
 			}
 			catch (BuildLogEventException Ex)
 			{
-				// BuildExceptions should have nicely formatted messages. We can log these directly.
-				Logger.Log(Ex.Event.Level, Ex.Event.Id, Ex.Event, Ex, (s, e) => s.ToString());
-				Logger.LogDebug(Ex, "{Ex}", ExceptionUtils.FormatExceptionDetails(Ex));
+				// BuildExceptions should have nicely formatted messages.
+				Ex.LogException(Logger);
 				return (int)CompilationResult.OtherCompilationError;
 			}
 			catch (BuildException Ex)
 			{
-				// BuildExceptions should have nicely formatted messages. We can log these directly.
-				Logger.LogError(Ex, "{Ex}", ExceptionUtils.FormatException(Ex));
-				Logger.LogDebug(Ex, "{Ex}", ExceptionUtils.FormatExceptionDetails(Ex));
+				// BuildExceptions should have nicely formatted messages.
+				Ex.LogException(Logger);
 				return (int)CompilationResult.OtherCompilationError;
 			}
 			catch (Exception Ex)
diff --git a/Engine/Source/Runtime/Core/Public/Misc/CompilationResult.h b/Engine/Source/Runtime/Core/Public/Misc/CompilationResult.h
--- a/Engine/Source/Runtime/Core/Public/Misc/CompilationResult.h
+++ b/Engine/Source/Runtime/Core/Public/Misc/CompilationResult.h
@@ -31,10 +31,14 @@ namespace ECompilationResult
 		OtherCompilationError = 6,
 		/** Compilation failed due to live coding limit reached */
 		LiveCodingLimitError = 7,
+		/** Compilation failed due to TargetRules or ModuleRules errors */
+		RulesError = 8,
+		/** Compilation failed due to invalid action graph */
+		ActionGraphInvalid = 9,
 		/** Compilation is not supported in the current build */
 		Unsupported,
 		/** Unknown error */
-		Unknown
+		Unknown 
 	};
 
 	/**
@@ -58,6 +62,10 @@ namespace ECompilationResult
 			return TEXT("CrashOrAssert");
 		case ECompilationResult::LiveCodingLimitError:
 			return TEXT("LiveCodingLimitError");
+		case ECompilationResult::RulesError:
+			return TEXT("RulesError");
+		case ECompilationResult::ActionGraphInvalid:
+			return TEXT("ActionGraphInvalid");
 		case ECompilationResult::Unsupported:
 			return TEXT("Unsupported");
 		};
