UnrealBuildTool: Don't show CompilationResult as the error message if no additional message is added

#rnx

diff --git a/Engine/Source/Programs/UnrealBuildTool/System/BuildException.cs b/Engine/Source/Programs/UnrealBuildTool/System/BuildException.cs
--- a/Engine/Source/Programs/UnrealBuildTool/System/BuildException.cs
+++ b/Engine/Source/Programs/UnrealBuildTool/System/BuildException.cs
@@ -143,6 +143,8 @@ namespace UnrealBuildTool
 		/// </summary>
 		public CompilationResult Result { get; }
 
+		readonly bool HasMessage = true;
+
 		/// <summary>
 		/// Constructor
 		/// </summary>
@@ -150,6 +152,7 @@ namespace UnrealBuildTool
 		public CompilationResultException(CompilationResult Result)
 			: base(LogEvent.Create(LogLevel.Error, "{CompilationResult}", Result))
 		{
+			HasMessage = false;
 			this.Result = Result;
 		}
 
@@ -214,6 +217,16 @@ namespace UnrealBuildTool
 		{
 			this.Result = Result;
 		}
+
+		/// <inheritdoc/>
+		public override void LogException(ILogger Logger)
+		{
+			if (HasMessage)
+			{
+				Logger.Log(Event.Level, Event.Id, Event, this, (s, e) => s.ToString());
+			}
+			Logger.LogDebug(this, "{Ex}", ExceptionUtils.FormatExceptionDetails(this));
+		}
 	}
 }
 
