Fix UBT to enable hot reload modules for engine only builds. This is done so that live coding checks still work when building from P4V or other tools that might not specify a project file.

diff --git a/Engine/Source/Programs/UnrealBuildTool/System/RulesCompiler.cs b/Engine/Source/Programs/UnrealBuildTool/System/RulesCompiler.cs
--- a/Engine/Source/Programs/UnrealBuildTool/System/RulesCompiler.cs
+++ b/Engine/Source/Programs/UnrealBuildTool/System/RulesCompiler.cs
@@ -55,7 +55,7 @@ namespace UnrealBuildTool
 			DirectoryReference Directory = DirectoryReference.Combine(BaseDirectory, SubDirectoryName);
 			if (DirectoryLookupCache.DirectoryExists(Directory))
 			{
-				ModuleRulesContext ModuleContext = new ModuleRulesContext(BaseModuleContext) { DefaultUHTModuleType = DefaultUHTModuleType };
+				ModuleRulesContext ModuleContext = new ModuleRulesContext(BaseModuleContext) { DefaultUHTModuleType = DefaultUHTModuleType, bCanHotReload = true };
 				AddModuleRulesWithContext(Directory, ModuleContext, ModuleFileToContext);
 			}
 		}
