UnrealBuildTool: Fix issue preventing PGO from running against a target without a project

#rnx

diff --git a/Engine/Source/Programs/UnrealBuildTool/Platform/Android/UEBuildAndroid.cs b/Engine/Source/Programs/UnrealBuildTool/Platform/Android/UEBuildAndroid.cs
--- a/Engine/Source/Programs/UnrealBuildTool/Platform/Android/UEBuildAndroid.cs
+++ b/Engine/Source/Programs/UnrealBuildTool/Platform/Android/UEBuildAndroid.cs
@@ -499,12 +499,12 @@ namespace UnrealBuildTool
 				CompileEnvironment.Definitions.Add("RHI_RAYTRACING=1");
 			}
 
-			if ((Target.bPGOOptimize || Target.bPGOProfile) && Target.ProjectFile != null)
+			if (Target.bPGOOptimize || Target.bPGOProfile)
 			{
 				Logger.LogInformation("PGO {PgoType} build", Target.bPGOOptimize ? "optimize" : "profile");
 				if (Target.bPGOOptimize)
 				{
-					CompileEnvironment.PGODirectory = DirectoryReference.Combine(Target.ProjectFile.Directory, "Platforms", "Android", "Build", "PGO").FullName;
+					CompileEnvironment.PGODirectory = DirectoryReference.Combine(Target.ProjectFile?.Directory ?? Unreal.WritableEngineDirectory, "Platforms", "Android", "Build", "PGO").FullName;
 					CompileEnvironment.PGOFilenamePrefix = String.Format("{0}-Android", Target.Name);
 
 					LinkEnvironment.PGODirectory = CompileEnvironment.PGODirectory;
diff --git a/Engine/Source/Programs/UnrealBuildTool/Platform/Linux/UEBuildLinux.cs b/Engine/Source/Programs/UnrealBuildTool/Platform/Linux/UEBuildLinux.cs
--- a/Engine/Source/Programs/UnrealBuildTool/Platform/Linux/UEBuildLinux.cs
+++ b/Engine/Source/Programs/UnrealBuildTool/Platform/Linux/UEBuildLinux.cs
@@ -494,7 +494,7 @@ namespace UnrealBuildTool
 
 			if (CompileEnvironment.bPGOOptimize)
 			{
-				DirectoryReference BaseDir = Unreal.EngineDirectory;
+				DirectoryReference BaseDir = Unreal.WritableEngineDirectory;
 				if (Target.ProjectFile != null)
 				{
 					BaseDir = DirectoryReference.FromFile(Target.ProjectFile);
diff --git a/Engine/Source/Programs/UnrealBuildTool/Platform/Windows/UEBuildWindows.cs b/Engine/Source/Programs/UnrealBuildTool/Platform/Windows/UEBuildWindows.cs
--- a/Engine/Source/Programs/UnrealBuildTool/Platform/Windows/UEBuildWindows.cs
+++ b/Engine/Source/Programs/UnrealBuildTool/Platform/Windows/UEBuildWindows.cs
@@ -1655,12 +1655,12 @@ namespace UnrealBuildTool
 
 			LinkEnvironment.ModuleDefinitionFile = Target.WindowsPlatform.ModuleDefinitionFile;
 
-			if ((Target.bPGOOptimize || Target.bPGOProfile) && Target.ProjectFile != null)
+			if (Target.bPGOOptimize || Target.bPGOProfile)
 			{
 				// Win64 PGO folder is Windows, the rest match the platform name
 				string PGOPlatform = Target.Platform == UnrealTargetPlatform.Win64 ? "Windows" : Target.Platform.ToString();
 
-				CompileEnvironment.PGODirectory = DirectoryReference.Combine(Target.ProjectFile.Directory, "Platforms", PGOPlatform, "Build", "PGO").FullName;
+				CompileEnvironment.PGODirectory = DirectoryReference.Combine(Target.ProjectFile?.Directory ?? Unreal.WritableEngineDirectory, "Platforms", PGOPlatform, "Build", "PGO").FullName;
 				CompileEnvironment.PGOFilenamePrefix = String.Format("{0}-{1}-{2}", Target.Name, Target.Platform, Target.Configuration);
 
 				LinkEnvironment.PGODirectory = CompileEnvironment.PGODirectory;
