- Fixing product name in Xcode project for launching, not building/installing, when there is no ACTION

diff --git a/Engine/Source/Programs/UnrealBuildTool/ProjectFiles/Xcode/XcodeProject.cs b/Engine/Source/Programs/UnrealBuildTool/ProjectFiles/Xcode/XcodeProject.cs
--- a/Engine/Source/Programs/UnrealBuildTool/ProjectFiles/Xcode/XcodeProject.cs
+++ b/Engine/Source/Programs/UnrealBuildTool/ProjectFiles/Xcode/XcodeProject.cs
@@ -1998,10 +1998,27 @@ namespace UnrealBuildTool.XcodeProjectXcconfig
 				ConfigXcconfig.AppendLine($"UE_TARGET_CONFIG = {Config.BuildConfig}");
 				ConfigXcconfig.AppendLine($"UE_UBT_BINARY_SUBPATH = {ExetuableSubPath}");
 				ConfigXcconfig.AppendLine($"{ExecutableKey} = {ExecutableName}");
+				if (EntitlementsMetadata != null && EntitlementsMetadata.Mode == MetadataMode.UsePremade)
+				{
+					ConfigXcconfig.AppendLine($"CODE_SIGN_ENTITLEMENTS = {EntitlementsMetadata.XcodeProjectRelative}");
+				}
+
+				// debug settings
+				if (Config.BuildConfig == UnrealTargetConfiguration.Debug)
+				{
+					ConfigXcconfig.AppendLine("ENABLE_TESTABILITY = YES");
+				}
+
 				if (Platform == UnrealTargetPlatform.Mac)
 				{
 					// on Mac, we need to name the .app nicely before pushing to App store, otherwise distributing, so use the ini setting if it's there ("Unreal Match 3"), otherwise use the uproject name (ie "Lyra" instead of "LyraGame")
-					ConfigXcconfig.AppendLine($"PRODUCT_NAME_build = {ProductName}");
+					ConfigXcconfig.AppendLine("");
+					ConfigXcconfig.AppendLine($"// this variable trickery will set the proper name for debugging, building, and archiving,");
+					ConfigXcconfig.AppendLine($"// where archiving (the '_install' action type) may need a differnet name so it shows up nicely");
+					ConfigXcconfig.AppendLine($"// on end-users machines in Finder, Spotlight, etc. The trailing _ on the next line is correct.");
+
+					ConfigXcconfig.AppendLine($"PRODUCT_NAME_ = {ProductName}");
+					ConfigXcconfig.AppendLine($"PRODUCT_NAME_build = $(PRODUCT_NAME_)");
 					string ArchivedName = ApplicationDisplayName ?? (UnrealData.UProjectFileLocation == null ? ProductName : UnrealData.UProjectFileLocation!.GetFileNameWithoutAnyExtensions());
 					ConfigXcconfig.AppendLine($"PRODUCT_NAME_install = {ArchivedName}");
 
@@ -2012,16 +2029,6 @@ namespace UnrealBuildTool.XcodeProjectXcconfig
 				{
 					ConfigXcconfig.AppendLine($"PRODUCT_NAME = {ProductName}");
 				}
-				if (EntitlementsMetadata != null && EntitlementsMetadata.Mode == MetadataMode.UsePremade)
-				{
-					ConfigXcconfig.AppendLine($"CODE_SIGN_ENTITLEMENTS = {EntitlementsMetadata.XcodeProjectRelative}");
-				}
-
-				// debug settings
-				if (Config.BuildConfig == UnrealTargetConfiguration.Debug)
-				{
-					ConfigXcconfig.AppendLine("ENABLE_TESTABILITY = YES");
-				}
 
 				ConfigXcconfig.Write();
 			}
