UnrealBuildTool: Option to use the experimental executor for local-only compiles, if no other distribution method is used

#rnx

diff --git a/Engine/Source/Programs/UnrealBuildTool/Configuration/BuildConfiguration.cs b/Engine/Source/Programs/UnrealBuildTool/Configuration/BuildConfiguration.cs
--- a/Engine/Source/Programs/UnrealBuildTool/Configuration/BuildConfiguration.cs
+++ b/Engine/Source/Programs/UnrealBuildTool/Configuration/BuildConfiguration.cs
@@ -39,7 +39,6 @@ namespace UnrealBuildTool
 		[Obsolete("HybridExecutor is no longer supported")]
 		public bool bAllowHybridExecutor = false;
 
-#if __BOXEXECUTOR_AVAILABLE__
 		/// <summary>
 		/// Whether the experimental box executor will be used.
 		/// </summary>
@@ -47,7 +46,19 @@ namespace UnrealBuildTool
 		[CommandLine("-Box", Value = "true")]
 		[CommandLine("-NoBox", Value = "false")]
 		public bool bAllowBoxExecutor = false;
-#endif // #if __BOXEXECUTOR_AVAILABLE__
+
+		/// <summary>
+		/// Whether the experimental box (local only) executor will be used.
+		/// </summary>
+		[XmlConfigFile]
+		[CommandLine("-BoxLocal", Value = "true")]
+		[CommandLine("-NoBoxLocal", Value = "false")]
+		public bool bAllowBoxLocalExecutor
+		{
+			get => bAllowBoxExecutor || bAllowBoxLocalExecutorPrivate;
+			set => bAllowBoxLocalExecutorPrivate = value;
+		}
+		private bool bAllowBoxLocalExecutorPrivate = false;
 
 		/// <summary>
 		/// Whether XGE may be used if available, default is true.
diff --git a/Engine/Source/Programs/UnrealBuildTool/System/ActionGraph.cs b/Engine/Source/Programs/UnrealBuildTool/System/ActionGraph.cs
--- a/Engine/Source/Programs/UnrealBuildTool/System/ActionGraph.cs
+++ b/Engine/Source/Programs/UnrealBuildTool/System/ActionGraph.cs
@@ -361,17 +361,15 @@ namespace UnrealBuildTool
 		/// </summary>
 		private static ActionExecutor SelectExecutor(BuildConfiguration BuildConfiguration, int ActionCount, List<TargetDescriptor> TargetDescriptors, ILogger Logger)
 		{
-#if __BOXEXECUTOR_AVAILABLE__
-			bool bAnySingleFile = TargetDescriptors.Any(Descriptor => Descriptor.SpecificFilesToCompile.Count == 1);
-			if (!bAnySingleFile && BuildConfiguration.bAllowBoxExecutor && BoxExecutor.IsAvailable(Logger))
+			if (ActionCount > ParallelExecutor.GetDefaultNumParallelProcesses(BuildConfiguration.MaxParallelActions, BuildConfiguration.bAllCores, Logger))
 			{
-				CommandLineArguments? AdditionalArguments = TargetDescriptors.FirstOrDefault()?.AdditionalArguments;
-				return new BoxExecutor(BuildConfiguration.MaxParallelActions, BuildConfiguration.bAllCores, BuildConfiguration.bCompactOutput, Logger, AdditionalArguments);
-			}
+#if __BOXEXECUTOR_AVAILABLE__
+				if (BuildConfiguration.bAllowBoxExecutor && BoxExecutor.IsAvailable(Logger))
+				{
+					return new BoxExecutor(BuildConfiguration.MaxParallelActions, BuildConfiguration.bAllCores, BuildConfiguration.bCompactOutput, Logger, TargetDescriptors.FirstOrDefault()?.AdditionalArguments);
+				}
 #endif // #if __BOXEXECUTOR_AVAILABLE__
 
-			if (ActionCount > ParallelExecutor.GetDefaultNumParallelProcesses(BuildConfiguration.MaxParallelActions, BuildConfiguration.bAllCores, Logger))
-			{
 				if (BuildConfiguration.bAllowXGE && XGE.IsAvailable(Logger) && ActionCount >= XGE.MinActions)
 				{
 					return new XGE(Logger);
@@ -386,6 +384,13 @@ namespace UnrealBuildTool
 				}
 			}
 
+#if __BOXEXECUTOR_AVAILABLE__
+			if (BuildConfiguration.bAllowBoxLocalExecutor && BoxLocalExecutor.IsAvailable(Logger))
+			{
+				return new BoxLocalExecutor(BuildConfiguration.MaxParallelActions, BuildConfiguration.bAllCores, BuildConfiguration.bCompactOutput, Logger, TargetDescriptors.FirstOrDefault()?.AdditionalArguments);
+			}
+#endif // #if __BOXEXECUTOR_AVAILABLE__
+
 			return new ParallelExecutor(BuildConfiguration.MaxParallelActions, BuildConfiguration.bAllCores, BuildConfiguration.bCompactOutput, Logger);
 		}
 
@@ -402,6 +407,7 @@ namespace UnrealBuildTool
 			{
 				// Figure out which executor to use
 				using ActionExecutor Executor = SelectExecutor(BuildConfiguration, ActionsToExecute.Count, TargetDescriptors, Logger);
+				Logger.LogInformation("Using {ExecutorName} executor to run {ActionCount} action(s)", Executor.Name, ActionsToExecute.Count);
 
 				// Execute the build
 				Stopwatch Timer = Stopwatch.StartNew();
