[UBT]
* Changed so cl.exe has a weight of 1.0 when not using unity files. When unity it is set to 1.5 but we very very rarely end up with 100% cpu utilization with 1.5.

Discussed with [REDACTED_PERSON]

diff --git a/Engine/Source/Programs/UnrealBuildTool/Platform/Windows/VCToolChain.cs b/Engine/Source/Programs/UnrealBuildTool/Platform/Windows/VCToolChain.cs
--- a/Engine/Source/Programs/UnrealBuildTool/Platform/Windows/VCToolChain.cs
+++ b/Engine/Source/Programs/UnrealBuildTool/Platform/Windows/VCToolChain.cs
@@ -1460,7 +1460,7 @@ namespace UnrealBuildTool
 			if (Target.WindowsPlatform.Compiler.IsMSVC())
 			{
 				// If deterministic is enabled, MSVC does not use multiple threads
-				if (!CompileEnvironment.bDeterministic)
+				if (!CompileEnvironment.bDeterministic && !CompileEnvironment.bUseUnity)
 				{
 					BaseCompileAction.Weight = Target.MSVCCompileActionWeight;
 				}
