UnrealBuildTool: Don't allow using a separate UnrealIntermediateEnvironment folder for installed engine as the intermediate path needs to always match

diff --git a/Engine/Source/Programs/UnrealBuildTool/Configuration/UEBuildTarget.cs b/Engine/Source/Programs/UnrealBuildTool/Configuration/UEBuildTarget.cs
--- a/Engine/Source/Programs/UnrealBuildTool/Configuration/UEBuildTarget.cs
+++ b/Engine/Source/Programs/UnrealBuildTool/Configuration/UEBuildTarget.cs
@@ -1676,7 +1676,7 @@ namespace UnrealBuildTool
 			Platform = InDescriptor.Platform;
 			Configuration = InDescriptor.Configuration;
 			Architectures = InDescriptor.Architectures;
-			IntermediateEnvironment = InDescriptor.IntermediateEnvironment != UnrealIntermediateEnvironment.Default ? InDescriptor.IntermediateEnvironment : InRules.IntermediateEnvironment;
+			IntermediateEnvironment = Unreal.IsEngineInstalled() ? UnrealIntermediateEnvironment.Default : InDescriptor.IntermediateEnvironment != UnrealIntermediateEnvironment.Default ? InDescriptor.IntermediateEnvironment : InRules.IntermediateEnvironment;
 			Rules = InRules;
 			RulesAssembly = InRulesAssembly;
 			TargetType = Rules.Type;
