PR #10682: When linking with clang pass -gcodeview-ghash with -Xclang

diff --git a/Engine/Source/Programs/UnrealBuildTool/Platform/Windows/VCToolChain.cs b/Engine/Source/Programs/UnrealBuildTool/Platform/Windows/VCToolChain.cs
--- a/Engine/Source/Programs/UnrealBuildTool/Platform/Windows/VCToolChain.cs
+++ b/Engine/Source/Programs/UnrealBuildTool/Platform/Windows/VCToolChain.cs
@@ -1074,7 +1074,7 @@ namespace UnrealBuildTool
 				// Enable codeview ghash for faster lld links
 				if (Target.WindowsPlatform.Compiler == WindowsCompiler.Clang && Target.WindowsPlatform.bAllowClangLinker)
 				{
-					Arguments.Add("-gcodeview-ghash");
+					Arguments.Add("-Xclang -gcodeview-ghash");
 				}
 
 				// Disable specific warnings that cause problems with Clang
