PR #10683: Fix linking Shipping with `lld-link`

diff --git a/Engine/Source/Programs/UnrealBuildTool/Platform/Windows/VCToolChain.cs b/Engine/Source/Programs/UnrealBuildTool/Platform/Windows/VCToolChain.cs
--- a/Engine/Source/Programs/UnrealBuildTool/Platform/Windows/VCToolChain.cs
+++ b/Engine/Source/Programs/UnrealBuildTool/Platform/Windows/VCToolChain.cs
@@ -1304,8 +1304,12 @@ namespace UnrealBuildTool
 			//
 			if (LinkEnvironment.Configuration == CppConfiguration.Shipping)
 			{
-				// Generate an EXE checksum.
-				Arguments.Add("/RELEASE");
+   				if (!Target.WindowsPlatform.Compiler.IsClang() || !Target.WindowsPlatform.bAllowClangLinker)
+	   			{
+					// Generate an EXE checksum.
+	 				// Not supported by lld-link
+					Arguments.Add("/RELEASE");
+				}
 			}
 
 			// Eliminate unreferenced symbols.
