PR #10781: Apply buildconfiguration.xml to SNDBS class

diff --git a/Engine/Source/Programs/UnrealBuildTool/Executors/SNDBS.cs b/Engine/Source/Programs/UnrealBuildTool/Executors/SNDBS.cs
--- a/Engine/Source/Programs/UnrealBuildTool/Executors/SNDBS.cs
+++ b/Engine/Source/Programs/UnrealBuildTool/Executors/SNDBS.cs
@@ -76,6 +76,7 @@ namespace UnrealBuildTool
 		public SNDBS(List<TargetDescriptor> InTargetDescriptors, ILogger Logger)
 			: base(Logger)
 		{
+			XmlConfig.ApplyTo(this);
 			TargetDescriptors = InTargetDescriptors;
 		}
 
