PR #10774: Add target rule for ALLOW_CONSOLE_IN_SHIPPING define

diff --git a/Engine/Source/Programs/UnrealBuildTool/Configuration/TargetRules.cs b/Engine/Source/Programs/UnrealBuildTool/Configuration/TargetRules.cs
--- a/Engine/Source/Programs/UnrealBuildTool/Configuration/TargetRules.cs
+++ b/Engine/Source/Programs/UnrealBuildTool/Configuration/TargetRules.cs
@@ -1383,6 +1383,12 @@ namespace UnrealBuildTool
 		[RequiresUniqueBuildEnvironment]
 		public bool bUseLoggingInShipping = false;
 
+		/// <summary>
+		/// Whether to turn on console for shipping builds.
+		/// </summary>
+		[RequiresUniqueBuildEnvironment]
+		public bool bUseConsoleInShipping = false;
+
 		/// <summary>
 		/// Whether to turn on logging to memory for test/shipping builds.
 		/// </summary>
@@ -3301,6 +3307,8 @@ namespace UnrealBuildTool
 
 		public bool bUseLoggingInShipping => Inner.bUseLoggingInShipping;
 
+		public bool bUseConsoleInShipping => Inner.bUseConsoleInShipping;
+
 		public bool bLoggingToMemoryEnabled => Inner.bLoggingToMemoryEnabled;
 
 		public bool bUseLauncherChecks => Inner.bUseLauncherChecks;
diff --git a/Engine/Source/Programs/UnrealBuildTool/Configuration/UEBuildTarget.cs b/Engine/Source/Programs/UnrealBuildTool/Configuration/UEBuildTarget.cs
--- a/Engine/Source/Programs/UnrealBuildTool/Configuration/UEBuildTarget.cs
+++ b/Engine/Source/Programs/UnrealBuildTool/Configuration/UEBuildTarget.cs
@@ -5043,6 +5043,15 @@ namespace UnrealBuildTool
 				GlobalCompileEnvironment.Definitions.Add("USE_LOGGING_IN_SHIPPING=0");
 			}
 
+			if (Rules.bUseConsoleInShipping)
+			{
+				GlobalCompileEnvironment.Definitions.Add("ALLOW_CONSOLE_IN_SHIPPING=1");
+			}
+			else
+			{
+				GlobalCompileEnvironment.Definitions.Add("ALLOW_CONSOLE_IN_SHIPPING=0");
+			}
+
 			if (Rules.bAllowProfileGPUInTest)
 			{
 				GlobalCompileEnvironment.Definitions.Add("ALLOW_PROFILEGPU_IN_TEST=1");
