[UBT]
* Fixed bug where cl.exe weight was changed from 1.5 to 1.0 for unity builds.. where it should have been the other way around. For non-unity we want 1.0

diff --git a/Engine/Source/Programs/UnrealBuildTool/Platform/Windows/VCToolChain.cs b/Engine/Source/Programs/UnrealBuildTool/Platform/Windows/VCToolChain.cs
--- a/Engine/Source/Programs/UnrealBuildTool/Platform/Windows/VCToolChain.cs
+++ b/Engine/Source/Programs/UnrealBuildTool/Platform/Windows/VCToolChain.cs
@@ -1464,7 +1464,8 @@ namespace UnrealBuildTool
 			if (Target.WindowsPlatform.Compiler.IsMSVC())
 			{
 				// If deterministic is enabled, MSVC does not use multiple threads
-				if (!CompileEnvironment.bDeterministic && !CompileEnvironment.bUseUnity)
+				// If we are building without unity files the multithread part balances out with the start/exit of all the actions
+				if (!CompileEnvironment.bDeterministic && CompileEnvironment.bUseUnity)
 				{
 					BaseCompileAction.Weight = Target.MSVCCompileActionWeight;
 				}
