[UBT]
* Fixed so CompileEnvironment bUseUnity is set properly for each module

diff --git a/Engine/Source/Programs/UnrealBuildTool/Configuration/UEBuildModuleCPP.cs b/Engine/Source/Programs/UnrealBuildTool/Configuration/UEBuildModuleCPP.cs
--- a/Engine/Source/Programs/UnrealBuildTool/Configuration/UEBuildModuleCPP.cs
+++ b/Engine/Source/Programs/UnrealBuildTool/Configuration/UEBuildModuleCPP.cs
@@ -637,6 +637,9 @@ namespace UnrealBuildTool
 				Logger.LogTrace("Module '{ModuleName}' not using unity build mode", Name);
 			}
 
+			CompileEnvironment.bUseUnity = bModuleUsesUnityBuild;
+			GeneratedCPPCompileEnvironment.bUseUnity = bModuleUsesUnityBuild;
+
 			// Compile Generated CPP Files
 			if (bModuleUsesUnityBuild)
 			{
