[UBT]
* Changed so non unity file cl.exe is weight ends up in the middle of 1 core and 1.5 core... seems like setting it to 1.0 is a bit too aggressive and saturates the machine too much

diff --git a/Engine/Source/Programs/UnrealBuildTool/Platform/Windows/VCToolChain.cs b/Engine/Source/Programs/UnrealBuildTool/Platform/Windows/VCToolChain.cs
--- a/Engine/Source/Programs/UnrealBuildTool/Platform/Windows/VCToolChain.cs
+++ b/Engine/Source/Programs/UnrealBuildTool/Platform/Windows/VCToolChain.cs
@@ -1464,10 +1464,13 @@ namespace UnrealBuildTool
 			if (Target.WindowsPlatform.Compiler.IsMSVC())
 			{
 				// If deterministic is enabled, MSVC does not use multiple threads
-				// If we are building without unity files the multithread part balances out with the start/exit of all the actions
-				if (!CompileEnvironment.bDeterministic && CompileEnvironment.bUseUnity)
+				if (!CompileEnvironment.bDeterministic)
 				{
 					BaseCompileAction.Weight = Target.MSVCCompileActionWeight;
+
+					// If we are building without unity files the multithread part balances out with the start/exit of all the actions
+					if (!CompileEnvironment.bUseUnity)
+						BaseCompileAction.Weight = 1.0f + (BaseCompileAction.Weight - 1.0f) * 0.5f;
 				}
 			}
 			else if (Target.WindowsPlatform.Compiler.IsClang())
