UnrealBuildTool: Don't mark additional buildproducts as executables, only the main target outputs should be marked as such

#rnx

diff --git a/Engine/Source/Programs/UnrealBuildTool/Configuration/UEBuildTarget.cs b/Engine/Source/Programs/UnrealBuildTool/Configuration/UEBuildTarget.cs
--- a/Engine/Source/Programs/UnrealBuildTool/Configuration/UEBuildTarget.cs
+++ b/Engine/Source/Programs/UnrealBuildTool/Configuration/UEBuildTarget.cs
@@ -2746,7 +2746,8 @@ namespace UnrealBuildTool
 					}
 					else if (ExecutableExt.Equals(Ext, StringComparison.OrdinalIgnoreCase))
 					{
-						ProductType = BuildProductType.Executable;
+						// Should not be marked as Executable as that enum value is used for the primary target
+						ProductType = BuildProductType.RequiredResource;
 					}
 					else if (DebugInfoExtensions.Contains(Ext, StringComparer.OrdinalIgnoreCase))
 					{
