UnrealBuildTool: C Standard doesn't matter for PCH variants as C files don't use PCH when compiled

#rnx

diff --git a/Engine/Source/Programs/UnrealBuildTool/Configuration/UEBuildModuleCPP.cs b/Engine/Source/Programs/UnrealBuildTool/Configuration/UEBuildModuleCPP.cs
--- a/Engine/Source/Programs/UnrealBuildTool/Configuration/UEBuildModuleCPP.cs
+++ b/Engine/Source/Programs/UnrealBuildTool/Configuration/UEBuildModuleCPP.cs
@@ -29,7 +29,9 @@ namespace UnrealBuildTool
 			public readonly List<FileItem> CPPFiles = new List<FileItem>();
 			public readonly List<FileItem> CFiles = new List<FileItem>();
 			public readonly List<FileItem> CCFiles = new List<FileItem>();
+			public readonly List<FileItem> MFiles = new List<FileItem>();
 			public readonly List<FileItem> MMFiles = new List<FileItem>();
+			public readonly List<FileItem> SwiftFiles = new List<FileItem>();
 			public readonly List<FileItem> RCFiles = new List<FileItem>();
 			public readonly List<FileItem> ISPCFiles = new List<FileItem>();
 		}
@@ -695,13 +697,33 @@ namespace UnrealBuildTool
 			// Compile CC files directly.
 			if (InputFiles.CCFiles.Count > 0)
 			{
-				LinkInputFiles.AddRange(ToolChain.CompileAllCPPFiles(CompileEnvironment, InputFiles.CCFiles, IntermediateDirectory, Name, Graph).ObjectFiles);
+				CppCompileEnvironment CCCompileEnvironment = new CppCompileEnvironment(CompileEnvironment);
+				CreateHeaderForDefinitions(CCCompileEnvironment, IntermediateDirectory, "cc", Graph);
+				LinkInputFiles.AddRange(ToolChain.CompileAllCPPFiles(CCCompileEnvironment, InputFiles.CCFiles, IntermediateDirectory, Name, Graph).ObjectFiles);
+			}
+
+			// Compile M files directly. Do not use a PCH here, because a C++ PCH is not compatible with C source files.
+			if (InputFiles.MFiles.Count > 0)
+			{
+				CppCompileEnvironment MCompileEnvironment = new CppCompileEnvironment(ModuleCompileEnvironment);
+				CreateHeaderForDefinitions(MCompileEnvironment, IntermediateDirectory, "m", Graph);
+				LinkInputFiles.AddRange(ToolChain.CompileAllCPPFiles(MCompileEnvironment, InputFiles.MFiles, IntermediateDirectory, Name, Graph).ObjectFiles);
 			}
 
 			// Compile MM files directly.
 			if (InputFiles.MMFiles.Count > 0)
 			{
-				LinkInputFiles.AddRange(ToolChain.CompileAllCPPFiles(CompileEnvironment, InputFiles.MMFiles, IntermediateDirectory, Name, Graph).ObjectFiles);
+				CppCompileEnvironment MMCompileEnvironment = new CppCompileEnvironment(CompileEnvironment);
+				CreateHeaderForDefinitions(MMCompileEnvironment, IntermediateDirectory, "mm", Graph);
+				LinkInputFiles.AddRange(ToolChain.CompileAllCPPFiles(MMCompileEnvironment, InputFiles.MMFiles, IntermediateDirectory, Name, Graph).ObjectFiles);
+			}
+
+			// Compile Swift files directly.
+			if (InputFiles.SwiftFiles.Count > 0)
+			{
+				CppCompileEnvironment SwiftCompileEnvironment = new CppCompileEnvironment(CompileEnvironment);
+				CreateHeaderForDefinitions(SwiftCompileEnvironment, IntermediateDirectory, "swift", Graph);
+				LinkInputFiles.AddRange(ToolChain.CompileAllCPPFiles(SwiftCompileEnvironment, InputFiles.SwiftFiles, IntermediateDirectory, Name, Graph).ObjectFiles);
 			}
 
 			// Compile RC files. The resource compiler does not work with response files, and using the regular compile environment can easily result in the 
@@ -1222,11 +1244,6 @@ namespace UnrealBuildTool
 				return false;
 			}
 
-			if (ModuleCompileEnvironment.CStandard != CompileEnvironment.CStandard)
-			{
-				return false;
-			}
-
 			if (ModuleCompileEnvironment.IncludeOrderVersion != CompileEnvironment.IncludeOrderVersion)
 			{
 				return false;
@@ -1314,19 +1331,6 @@ namespace UnrealBuildTool
 				case CppStandardVersion.Latest: Variant += ".CppLatest"; break;
 			}
 
-			if (CompileEnvironment.CStandard != BaseCompileEnvironment.CStandard)
-			{
-				switch (CompileEnvironment.CStandard)
-				{
-					case CStandardVersion.None: Variant += ".CNone"; break;
-					case CStandardVersion.C89: Variant += ".C89"; break;
-					case CStandardVersion.C99: Variant += ".C99"; break;
-					case CStandardVersion.C11: Variant += ".C11"; break;
-					case CStandardVersion.C17: Variant += ".C17"; break;
-					case CStandardVersion.Latest: Variant += ".CLatest"; break;
-				}
-			}
-
 			if (CompileEnvironment.IncludeOrderVersion != BaseCompileEnvironment.IncludeOrderVersion)
 			{
 				if (CompileEnvironment.IncludeOrderVersion != EngineIncludeOrderVersion.Latest)
@@ -1351,7 +1355,6 @@ namespace UnrealBuildTool
 			CompileEnvironment.bEnableExceptions = ModuleCompileEnvironment.bEnableExceptions;
 			CompileEnvironment.CppStandardEngine = ModuleCompileEnvironment.CppStandardEngine;
 			CompileEnvironment.CppStandard = ModuleCompileEnvironment.CppStandard;
-			CompileEnvironment.CStandard = ModuleCompileEnvironment.CStandard;
 			CompileEnvironment.IncludeOrderVersion = ModuleCompileEnvironment.IncludeOrderVersion;
 			CompileEnvironment.bUseAutoRTFMCompiler = ModuleCompileEnvironment.bUseAutoRTFMCompiler;
 			CompileEnvironment.bAllowAutoRTFMInstrumentation = ModuleCompileEnvironment.bAllowAutoRTFMInstrumentation;
@@ -2276,11 +2279,21 @@ namespace UnrealBuildTool
 					SourceFiles.Add(InputFile);
 					InputFiles.CCFiles.Add(InputFile);
 				}
-				else if (InputFile.HasExtension(".m") || InputFile.HasExtension(".mm") || InputFile.HasExtension(".swift"))
+				else if (InputFile.HasExtension(".m"))
+				{
+					SourceFiles.Add(InputFile);
+					InputFiles.MFiles.Add(InputFile);
+				}
+				else if (InputFile.HasExtension(".mm"))
 				{
 					SourceFiles.Add(InputFile);
 					InputFiles.MMFiles.Add(InputFile);
 				}
+				else if (InputFile.HasExtension(".swift"))
+				{
+					SourceFiles.Add(InputFile);
+					InputFiles.SwiftFiles.Add(InputFile);
+				}
 				else if (InputFile.HasExtension(".rc"))
 				{
 					SourceFiles.Add(InputFile);
