UnrealBuildTool: Show programs in generated project if under a uproject path

#rnx

diff --git a/Engine/Source/Programs/UnrealBuildTool/ProjectFiles/ProjectFileGenerator.cs b/Engine/Source/Programs/UnrealBuildTool/ProjectFiles/ProjectFileGenerator.cs
--- a/Engine/Source/Programs/UnrealBuildTool/ProjectFiles/ProjectFileGenerator.cs
+++ b/Engine/Source/Programs/UnrealBuildTool/ProjectFiles/ProjectFileGenerator.cs
@@ -433,6 +433,8 @@ namespace UnrealBuildTool
 		/// </summary>
 		public readonly FileReference? OnlyGameProject;
 
+		readonly ProjectDescriptor? OnlyGameProjectDescriptor;
+
 		/// <summary>
 		/// File extension for project files we'll be generating (e.g. ".vcxproj")
 		/// </summary>
@@ -476,6 +478,7 @@ namespace UnrealBuildTool
 			XmlConfig.ApplyTo(this);
 
 			RootFolder = new PrimaryProjectFolder(this, "<Root>");
+			OnlyGameProjectDescriptor = OnlyGameProject != null && FileReference.Exists(OnlyGameProject) ? ProjectDescriptor.FromFile(OnlyGameProject) : null;
 		}
 
 		/// <summary>
@@ -1353,6 +1356,15 @@ namespace UnrealBuildTool
 				return false;
 			}
 
+			if (TargetFile.IsUnderDirectory(OnlyGameProject.Directory))
+			{
+				return true;
+			}
+			if (OnlyGameProjectDescriptor != null && OnlyGameProjectDescriptor.AdditionalRootDirectories.Any(x => TargetFile.IsUnderDirectory(x)))
+			{
+				return true;
+			}
+
 			// programs have the Target.cs name match the .uprojet name (for the rare program with a .uproject)
 			string TargetName = TargetFile.GetFileNameWithoutAnyExtensions();
 			string UProjectName = OnlyGameProject.GetFileNameWithoutAnyExtensions();
@@ -2636,6 +2648,7 @@ namespace UnrealBuildTool
 				// Check to see if this is an Engine target.  That is, the target is located under the "Engine" folder
 				bool IsEngineTarget = false;
 				bool WantProjectFileForTarget = true;
+				bool ForceProgramInProject = false;
 				if (TargetFilePath.IsUnderDirectory(Unreal.EngineDirectory))
 				{
 					// This is an engine target
@@ -2645,6 +2658,7 @@ namespace UnrealBuildTool
 					if (DoesProgramMatchOnlyGameProject(TargetFilePath))
 					{
 						WantProjectFileForTarget = true;
+						ForceProgramInProject = true;
 					}
 					else if (Unreal.GetExtensionDirs(Unreal.EngineDirectory, "Source/Programs").Any(x => TargetFilePath.IsUnderDirectory(x)))
 					{
@@ -2700,7 +2714,7 @@ namespace UnrealBuildTool
 					string ProjectFileNameBase;
 					if (TargetRulesObject.Type == TargetType.Program)
 					{
-						if (!bIncludeEnginePrograms && IsEngineTarget)
+						if (!ForceProgramInProject && !bIncludeEnginePrograms && IsEngineTarget)
 						{
 							continue;
 						}
