UnrealBuildTool: Fix vs output path when TargetRules.OutputFile is set

#rnx

diff --git a/Engine/Source/Programs/UnrealBuildTool/ProjectFiles/VisualStudio/VCProject.cs b/Engine/Source/Programs/UnrealBuildTool/ProjectFiles/VisualStudio/VCProject.cs
--- a/Engine/Source/Programs/UnrealBuildTool/ProjectFiles/VisualStudio/VCProject.cs
+++ b/Engine/Source/Programs/UnrealBuildTool/ProjectFiles/VisualStudio/VCProject.cs
@@ -2183,6 +2183,11 @@ namespace UnrealBuildTool
 					}
 					NMakePath += BuildPlatform.GetBinaryExtension(UEBuildBinaryType.Executable);
 
+					if (TargetRulesObject.OutputFile != null)
+					{
+						NMakePath = FileReference.Combine(RootOutputDirectory, TargetRulesObject.OutputFile);
+					}
+
 					VCProjectFileContent.AppendLine("  <PropertyGroup {0}>", ConditionString);
 
 					if (ProjGenerator != null)
