UnrealBuildTool: Update action description when preprocessing

#rnx

diff --git a/Engine/Source/Programs/UnrealBuildTool/Platform/Windows/VCCompileAction.cs b/Engine/Source/Programs/UnrealBuildTool/Platform/Windows/VCCompileAction.cs
--- a/Engine/Source/Programs/UnrealBuildTool/Platform/Windows/VCCompileAction.cs
+++ b/Engine/Source/Programs/UnrealBuildTool/Platform/Windows/VCCompileAction.cs
@@ -202,7 +202,21 @@ namespace UnrealBuildTool
 
 		IEnumerable<FileItem> IExternalAction.DeleteItems => DeleteItems;
 		public DirectoryReference WorkingDirectory => Unreal.EngineSourceDirectory;
-		string IExternalAction.CommandDescription => $"{(bIsAnalyzing ? "Analyze" : "Compile")} [{Architecture}]";
+		string IExternalAction.CommandDescription
+		{
+			get
+			{
+				if (PreprocessedFile != null)
+				{
+					return $"Preprocess [{Architecture}]";
+				}
+				else if (bIsAnalyzing)
+				{
+					return $"Analyze [{Architecture}]";
+				}
+				return $"Compile [{Architecture}]";
+			}
+		}
 		bool IExternalAction.bIsGCCCompiler => false;
 		bool IExternalAction.bProducesImportLibrary => false;
 		string IExternalAction.StatusDescription => (SourceFile == null) ? "Compiling" : SourceFile.Location.GetFileName();
@@ -509,7 +523,10 @@ namespace UnrealBuildTool
 				VCToolChain.AddPreprocessedFile(Arguments, PreprocessedFile, Logger);
 
 				// this is parsed by external tools wishing to open this file directly.
-				Logger.LogInformation("PreProcessPath: {File}", PreprocessedFile);
+				if (!bIsAnalyzing)
+				{
+					Logger.LogInformation("PreProcessPath: {File}", PreprocessedFile);
+				}
 			}
 
 			if (ObjectFile != null)
diff --git a/Engine/Source/Programs/UnrealBuildTool/Platform/Windows/VCToolChain.cs b/Engine/Source/Programs/UnrealBuildTool/Platform/Windows/VCToolChain.cs
--- a/Engine/Source/Programs/UnrealBuildTool/Platform/Windows/VCToolChain.cs
+++ b/Engine/Source/Programs/UnrealBuildTool/Platform/Windows/VCToolChain.cs
@@ -269,9 +269,6 @@ namespace UnrealBuildTool
 			Arguments.Add("/P"); // Preprocess
 			Arguments.Add("/C"); // Preserve comments when preprocessing
 			Arguments.Add($"/Fi\"{PreprocessedFileString}\""); // Preprocess to a file
-
-			// this is parsed by external tools wishing to open this file directly.
-			Logger.LogInformation("PreProcessPath: {Path}", PreprocessedFile);
 		}
 
 		public static void AddObjectFile(List<string> Arguments, FileItem ObjectFile)
