Update support for PVS-Studio to 7.27.75620.507 (x64)

#rnx

diff --git a/Engine/Source/Programs/UnrealBuildTool/Platform/Windows/PVSToolChain.cs b/Engine/Source/Programs/UnrealBuildTool/Platform/Windows/PVSToolChain.cs
--- a/Engine/Source/Programs/UnrealBuildTool/Platform/Windows/PVSToolChain.cs
+++ b/Engine/Source/Programs/UnrealBuildTool/Platform/Windows/PVSToolChain.cs
@@ -762,6 +762,11 @@ namespace UnrealBuildTool
 				{
 					ConfigFileContents.Append("report-disabled-rules=yes\n");
 				}
+				
+				if (SourceFileItem.Location.IsUnderDirectory(Unreal.RootDirectory))
+				{
+					ConfigFileContents.AppendFormat("errors-off=V1102\n");
+				}
 
 				int Timeout = (int)(Settings.AnalysisTimeoutFlag == AnalysisTimeoutFlags.No_timeout ? 0 : Settings.AnalysisTimeoutFlag);
 				ConfigFileContents.AppendFormat("timeout={0}\n", Timeout);
@@ -795,7 +800,9 @@ namespace UnrealBuildTool
 				AnalyzeAction.PrerequisiteItems.UnionWith(InputFiles); // Add the InputFiles as PrerequisiteItems so that in SingleFileCompile mode the PVSAnalyze step is not filtered out
 				AnalyzeAction.ProducedItems.Add(OutputFileItem);
 				AnalyzeAction.DeleteItems.Add(OutputFileItem); // PVS Studio will append by default, so need to delete produced items
-				AnalyzeAction.bCanExecuteRemotely = false;
+				AnalyzeAction.bCanExecuteRemotely = true;
+				AnalyzeAction.bCanExecuteRemotelyWithXGE = false;
+				AnalyzeAction.bCanExecuteRemotelyWithSNDBS = false;
 
 				Result.ObjectFiles.AddRange(AnalyzeAction.ProducedItems);
 			}
