Add option to update .vcxproj.user files during project generation

diff --git a/Engine/Source/Programs/UnrealBuildTool/ProjectFiles/ProjectFileGenerator.cs b/Engine/Source/Programs/UnrealBuildTool/ProjectFiles/ProjectFileGenerator.cs
--- a/Engine/Source/Programs/UnrealBuildTool/ProjectFiles/ProjectFileGenerator.cs
+++ b/Engine/Source/Programs/UnrealBuildTool/ProjectFiles/ProjectFileGenerator.cs
@@ -215,6 +215,11 @@ namespace UnrealBuildTool
 		/// </summary>
 		public bool bGeneratingTemporaryProjects = false;
 
+		/// <summary>
+		/// True will override any extra settings files like .vcxproj.user
+		/// </summary>
+		public static bool bForceUpdateAllFiles = false;
+
 		/// <summary>
 		/// Optional list of platforms to generate projects for
 		/// </summary>
@@ -1558,6 +1563,10 @@ namespace UnrealBuildTool
 							case "-INCLUDETEMPTARGETS":
 								bIncludeTempTargets = true;
 								break;
+							
+							case "-FORCEUPDATEALL":
+								bForceUpdateAllFiles = true;
+								break;
 						}
 					}
 				}
diff --git a/Engine/Source/Programs/UnrealBuildTool/ProjectFiles/VisualStudio/VCProject.cs b/Engine/Source/Programs/UnrealBuildTool/ProjectFiles/VisualStudio/VCProject.cs
--- a/Engine/Source/Programs/UnrealBuildTool/ProjectFiles/VisualStudio/VCProject.cs
+++ b/Engine/Source/Programs/UnrealBuildTool/ProjectFiles/VisualStudio/VCProject.cs
@@ -1771,7 +1771,7 @@ namespace UnrealBuildTool
 				// Create a path to the project file's user file
 				string VCUserFilePath = ProjectFilePath.FullName + ".user";
 				// Never overwrite the existing user path as it will cause them to lose their settings
-				if (File.Exists(VCUserFilePath) == false)
+				if (ProjectFileGenerator.bForceUpdateAllFiles || File.Exists(VCUserFilePath) == false)
 				{
 					bSuccess = ProjectFileGenerator.WriteFileIfChanged(VCUserFilePath, VCUserFileContent.ToString(), Logger);
 				}
