Don't bundle libUnreal.so in .apk for faster iteration

diff --git a/Engine/Build/Android/Java/gradle/.idea/runConfigurations/app.xml b/Engine/Build/Android/Java/gradle/.idea/runConfigurations/app.xml
--- a/Engine/Build/Android/Java/gradle/.idea/runConfigurations/app.xml
+++ b/Engine/Build/Android/Java/gradle/.idea/runConfigurations/app.xml
@@ -2,7 +2,7 @@
 	<configuration default="false" name="app" type="AndroidRunConfigurationType" factoryName="Android App">
 		<module name="app.app.main" />
 		<option name="DEBUGGER_TYPE" value="Hybrid" />
-		<Hybrid>
+		<Hybrid>${IDEA_RUN_CONFIGURATION_SYMBOL_PATHS}
 			<startup_commands startup_commands_attr="command script import &quot;${PY_VISUALIZER_PATH}&quot;" />
 		</Hybrid>
 	</configuration>
diff --git a/Engine/Build/Android/Java/src/com/epicgames/unreal/GameActivity.java.template b/Engine/Build/Android/Java/src/com/epicgames/unreal/GameActivity.java.template
--- a/Engine/Build/Android/Java/src/com/epicgames/unreal/GameActivity.java.template
+++ b/Engine/Build/Android/Java/src/com/epicgames/unreal/GameActivity.java.template
@@ -7032,7 +7032,7 @@ public class GameActivity extends $${gameActivitySuperClass}$$ implements Surfac
 
 	static
 	{
-		System.loadLibrary("Unreal");
+//$${soLoadUnreal}$$
 //$${soLoadLibrary}$$
 	}
 
diff --git a/Engine/Build/Android/Java/src/com/epicgames/unreal/GameApplication.java.template b/Engine/Build/Android/Java/src/com/epicgames/unreal/GameApplication.java.template
--- a/Engine/Build/Android/Java/src/com/epicgames/unreal/GameApplication.java.template
+++ b/Engine/Build/Android/Java/src/com/epicgames/unreal/GameApplication.java.template
@@ -16,9 +16,14 @@ public class GameApplication extends Application implements LifecycleObserver {
 
 	private static boolean isForeground = false;
 
+	private static Context context;
+
+//$${gameApplicationClassAdditions}$$
+
 	@Override
 	public void onCreate() {
 		super.onCreate();
+		GameApplication.context = getApplicationContext();
 //$${gameApplicationOnCreateAdditions}$$
 
 		ProcessLifecycleOwner.get().getLifecycle().addObserver(this);
@@ -26,6 +31,10 @@ public class GameApplication extends Application implements LifecycleObserver {
 		NetworkChangedManager.getInstance().initNetworkCallback(this);
 	}
 
+	public static Context getAppContext() {
+		return GameApplication.context;
+	}
+
 	@Override
 	public void attachBaseContext(Context base) {
 		super.attachBaseContext(base);
@@ -47,7 +56,7 @@ public class GameApplication extends Application implements LifecycleObserver {
 	@Override
 	public void onConfigurationChanged (Configuration newConfig) {
 		super.onConfigurationChanged(newConfig);
-//$${gameApplicationOnConfigurationChangedAdditions}$$		
+//$${gameApplicationOnConfigurationChangedAdditions}$$
 	}
 
 	@OnLifecycleEvent(Lifecycle.Event.ON_START)
diff --git a/Engine/Plugins/Runtime/AndroidFileServer/Source/AFSStub/app/src/main/java/com/epicgames/ue4/GameApplication.java b/Engine/Plugins/Runtime/AndroidFileServer/Source/AFSStub/app/src/main/java/com/epicgames/ue4/GameApplication.java
--- a/Engine/Plugins/Runtime/AndroidFileServer/Source/AFSStub/app/src/main/java/com/epicgames/ue4/GameApplication.java
+++ b/Engine/Plugins/Runtime/AndroidFileServer/Source/AFSStub/app/src/main/java/com/epicgames/ue4/GameApplication.java
@@ -14,6 +14,11 @@ public class GameApplication extends Application implements LifecycleObserver {
 
 	private static boolean isForeground = false;
 
+	public static boolean AndroidFileServer_Verify(String Token)
+	{
+		return true;
+	}
+
 	@Override
 	public void onCreate() {
 		super.onCreate();
@@ -23,6 +28,7 @@ public class GameApplication extends Application implements LifecycleObserver {
 
 	}
 
+
 	@Override
 	public void attachBaseContext(Context base) {
 		super.attachBaseContext(base);
diff --git a/Engine/Plugins/Runtime/AndroidFileServer/Source/AFSStub/app/src/main/java/com/epicgames/unreal/GameApplication.java b/Engine/Plugins/Runtime/AndroidFileServer/Source/AFSStub/app/src/main/java/com/epicgames/unreal/GameApplication.java
--- a/Engine/Plugins/Runtime/AndroidFileServer/Source/AFSStub/app/src/main/java/com/epicgames/unreal/GameApplication.java
+++ b/Engine/Plugins/Runtime/AndroidFileServer/Source/AFSStub/app/src/main/java/com/epicgames/unreal/GameApplication.java
@@ -14,6 +14,11 @@ public class GameApplication extends Application implements LifecycleObserver {
 
 	private static boolean isForeground = false;
 
+	public static boolean AndroidFileServer_Verify(String Token)
+	{
+		return true;
+	}
+
 	@Override
 	public void onCreate() {
 		super.onCreate();
diff --git a/Engine/Plugins/Runtime/AndroidFileServer/Source/AndroidFileServer/Android/java/com/epicgames/unreal/RemoteFileManagerActivity.java b/Engine/Plugins/Runtime/AndroidFileServer/Source/AndroidFileServer/Android/java/com/epicgames/unreal/RemoteFileManagerActivity.java
--- a/Engine/Plugins/Runtime/AndroidFileServer/Source/AndroidFileServer/Android/java/com/epicgames/unreal/RemoteFileManagerActivity.java
+++ b/Engine/Plugins/Runtime/AndroidFileServer/Source/AndroidFileServer/Android/java/com/epicgames/unreal/RemoteFileManagerActivity.java
@@ -41,7 +41,7 @@ public class RemoteFileManagerActivity extends Activity {
 			{
 				stopAllThreads();
 				stopService(new Intent(this, RemoteFileManagerService.class));
-			} else 	if (cmd != null && cmd.equals("start") && GameActivity.AndroidFileServer_Verify(token)) {
+			} else 	if (cmd != null && cmd.equals("start") && GameApplication.AndroidFileServer_Verify(token)) {
 
 				// make a service intent with same extras
 				Intent serviceIntent = new Intent(this, RemoteFileManagerService.class);
diff --git a/Engine/Plugins/Runtime/AndroidFileServer/Source/AndroidFileServer/Android/java/com/epicgames/unreal/RemoteFileManagerService.java b/Engine/Plugins/Runtime/AndroidFileServer/Source/AndroidFileServer/Android/java/com/epicgames/unreal/RemoteFileManagerService.java
--- a/Engine/Plugins/Runtime/AndroidFileServer/Source/AndroidFileServer/Android/java/com/epicgames/unreal/RemoteFileManagerService.java
+++ b/Engine/Plugins/Runtime/AndroidFileServer/Source/AndroidFileServer/Android/java/com/epicgames/unreal/RemoteFileManagerService.java
@@ -75,7 +75,7 @@ public class RemoteFileManagerService extends Service {
 			String cmd = intent.getStringExtra("cmd");
 			String token = intent.getStringExtra("token");
 			int port = intent.getIntExtra("port", SERVER_PORT);
-			if (cmd != null && cmd.equals("start") && GameActivity.AndroidFileServer_Verify(token)) {
+			if (cmd != null && cmd.equals("start") && GameApplication.AndroidFileServer_Verify(token)) {
 				String packageName = getPackageName();
 				PackageManager pm = getPackageManager();
 
diff --git a/Engine/Plugins/Runtime/AndroidFileServer/Source/AndroidFileServer/AndroidFileServer_UPL.xml b/Engine/Plugins/Runtime/AndroidFileServer/Source/AndroidFileServer/AndroidFileServer_UPL.xml
--- a/Engine/Plugins/Runtime/AndroidFileServer/Source/AndroidFileServer/AndroidFileServer_UPL.xml
+++ b/Engine/Plugins/Runtime/AndroidFileServer/Source/AndroidFileServer/AndroidFileServer_UPL.xml
@@ -365,7 +365,7 @@ import com.epicgames.unreal.RemoteFileManager;
 		</if>
 	</gameActivityImportAdditions>
 
-	<gameActivityClassAdditions>
+	<gameApplicationClassAdditions>
 		<if condition="bEmbedAFS">
 			<true>
 				<insert>
@@ -383,7 +383,21 @@ import com.epicgames.unreal.RemoteFileManager;
 				<insertNewline/>
 				<insert>
 	}
+				</insert>
+			</true>
+			<false>
+	public static boolean AndroidFileServer_Verify(String Token)
+	{
+		return false;
+	}
+			</false>
+		</if>
+	</gameApplicationClassAdditions>
 
+	<gameActivityClassAdditions>
+		<if condition="bEmbedAFS">
+			<true>
+				<insert>
 <![CDATA[
 	public int AndroidThunkJava_AndroidFileServer_IsRunning()
 	{
@@ -544,13 +558,8 @@ import com.epicgames.unreal.RemoteFileManager;
 				</insert>
 			</true>
 			<false>
-	public static boolean AndroidFileServer_Verify(String Token)
-	{
-		return false;
-	}
-
 	public boolean AndroidThunkJava_AndroidFileServer_IsRunning()
-	{
+	{ 
 		return false;
 	}
 
diff --git a/Engine/Source/Programs/AutomationTool/Android/AndroidPlatform.Automation.cs b/Engine/Source/Programs/AutomationTool/Android/AndroidPlatform.Automation.cs
--- a/Engine/Source/Programs/AutomationTool/Android/AndroidPlatform.Automation.cs
+++ b/Engine/Source/Programs/AutomationTool/Android/AndroidPlatform.Automation.cs
@@ -520,6 +520,12 @@ public class AndroidPlatform : Platform
 		return Path.Combine(Path.GetDirectoryName(Params.GetProjectExeForPlatform(UnrealTargetPlatform.Android).ToString()), DecoratedExeName) + ".so";
 	}
 
+	private static string GetSOName(ProjectParams Params, string DecoratedExeName, UnrealArch? Architecture)
+	{
+		string ArchName = Architecture == null ? "" : "-" + Architecture.ToString();
+		return Path.Combine(Path.GetDirectoryName(Params.GetProjectExeForPlatform(UnrealTargetPlatform.Android).ToString()), DecoratedExeName) + ArchName + ".so";
+	}
+
 	private static string GetFinalApkName(ProjectParams Params, string DecoratedExeName, bool bRenameUnrealGame, UnrealArch? Architecture)
 	{
 		string ProjectDir = Path.Combine(Path.GetDirectoryName(Path.GetFullPath(Params.RawProjectPath.FullName)), "Binaries/Android");
@@ -1353,7 +1359,7 @@ public class AndroidPlatform : Platform
 					Logger.LogInformation("SavePackageInfo");
 					Deploy.SavePackageInfo(Params.ShortProjectName, SC.ProjectRoot.FullName, Type, true);
 				}
-				Deploy.PrepForUATPackageOrDeploy(Params.RawProjectPath, Params.ShortProjectName, SC.ProjectRoot, SOName, SC.LocalRoot + "/Engine", Params.Distribution, CookFlavor, SC.StageTargets[0].Receipt.Configuration, false, bShouldCompileAsDll);
+				Deploy.PrepForUATPackageOrDeploy(Params.RawProjectPath, Params.ShortProjectName, SC.ProjectRoot, SOName, SC.LocalRoot + "/Engine", Params.Distribution, CookFlavor, SC.StageTargets[0].Receipt.Configuration, false, bShouldCompileAsDll, SC.Archive);
 			}
 
 			// Create APK specific OBB in case we have a detached OBB.
@@ -1555,10 +1561,10 @@ public class AndroidPlatform : Platform
 	private List<string> GenerateInstallBatchFile(bool bPackageDataInsideApk, string PackageName, string ApkName, ProjectParams Params, string ObbName, string DeviceObbName, bool bNoObbInstall,
 		string PatchName, string DevicePatchName, bool bNoPatchInstall, List<OverflowBatchInstallInfo> OverflowInfo,
 		bool bIsPC, bool bIsDistribution, bool bRequireRuntimeStoragePermission, bool bDisablePerfHarden, bool bUseAFS, bool bUseAFSProject, string AFSToken, UnrealTargetPlatform Target)
-    {
+	{
 		List<string> BatchLines = new List<string>();
-        string ReadPermissionGrantCommand = "shell pm grant " + PackageName + " android.permission.READ_EXTERNAL_STORAGE";
-        string WritePermissionGrantCommand = "shell pm grant " + PackageName + " android.permission.WRITE_EXTERNAL_STORAGE";
+		string ReadPermissionGrantCommand = "shell pm grant " + PackageName + " android.permission.READ_EXTERNAL_STORAGE";
+		string WritePermissionGrantCommand = "shell pm grant " + PackageName + " android.permission.WRITE_EXTERNAL_STORAGE";
 		string ForegroundPermissionGrantCommand = "shell pm grant " + PackageName + " android.permission.FOREGROUND_SERVICE";
 		string ForegroundDataSyncPermissionGrantCommand = "shell pm grant " + PackageName + " android.permission.FOREGROUND_SERVICE_DATA_SYNC";
 		string NotificationPermissionGrantCommand = "shell pm grant " + PackageName + " android.permission.POST_NOTIFICATIONS";
@@ -1736,7 +1742,7 @@ public class AndroidPlatform : Platform
 				{
 					foreach (OverflowBatchInstallInfo Overflow in OverflowInfo)
 					{
-						OverflowInstallCommands.Add(Overflow.bNoOverflowInstall  ? "%ADB% %DEVICE% shell rm -r %STORAGE%/" + Overflow.DeviceOverflowName : "%ADB% %DEVICE% push " + Path.GetFileName(Overflow.OverflowName) + " " + TargetAndroidTemp + Overflow.DeviceOverflowName);
+						OverflowInstallCommands.Add(Overflow.bNoOverflowInstall ? "%ADB% %DEVICE% shell rm -r %STORAGE%/" + Overflow.DeviceOverflowName : "%ADB% %DEVICE% push " + Path.GetFileName(Overflow.OverflowName) + " " + TargetAndroidTemp + Overflow.DeviceOverflowName);
 						OverflowInstallCommands.Add("if \"%ERRORLEVEL%\" NEQ \"0\" goto Error");
 					}
 				}
@@ -1776,9 +1782,9 @@ public class AndroidPlatform : Platform
 						!bHavePatch ? "" : (bPackageDataInsideApk ? "" : PatchInstallCommand),
 						!bHavePatch ? "" : (bPackageDataInsideApk ? "" : "if \"%ERRORLEVEL%\" NEQ \"0\" goto Error")});
 
-					BatchLines.AddRange(OverflowInstallCommands);
+			BatchLines.AddRange(OverflowInstallCommands);
 
-					BatchLines.AddRange(new string[] {
+			BatchLines.AddRange(new string[] {
 						bDontMoveOBB ? "" : "%ADB% %DEVICE% shell mkdir %STORAGE%/Android/" + TargetAndroidLocation + PackageName, // don't check for error since installing may create the obb directory
 						bDontMoveOBB ? "" : "%ADB% %DEVICE% shell mv " + TargetAndroidTemp + TargetAndroidLocation + PackageName + " %STORAGE%/Android/" + TargetAndroidLocation,
 						bDontMoveOBB ? "" : "if \"%ERRORLEVEL%\" NEQ \"0\" goto Error",
@@ -1790,8 +1796,8 @@ public class AndroidPlatform : Platform
 						bNeedGrantStoragePermission ? "%ADB% %DEVICE% " + ReadPermissionGrantCommand + NullCmd : "",
 						bNeedGrantStoragePermission ? "%ADB% %DEVICE% " + WritePermissionGrantCommand + NullCmd : "",
 						bDisablePerfHarden ? "%ADB% %DEVICE% " + DisablePerfHardenCommand : "",
-                        "@echo.",
-                        "@echo Installation successful",
+						"@echo.",
+						"@echo Installation successful",
 						"goto:eof",
 						":Error",
 						"@echo.",
@@ -1803,9 +1809,10 @@ public class AndroidPlatform : Platform
 						"@echo Check that the device has an SD card.",
 						"@pause"
 					});
-        }
-        return BatchLines;
-    }
+		}
+		return BatchLines;
+	}
+
 
 	private string[] GenerateUninstallBatchFile(bool bPackageDataInsideApk, string PackageName, string ApkName, ProjectParams Params, bool bIsPC)
 	{
@@ -2443,11 +2450,17 @@ public class AndroidPlatform : Platform
 		}
 	}
 
+	private bool GetDontBundleLibrariesInAPK(ProjectParams Params, DeploymentContext SC, bool bVerbose = false)
+	{
+		return AndroidExports.GetDontBundleLibrariesInAPK(Params.RawProjectPath, null, SC.StageTargets[0].Receipt.Configuration, SC.Archive, false,
+			true, bVerbose ? Logger : null);
+	}
+
 	private void DeployAndroidFileServer(ProjectParams Params, DeploymentContext SC, string AFSToken)
 	{
 		ConfigHierarchy Ini = ConfigCache.ReadHierarchy(ConfigHierarchyType.Engine, DirectoryReference.FromFile(Params.RawProjectPath), UnrealTargetPlatform.Android);
-		bool bDisablePerfHarden;
-		Ini.GetBool("/Script/AndroidRuntimeSettings.AndroidRuntimeSettings", "bEnableMaliPerfCounters", out bDisablePerfHarden);
+		Ini.GetBool("/Script/AndroidRuntimeSettings.AndroidRuntimeSettings", "bEnableMaliPerfCounters", out bool bDisablePerfHarden);
+		bool bDontBundleLibrariesInAPK = GetDontBundleLibrariesInAPK(Params, SC, true);
 
 		bool bUseCompression;
 		bool bLogFiles;
@@ -2463,6 +2476,7 @@ public class AndroidPlatform : Platform
 		{
 			UnrealArch? DeviceArchitecture = GetBestDeviceArchitecture(Params, DeviceName);
 			string ApkName = GetFinalApkName(Params, SC.StageExecutables[0], true, DeviceArchitecture);
+			string FinalSOName = GetSOName(Params, SC.StageExecutables[0], DeviceArchitecture);
 
 			// make sure APK is up to date (this is fast if so)
 			var Deploy = AndroidExports.CreateDeploymentHandler(Params.RawProjectPath, Params.ForcePackageData);
@@ -2471,7 +2485,7 @@ public class AndroidPlatform : Platform
 				string CookFlavor = SC.FinalCookPlatform.IndexOf("_") > 0 ? SC.FinalCookPlatform.Substring(SC.FinalCookPlatform.IndexOf("_")) : "";
 				string SOName = GetSONameWithoutArchitecture(Params, SC.StageExecutables[0]);
 				Deploy.SetAndroidPluginData(GetDeploymentArchitectures(Params, SC), CollectPluginDataPaths(SC));
-				Deploy.PrepForUATPackageOrDeploy(Params.RawProjectPath, Params.ShortProjectName, SC.ProjectRoot, SOName, SC.LocalRoot + "/Engine", Params.Distribution, CookFlavor, SC.StageTargets[0].Receipt.Configuration, true, false);
+				Deploy.PrepForUATPackageOrDeploy(Params.RawProjectPath, Params.ShortProjectName, SC.ProjectRoot, SOName, SC.LocalRoot + "/Engine", Params.Distribution, CookFlavor, SC.StageTargets[0].Receipt.Configuration, true, false, SC.Archive);
 			}
 
 			// now we can use the apk to get more info
@@ -2722,6 +2736,14 @@ public class AndroidPlatform : Platform
 				client.FileWriteString("APK: " + APKLastUpdateTime + "\n", "^ext/APKFileStamp.txt");
 			}
 
+			// always update libUnreal.so
+			// TODO potential optimization not to push it every time but compare filestamp instead to check if we need to update it
+			if (bDontBundleLibrariesInAPK)
+			{
+				string FinalSONameStripped = Path.Combine(Path.GetDirectoryName(FinalSOName), Path.GetFileNameWithoutExtension(FinalSOName) + "-stripped" + Path.GetExtension(FinalSOName));
+				client.PushFile(FinalSONameStripped, "^int/libUnreal.so", true);
+			}
+
 			// update the uecommandline.txt
 			// update and deploy uecommandline.txt
 			// always delete the existing commandline text file, so it doesn't reuse an old one
@@ -3058,14 +3080,15 @@ public class AndroidPlatform : Platform
 	private void DeployADB(ProjectParams Params, DeploymentContext SC)
     {
 		ConfigHierarchy Ini = ConfigCache.ReadHierarchy(ConfigHierarchyType.Engine, DirectoryReference.FromFile(Params.RawProjectPath), UnrealTargetPlatform.Android);
-		bool bDisablePerfHarden = false;
-		Ini.GetBool("/Script/AndroidRuntimeSettings.AndroidRuntimeSettings", "bEnableMaliPerfCounters", out bDisablePerfHarden);
+		Ini.GetBool("/Script/AndroidRuntimeSettings.AndroidRuntimeSettings", "bEnableMaliPerfCounters", out bool bDisablePerfHarden);
+		bool bDontBundleLibrariesInAPK = GetDontBundleLibrariesInAPK(Params, SC, true);
 		int AllowOverflowOBBLimit = AllowOverflowOBBFiles(SC);
 
 		foreach (var DeviceName in Params.DeviceNames)
         {
             UnrealArch? DeviceArchitecture = GetBestDeviceArchitecture(Params, DeviceName);
             string ApkName = GetFinalApkName(Params, SC.StageExecutables[0], true, DeviceArchitecture);
+            string FinalSOName = GetSOName(Params, SC.StageExecutables[0], DeviceArchitecture);
 
             // make sure APK is up to date (this is fast if so)
             var Deploy = AndroidExports.CreateDeploymentHandler(Params.RawProjectPath, Params.ForcePackageData);
@@ -3074,7 +3097,7 @@ public class AndroidPlatform : Platform
                 string CookFlavor = SC.FinalCookPlatform.IndexOf("_") > 0 ? SC.FinalCookPlatform.Substring(SC.FinalCookPlatform.IndexOf("_")) : "";
 				string SOName = GetSONameWithoutArchitecture(Params, SC.StageExecutables[0]);
 				Deploy.SetAndroidPluginData(GetDeploymentArchitectures(Params, SC), CollectPluginDataPaths(SC));
-                Deploy.PrepForUATPackageOrDeploy(Params.RawProjectPath, Params.ShortProjectName, SC.ProjectRoot, SOName, SC.LocalRoot + "/Engine", Params.Distribution, CookFlavor, SC.StageTargets[0].Receipt.Configuration, true, false);
+				Deploy.PrepForUATPackageOrDeploy(Params.RawProjectPath, Params.ShortProjectName, SC.ProjectRoot, SOName, SC.LocalRoot + "/Engine", Params.Distribution, CookFlavor, SC.StageTargets[0].Receipt.Configuration, true, false, SC.Archive);
             }
 
             // now we can use the apk to get more info
@@ -3178,6 +3201,27 @@ public class AndroidPlatform : Platform
                 }
             }
 
+			if (bDontBundleLibrariesInAPK)
+			{
+				int SuccessCode = 0;
+
+				string FinalSOPathStripped = Path.Combine(Path.GetDirectoryName(FinalSOName), Path.GetFileNameWithoutExtension(FinalSOName) + "-stripped" + Path.GetExtension(FinalSOName));
+				string FinalSOFileNameStripped = Path.GetFileName(FinalSOPathStripped);
+
+				string PushSO = $"push -z lz4 {FinalSOPathStripped} /data/local/tmp/{FinalSOFileNameStripped}";
+				string CopySO = $"shell run-as {PackageName} cp /data/local/tmp/{FinalSOFileNameStripped} ./files/libUnreal.so";
+				string DeleteSO = $"shell rm /data/local/tmp/{FinalSOFileNameStripped}";
+				RunAndLogAdbCommand(Params, DeviceName, PushSO, out SuccessCode);
+				RunAndLogAdbCommand(Params, DeviceName, CopySO, out SuccessCode);
+				RunAndLogAdbCommand(Params, DeviceName, DeleteSO, out SuccessCode);
+
+				if (SuccessCode != 0)
+				{
+					string ErrorMessage = $"Installation of '{PackageName}' failed due to failing to push libUnreal.so outside of {ApkName}";
+					throw new AutomationException(ExitCode.Error_AppInstallFailed, ErrorMessage);
+				}
+			}
+
             // update the uecommandline.txt
             // update and deploy uecommandline.txt
             // always delete the existing commandline text file, so it doesn't reuse an old one
diff --git a/Engine/Source/Programs/UnrealBuildTool/Platform/Android/AndroidExports.cs b/Engine/Source/Programs/UnrealBuildTool/Platform/Android/AndroidExports.cs
--- a/Engine/Source/Programs/UnrealBuildTool/Platform/Android/AndroidExports.cs
+++ b/Engine/Source/Programs/UnrealBuildTool/Platform/Android/AndroidExports.cs
@@ -56,8 +56,9 @@ namespace UnrealBuildTool
 		/// <param name="Configuration"></param>
 		/// <param name="bIsDataDeploy"></param>
 		/// <param name="bSkipGradleBuild"></param>
+		/// <param name="bIsArchive"></param>
 		/// <returns></returns>
-		bool PrepForUATPackageOrDeploy(FileReference ProjectFile, string ProjectName, DirectoryReference ProjectDirectory, string ExecutablePath, string EngineDirectory, bool bForDistribution, string CookFlavor, UnrealTargetConfiguration Configuration, bool bIsDataDeploy, bool bSkipGradleBuild);
+		bool PrepForUATPackageOrDeploy(FileReference ProjectFile, string ProjectName, DirectoryReference ProjectDirectory, string ExecutablePath, string EngineDirectory, bool bForDistribution, string CookFlavor, UnrealTargetConfiguration Configuration, bool bIsDataDeploy, bool bSkipGradleBuild, bool bIsArchive);
 
 		/// <summary>
 		/// 
@@ -115,5 +116,20 @@ namespace UnrealBuildTool
 			AndroidToolChain ToolChain = new AndroidToolChain(null, Logger);
 			ToolChain.StripSymbols(SourceFile, TargetFile, Logger);
 		}
+
+		/// <summary>
+		/// 
+		/// </summary>
+		/// <param name="ProjectFile"></param>
+		/// <param name="bForceDontBundleLibrariesInAPK"></param>
+		/// <param name="Configuration"></param>
+		/// <param name="bIsArchive"></param>
+		/// <param name="bFromMSBuild"></param>
+		/// <param name="bIsFromUAT"></param>
+		/// <param name="Logger"></param>
+		public static bool GetDontBundleLibrariesInAPK(FileReference? ProjectFile, bool? bForceDontBundleLibrariesInAPK, UnrealTargetConfiguration Configuration, bool bIsArchive, bool bFromMSBuild, bool bIsFromUAT, ILogger? Logger)
+		{
+			return UEDeployAndroid.GetDontBundleLibrariesInAPK(ProjectFile, bForceDontBundleLibrariesInAPK, Configuration, bIsArchive, bFromMSBuild, bIsFromUAT, Logger);
+		}
 	}
 }
diff --git a/Engine/Source/Programs/UnrealBuildTool/Platform/Android/AndroidPlatformSDK.cs b/Engine/Source/Programs/UnrealBuildTool/Platform/Android/AndroidPlatformSDK.cs
--- a/Engine/Source/Programs/UnrealBuildTool/Platform/Android/AndroidPlatformSDK.cs
+++ b/Engine/Source/Programs/UnrealBuildTool/Platform/Android/AndroidPlatformSDK.cs
@@ -22,7 +22,7 @@ namespace UnrealBuildTool
 		{
 		}
 
-		private string? GetNDKRoot()
+		public static string? GetNDKRoot()
 		{
 			string? NDKPath = Environment.GetEnvironmentVariable("NDKROOT");
 
diff --git a/Engine/Source/Programs/UnrealBuildTool/Platform/Android/AndroidProjectGenerator.cs b/Engine/Source/Programs/UnrealBuildTool/Platform/Android/AndroidProjectGenerator.cs
--- a/Engine/Source/Programs/UnrealBuildTool/Platform/Android/AndroidProjectGenerator.cs
+++ b/Engine/Source/Programs/UnrealBuildTool/Platform/Android/AndroidProjectGenerator.cs
@@ -112,9 +112,22 @@ namespace UnrealBuildTool
 					Path.GetFileNameWithoutExtension(NMakeOutputPath.FullName) + "-arm64.apk");
 
 				ProjectFileBuilder.AppendLine($"    <AndroidApkLocation>{apkLocation}</AndroidApkLocation>");
-				string intermediatePath = Path.GetFullPath(Path.GetDirectoryName(NMakeOutputPath.FullName) + @"\..\..\Intermediate\Android\arm64\");
-				string symbolLocations = $@"{intermediatePath}jni\arm64-v8a;{intermediatePath}libs\arm64-v8a";
-				ProjectFileBuilder.AppendLine($"    <AndroidSymbolDirectories>{symbolLocations}</AndroidSymbolDirectories>");
+				string intermediateRootPath = Path.GetFullPath(Path.GetDirectoryName(NMakeOutputPath.FullName) + @"\..\..\Intermediate\Android\");
+				string intermediatePath = Path.Combine(intermediateRootPath, "arm64");
+				string intermediateAGDESymbolsPath = Path.Combine(intermediateRootPath, "LLDBSymbolsLibs", "arm64");
+				List<string> symbolLocations = new List<string>
+				{
+					$@"{intermediatePath}jni\arm64-v8a",
+					$@"{intermediatePath}libs\arm64-v8a",
+					intermediateAGDESymbolsPath // support bDontBundleLibrariesInAPK
+				};
+				ProjectFileBuilder.AppendLine($"    <AndroidSymbolDirectories>{string.Join(";", symbolLocations)}</AndroidSymbolDirectories>");
+
+				// At this stage we don't know if bDontBundleLibrariesInAPK is enabled or not, so make a fail-safe check.
+				string pushSOScript = Path.Combine(
+					Path.GetDirectoryName(NMakeOutputPath.FullName)!,
+					"Push_" + Path.GetFileNameWithoutExtension(NMakeOutputPath.FullName) + "-arm64_so.bat");
+				ProjectFileBuilder.AppendLine($"    <AndroidPostApkInstallCommands>IF EXIST {pushSOScript} {pushSOScript};$(AndroidPostApkInstallCommands)</AndroidPostApkInstallCommands>");
 			}
 			else
 			{
diff --git a/Engine/Source/Programs/UnrealBuildTool/Platform/Android/UEDeployAndroid.cs b/Engine/Source/Programs/UnrealBuildTool/Platform/Android/UEDeployAndroid.cs
--- a/Engine/Source/Programs/UnrealBuildTool/Platform/Android/UEDeployAndroid.cs
+++ b/Engine/Source/Programs/UnrealBuildTool/Platform/Android/UEDeployAndroid.cs
@@ -74,6 +74,19 @@ namespace UnrealBuildTool
 		[CommandLine("-BypassGradlePackaging", Value = "true")]
 		public bool BypassGradlePackaging = false;
 
+		/// <summary>
+		/// Whether UBT is invoked from MSBuild.
+		/// If false will, disable bDontBundleLibrariesInAPK, unless forced .
+		/// </summary>
+		[CommandLine(Prefix = "-FromMsBuild", Value = "true")]
+		public bool bFromMSBuild = false;
+
+		/// <summary>
+		/// Forcing bDontBundleLibrariesInAPK to "true" or "false" ignoring any other option.
+		/// </summary>
+		[CommandLine("-ForceDontBundleLibrariesInAPK=")]
+		public bool? ForceDontBundleLibrariesInAPK;
+
 		public UEDeployAndroid(FileReference? InProjectFile, bool InForcePackageData, ILogger InLogger)
 			: base(InLogger)
 		{
@@ -1909,7 +1922,7 @@ namespace UnrealBuildTool
 		}
 
 		private string GetAllBuildSettings(AndroidToolChain ToolChain, UnrealPluginLanguage UPL, bool bForDistribution, bool bMakeSeparateApks, bool bPackageDataInsideApk,
-			bool bDisableVerifyOBBOnStartUp, bool bUseExternalFilesDir, string TemplatesHashCode)
+			bool bDisableVerifyOBBOnStartUp, bool bUseExternalFilesDir, bool bDontBundleLibrariesInAPK, string TemplatesHashCode)
 		{
 			// make the settings string - this will be char by char compared against last time
 			StringBuilder CurrentSettings = new StringBuilder();
@@ -1925,6 +1938,7 @@ namespace UnrealBuildTool
 			CurrentSettings.AppendLine(String.Format("bUseExternalFilesDir={0}", bUseExternalFilesDir));
 			CurrentSettings.AppendLine(String.Format("UPLHashCode={0}", UPLHashCode));
 			CurrentSettings.AppendLine(String.Format("TemplatesHashCode={0}", TemplatesHashCode));
+			CurrentSettings.AppendLine(String.Format("bDontBundleLibrariesInAPK={0}", bDontBundleLibrariesInAPK));
 
 			// all AndroidRuntimeSettings ini settings in here
 			ConfigHierarchy Ini = GetConfigCacheIni(ConfigHierarchyType.Engine);
@@ -2050,7 +2064,7 @@ namespace UnrealBuildTool
 		}
 
 		private bool CheckDependencies(UnrealArchitectures Architectures, string ProjectName, string ProjectDirectory, string UnrealBuildFilesPath, string GameBuildFilesPath, string EngineDirectory, List<string> SettingsFiles,
-			string CookFlavor, string OutputPath, bool bMakeSeparateApks, bool bPackageDataInsideApk)
+			string CookFlavor, string OutputPath, bool bMakeSeparateApks, bool bPackageDataInsideApk, bool bDontBundleLibrariesInAPK)
 		{
 			// check all input files (.so, java files, .ini files, etc)
 			bool bAllInputsCurrent = true;
@@ -2069,11 +2083,14 @@ namespace UnrealBuildTool
 
 				// check to see if it's out of date before trying the slow make apk process (look at .so and all Engine and Project build files to be safe)
 				List<String> InputFiles = new List<string>();
-				InputFiles.Add(SourceSOName);
-				InputFiles.AddRange(Directory.EnumerateFiles(UnrealBuildFilesPath, "*.*", SearchOption.AllDirectories));
-				if (Directory.Exists(GameBuildFilesPath))
+				if (!bDontBundleLibrariesInAPK)
 				{
-					InputFiles.AddRange(Directory.EnumerateFiles(GameBuildFilesPath, "*.*", SearchOption.AllDirectories));
+					InputFiles.Add(SourceSOName);
+					InputFiles.AddRange(Directory.EnumerateFiles(UnrealBuildFilesPath, "*.*", SearchOption.AllDirectories));
+					if (Directory.Exists(GameBuildFilesPath))
+					{
+						InputFiles.AddRange(Directory.EnumerateFiles(GameBuildFilesPath, "*.*", SearchOption.AllDirectories));
+					}
 				}
 
 				// make sure changed java files will rebuild apk
@@ -2109,7 +2126,6 @@ namespace UnrealBuildTool
 						{
 							bAllInputsCurrent = false;
 							Logger.LogInformation("{DestApkName} is out of date due to newer input file {InputFileName}", DestApkName, InputFileName);
-							break;
 						}
 					}
 				}
@@ -3955,8 +3971,47 @@ namespace UnrealBuildTool
 			File.WriteAllText(GradleBuildScriptAdditionsFilename, UPL.ProcessPluginNode(NDKArch, "buildscriptGradleAdditions", ""));
 		}
 
+		public static bool GetDontBundleLibrariesInAPK(FileReference? ProjectFile, bool? bForceDontBundleLibrariesInAPK, UnrealTargetConfiguration Configuration, bool bIsArchive, bool bFromMSBuild, bool bIsFromUAT, ILogger? Logger)
+		{
+			if (bForceDontBundleLibrariesInAPK.HasValue)
+			{
+				Logger?.LogInformation("bDontBundleLibrariesInAPK is force set to {bForceDontBundleLibrariesInAPK} via command line option.", bForceDontBundleLibrariesInAPK.Value);
+				return bForceDontBundleLibrariesInAPK.Value;
+			}
+
+			// this feature requires extra step to install the app besides just .apk install
+			// limit where the feature is enabled because we're not in control of deployment step but in VS/AGDE and UAT
+			if (!(bFromMSBuild || bIsFromUAT)) 
+			{
+				Logger?.LogInformation("bDontBundleLibrariesInAPK is set to false, due to not called from MSBuild or UAT (bFromMSBuild={bFromMSBuild}, bIsFromUAT={bIsFromUAT}), use \"-ForceDontBundleLibrariesInAPK=true\" to override.", bFromMSBuild, bIsFromUAT);
+				return false;
+			}
+
+			if (Configuration == UnrealTargetConfiguration.Shipping)
+			{
+				Logger?.LogInformation("bDontBundleLibrariesInAPK is set to false, due to a shipping build.");
+				return false;
+			}
+
+			if (bIsArchive)
+			{
+				Logger?.LogInformation("bDontBundleLibrariesInAPK is set to false, due to an archive build.");
+				return false;
+			}
+
+			ConfigHierarchy Ini = ConfigCache.ReadHierarchy(ConfigHierarchyType.Engine, DirectoryReference.FromFile(ProjectFile), UnrealTargetPlatform.Android);
+			Ini.GetBool("/Script/AndroidRuntimeSettings.AndroidRuntimeSettings", "bDontBundleLibrariesInAPK", out bool bDontBundleLibrariesInAPK);
+			Logger?.LogInformation("bDontBundleLibrariesInAPK is set to {bDontBundleLibrariesInAPK}, based on AndroidRuntimeSettings.", bDontBundleLibrariesInAPK);
+			return bDontBundleLibrariesInAPK;
+		}
+
+		private bool GetDontBundleLibrariesInAPK(UnrealTargetConfiguration Configuration, bool bIsArchive, bool bIsFromUAT, bool bVerbose = false)
+		{
+			return GetDontBundleLibrariesInAPK(ProjectFile, ForceDontBundleLibrariesInAPK, Configuration, bIsArchive, bFromMSBuild, bIsFromUAT, bVerbose ? Logger : null);
+		}
+
 		private void MakeApk(AndroidToolChain ToolChain, string ProjectName, TargetType InTargetType, string ProjectDirectory, string OutputPath, string EngineDirectory, bool bForDistribution, string CookFlavor,
-			UnrealTargetConfiguration Configuration, bool bMakeSeparateApks, bool bIncrementalPackage, bool bDisallowPackagingDataInApk, bool bDisallowExternalFilesDir, bool bSkipGradleBuild)
+			UnrealTargetConfiguration Configuration, bool bMakeSeparateApks, bool bIncrementalPackage, bool bDisallowPackagingDataInApk, bool bDisallowExternalFilesDir, bool bSkipGradleBuild, bool bIsArchive, bool bIsFromUAT)
 		{
 			Logger.LogInformation("");
 			Logger.LogInformation("===={Time}====PREPARING TO MAKE APK=================================================================", DateTime.Now.ToString());
@@ -4024,6 +4079,8 @@ namespace UnrealBuildTool
 
 			LogBuildSetup();
 
+			ConfigHierarchy Ini = GetConfigCacheIni(ConfigHierarchyType.Engine);
+
 			// bundles disabled for launch-on
 			bool bEnableBundle = BundleEnabled() && !bDisallowPackagingDataInApk;
 
@@ -4255,9 +4312,11 @@ namespace UnrealBuildTool
 				}
 			}
 
+			bool bDontBundleLibrariesInAPK = GetDontBundleLibrariesInAPK(Configuration, bIsArchive, bIsFromUAT, bVerbose: !bIsFromUAT);
+
 			// check to see if any "meta information" is newer than last time we build
 			string TemplatesHashCode = GenerateTemplatesHashCode(EngineDirectory);
-			string CurrentBuildSettings = GetAllBuildSettings(ToolChain, UPL!, bForDistribution, bMakeSeparateApks, bPackageDataInsideApk, bDisableVerifyOBBOnStartUp, bUseExternalFilesDir, TemplatesHashCode);
+			string CurrentBuildSettings = GetAllBuildSettings(ToolChain, UPL!, bForDistribution, bMakeSeparateApks, bPackageDataInsideApk, bDisableVerifyOBBOnStartUp, bUseExternalFilesDir, bDontBundleLibrariesInAPK, TemplatesHashCode);
 			string BuildSettingsCacheFile = Path.Combine(IntermediateAndroidPath, "UEBuildSettings.txt");
 
 			// Architecture remapping
@@ -4303,7 +4362,7 @@ namespace UnrealBuildTool
 				JavaFiles.AddRange(from t in templates select t.DestinationFile);
 
 				bBuildSettingsMatch = CheckDependencies(Architectures, ProjectName, ProjectDirectory, UnrealBuildFilesPath, GameBuildFilesPath,
-					EngineDirectory, JavaFiles, CookFlavor, OutputPath, bMakeSeparateApks, bPackageDataInsideApk);
+					EngineDirectory, JavaFiles, CookFlavor, OutputPath, bMakeSeparateApks, bPackageDataInsideApk, bDontBundleLibrariesInAPK);
 
 			}
 
@@ -4343,6 +4402,72 @@ namespace UnrealBuildTool
 							select Tuple.Create(Arch, manifest);
 			}
 
+			List<string> LLDBExtraSymbolsDirectories = new();
+
+			// Need to do stripping here because .apk will stay the same during iteration
+			if (bDontBundleLibrariesInAPK)
+			{
+				foreach (string NDKArch in NDKArches)
+				{
+					UnrealArch Arch = GetUnrealArch(NDKArch);
+
+					string SOName = AndroidToolChain.InlineArchName(OutputPath, Arch);
+					if (!File.Exists(SOName))
+					{
+						Logger.LogWarning("Did not find compiled .so [{SOName}]", SOName);
+					}
+
+					string SONameStripped = Path.Combine(Path.GetDirectoryName(SOName)!, Path.GetFileNameWithoutExtension(SOName) + "-stripped" + Path.GetExtension(SOName));
+					StripDebugSymbols(SOName, SONameStripped, Arch, Logger, true);
+
+					// LLDB (in Android Studio and AGDE) needs to see a file that is matching the filename on the device ("libUnreal.so") to resolve symbols,
+					// otherwise it will try to download libUnreal.so and fail due to adb pull not having permissions to access it.
+					// Newer lldb has "platform.plugin.remote-android.package-name" in lldb to avoid that.
+					string LLDBExtraSymbolsDirectory = Path.Combine(IntermediateAndroidPath, "LLDBSymbolsLibs", Arch.ToString().ToLower());
+					if (!Directory.Exists(LLDBExtraSymbolsDirectory))
+					{
+						Directory.CreateDirectory(LLDBExtraSymbolsDirectory);
+					}
+					
+					string LLDBExtraSymbolsFile = Path.Combine(LLDBExtraSymbolsDirectory, "libUnreal.so");
+					bool LLDBExtraSymbolsFileIsPresent = File.Exists(LLDBExtraSymbolsFile);
+
+					string LLDBExtraSymbolsFileIsSymlinkMarkerFile = LLDBExtraSymbolsFile + ".isSymLink";
+					bool LLDBExtraSymbolsFileIsSymlink = File.Exists(LLDBExtraSymbolsFileIsSymlinkMarkerFile);
+
+					if (LLDBExtraSymbolsFileIsSymlink)
+					{
+						// There is no easy way to check to which location a hardlink is pointing to,
+						// so to ensure we have the correct redirections let's remove the link first.
+						SafeDeleteFile(LLDBExtraSymbolsFile);
+					}
+
+					// try to create symlink again if we already know that file was a symlink or it's our first time and file doesn't exist yet
+					if ((LLDBExtraSymbolsFileIsSymlink || !LLDBExtraSymbolsFileIsPresent) && Utils.TryCreateSymlink(LLDBExtraSymbolsFile, SOName, Logger))
+					{
+						if (!LLDBExtraSymbolsFileIsSymlink)
+						{
+							// it's our first time creating a symlink, so create a marker file as well
+							File.Create(LLDBExtraSymbolsFileIsSymlinkMarkerFile).Close();
+						}
+					}
+					else
+					{
+						// copy the file if we couldn't create a symlink or file exists without the marker 
+						Logger.LogWarning("Failed to create symlink '{Path}' -> '{PathToTarget}', copying the file instead", LLDBExtraSymbolsFile, SOName);
+						CopyIfDifferent(SOName, LLDBExtraSymbolsFile, false, false);
+
+						// delete the marker if it was present
+						if (LLDBExtraSymbolsFileIsSymlink)
+						{
+							SafeDeleteFile(LLDBExtraSymbolsFileIsSymlinkMarkerFile);
+						}
+					}
+
+					LLDBExtraSymbolsDirectories.Add(LLDBExtraSymbolsDirectory);
+				}
+			}
+
 			// Now we have to spin over all the arch/gpu combinations to make sure they all match
 			int BuildListComboTotal = BuildList.Count();
 			if (BuildListComboTotal == 0)
@@ -4412,8 +4537,6 @@ namespace UnrealBuildTool
 
 					Logger.LogInformation("\n===={Time}====PERFORMING FINAL APK PACKAGE OPERATION====={Arch}===========================================", DateTime.Now.ToString(), Arch);
 
-					ConfigHierarchy Ini = GetConfigCacheIni(ConfigHierarchyType.Engine);
-
 					string TempAPK = Path.Combine(DestApkDirectory, "TEMP_" + Path.GetFileName(DestApkName));
 					CopyAPKAndReplaceSO(Logger, DestApkName, TempAPK, StrippedSOName, "lib/" + NDKArch + "/libUnreal.so");
 
@@ -4544,9 +4667,12 @@ namespace UnrealBuildTool
 			}
 
 			// make up a dictionary of strings to replace in xml files (strings.xml)
-			Dictionary<string, string> Replacements = new Dictionary<string, string>();
-			Replacements.Add("${EXECUTABLE_NAME}", ApplicationDisplayName!);
-			Replacements.Add("${PY_VISUALIZER_PATH}", Path.GetFullPath(Path.Combine(EngineDirectory, "Extras", "LLDBDataFormatters", "UEDataFormatters_2ByteChars.py")));
+			Dictionary<string, string> Replacements = new Dictionary<string, string>
+			{
+				{"${EXECUTABLE_NAME}", ApplicationDisplayName!},
+				{"${PY_VISUALIZER_PATH}", Path.GetFullPath(Path.Combine(EngineDirectory, "Extras", "LLDBDataFormatters", "UEDataFormatters_2ByteChars.py"))},
+				{"${IDEA_RUN_CONFIGURATION_SYMBOL_PATHS}", string.Join("", LLDBExtraSymbolsDirectories.Select(x => $"\n\t\t\t<symbol_dirs symbol_path=\"{x}\" />"))}
+			};
 
 			// steps run for each build combination (note: there should only be one GPU in future)
 			foreach (Tuple<UnrealArch, string> build in BuildList)
@@ -4666,7 +4792,7 @@ namespace UnrealBuildTool
 				WriteJavaOBBDataFile(UnrealOBBDataFileName, PackageName, RequiredOBBFiles, CookFlavor, bPackageDataInsideApk, Arch);
 
 				// update GameActivity.java and GameApplication.java if out of date
-				UpdateGameActivity(Arch, NDKArch, EngineDirectory, UnrealBuildPath);
+				UpdateGameActivity(Arch, NDKArch, EngineDirectory, UnrealBuildPath, bDontBundleLibrariesInAPK);
 				UpdateGameApplication(Arch, NDKArch, EngineDirectory, UnrealBuildPath);
 
 				// we don't actually need the SO for the bSkipGradleBuild case
@@ -4681,6 +4807,74 @@ namespace UnrealBuildTool
 						Logger.LogWarning("Did not find compiled .so [{FinalSOName}]", FinalSOName);
 					}
 				}
+				else if (bDontBundleLibrariesInAPK)
+				{
+					FinalSOName = AndroidToolChain.InlineArchName(OutputPath, Arch);
+					string FinalSONameStripped = Path.Combine(Path.GetDirectoryName(FinalSOName)!, Path.GetFileNameWithoutExtension(FinalSOName) + "-stripped" + Path.GetExtension(FinalSOName));
+					if (!File.Exists(FinalSONameStripped))
+					{
+						Logger.LogWarning("Did not find compiled .so [{FinalSONameStripped}]", FinalSONameStripped);
+					}
+
+					string ApkFilename = Path.GetFileNameWithoutExtension(OutputPath).Replace("UnrealGame", ProjectName);
+					DestApkName = Path.Combine(DestApkDirectory, ApkFilename + ".apk");
+
+					// As we are always making seperate APKs we need to put the architecture into the name
+					DestApkName = AndroidToolChain.InlineArchName(DestApkName, Arch);
+
+					bool bUseAFS = true;
+					if (!Ini.GetBool("/Script/AndroidFileServerEditor.AndroidFileServerRuntimeSettings", "bEnablePlugin", out bUseAFS))
+					{
+						bUseAFS = true;
+					}
+
+					Ini.GetString("/Script/AndroidFileServerEditor.AndroidFileServerRuntimeSettings", "SecurityToken", out string AFSToken);
+
+					AFSToken = string.IsNullOrEmpty(AFSToken) ? "" : " -k " + AFSToken;
+
+					string AFSExecutable = Path.Combine(Unreal.EngineDirectory.ToString(), @"Binaries/DotNET/Android/UnrealAndroidFileTool/UnrealAndroidFileTool");
+					string AFS = $"{AFSExecutable} -p {PackageName}{AFSToken}";
+
+					string? SOPushScriptLocation = Path.GetDirectoryName(FinalSOName)!;
+
+					string FinalSONameStrippedRelative = Path.GetRelativePath(SOPushScriptLocation, FinalSONameStripped);
+
+					// MakeApk will not be called in bDontBundleLibrariesInAPK mode, so we need to run stripping outside of it
+					string SOPushScript = @$"
+set ADB=adb
+set AFS={AFS.Replace("/", "\\")}
+set DEVICE=
+if not \""%1\""==\""\"" set DEVICE=-s %1
+pushd %~dp0
+";
+					if (bUseAFS)
+					{
+						SOPushScript += @$"
+%AFS% %DEVICE% push {FinalSONameStrippedRelative} ""^int/libUnreal.so""
+";
+					}
+					else
+					{
+						SOPushScript += @$"
+%ADB% %DEVICE% push -z lz4 {FinalSONameStrippedRelative} /data/local/tmp/{FinalSONameStrippedRelative}
+%ADB% %DEVICE% shell run-as {PackageName} cp /data/local/tmp/{FinalSONameStrippedRelative} ./files/libUnreal.so
+%ADB% %DEVICE% shell rm /data/local/tmp/{FinalSONameStrippedRelative}
+";
+					}
+
+					SOPushScript +=
+@"
+popd
+";
+					// name is also used in the AGDE project, see AndroidProjectGenerator.cs
+					string SOPushScriptName = Path.Combine(Path.GetDirectoryName(FinalSOName)!, "Push_" + Path.GetFileNameWithoutExtension(FinalSOName) + "_so.bat");
+					File.WriteAllText(SOPushScriptName, SOPushScript);
+
+					// Remove libUnreal.so from gradle project if any
+					string JniDir = UnrealBuildPath + "/jni/" + NDKArch;
+					string TempFinalSOName = JniDir + "/libUnreal.so";
+					SafeDeleteFile(TempFinalSOName);
+				}
 				else
 				{
 					string SourceSOName = AndroidToolChain.InlineArchName(OutputPath, Arch);
@@ -4699,6 +4893,10 @@ namespace UnrealBuildTool
 					{
 						throw new BuildException("Can't make an APK without the jni directory [{0}/jni]", UnrealBuildFilesPath);
 					}
+					
+					// Delete the push script if exists
+					string SOPushScriptName = Path.Combine(Path.GetDirectoryName(DestApkName)!, "Push_" + Path.GetFileNameWithoutExtension(DestApkName) + "_so.bat");
+					SafeDeleteFile(SOPushScriptName);
 
 					string JniDir = UnrealBuildPath + "/jni/" + NDKArch;
 					FinalSOName = JniDir + "/libUnreal.so";
@@ -4724,7 +4922,6 @@ namespace UnrealBuildTool
 					File.SetLastWriteTimeUtc(FinalSOName, File.GetLastWriteTimeUtc(SourceSOName));
 				}
 
-				ConfigHierarchy Ini = GetConfigCacheIni(ConfigHierarchyType.Engine);
 				bool bSkipLibCpp = false;
 				Ini.GetBool("/Script/AndroidRuntimeSettings.AndroidRuntimeSettings", "bSkipLibCpp", out bSkipLibCpp);
 				if (!bSkipLibCpp)
@@ -5205,7 +5402,6 @@ namespace UnrealBuildTool
 					}
 					File.WriteAllLines(ManifestFilename, ManifestContents);
 
-					ConfigHierarchy Ini = GetConfigCacheIni(ConfigHierarchyType.Engine);
 					bool bEnableUniversalAPK = false;
 					Ini.GetBool("/Script/AndroidRuntimeSettings.AndroidRuntimeSettings", "bEnableUniversalAPK", out bEnableUniversalAPK);
 
@@ -5392,7 +5588,7 @@ namespace UnrealBuildTool
 			string RelativeEnginePath = Unreal.EngineDirectory.MakeRelativeTo(DirectoryReference.GetCurrentDirectory());
 
 			MakeApk(ToolChain, TargetName, Receipt.TargetType, ProjectDirectory.FullName, BaseSoName, RelativeEnginePath, bForDistribution: false, CookFlavor: "", Configuration: Receipt.Configuration,
-				bMakeSeparateApks: ShouldMakeSeparateApks(), bIncrementalPackage: true, bDisallowPackagingDataInApk: false, bDisallowExternalFilesDir: true, bSkipGradleBuild: bShouldCompileAsDll);
+				bMakeSeparateApks: ShouldMakeSeparateApks(), bIncrementalPackage: true, bDisallowPackagingDataInApk: false, bDisallowExternalFilesDir: true, bSkipGradleBuild: bShouldCompileAsDll, bIsArchive: false, bIsFromUAT: false);
 
 			// if we made any non-standard .apk files, the generated debugger settings may be wrong
 			if (ShouldMakeSeparateApks() && (OutputPaths.Count > 1 || !OutputPaths[0].FullName.Contains("-armv7")))
@@ -5445,7 +5641,7 @@ namespace UnrealBuildTool
 			// 			return bSeparateApks;
 		}
 
-		public bool PrepForUATPackageOrDeploy(FileReference ProjectFile, string ProjectName, DirectoryReference ProjectDirectory, string ExecutablePath, string EngineDirectory, bool bForDistribution, string CookFlavor, UnrealTargetConfiguration Configuration, bool bIsDataDeploy, bool bSkipGradleBuild)
+		public bool PrepForUATPackageOrDeploy(FileReference ProjectFile, string ProjectName, DirectoryReference ProjectDirectory, string ExecutablePath, string EngineDirectory, bool bForDistribution, string CookFlavor, UnrealTargetConfiguration Configuration, bool bIsDataDeploy, bool bSkipGradleBuild, bool bIsArchive)
 		{
 			//Logger.LogInformation("$$$$$$$$$$$$$$ PrepForUATPackageOrDeploy $$$$$$$$$$$$$$$$$");
 
@@ -5478,7 +5674,8 @@ namespace UnrealBuildTool
 			SavePackageInfo(ProjectName, ProjectDirectory.FullName, Type, bSkipGradleBuild);
 
 			MakeApk(ToolChain, ProjectName, Type, ProjectDirectory.FullName, ExecutablePath, EngineDirectory, bForDistribution: bForDistribution, CookFlavor: CookFlavor, Configuration: Configuration,
-				bMakeSeparateApks: ShouldMakeSeparateApks(), bIncrementalPackage: false, bDisallowPackagingDataInApk: bIsDataDeploy, bDisallowExternalFilesDir: !bForDistribution || bIsDataDeploy, bSkipGradleBuild: bSkipGradleBuild);
+				bMakeSeparateApks: ShouldMakeSeparateApks(), bIncrementalPackage: false, bDisallowPackagingDataInApk: bIsDataDeploy, bDisallowExternalFilesDir: !bForDistribution || bIsDataDeploy,
+				bSkipGradleBuild: bSkipGradleBuild, bIsArchive: bIsArchive, bIsFromUAT: true);
 			return true;
 		}
 
@@ -5537,7 +5734,7 @@ namespace UnrealBuildTool
 			return "empty";
 		}
 
-		private void UpdateGameActivity(UnrealArch UnrealArch, string NDKArch, string EngineDir, string UnrealBuildPath)
+		private void UpdateGameActivity(UnrealArch UnrealArch, string NDKArch, string EngineDir, string UnrealBuildPath, bool bDontBundleLibrariesInAPK)
 		{
 			string SourceFilename = Path.Combine(EngineDir, "Build", "Android", "Java", "src", "com", "epicgames", "unreal", "GameActivity.java.template");
 			string DestFilename = Path.Combine(UnrealBuildPath, "src", "com", "epicgames", "unreal", "GameActivity.java");
@@ -5547,6 +5744,9 @@ namespace UnrealBuildTool
 
 			ConfigHierarchy Ini = GetConfigCacheIni(ConfigHierarchyType.Engine);
 
+			string GameActivityImportAdditionsDefault = "";
+			string GameActivityClassAdditionsDefault = "";
+			string LoadUnrealDefault = "";
 			string LoadLibraryDefaults = "";
 
 			string SuperClassDefault;
@@ -5559,9 +5759,47 @@ namespace UnrealBuildTool
 				}
 			}
 
+			if (bDontBundleLibrariesInAPK) {
+				LoadUnrealDefault = @"		System.load(GameApplication.getAppContext().getFilesDir().getAbsolutePath() + ""/libUnreal.so"");";
+
+				GameActivityImportAdditionsDefault = @"
+import dalvik.system.BaseDexClassLoader;
+import java.util.Collection;
+";
+
+				GameActivityClassAdditionsDefault = @"
+	@Override
+	public ClassLoader getClassLoader() {
+		ClassLoader baseClassLoader = super.getClassLoader();
+
+		try {
+			BaseDexClassLoader dexClassLoader = (BaseDexClassLoader) baseClassLoader;
+			if (dexClassLoader.findLibrary(""Unreal"") == null) {
+				Field pathListField = BaseDexClassLoader.class.getDeclaredField(""pathList"");
+				pathListField.setAccessible(true);
+
+				Object pathListObj = pathListField.get(dexClassLoader);
+
+				Method addNativePathMethod = pathListObj.getClass().getMethod(""addNativePath"", Collection.class);
+
+				String filesDir = getApplicationContext().getFilesDir().getAbsolutePath();
+				Log.verbose(""Adding '"" + filesDir + ""' to GameActivity BaseDexClassLoader"");
+				Collection<String> paths = Arrays.asList(filesDir);
+				addNativePathMethod.invoke(pathListObj, paths);
+			}
+		} catch (Exception e) {
+			Log.warn(""Failed to add native library path due to "" + e);
+		}
+
+		return baseClassLoader;
+	}
+";
+			} else {
+				LoadUnrealDefault = @"		System.loadLibrary(""Unreal"");";
+			}
+
 			string AndroidGraphicsDebugger;
 			Ini.GetString("/Script/AndroidRuntimeSettings.AndroidRuntimeSettings", "AndroidGraphicsDebugger", out AndroidGraphicsDebugger);
-
 			switch (AndroidGraphicsDebugger.ToLower())
 			{
 				case "mali":
@@ -5578,10 +5816,10 @@ namespace UnrealBuildTool
 
 			Dictionary<string, string> Replacements = new Dictionary<string, string>{				
 				{ "//$${gameActivityMemStatAdditions}$$", UPL.ProcessPluginNode(NDKArch, "gameActivityMemStatAdditions", "")},
-				{ "//$${gameActivityImportAdditions}$$", UPL.ProcessPluginNode(NDKArch, "gameActivityImportAdditions", "")},
+				{ "//$${gameActivityImportAdditions}$$", UPL.ProcessPluginNode(NDKArch, "gameActivityImportAdditions", GameActivityImportAdditionsDefault)},
 				{ "//$${gameActivityPostImportAdditions}$$", UPL.ProcessPluginNode(NDKArch, "gameActivityPostImportAdditions", "")},
 				{ "//$${gameActivityImplementsAdditions}$$", UPL.ProcessPluginNode(NDKArch, "gameActivityImplementsAdditions", "")},
-				{ "//$${gameActivityClassAdditions}$$", UPL.ProcessPluginNode(NDKArch, "gameActivityClassAdditions", "")},
+				{ "//$${gameActivityClassAdditions}$$", UPL.ProcessPluginNode(NDKArch, "gameActivityClassAdditions", GameActivityClassAdditionsDefault)},
 				{ "//$${gameActivityReadMetadataAdditions}$$", UPL.ProcessPluginNode(NDKArch, "gameActivityReadMetadataAdditions", "")},
 				{ "//$${gameActivityOnCreateBeginningAdditions}$$", UPL.ProcessPluginNode(NDKArch, "gameActivityOnCreateBeginningAdditions", "")},
 				{ "//$${gameActivityOnCreateAdditions}$$", UPL.ProcessPluginNode(NDKArch, "gameActivityOnCreateAdditions", "")},
@@ -5614,6 +5852,7 @@ namespace UnrealBuildTool
 				{ "//$${gameActivityOnRestartApplicationAdditions}$$", UPL.ProcessPluginNode(NDKArch, "gameActivityOnRestartApplicationAdditions", "")},
 				{ "//$${gameActivityForceQuitAdditions}$$", UPL.ProcessPluginNode(NDKArch, "gameActivityForceQuitAdditions", "")},
 				{ "//$${gameActivitySetCrashContextData}$$", UPL.ProcessPluginNode(NDKArch, "gameActivitySetCrashContextData", "")},
+				{ "//$${soLoadUnreal}$$", UPL.ProcessPluginNode(NDKArch, "soLoadUnreal", LoadUnrealDefault)},
 				{ "//$${soLoadLibrary}$$", UPL.ProcessPluginNode(NDKArch, "soLoadLibrary", LoadLibraryDefaults)},
 				{ "$${gameActivitySuperClass}$$", SuperClassDefault},
 			};
@@ -5667,6 +5906,7 @@ namespace UnrealBuildTool
 				{ "//$${gameApplicationOnLowMemoryAdditions}$$", UPL.ProcessPluginNode(NDKArch, "gameApplicationOnLowMemoryAdditions", "")},
 				{ "//$${gameApplicationOnTrimMemoryAdditions}$$", UPL.ProcessPluginNode(NDKArch, "gameApplicationOnTrimMemoryAdditions", "")},
 				{ "//$${gameApplicationOnConfigurationChangedAdditions}$$", UPL.ProcessPluginNode(NDKArch, "gameApplicationOnConfigurationChangedAdditions", "")},
+				{ "//$${gameApplicationClassAdditions}$$", UPL.ProcessPluginNode(NDKArch, "gameApplicationClassAdditions", "")},
 			};
 
 			string[] TemplateSrc = File.ReadAllLines(SourceFilename);
diff --git a/Engine/Source/Programs/UnrealBuildTool/System/Utils.cs b/Engine/Source/Programs/UnrealBuildTool/System/Utils.cs
--- a/Engine/Source/Programs/UnrealBuildTool/System/Utils.cs
+++ b/Engine/Source/Programs/UnrealBuildTool/System/Utils.cs
@@ -1826,5 +1826,55 @@ namespace UnrealBuildTool
 		{
 			return Strings.Any(S => Encoding.UTF8.GetByteCount(S) != S.Length) ? new UTF8Encoding(false) : Encoding.ASCII;
 		}
+
+		/// <summary>
+		/// Attempts to create a symbolic link at location specified by Path pointing to location specified by PathToTarget.
+		/// Soft symlinks are available since Windows 10 build 14972 without elevated privileges if developer mode is enabled.
+		/// Hard links are available since Windows 8 for NTFS/NFS file systems.
+		/// </summary>
+		/// <param name="Path">Path to create the symbolic link at.</param>
+		/// <param name="PathToTarget">Path to which the symbolic link should point to.</param>
+		/// <param name="Logger">Logger for output.</param>
+		/// <returns>True if symlink was created, false if failed.</returns>
+		internal static bool TryCreateSymlink(string Path, string PathToTarget, ILogger Logger)
+		{
+			try
+			{
+				// passes SYMBOLIC_LINK_FLAG_ALLOW_UNPRIVILEGED_CREATE on valid Windows versions
+				FileSystemInfo Result = File.CreateSymbolicLink(Path, PathToTarget);
+				if (Result.Exists)
+				{
+					Logger.LogInformation("Created symlink '{Path}' -> '{PathToTarget}'", Path, PathToTarget);
+					return true;
+				}
+			}
+			catch
+			{
+				// ignored
+			}
+
+			if (RuntimePlatform.IsWindows)
+			{
+				try
+				{
+					WindowsKernelCreateHardLink(Path, PathToTarget, IntPtr.Zero);
+					// not 100% confident in a result value of CreateHardLink, so let's check for file to be extra sure 
+					if (File.Exists(Path))
+					{
+						Logger.LogInformation("Created hard link '{Path}' -> '{PathToTarget}'", Path, PathToTarget);
+						return true;
+					}
+				}
+				catch
+				{
+					// ignored
+				}
+			}
+
+			return false;
+		}
+
+		[DllImport("kernel32.dll", EntryPoint = "CreateHardLink", SetLastError = true, CharSet = CharSet.Auto)]
+		private static extern bool WindowsKernelCreateHardLink(string lpFileName, string lpExistingFileName, IntPtr lpSecurityAttributes);
 	}
 }
diff --git a/Engine/Source/Runtime/Android/AndroidRuntimeSettings/Classes/AndroidRuntimeSettings.h b/Engine/Source/Runtime/Android/AndroidRuntimeSettings/Classes/AndroidRuntimeSettings.h
--- a/Engine/Source/Runtime/Android/AndroidRuntimeSettings/Classes/AndroidRuntimeSettings.h
+++ b/Engine/Source/Runtime/Android/AndroidRuntimeSettings/Classes/AndroidRuntimeSettings.h
@@ -262,6 +262,9 @@ public:
 	UPROPERTY(GlobalConfig, EditAnywhere, Category = "APK Packaging", Meta = (DisplayName = "Allow overflow OBB files."))
 	bool bAllowOverflowOBBFiles;
 
+	UPROPERTY(GlobalConfig, EditAnywhere, Category = "APK Packaging", Meta = (DisplayName = "Don't bundle libraries into .apk for quicker iteration [Experimental]"))
+	bool bDontBundleLibrariesInAPK;
+
 	// If checked, UnrealGame files will be placed in ExternalFilesDir which is removed on uninstall.
 	// You should also check this if you need to save you game progress without requesting runtime WRITE_EXTERNAL_STORAGE permission in android api 23+
 	UPROPERTY(GlobalConfig, EditAnywhere, Category = "APK Packaging", Meta = (DisplayName = "Use ExternalFilesDir for UnrealGame files?"))
