UnrealBuildTool: Don't try to get files in a directory if null

#rnx

diff --git a/Engine/Source/Programs/UnrealBuildTool/Configuration/UEBuildModuleCPP.cs b/Engine/Source/Programs/UnrealBuildTool/Configuration/UEBuildModuleCPP.cs
--- a/Engine/Source/Programs/UnrealBuildTool/Configuration/UEBuildModuleCPP.cs
+++ b/Engine/Source/Programs/UnrealBuildTool/Configuration/UEBuildModuleCPP.cs
@@ -851,6 +851,10 @@ namespace UnrealBuildTool
 		static ConcurrentDictionary<DirectoryReference, HashSet<FileReference>?> KnownIncludeFilesDict = new();
 		static HashSet<FileReference>? GetIncludeFiles(DirectoryReference Directory)
 		{
+			if (Directory == null)
+			{
+				return null;
+			}
 			return KnownIncludeFilesDict.GetOrAdd(Directory, (_) =>
 			{
 				if (DirectoryLookupCache.DirectoryExists(Directory))
