UnrealBuildTool: Revert pvs-studio distribution

#rnx

diff --git a/Engine/Source/Programs/UnrealBuildTool/Platform/Windows/PVSToolChain.cs b/Engine/Source/Programs/UnrealBuildTool/Platform/Windows/PVSToolChain.cs
--- a/Engine/Source/Programs/UnrealBuildTool/Platform/Windows/PVSToolChain.cs
+++ b/Engine/Source/Programs/UnrealBuildTool/Platform/Windows/PVSToolChain.cs
@@ -800,7 +800,7 @@ namespace UnrealBuildTool
 				AnalyzeAction.PrerequisiteItems.UnionWith(InputFiles); // Add the InputFiles as PrerequisiteItems so that in SingleFileCompile mode the PVSAnalyze step is not filtered out
 				AnalyzeAction.ProducedItems.Add(OutputFileItem);
 				AnalyzeAction.DeleteItems.Add(OutputFileItem); // PVS Studio will append by default, so need to delete produced items
-				AnalyzeAction.bCanExecuteRemotely = true;
+				AnalyzeAction.bCanExecuteRemotely = false;
 				AnalyzeAction.bCanExecuteRemotelyWithXGE = false;
 				AnalyzeAction.bCanExecuteRemotelyWithSNDBS = false;
 
