Fix the AutoRTFM clang compiler detection.

diff --git a/Engine/Source/Programs/UnrealBuildTool/Platform/Windows/VCToolChain.cs b/Engine/Source/Programs/UnrealBuildTool/Platform/Windows/VCToolChain.cs
--- a/Engine/Source/Programs/UnrealBuildTool/Platform/Windows/VCToolChain.cs
+++ b/Engine/Source/Programs/UnrealBuildTool/Platform/Windows/VCToolChain.cs
@@ -1106,7 +1106,7 @@ namespace UnrealBuildTool
 			}
 		}
 
-		static void AppendCLArguments_C(CppCompileEnvironment CompileEnvironment, List<string> Arguments)
+		void AppendCLArguments_C(CppCompileEnvironment CompileEnvironment, List<string> Arguments)
 		{
 			// Explicitly compile the file as C.
 			Arguments.Add("/TC");
@@ -1132,6 +1132,14 @@ namespace UnrealBuildTool
 				default:
 					throw new BuildException($"Unsupported C standard type set: {CompileEnvironment.CStandard}");
 			}
+
+			if (Target.WindowsPlatform.Compiler.IsClang())
+			{
+				if (CompileEnvironment.bAllowAutoRTFMInstrumentation)
+				{
+					Arguments.Add("-fautortfm");
+				}
+			}
 		}
 
 		protected virtual void AppendLinkArguments(LinkEnvironment LinkEnvironment, List<string> Arguments)
@@ -1655,32 +1663,26 @@ namespace UnrealBuildTool
 				if (bIsPlainCFile)
 				{
 					AppendCLArguments_C(CompileEnvironment, CompileAction.Arguments);
+				}
+				else
+				{
+					AppendCLArguments_CPP(CompileEnvironment, CompileAction.Arguments);
+				}
 
-					// AppendCLArguments_C is a static and we cannot use the Target to check if we are using clang, so shoving here for now
-					if (Target.WindowsPlatform.Compiler.IsClang())
+				if (Target.WindowsPlatform.Compiler.IsClang())
+				{
+					// If we are using the AutoRTFM compiler, we make the compile action depend on the version of the compiler itself.
+					// This lets us update the compiler (which might not cause a version update of the compiler, which instead tracks
+					// the LLVM versioning scheme that Clang uses), but ensure that we rebuild the source if the compiler has changed.
+					if (CompileEnvironment.bUseAutoRTFMCompiler)
 					{
-						if (CompileEnvironment.bAllowAutoRTFMInstrumentation)
-						{
-							CompileAction.Arguments.Add("-fautortfm");
-						}
-
-						// If we are using the AutoRTFM compiler, we make the compile action depend on the version of the compiler itself.
-						// This lets us update the compiler (which might not cause a version update of the compiler, which instead tracks
-						// the LLVM versioning scheme that Clang uses), but ensure that we rebuild the source if the compiler has changed.
-						if (CompileEnvironment.bUseAutoRTFMCompiler)
+						FileReference? CompilerPath = GetCppCompilerPath();
+						if (null != CompilerPath)
 						{
-							FileReference? CompilerPath = GetCppCompilerPath();
-							if (null != CompilerPath)
-							{
-								BaseCompileAction.AdditionalPrerequisiteItems.Add(FileItem.GetItemByFileReference(CompilerPath));
-							}
+							CompileAction.AdditionalPrerequisiteItems.Add(FileItem.GetItemByFileReference(CompilerPath));
 						}
 					}
 				}
-				else
-				{
-					AppendCLArguments_CPP(CompileEnvironment, CompileAction.Arguments);
-				}
 
 				List<FileItem>? InlinedFiles;
 				if (CompileEnvironment.FileInlineGenCPPMap.TryGetValue(SourceFile, out InlinedFiles))
