UnrealBuildTool: Don't run pvs with experimental

#rnx

diff --git a/Engine/Source/Programs/UnrealBuildTool/Platform/Windows/PVSToolChain.cs b/Engine/Source/Programs/UnrealBuildTool/Platform/Windows/PVSToolChain.cs
--- a/Engine/Source/Programs/UnrealBuildTool/Platform/Windows/PVSToolChain.cs
+++ b/Engine/Source/Programs/UnrealBuildTool/Platform/Windows/PVSToolChain.cs
@@ -803,6 +803,7 @@ namespace UnrealBuildTool
 				AnalyzeAction.bCanExecuteRemotely = false;
 				AnalyzeAction.bCanExecuteRemotelyWithXGE = false;
 				AnalyzeAction.bCanExecuteRemotelyWithSNDBS = false;
+				AnalyzeAction.bCanExecuteInBox = false;
 
 				Result.ObjectFiles.AddRange(AnalyzeAction.ProducedItems);
 			}
