Roll back to 7.17.57873.3433 (x64)
Reenable box distribution
#rnx

diff --git a/Engine/Source/Programs/UnrealBuildTool/Platform/Windows/PVSToolChain.cs b/Engine/Source/Programs/UnrealBuildTool/Platform/Windows/PVSToolChain.cs
--- a/Engine/Source/Programs/UnrealBuildTool/Platform/Windows/PVSToolChain.cs
+++ b/Engine/Source/Programs/UnrealBuildTool/Platform/Windows/PVSToolChain.cs
@@ -800,10 +800,9 @@ namespace UnrealBuildTool
 				AnalyzeAction.PrerequisiteItems.UnionWith(InputFiles); // Add the InputFiles as PrerequisiteItems so that in SingleFileCompile mode the PVSAnalyze step is not filtered out
 				AnalyzeAction.ProducedItems.Add(OutputFileItem);
 				AnalyzeAction.DeleteItems.Add(OutputFileItem); // PVS Studio will append by default, so need to delete produced items
-				AnalyzeAction.bCanExecuteRemotely = false;
+				AnalyzeAction.bCanExecuteRemotely = true;
 				AnalyzeAction.bCanExecuteRemotelyWithXGE = false;
 				AnalyzeAction.bCanExecuteRemotelyWithSNDBS = false;
-				AnalyzeAction.bCanExecuteInBox = false;
 
 				Result.ObjectFiles.AddRange(AnalyzeAction.ProducedItems);
 			}
