Use correct AFS in the push script

diff --git a/Engine/Source/Programs/AutomationTool/Android/AndroidPlatform.Automation.cs b/Engine/Source/Programs/AutomationTool/Android/AndroidPlatform.Automation.cs
--- a/Engine/Source/Programs/AutomationTool/Android/AndroidPlatform.Automation.cs
+++ b/Engine/Source/Programs/AutomationTool/Android/AndroidPlatform.Automation.cs
@@ -1526,20 +1526,7 @@ public class AndroidPlatform : Platform
 
 	string GetAFSExecutable(UnrealTargetPlatform Target)
 	{
-		if (Target == UnrealTargetPlatform.Win64)
-		{
-			return "win-x64/UnrealAndroidFileTool.exe";
-		}
-		if (Target == UnrealTargetPlatform.Mac)
-		{
-			return "osx-x64/UnrealAndroidFileTool";
-		}
-		if (Target == UnrealTargetPlatform.Linux)
-		{
-			return "linux-x64/UnrealAndroidFileTool";
-		}
-		Logger.LogWarning("GetAFSExecutable unsupported target, assuming Win64");
-		return "win-x64/UnrealAndroidFileTool.exe";
+		return AndroidExports.GetAFSExecutable(Target, Logger);
 	}
 
 	private class OverflowBatchInstallInfo
diff --git a/Engine/Source/Programs/UnrealBuildTool/Platform/Android/AndroidExports.cs b/Engine/Source/Programs/UnrealBuildTool/Platform/Android/AndroidExports.cs
--- a/Engine/Source/Programs/UnrealBuildTool/Platform/Android/AndroidExports.cs
+++ b/Engine/Source/Programs/UnrealBuildTool/Platform/Android/AndroidExports.cs
@@ -131,5 +131,15 @@ namespace UnrealBuildTool
 		{
 			return UEDeployAndroid.GetDontBundleLibrariesInAPK(ProjectFile, bForceDontBundleLibrariesInAPK, Configuration, bIsArchive, bFromMSBuild, bIsFromUAT, Logger);
 		}
+
+		/// <summary>
+		/// 
+		/// </summary>
+		/// <param name="Target"></param>
+		/// <param name="Logger"></param>
+		public static string GetAFSExecutable(UnrealTargetPlatform Target, ILogger Logger)
+		{
+			return UEDeployAndroid.GetAFSExecutable(Target, Logger);
+		}
 	}
 }
diff --git a/Engine/Source/Programs/UnrealBuildTool/Platform/Android/UEDeployAndroid.cs b/Engine/Source/Programs/UnrealBuildTool/Platform/Android/UEDeployAndroid.cs
--- a/Engine/Source/Programs/UnrealBuildTool/Platform/Android/UEDeployAndroid.cs
+++ b/Engine/Source/Programs/UnrealBuildTool/Platform/Android/UEDeployAndroid.cs
@@ -4832,7 +4832,7 @@ namespace UnrealBuildTool
 
 					AFSToken = string.IsNullOrEmpty(AFSToken) ? "" : " -k " + AFSToken;
 
-					string AFSExecutable = Path.Combine(Unreal.EngineDirectory.ToString(), @"Binaries/DotNET/Android/UnrealAndroidFileTool/UnrealAndroidFileTool");
+					string AFSExecutable = Path.Combine(Unreal.EngineDirectory.ToString(), @"Binaries/DotNET/Android/UnrealAndroidFileTool", GetAFSExecutable(UnrealTargetPlatform.Win64, Logger));
 					string AFS = $"{AFSExecutable} -p {PackageName}{AFSToken}";
 
 					string? SOPushScriptLocation = Path.GetDirectoryName(FinalSOName)!;
@@ -5533,6 +5533,24 @@ popd
 			Logger.LogInformation("\n===={Time}====COMPLETED MAKE APK=======================================================================", DateTime.Now.ToString());
 		}
 
+		public static string GetAFSExecutable(UnrealTargetPlatform Target, ILogger Logger)
+		{
+			if (Target == UnrealTargetPlatform.Win64)
+			{
+				return "win-x64/UnrealAndroidFileTool.exe";
+			}
+			if (Target == UnrealTargetPlatform.Mac)
+			{
+				return "osx-x64/UnrealAndroidFileTool";
+			}
+			if (Target == UnrealTargetPlatform.Linux)
+			{
+				return "linux-x64/UnrealAndroidFileTool";
+			}
+			Logger.LogWarning("GetAFSExecutable unsupported target, assuming Win64");
+			return "win-x64/UnrealAndroidFileTool.exe";
+		}
+
 		private List<string> CollectPluginDataPaths(TargetReceipt Receipt)
 		{
 			List<string> PluginExtras = new List<string>();
