UnrealBuildTool: QueryMode improvements

* Async
* Sort dirToModule & moduleToCompileSettings so the order is stable now that the binaries and modules are processed async
* Find the binaryPath by searching for the first Binary that includes the launch module
* Remove BrowseConfiguration which is unused

#rnx

diff --git a/Engine/Source/Programs/UnrealBuildTool/Modes/QueryMode.cs b/Engine/Source/Programs/UnrealBuildTool/Modes/QueryMode.cs
--- a/Engine/Source/Programs/UnrealBuildTool/Modes/QueryMode.cs
+++ b/Engine/Source/Programs/UnrealBuildTool/Modes/QueryMode.cs
@@ -1,10 +1,13 @@
 // Copyright Epic Games, Inc. All Rights Reserved.
 
 using System;
+using System.Collections.Concurrent;
 using System.Collections.Generic;
+using System.Collections.Immutable;
 using System.IO;
 using System.Linq;
 using System.Text.Json;
+using System.Threading;
 using System.Threading.Tasks;
 using EpicGames.Core;
 using Microsoft.CodeAnalysis.CSharp.Syntax;
@@ -34,8 +37,8 @@ namespace UnrealBuildTool
 			public string? WindowsSdkVersion { get; set; }
 		}
 
-		public Dictionary<UEBuildModule, CompileSettings> ModuleToCompileSettings = new();
-		public Dictionary<DirectoryReference, UEBuildModule> DirToModule = new();
+		public ConcurrentDictionary<UEBuildModule, CompileSettings> ModuleToCompileSettings = new();
+		public ConcurrentDictionary<DirectoryReference, UEBuildModule> DirToModule = new();
 
 		public UEBuildModule? FindModuleForFile(FileReference File)
 		{
@@ -154,19 +157,19 @@ namespace UnrealBuildTool
 			{
 				case QueryType.Capabilities:
 					Logger.LogInformation("QueryCapabilities");
-					return Task.FromResult(QueryCapabilities(Arguments, Logger, ResponseOptions));
+					return QueryCapabilities(Arguments, Logger, ResponseOptions);
 				case QueryType.AvailableTargets:
 					Logger.LogInformation("QueryAvailableTargets");
-					return Task.FromResult(QueryAvailableTargets(Arguments, Logger, ResponseOptions));
+					return QueryAvailableTargets(Arguments, Logger, ResponseOptions);
 				case QueryType.TargetDetails:
 					Logger.LogInformation("QueryTargetDetails");
-					return Task.FromResult(QueryTargetDetails(Arguments, Logger, ResponseOptions));
+					return QueryTargetDetails(Arguments, Logger, ResponseOptions);
 			}
 
 			return Task.FromResult(0);
 		}
 
-		private void WriteResults(object? Value, JsonSerializerOptions JsonOptions, ILogger Logger)
+		private async Task WriteResultsAsync(object? Value, JsonSerializerOptions JsonOptions, ILogger Logger)
 		{
 			if (OutputPath == null)
 			{
@@ -175,20 +178,20 @@ namespace UnrealBuildTool
 			}
 			using FileStream Stream = new FileStream(OutputPath.FullName, FileMode.Create, FileAccess.Write);
 			Logger.LogInformation("Writing {File}...", OutputPath.FullName);
-			JsonSerializer.Serialize(Stream, Value, JsonOptions);
+			await JsonSerializer.SerializeAsync(Stream, Value, JsonOptions);
 		}
 
-		private int QueryCapabilities(CommandLineArguments Arguments, ILogger Logger, JsonSerializerOptions JsonOptions)
+		private async Task<int> QueryCapabilities(CommandLineArguments Arguments, ILogger Logger, JsonSerializerOptions JsonOptions)
 		{
 			var Reply = new
 			{
 				Queries = new List<string> { QueryType.Capabilities.ToString(), QueryType.AvailableTargets.ToString(), QueryType.TargetDetails.ToString() }
 			};
-			WriteResults(Reply, JsonOptions, Logger);
+			await WriteResultsAsync(Reply, JsonOptions, Logger);
 			return 0;
 		}
 
-		private int QueryAvailableTargets(CommandLineArguments Arguments, ILogger Logger, JsonSerializerOptions JsonOptions)
+		private async Task<int> QueryAvailableTargets(CommandLineArguments Arguments, ILogger Logger, JsonSerializerOptions JsonOptions)
 		{
 			try
 			{
@@ -249,16 +252,18 @@ namespace UnrealBuildTool
 					DefaultConfiguration = UnrealTargetConfiguration.Development.ToString(),
 				};
 
-				WriteResults(Reply, JsonOptions, Logger);
+				await WriteResultsAsync(Reply, JsonOptions, Logger);
 				return 0;
 			}
-			catch (Exception e)
+			catch (Exception Ex)
 			{
-				Logger.LogError("Failed to query available targets: {0}", e.Message);
+				Logger.LogError("Failed to query available targets: {Error}", Ex.Message);
+				Logger.LogDebug(Ex, "Unhandled exception: {Ex}", ExceptionUtils.FormatExceptionDetails(Ex));
 				return 1;
 			}
 		}
-		private int QueryTargetDetails(CommandLineArguments Arguments, ILogger Logger, JsonSerializerOptions JsonOptions)
+
+		private async Task<int> QueryTargetDetails(CommandLineArguments Arguments, ILogger Logger, JsonSerializerOptions JsonOptions)
 		{
 			if (TargetName == null)
 			{
@@ -293,7 +298,6 @@ namespace UnrealBuildTool
 				return 1;
 			}
 
-			HashSet<string> BrowseConfigurationFolders = new HashSet<string>();
 			try
 			{
 				UEBuildTarget CurrentTarget;
@@ -309,9 +313,7 @@ namespace UnrealBuildTool
 					}
 				}
 
-				LaunchSettings? CurrentLaunchSettings = null;
 				TargetIntellisenseInfo CurrentTargetIntellisenseInfo = new TargetIntellisenseInfo();
-				GetBrowseConfigurationResponse CurrentBrowseConfiguration = new GetBrowseConfigurationResponse { Success = true };
 
 				// Partially duplicated from UEBuildTarget.Build because we just want to get C++ compile actions without running UHT
 				// or generating link actions / full dependency graph 
@@ -324,37 +326,36 @@ namespace UnrealBuildTool
 				TargetToolChain.SetEnvironmentVariables();
 				CurrentTarget.SetupGlobalEnvironment(TargetToolChain, GlobalCompileEnvironment, GlobalLinkEnvironment);
 
-				CurrentBrowseConfiguration.WindowsSdkVersion = CurrentTarget.Rules.WindowsPlatform.WindowsSdkVersion;
+				UEBuildBinary? LaunchBinary = CurrentTarget.Binaries.FirstOrDefault(Binary => Binary.Modules.Any(Module => Module.Name == CurrentTarget.Rules.LaunchModuleName));
 
-				// TODO: For installed builds, filter out all the binaries that aren't in mods
-				foreach (UEBuildBinary Binary in CurrentTarget.Binaries)
+				if (LaunchBinary == null)
 				{
-					if (Binary.Type == UEBuildBinaryType.Executable && CurrentLaunchSettings == null && Binary.OutputFilePaths.Count == 1)
-					{
-						CurrentLaunchSettings = new LaunchSettings();
-						CurrentLaunchSettings.Description = $"{CurrentTarget.TargetName} {CurrentTarget.Configuration} {CurrentTarget.Platform}";
-						CurrentLaunchSettings.BinaryPath = Binary.OutputFilePath.ToString();
-						if (CurrentTarget.ProjectFile != null && CurrentTarget.TargetType != TargetType.Program)
-						{
-							CurrentLaunchSettings.Arguments.Add(CurrentTarget.ProjectFile.ToString());
-						}
-					}
+					throw new BuildException("Unable to find launch binary for target");
+				}
 
-					HashSet<UEBuildModule> LinkEnvironmentVisitedModules = new HashSet<UEBuildModule>();
+				LaunchSettings CurrentLaunchSettings = new LaunchSettings();
+				CurrentLaunchSettings.Description = $"{CurrentTarget.TargetName} {CurrentTarget.Configuration} {CurrentTarget.Platform}";
+				CurrentLaunchSettings.BinaryPath = LaunchBinary.OutputFilePath.FullName;
+				if (CurrentTarget.ProjectFile != null && CurrentTarget.TargetType != TargetType.Program)
+				{
+					CurrentLaunchSettings.Arguments.Add(CurrentTarget.ProjectFile.FullName);
+				}
+
+				await Parallel.ForEachAsync(CurrentTarget.Binaries, async (Binary, CancellationToken) =>
+				{
 					CppCompileEnvironment BinaryCompileEnvironment = Binary.CreateBinaryCompileEnvironment(GlobalCompileEnvironment);
-					CurrentBrowseConfiguration.Standard = BinaryCompileEnvironment.CppStandard.ToString();
 
-					foreach (UEBuildModuleCPP Module in Binary.Modules.OfType<UEBuildModuleCPP>())
+					IEnumerable<UEBuildModuleCPP> Modules = Binary.Modules.OfType<UEBuildModuleCPP>().Where(x => x.Binary == Binary);
+
+					await Parallel.ForEachAsync(Modules, (Module, CancellationToken) =>
 					{
-						if (Module.Binary != null && Module.Binary != Binary)
+						if (CurrentTargetIntellisenseInfo.ModuleToCompileSettings.ContainsKey(Module))
 						{
-							continue;
+							return ValueTask.CompletedTask;
 						}
 
-						CppCompileEnvironment ModuleCompileEnvironment = Module.CreateModuleCompileEnvironment(CurrentTarget.Rules, BinaryCompileEnvironment, Logger);
 						foreach (DirectoryReference Dir in Module.ModuleDirectories)
 						{
-							BrowseConfigurationFolders.Add(Dir.ToString());
 							CurrentTargetIntellisenseInfo.DirToModule.TryAdd(Dir, Module);
 						}
 						if (Module.GeneratedCodeDirectory != null)
@@ -362,14 +363,7 @@ namespace UnrealBuildTool
 							CurrentTargetIntellisenseInfo.DirToModule.TryAdd(Module.GeneratedCodeDirectory, Module);
 						}
 
-						foreach (DirectoryReference Dir in ModuleCompileEnvironment.SystemIncludePaths)
-						{
-							BrowseConfigurationFolders.Add(Dir.ToString());
-						}
-						foreach (DirectoryReference Dir in ModuleCompileEnvironment.UserIncludePaths)
-						{
-							BrowseConfigurationFolders.Add(Dir.ToString());
-						}
+						CppCompileEnvironment ModuleCompileEnvironment = Module.CreateModuleCompileEnvironment(CurrentTarget.Rules, BinaryCompileEnvironment, Logger);
 
 						TargetIntellisenseInfo.CompileSettings Settings = new TargetIntellisenseInfo.CompileSettings();
 						if (TargetToolChain is VCToolChain TargetVCToolChain)
@@ -384,25 +378,26 @@ namespace UnrealBuildTool
 						Settings.CompilerPath = TargetToolChain.GetCppCompilerPath()?.ToString();
 						Settings.CompilerArgs = TargetToolChain.GetGlobalCommandLineArgs(ModuleCompileEnvironment).ToList();
 						Settings.WindowsSdkVersion = CurrentTarget.Rules.WindowsPlatform.WindowsSdkVersion;
-						CurrentTargetIntellisenseInfo.ModuleToCompileSettings.Add(Module, Settings);
-					}
-				}
+						CurrentTargetIntellisenseInfo.ModuleToCompileSettings.TryAdd(Module, Settings);
 
-				CurrentBrowseConfiguration.Paths = BrowseConfigurationFolders.ToList();
+						return ValueTask.CompletedTask;
+					});
+				});
 
 				var Result = new
 				{
-					DirToModule = CurrentTargetIntellisenseInfo.DirToModule.ToDictionary(x => x.Key.ToString(), x => x.Value.Name),
-					ModuleToCompileSettings = CurrentTargetIntellisenseInfo.ModuleToCompileSettings.ToDictionary(x => x.Key.Name, x => x.Value),
+					DirToModule = CurrentTargetIntellisenseInfo.DirToModule.ToImmutableSortedDictionary(x => x.Key.ToString(), x => x.Value.Name),
+					ModuleToCompileSettings = CurrentTargetIntellisenseInfo.ModuleToCompileSettings.ToImmutableSortedDictionary(x => x.Key.Name, x => x.Value),
 					LaunchSettings = CurrentLaunchSettings,
 				};
 
-				WriteResults(Result, JsonOptions, Logger);
+				await WriteResultsAsync(Result, JsonOptions, Logger);
 				return 0;
 			}
-			catch (Exception e)
+			catch (Exception Ex)
 			{
-				Logger.LogError("Caught exception setting up target: {0}", e);
+				Logger.LogError("Failed to query available targets: {Error}", Ex.Message);
+				Logger.LogDebug(Ex, "Unhandled exception: {Ex}", ExceptionUtils.FormatExceptionDetails(Ex));
 				return 1;
 			}
 		}
diff --git a/Engine/Source/Programs/UnrealBuildTool/Modes/ServerMode.cs b/Engine/Source/Programs/UnrealBuildTool/Modes/ServerMode.cs
--- a/Engine/Source/Programs/UnrealBuildTool/Modes/ServerMode.cs
+++ b/Engine/Source/Programs/UnrealBuildTool/Modes/ServerMode.cs
@@ -382,7 +382,7 @@ namespace UnrealBuildTool
 						Settings.ForcedIncludes = ModuleCompileEnvironment.ForceIncludeFiles.Select(x => x.ToString()).ToList();
 						Settings.CompilerPath = TargetToolChain.GetCppCompilerPath()?.ToString();
 						Settings.WindowsSdkVersion = CurrentTarget.Rules.WindowsPlatform.WindowsSdkVersion;
-						CurrentTargetIntellisenseInfo.ModuleToCompileSettings.Add(Module, Settings);
+						CurrentTargetIntellisenseInfo.ModuleToCompileSettings.TryAdd(Module, Settings);
 					}
 				}
 			}
