UnrealBuildTool: Move VNIExecution task to async function

#rnx

diff --git a/Engine/Source/Programs/UnrealBuildTool/Configuration/UEBuildTarget.cs b/Engine/Source/Programs/UnrealBuildTool/Configuration/UEBuildTarget.cs
--- a/Engine/Source/Programs/UnrealBuildTool/Configuration/UEBuildTarget.cs
+++ b/Engine/Source/Programs/UnrealBuildTool/Configuration/UEBuildTarget.cs
@@ -2573,24 +2573,10 @@ namespace UnrealBuildTool
 			}
 
 #if __VPROJECT_AVAILABLE__
-			Task VNITask = Task.Run(() =>
-			{
-				// Prepare cached data for VNI header generation
-				using (Timeline.ScopeEvent("ExternalExecution.SetupVNIModules()"))
-				{
-					VNIExecution.SetupVNIModules(ModulesToGenerateHeadersFor, RulesAssembly, out Makefile.VNIModules);
-				}
-
-				// Copy Verse BPVM usage flag
-				Makefile.bUseVerseBPVM = Rules.bUseVerseBPVM;
+			// Copy Verse BPVM usage flag
+			Makefile.bUseVerseBPVM = Rules.bUseVerseBPVM;
 
-				// NOTE: Even in Gather mode, we need to run VNI to make sure the files exist for the static action graph to be setup correctly.  This is because VNI generates .cpp
-				// files that are injected as top level prerequisites.  If VNI only emitted included header files, we wouldn't need to run it during the Gather phase at all.
-				if (Makefile.VNIModules.Count > 0)
-				{
-					VNIExecution.ExecuteVNITool(Makefile, TargetDescriptor, Logger);
-				}
-			});
+			Task VNITask = VNIExecution.RunVNIAsync(ModulesToGenerateHeadersFor, RulesAssembly, Rules, Makefile, TargetDescriptor, Logger);
 #endif
 
 			// NOTE: Even in Gather mode, we need to run UHT to make sure the files exist for the static action graph to be setup correctly.  This is because UHT generates .cpp
@@ -2620,7 +2606,7 @@ namespace UnrealBuildTool
 
 			// Can probably be moved further down?
 #if __VPROJECT_AVAILABLE__
-			VNITask.Wait();
+			await VNITask;
 #endif
 
 			foreach (UEBuildModuleCPP Module in Modules.Values.OfType<UEBuildModuleCPP>())
