UnrealBuildTool: -GenFiles arg to QueryMode

#rnx

diff --git a/Engine/Source/Programs/UnrealBuildTool/Configuration/UEBuildTarget.cs b/Engine/Source/Programs/UnrealBuildTool/Configuration/UEBuildTarget.cs
--- a/Engine/Source/Programs/UnrealBuildTool/Configuration/UEBuildTarget.cs
+++ b/Engine/Source/Programs/UnrealBuildTool/Configuration/UEBuildTarget.cs
@@ -2336,7 +2336,7 @@ namespace UnrealBuildTool
 		/// <summary>
 		/// Builds the target, appending list of output files and returns building result.
 		/// </summary>
-		public async Task<TargetMakefile> BuildAsync(BuildConfiguration BuildConfiguration, ISourceFileWorkingSet WorkingSet, TargetDescriptor TargetDescriptor, ILogger Logger, bool bInitOnly = false)
+		public async Task<TargetMakefile> BuildAsync(BuildConfiguration BuildConfiguration, ISourceFileWorkingSet WorkingSet, TargetDescriptor TargetDescriptor, ILogger Logger, bool bInitOnly = false, bool bGenUHTOnly = false)
 		{
 			CppConfiguration CppConfiguration = GetCppConfiguration(Configuration);
 
@@ -2553,17 +2553,14 @@ namespace UnrealBuildTool
 				ExternalExecution.SetupUObjectModules(ModulesToGenerateHeadersFor, Rules.Platform, ProjectDescriptor, Makefile.UObjectModules, Makefile.UObjectModuleHeaders, Rules.GeneratedCodeVersion, MetadataCache, Logger);
 			}
 
-			List<string> UHTAdditionalArguments = new List<string>();
 			if (Rules.NativePointerMemberBehaviorOverride != null)
 			{
-				UHTAdditionalArguments.Add("-ini:Engine:[UnrealHeaderTool]:EngineNativePointerMemberBehavior=" + Rules.NativePointerMemberBehaviorOverride);
-				UHTAdditionalArguments.Add("-ini:Engine:[UnrealHeaderTool]:EnginePluginNativePointerMemberBehavior=" + Rules.NativePointerMemberBehaviorOverride);
-				UHTAdditionalArguments.Add("-ini:Engine:[UnrealHeaderTool]:NonEngineNativePointerMemberBehavior=" + Rules.NativePointerMemberBehaviorOverride);
-			}
-
-			if ((UHTAdditionalArguments != null) && (UHTAdditionalArguments.Any()))
-			{
-				Makefile.UHTAdditionalArguments = UHTAdditionalArguments.ToArray();
+				Makefile.UHTAdditionalArguments = new string[]
+				{
+					$"-ini:Engine:[UnrealHeaderTool]:EngineNativePointerMemberBehavior={Rules.NativePointerMemberBehaviorOverride}",
+					$"-ini:Engine:[UnrealHeaderTool]:EnginePluginNativePointerMemberBehavior={Rules.NativePointerMemberBehaviorOverride}",
+					$"-ini:Engine:[UnrealHeaderTool]:NonEngineNativePointerMemberBehavior={Rules.NativePointerMemberBehaviorOverride}",
+				};
 			}
 
 			// UHT mode uses this to create the makefile at least to the point where UHT would have valid manifest information
@@ -2609,6 +2606,11 @@ namespace UnrealBuildTool
 			await VNITask;
 #endif
 
+			if (bGenUHTOnly)
+			{
+				return Makefile;
+			}
+
 			foreach (UEBuildModuleCPP Module in Modules.Values.OfType<UEBuildModuleCPP>())
 			{
 				foreach ((string Subdirectory, Action<ILogger, DirectoryReference> Func) in Module.Rules.GenerateHeaderFuncs)
diff --git a/Engine/Source/Programs/UnrealBuildTool/Modes/QueryMode.cs b/Engine/Source/Programs/UnrealBuildTool/Modes/QueryMode.cs
--- a/Engine/Source/Programs/UnrealBuildTool/Modes/QueryMode.cs
+++ b/Engine/Source/Programs/UnrealBuildTool/Modes/QueryMode.cs
@@ -109,6 +109,9 @@ namespace UnrealBuildTool
 		[CommandLine("-Indented")]
 		public bool bIndented;
 
+		[CommandLine("-GenFiles")]
+		public bool bGenFiles = false;
+
 		private BuildConfiguration BuildConfiguration = new();
 
 		public override Task<int> ExecuteAsync(CommandLineArguments Arguments, ILogger Logger)
@@ -313,6 +316,13 @@ namespace UnrealBuildTool
 					}
 				}
 
+				if (bGenFiles)
+				{
+					// Create the makefile for the target and export the module information
+					using ISourceFileWorkingSet WorkingSet = new EmptySourceFileWorkingSet();
+					TargetMakefile Makefile = await CurrentTarget.BuildAsync(BuildConfiguration, WorkingSet, TargetDescriptors[0], Logger, bGenUHTOnly: true);
+				}
+
 				TargetIntellisenseInfo CurrentTargetIntellisenseInfo = new TargetIntellisenseInfo();
 
 				// Partially duplicated from UEBuildTarget.Build because we just want to get C++ compile actions without running UHT
@@ -390,6 +400,8 @@ namespace UnrealBuildTool
 					ModuleToCompileSettings = CurrentTargetIntellisenseInfo.ModuleToCompileSettings.ToImmutableSortedDictionary(x => x.Key.Name, x => x.Value),
 					LaunchSettings = CurrentLaunchSettings,
 				};
+				
+			
 
 				await WriteResultsAsync(Result, JsonOptions, Logger);
 				return 0;
