UnrealBuildTool: Fix not building files in override paths with -SingleFile

#rnx

diff --git a/Engine/Source/Programs/UnrealBuildTool/Configuration/UEBuildModuleCPP.cs b/Engine/Source/Programs/UnrealBuildTool/Configuration/UEBuildModuleCPP.cs
--- a/Engine/Source/Programs/UnrealBuildTool/Configuration/UEBuildModuleCPP.cs
+++ b/Engine/Source/Programs/UnrealBuildTool/Configuration/UEBuildModuleCPP.cs
@@ -1425,7 +1425,10 @@ namespace UnrealBuildTool
 			{
 				PrecompiledHeaderAction OldAction = CompileEnvironment.PrecompiledHeaderAction;
 				CompileEnvironment.PrecompiledHeaderAction = PrecompiledHeaderAction.None;
-				ToolChain.CreateSpecificFileAction(CompileEnvironment, ModuleDirectory, IntermediateDirectory, Graph);
+				foreach (DirectoryReference Directory in ModuleDirectories)
+				{
+					ToolChain.CreateSpecificFileAction(CompileEnvironment, Directory, IntermediateDirectory, Graph);
+				}
 				CompileEnvironment.PrecompiledHeaderAction = OldAction;
 			}
 
diff --git a/Engine/Source/Programs/UnrealBuildTool/Modes/BuildMode.cs b/Engine/Source/Programs/UnrealBuildTool/Modes/BuildMode.cs
--- a/Engine/Source/Programs/UnrealBuildTool/Modes/BuildMode.cs
+++ b/Engine/Source/Programs/UnrealBuildTool/Modes/BuildMode.cs
@@ -831,41 +831,25 @@ namespace UnrealBuildTool
 			// First, find all the SpecificFileActions.
 			// These are used to create a custom action for the source file in case it is inside a unity or normally not part of the build (headers for example)
 			Dictionary<DirectoryReference, ISpecificFileAction> SpecificFileActions = new();
-			foreach (LinkedAction Action in Actions)
+			foreach (ISpecificFileAction SpecificFileAction in Actions.Select(x => x.Inner).OfType<ISpecificFileAction>())
 			{
-				if (Action.Inner is ISpecificFileAction SpecificFileAction)
-				{
-					SpecificFileActions.TryAdd(SpecificFileAction.RootDirectory, SpecificFileAction);
-				}
+				SpecificFileActions.TryAdd(SpecificFileAction.RootDirectory, SpecificFileAction);
 			}
 
 			// Now traverse all specific files and make sure we find or create actions for them.
-			foreach (FileReference FileRef in FileList)
+			foreach (FileItem SourceFile in FileList.Select(x => FileItem.GetItemByFileReference(x)).Where(x => x.Exists))
 			{
-				FileItem SourceFile = FileItem.GetItemByFileReference(FileRef);
-
-				// Ensure the file requested exists on disk
-				if (!SourceFile.Exists)
-				{
-					continue;
-				}
-
 				ISpecificFileAction? SpecificFileAction = null;
 
 				// Traverse upwards from file directory to try to find a SpecificFileAction.
-				DirectoryItem Dir = SourceFile.Directory;
-				while (true)
+				DirectoryItem? Dir = SourceFile.Directory;
+				while (Dir != null)
 				{
 					if (SpecificFileActions.TryGetValue(Dir.Location, out SpecificFileAction))
 					{
 						break;
 					}
-					DirectoryItem? ParentDir = Dir.GetParentDirectoryItem();
-					if (ParentDir == null)
-					{
-						break;
-					}
-					Dir = ParentDir;
+					Dir = Dir.GetParentDirectoryItem();
 				}
 
 				// We found an action to take care of this file.
@@ -881,27 +865,15 @@ namespace UnrealBuildTool
 				}
 
 				// There is no special action for this file.. let's look for an action that has this exact file as input and use that (for example ispc files)
-				bool FoundAction = false;
-				foreach (LinkedAction Action in Actions)
+				LinkedAction? Action = Actions.Find(x => x.PrerequisiteItems.Contains(SourceFile));
+
+				if (Action != null)
 				{
-					foreach (FileItem PrereqItem in Action.PrerequisiteItems)
-					{
-						if (PrereqItem == SourceFile)
-						{
-							FoundAction = true;
-							ProducedItems.AddRange(Action.ProducedItems);
-							break;
-						}
-					}
-					if (FoundAction)
-					{
-						break;
-					}
+					ProducedItems.AddRange(Action.ProducedItems);
 				}
-
-				if (!FoundAction && !BuildConfiguration.bIgnoreInvalidFiles)
+				else if (!BuildConfiguration.bIgnoreInvalidFiles)
 				{
-					Logger.LogError($"{FileRef.FullName} - ERROR: Failed to find an Action that can be used to build this file (does target use this file?)");
+					Logger.LogError("ERROR: Failed to find an Action that can be used to build \"{Path}\" (does target use this file?)", SourceFile);
 				}
 			}
 
