UnrealBuildTool: Fix a rare exeception I saw when generating project files

#rnx

diff --git a/Engine/Source/Programs/UnrealBuildTool/Configuration/TestTargetRules.cs b/Engine/Source/Programs/UnrealBuildTool/Configuration/TestTargetRules.cs
--- a/Engine/Source/Programs/UnrealBuildTool/Configuration/TestTargetRules.cs
+++ b/Engine/Source/Programs/UnrealBuildTool/Configuration/TestTargetRules.cs
@@ -342,13 +342,11 @@ namespace UnrealBuildTool
 				GlobalDefinitions.Add("USE_ANDROID_LAUNCH=0");
 				GlobalDefinitions.Add("USE_ANDROID_JNI=0");
 
-				string VersionScriptFile = Path.GetTempPath() + "LLTWorkaroundScrip.ldscript";
-				using (StreamWriter Writer = System.IO.File.CreateText(VersionScriptFile))
-				{
-					// Workaround for a linker bug when building LowLevelTests for Android
-					Writer.WriteLine("{ local: *; };");
-				}
-				AdditionalLinkerArguments = " -Wl,--version-script=\"" + VersionScriptFile + "\"";
+				// Workaround for a linker bug when building LowLevelTests for Android
+				// TODO: This should be written to the intermediate directory, somewhere else not when TargetRules are being created
+				FileReference VersionScriptFile = new FileReference(Path.GetTempPath() + $"LLTWorkaroundScrip-{Name}.ldscript");
+				FileReference.WriteAllTextIfDifferent(VersionScriptFile, "{ local: *; };");
+				AdditionalLinkerArguments = " -Wl,--version-script=\"" + VersionScriptFile.FullName + "\"";
 			}
 			else if (Target.Platform == UnrealTargetPlatform.IOS)
 			{
