UnrealBuildTool: Fix overridden output path in vcxproj for non-win64

#rnx

diff --git a/Engine/Source/Programs/UnrealBuildTool/ProjectFiles/VisualStudio/VCProject.cs b/Engine/Source/Programs/UnrealBuildTool/ProjectFiles/VisualStudio/VCProject.cs
--- a/Engine/Source/Programs/UnrealBuildTool/ProjectFiles/VisualStudio/VCProject.cs
+++ b/Engine/Source/Programs/UnrealBuildTool/ProjectFiles/VisualStudio/VCProject.cs
@@ -2114,7 +2114,15 @@ namespace UnrealBuildTool
 				else
 				{
 					UnrealTargetPlatform Platform = Combination.Platform!.Value;
-					TargetRules TargetRulesObject = Combination.ProjectTarget!.TargetRules!;
+					TargetRules TargetRulesObject;
+					try
+					{
+						TargetRulesObject = Combination.ProjectTarget!.CreateRulesDelegate(Platform, Configuration);
+					}
+					catch (BuildException)
+					{
+						TargetRulesObject = Combination.ProjectTarget!.TargetRules!;
+					}
 					FileReference TargetFilePath = Combination.ProjectTarget.TargetFilePath;
 					string TargetName = TargetFilePath.GetFileNameWithoutAnyExtensions();
 					string UBTPlatformName = Platform.ToString();
@@ -2125,14 +2133,7 @@ namespace UnrealBuildTool
 					UEBuildPlatform BuildPlatform = UEBuildPlatform.GetBuildPlatform(Platform);
 
 					// Figure out if this is a monolithic build
-					bool bShouldCompileMonolithic = BuildPlatform.ShouldCompileMonolithicBinary(Platform);
-					try
-					{
-						bShouldCompileMonolithic |= (Combination.ProjectTarget.CreateRulesDelegate(Platform, Configuration).LinkType == TargetLinkType.Monolithic);
-					}
-					catch (BuildException)
-					{
-					}
+					bool bShouldCompileMonolithic = BuildPlatform.ShouldCompileMonolithicBinary(Platform) | (TargetRulesObject.LinkType == TargetLinkType.Monolithic);
 
 					// Get the .uproject directory
 					DirectoryReference? UProjectDirectory = DirectoryReference.FromFile(Combination.ProjectTarget.UnrealProjectFilePath);
