UnrealBuildTool: Experimental -VisualStudioLinux for project gen, which creates .vcxproj that allow remote debugging on linux. This currently probably only works for monolithic programs with no extra requried files

#rnx

diff --git a/Engine/Source/Programs/UnrealBuildTool/ProjectFiles/ProjectFileGenerator.cs b/Engine/Source/Programs/UnrealBuildTool/ProjectFiles/ProjectFileGenerator.cs
--- a/Engine/Source/Programs/UnrealBuildTool/ProjectFiles/ProjectFileGenerator.cs
+++ b/Engine/Source/Programs/UnrealBuildTool/ProjectFiles/ProjectFileGenerator.cs
@@ -259,6 +259,12 @@ namespace UnrealBuildTool
 		[XmlConfigFile]
 		bool bGenerateIntelliSenseData = true;
 
+		/// <summary>
+		/// True if visual studio project should be generated in linux mode.
+		/// </summary>
+		[XmlConfigFile]
+		public static bool bVisualStudioLinux = false;
+
 		/// <summary>
 		/// True if we should include documentation in the generated projects.
 		/// </summary>
@@ -1567,6 +1573,10 @@ namespace UnrealBuildTool
 							case "-FORCEUPDATEALL":
 								bForceUpdateAllFiles = true;
 								break;
+
+							case "-VISUALSTUDIOLINUX":
+								bVisualStudioLinux = true;
+								break;
 						}
 					}
 				}
diff --git a/Engine/Source/Programs/UnrealBuildTool/ProjectFiles/VisualStudio/VCProject.cs b/Engine/Source/Programs/UnrealBuildTool/ProjectFiles/VisualStudio/VCProject.cs
--- a/Engine/Source/Programs/UnrealBuildTool/ProjectFiles/VisualStudio/VCProject.cs
+++ b/Engine/Source/Programs/UnrealBuildTool/ProjectFiles/VisualStudio/VCProject.cs
@@ -1020,6 +1020,14 @@ namespace UnrealBuildTool
 				{
 					VCProjectFileContent.AppendLine("    <IsTestTarget>true</IsTestTarget>");
 				}
+				if (ProjectFileGenerator.bVisualStudioLinux)
+				{
+					VCProjectFileContent.AppendLine("      <Keyword>Linux</Keyword>");
+					VCProjectFileContent.AppendLine("      <ApplicationType>Linux</ApplicationType>");
+					VCProjectFileContent.AppendLine("      <TargetLinuxPlatform>Generic</TargetLinuxPlatform>");
+					VCProjectFileContent.AppendLine("      <ApplicationTypeRevision>1.0</ApplicationTypeRevision>");
+					VCProjectFileContent.AppendLine("      <LinuxProjectType>{D51BCBC9-82E9-4017-911E-C93873C4EA2B}</LinuxProjectType>");
+				}
 				VCProjectFileContent.AppendLine("  </PropertyGroup>");
 			}
 
@@ -2095,7 +2103,14 @@ namespace UnrealBuildTool
 					VCProjectFileContent.AppendLine("    <NMakeBuildCommandLine>@rem Nothing to do.</NMakeBuildCommandLine>");
 					VCProjectFileContent.AppendLine("    <NMakeReBuildCommandLine>@rem Nothing to do.</NMakeReBuildCommandLine>");
 					VCProjectFileContent.AppendLine("    <NMakeCleanCommandLine>@rem Nothing to do.</NMakeCleanCommandLine>");
-					VCProjectFileContent.AppendLine("    <NMakeOutput/>");
+					if (ProjectFileGenerator.bVisualStudioLinux)
+					{
+						VCProjectFileContent.AppendLine("    <BuildCommandLine>$(NMakeBuildCommandLine)</BuildCommandLine>");
+						VCProjectFileContent.AppendLine("    <ReBuildCommandLine>$(NMakeReBuildCommandLine)</ReBuildCommandLine>");
+						VCProjectFileContent.AppendLine("    <CleanCommandLine>$(NMakeCleanCommandLine)</CleanCommandLine>");
+						VCProjectFileContent.AppendLine("    <LocalBuildOutputs />");
+						VCProjectFileContent.AppendLine("    <SourcesToCopyRemotelyOverride />");
+					}
 					VCProjectFileContent.AppendLine("  </PropertyGroup>");
 				}
 				else if (Unreal.IsEngineInstalled() && Combination.ProjectTarget != null && Combination.ProjectTarget.TargetRules != null &&
@@ -2109,6 +2124,14 @@ namespace UnrealBuildTool
 					VCProjectFileContent.AppendLine("    <NMakeReBuildCommandLine>@echo {0} is not a supported platform for {1}. Valid platforms are {2}.</NMakeReBuildCommandLine>", Combination.Platform!, TargetName, ValidPlatforms);
 					VCProjectFileContent.AppendLine("    <NMakeCleanCommandLine>@echo {0} is not a supported platform for {1}. Valid platforms are {2}.</NMakeCleanCommandLine>", Combination.Platform!, TargetName, ValidPlatforms);
 					VCProjectFileContent.AppendLine("    <NMakeOutput/>");
+					if (ProjectFileGenerator.bVisualStudioLinux)
+					{
+						VCProjectFileContent.AppendLine("    <BuildCommandLine>$(NMakeBuildCommandLine)</BuildCommandLine>");
+						VCProjectFileContent.AppendLine("    <ReBuildCommandLine>$(NMakeReBuildCommandLine)</ReBuildCommandLine>");
+						VCProjectFileContent.AppendLine("    <CleanCommandLine>$(NMakeCleanCommandLine)</CleanCommandLine>");
+						VCProjectFileContent.AppendLine("    <LocalBuildOutputs />");
+						VCProjectFileContent.AppendLine("    <SourcesToCopyRemotelyOverride />");
+					}
 					VCProjectFileContent.AppendLine("  </PropertyGroup>");
 				}
 				else
@@ -2219,6 +2242,19 @@ namespace UnrealBuildTool
 						VCProjectFileContent.AppendLine("    <NMakeCleanCommandLine>$(CleanBatchScript) {0}</NMakeCleanCommandLine>", BuildArguments);
 					}
 					VCProjectFileContent.AppendLine("    <NMakeOutput>{0}</NMakeOutput>", NormalizeProjectPath(NMakePath.FullName));
+					if (ProjectFileGenerator.bVisualStudioLinux)
+					{
+						VCProjectFileContent.AppendLine("    <BuildCommandLine>$(NMakeBuildCommandLine)</BuildCommandLine>");
+						VCProjectFileContent.AppendLine("    <ReBuildCommandLine>$(NMakeReBuildCommandLine)</ReBuildCommandLine>");
+						VCProjectFileContent.AppendLine("    <CleanCommandLine>$(NMakeCleanCommandLine)</CleanCommandLine>");
+						if (TargetRulesObject.Platform.IsInGroup(UnrealPlatformGroup.Linux))
+						{
+							VCProjectFileContent.AppendLine("    <LocalBuildOutputs>{0};{1};{2}</LocalBuildOutputs>",
+								NormalizeProjectPath(NMakePath), NormalizeProjectPath(NMakePath.ChangeExtension(".debug")), NormalizeProjectPath(NMakePath.ChangeExtension(".sym")));
+							VCProjectFileContent.AppendLine("    <PreLaunchCommand>chmod +x $(RemoteDeployDir)/{0}</PreLaunchCommand>", NMakePath.GetFileName());
+							VCProjectFileContent.AppendLine("    <RemoteDebuggerCommand>$(RemoteDeployDir)/{0}</RemoteDebuggerCommand>", NMakePath.GetFileName());
+						}
+					}
 
 					if (TargetRulesObject.Type == TargetType.Game || TargetRulesObject.Type == TargetType.Client || TargetRulesObject.Type == TargetType.Server)
 					{
@@ -2237,6 +2273,15 @@ namespace UnrealBuildTool
 					VCProjectFileContent.AppendLine("    <NMakeCompile>");
 					VCProjectFileContent.AppendLine("      <NMakeCompileFileCommandLine>$(BuildBatchScript) {0} -WorkingDir=$(MSBuildProjectDirectory) -Files=$(SelectedFiles)</NMakeCompileFileCommandLine>", BuildArguments);
 					VCProjectFileContent.AppendLine("    </NMakeCompile>");
+					if (ProjectFileGenerator.bVisualStudioLinux && TargetRulesObject.Platform.IsInGroup(UnrealPlatformGroup.Linux))
+					{
+						VCProjectFileContent.AppendLine("    <PostBuildEvent>");
+						VCProjectFileContent.AppendLine("      <AdditionalSourcesToCopyMapping>{0}:=$(RemoteDeployDir)/{1};{2}:=$(RemoteDeployDir)/{3};{4}:=$(RemoteDeployDir)/{5}</AdditionalSourcesToCopyMapping>",
+							NormalizeProjectPath(NMakePath), NMakePath.GetFileName(),
+							NormalizeProjectPath(NMakePath.ChangeExtension(".debug")), NMakePath.ChangeExtension(".debug").GetFileName(),
+							NormalizeProjectPath(NMakePath.ChangeExtension(".sym")), NMakePath.ChangeExtension(".sym").GetFileName());
+						VCProjectFileContent.AppendLine("    </PostBuildEvent>");
+					}
 					VCProjectFileContent.AppendLine("  </ItemDefinitionGroup>");
 				}
 
