Fix broken AGDE symbol paths

diff --git a/Engine/Source/Programs/UnrealBuildTool/Platform/Android/AndroidProjectGenerator.cs b/Engine/Source/Programs/UnrealBuildTool/Platform/Android/AndroidProjectGenerator.cs
--- a/Engine/Source/Programs/UnrealBuildTool/Platform/Android/AndroidProjectGenerator.cs
+++ b/Engine/Source/Programs/UnrealBuildTool/Platform/Android/AndroidProjectGenerator.cs
@@ -113,13 +113,11 @@ namespace UnrealBuildTool
 
 				ProjectFileBuilder.AppendLine($"    <AndroidApkLocation>{apkLocation}</AndroidApkLocation>");
 				string intermediateRootPath = Path.GetFullPath(Path.GetDirectoryName(NMakeOutputPath.FullName) + @"\..\..\Intermediate\Android\");
-				string intermediatePath = Path.Combine(intermediateRootPath, "arm64");
-				string intermediateAGDESymbolsPath = Path.Combine(intermediateRootPath, "LLDBSymbolsLibs", "arm64");
 				List<string> symbolLocations = new List<string>
 				{
-					$@"{intermediatePath}jni\arm64-v8a",
-					$@"{intermediatePath}libs\arm64-v8a",
-					intermediateAGDESymbolsPath // support bDontBundleLibrariesInAPK
+					Path.Combine(intermediateRootPath, "arm64", "jni", "arm64-v8a"),
+					Path.Combine(intermediateRootPath, "arm64", "libs", "arm64-v8a"),
+					Path.Combine(intermediateRootPath, "LLDBSymbolsLibs", "arm64") // support bDontBundleLibrariesInAPK
 				};
 				ProjectFileBuilder.AppendLine($"    <AndroidSymbolDirectories>{string.Join(";", symbolLocations)}</AndroidSymbolDirectories>");
 
