UnrealBuildTool: pgort140.dll path for arm64

#rnx

diff --git a/Engine/Source/Programs/UnrealBuildTool/Platform/Windows/VCToolChain.cs b/Engine/Source/Programs/UnrealBuildTool/Platform/Windows/VCToolChain.cs
--- a/Engine/Source/Programs/UnrealBuildTool/Platform/Windows/VCToolChain.cs
+++ b/Engine/Source/Programs/UnrealBuildTool/Platform/Windows/VCToolChain.cs
@@ -1975,7 +1975,9 @@ namespace UnrealBuildTool
 				};
 				foreach (string PGOFile in PGOFiles)
 				{
-					FileReference SrcFile = FileReference.Combine(EnvVars.ToolChainDir, "bin", "Hostx64", "x64", PGOFile);
+					FileReference SrcFile = EnvVars.Architecture.bIsX64
+						? FileReference.Combine(EnvVars.ToolChainDir, "bin", "Hostx64", "x64", PGOFile)
+						: FileReference.Combine(EnvVars.ToolChainDir, "bin", "arm64", PGOFile);
 					FileReference DstFile = FileReference.Combine(ExeDir, PGOFile);
 					TargetFileToSourceFile[DstFile] = SrcFile;
 					RuntimeDependencies.Add(new RuntimeDependency(DstFile, StagedFileType.NonUFS));
