UnrealBuildTool: If there is an exeception validating TargetRules, print any upgrade warnings as those settings are probably related

#rnx

diff --git a/Engine/Source/Programs/UnrealBuildTool/Configuration/TargetRules.cs b/Engine/Source/Programs/UnrealBuildTool/Configuration/TargetRules.cs
--- a/Engine/Source/Programs/UnrealBuildTool/Configuration/TargetRules.cs
+++ b/Engine/Source/Programs/UnrealBuildTool/Configuration/TargetRules.cs
@@ -3070,6 +3070,16 @@ namespace UnrealBuildTool
 
 			Logger.LogDebug("Using EngineIncludeOrderVersion.{Version} for target {Target}", IncludeOrderVersion, File!.GetFileName());
 		}
+
+		/// <summary>
+		/// Prints diagnostic messages about default settings which have changed in newer versions of the engine
+		/// </summary>
+		public void PrintBuildSettingsInfoWarnings()
+		{
+			List<string> Diagnostics = new();
+			GetBuildSettingsInfo(Diagnostics);
+			Diagnostics.ForEach(x => Logger.LogWarning("Warning: {Message}", x));
+		}
 	}
 
 	/// <summary>
diff --git a/Engine/Source/Programs/UnrealBuildTool/Configuration/UEBuildTarget.cs b/Engine/Source/Programs/UnrealBuildTool/Configuration/UEBuildTarget.cs
--- a/Engine/Source/Programs/UnrealBuildTool/Configuration/UEBuildTarget.cs
+++ b/Engine/Source/Programs/UnrealBuildTool/Configuration/UEBuildTarget.cs
@@ -1302,7 +1302,15 @@ namespace UnrealBuildTool
 			// If we're using the shared build environment, make sure all the settings are valid
 			if (RulesObject.BuildEnvironment == TargetBuildEnvironment.Shared)
 			{
-				ValidateSharedEnvironment(RulesAssembly, Descriptor.Name, Descriptor.AdditionalArguments, RulesObject, Logger);
+				try
+				{
+					ValidateSharedEnvironment(RulesAssembly, Descriptor.Name, Descriptor.AdditionalArguments, RulesObject, Logger);
+				}
+				catch (Exception)
+				{
+					RulesObject.PrintBuildSettingsInfoWarnings();
+					throw;
+				}
 			}
 
 			// If we're precompiling, generate a list of all the files that we depend on
