Adding a module groups attribute for module rules.

Expected usage:

```
[ModuleGroups("MyGroup"]
public class MyModule : ModuleRules {... }

public class MyOtherModule : ModuleRules
{
    ...
    PrivateDependencyModuleNames.AddRange( GetModulesInGroup("MyGroup") );
}
```

#rnx

diff --git a/Engine/Source/Programs/UnrealBuildTool/Configuration/ModuleRules.cs b/Engine/Source/Programs/UnrealBuildTool/Configuration/ModuleRules.cs
--- a/Engine/Source/Programs/UnrealBuildTool/Configuration/ModuleRules.cs
+++ b/Engine/Source/Programs/UnrealBuildTool/Configuration/ModuleRules.cs
@@ -1728,6 +1728,43 @@ namespace UnrealBuildTool
 			return UEBuildPlatform.IsPlatformAvailableForTarget(InPlatform, Target);
 		}
 
+		/// <summary>
+		/// Returns all the modules that are in the given module group
+		/// </summary>
+		/// <param name="ModuleGroup">The name of the module group, as defined by a [ModuleGroup("Name")] attribute</param>
+		/// <param name="bOnlyValid">Only include modules that are valid for the current Target</param>
+		/// <returns></returns>
+		public IEnumerable<string> GetModulesInGroup( string ModuleGroup, bool bOnlyValid = true )
+		{
+			// figure out what platforms/groups aren't allowed with this opted in list
+			List<string>? DisallowedPlatformsAndGroups = (bOnlyValid && Target.OptedInModulePlatforms != null) ? Utils.MakeListOfUnsupportedPlatforms(Target.OptedInModulePlatforms.ToList(), false, Logger) : null;
+
+			foreach (Type ModuleType in RulesAssembly.GetTypes()
+				.Where( T => T.IsSubclassOf(typeof(ModuleRules)) && T.IsDefined(typeof(ModuleGroupsAttribute), true ))
+				.Where( T => !bOnlyValid || IsValidForTarget(T, Target, out string? _))
+				)
+			{
+				// check if the module file is disallowed
+				if (DisallowedPlatformsAndGroups != null)
+				{
+					FileReference? ModuleFileName = RulesAssembly.GetModuleFileName(ModuleType.Name);
+					if (ModuleFileName != null)
+					{
+						if (ModuleFileName.ContainsAnyNames(DisallowedPlatformsAndGroups, Unreal.EngineDirectory) ||
+							(Target.ProjectFile != null && ModuleFileName.ContainsAnyNames(DisallowedPlatformsAndGroups, Target.ProjectFile.Directory)))
+						{
+							continue;
+						}
+					}
+				}
+
+				if (ModuleType.GetCustomAttributes<ModuleGroupsAttribute>().Any( X => X.ModuleGroups.Contains(ModuleGroup)))
+				{
+					yield return ModuleType.Name;
+				}
+			}
+		}
+
 		/// <summary>
 		/// Prepares a module for building a low level tests executable.
 		/// If we're building a module as part of a test module chain and there's a Tests folder with low level tests, then they require the LowLevelTestsRunner dependency.
diff --git a/Engine/Source/Programs/UnrealBuildTool/System/ModuleGroupsAttribute.cs b/Engine/Source/Programs/UnrealBuildTool/System/ModuleGroupsAttribute.cs
new file mode 100644
--- /dev/null
+++ b/Engine/Source/Programs/UnrealBuildTool/System/ModuleGroupsAttribute.cs
@@ -0,0 +1,30 @@
+// Copyright Epic Games, Inc. All Rights Reserved.
+
+using System;
+using System.Collections.Generic;
+
+namespace UnrealBuildTool
+{
+	#nullable enable
+
+	/// <summary>
+	/// Attribute which can be applied to a ModuleRules-dervied class to indicate which module groups it belongs to
+	/// </summary>
+	[AttributeUsage(AttributeTargets.Class, AllowMultiple = true)]
+	public class ModuleGroupsAttribute : Attribute
+	{
+		/// <summary>
+		/// Array of module group names
+		/// </summary>
+		public readonly string[] ModuleGroups;
+
+		/// <summary>
+		/// Initialize the attribute with a list of module groups
+		/// </summary>
+		/// <param name="ModuleGroups">Variable-length array of module group arguments</param>
+		public ModuleGroupsAttribute(params string[] ModuleGroups)
+		{
+			this.ModuleGroups = ModuleGroups;
+		}
+	}
+}
diff --git a/Engine/Source/Programs/UnrealBuildTool/System/RulesAssembly.cs b/Engine/Source/Programs/UnrealBuildTool/System/RulesAssembly.cs
--- a/Engine/Source/Programs/UnrealBuildTool/System/RulesAssembly.cs
+++ b/Engine/Source/Programs/UnrealBuildTool/System/RulesAssembly.cs
@@ -43,6 +43,20 @@ namespace UnrealBuildTool
 			}
 		}
 
+		/// <summary>
+		/// Returns all the types contained in the compiled rules assembly
+		/// </summary>
+		/// <returns></returns>
+		public IEnumerable<Type> GetTypes()
+		{
+			if (CompiledAssembly != null)
+			{
+				return CompiledAssembly.GetTypes();
+			}
+			return Enumerable.Empty<Type>();
+		}
+
+
 		/// <summary>
 		/// The base directories for this assembly
 		/// </summary>
