Add an optional Verse language version field to uplugins. There's currently only one valid version: 0.

diff --git a/Engine/Source/Programs/Shared/EpicGames.Core/JsonWriter.cs b/Engine/Source/Programs/Shared/EpicGames.Core/JsonWriter.cs
--- a/Engine/Source/Programs/Shared/EpicGames.Core/JsonWriter.cs
+++ b/Engine/Source/Programs/Shared/EpicGames.Core/JsonWriter.cs
@@ -281,6 +281,16 @@ namespace EpicGames.Core
 			WriteValueInternal(name, value.ToString());
 		}
 
+		/// <summary>
+		/// Write a field name and unsigned integer value
+		/// </summary>
+		/// <param name="name">Name of the field</param>
+		/// <param name="value">Value for the field</param>
+		public void WriteValue(string name, uint value)
+		{
+			WriteValueInternal(name, value.ToString());
+		}
+
 		/// <summary>
 		/// Write a field name and double value
 		/// </summary>
diff --git a/Engine/Source/Programs/UnrealBuildTool/System/PluginDescriptor.cs b/Engine/Source/Programs/UnrealBuildTool/System/PluginDescriptor.cs
--- a/Engine/Source/Programs/UnrealBuildTool/System/PluginDescriptor.cs
+++ b/Engine/Source/Programs/UnrealBuildTool/System/PluginDescriptor.cs
@@ -148,6 +148,12 @@ namespace UnrealBuildTool
 		/// </summary>
 		public VerseScope VerseScope = VerseScope.User;
 
+		/// <summary>
+		/// The version of the Verse language that this plugin targets.
+		/// If no value is specified, the latest stable version is used.
+		/// </summary>
+		public uint? VerseVersion;
+
 		/// <summary>
 		/// Whether this plugin should be enabled by default for all projects
 		/// </summary>
@@ -319,6 +325,12 @@ namespace UnrealBuildTool
 				VerseScope = PluginVerseScope;
 			}
 
+			uint PluginVerseVersion;
+			if (RawObject.TryGetUnsignedIntegerField("VerseVersion", out PluginVerseVersion))
+			{
+				VerseVersion = PluginVerseVersion;
+			}
+
 			bool bEnabledByDefaultValue;
 			if (RawObject.TryGetBoolField("EnabledByDefault", out bEnabledByDefaultValue))
 			{
@@ -431,6 +443,10 @@ namespace UnrealBuildTool
 			{
 				Writer.WriteValue("VerseScope", VerseScope.ToString());
 			}
+			if (VerseVersion.HasValue)
+			{
+				Writer.WriteValue("VerseVersion", VerseVersion.Value);
+			}
 			if (bEnabledByDefault.HasValue)
 			{
 				Writer.WriteValue("EnabledByDefault", bEnabledByDefault.Value);
diff --git a/Engine/Source/Runtime/Projects/Private/PluginDescriptor.cpp b/Engine/Source/Runtime/Projects/Private/PluginDescriptor.cpp
--- a/Engine/Source/Runtime/Projects/Private/PluginDescriptor.cpp
+++ b/Engine/Source/Runtime/Projects/Private/PluginDescriptor.cpp
@@ -233,6 +233,25 @@ bool FPluginDescriptor::Read(const FJsonObject& Object, FText* OutFailReason /*=
 		}
 	}
 
+	// Read the Verse version.
+	TSharedPtr<FJsonValue> VerseVersionValue = Object.TryGetField(TEXT("VerseVersion"));
+	if (VerseVersionValue.IsValid())
+	{
+		uint32 PluginVerseVersion;
+		if (VerseVersionValue->TryGetNumber(PluginVerseVersion))
+		{
+			VerseVersion = PluginVerseVersion;
+		}
+		else
+		{
+			if (OutFailReason)
+			{
+				*OutFailReason = FText::Format(LOCTEXT("PluginWithInvalidVerseVersion", "Plugin entry 'VerseVersion' specified an unrecognized value '{1}'"), FText::FromString(VerseVersionValue->AsString()));
+			}
+			return false;
+		}
+	}
+
 	Object.TryGetBoolField(TEXT("EnableVerseAssetReflection"), bEnableVerseAssetReflection);
 
 	bool bEnabledByDefault;
diff --git a/Engine/Source/Runtime/Projects/Public/PluginDescriptor.h b/Engine/Source/Runtime/Projects/Public/PluginDescriptor.h
--- a/Engine/Source/Runtime/Projects/Public/PluginDescriptor.h
+++ b/Engine/Source/Runtime/Projects/Public/PluginDescriptor.h
@@ -9,6 +9,7 @@
 #include "CustomBuildSteps.h"
 #include "HAL/Platform.h"
 #include "LocalizationDescriptor.h"
+#include "Misc/Optional.h"
 #include "ModuleDescriptor.h"
 #include "PluginDisallowedDescriptor.h"
 #include "PluginReferenceDescriptor.h"
@@ -96,6 +97,10 @@ struct FPluginDescriptor
 	/** Origin/visibility of Verse code in this plugin's Content/Verse folder */
 	EVerseScope::Type VerseScope = EVerseScope::User;
 
+	/** The version of the Verse language that this plugin targets.
+		If no value is specified, the latest stable version is used. */
+	TOptional<uint32> VerseVersion;
+
 	/** If to generate Verse source code definitions from assets contained in this plugin */
 	bool bEnableVerseAssetReflection = false;
 
