UnrealBuildTool: Fix command line args containing include paths and defines when they are already listed. Add c & cpp args into compilerAdditionalArgs dictionary

#rnx

diff --git a/Engine/Source/Programs/UnrealBuildTool/Modes/QueryMode.cs b/Engine/Source/Programs/UnrealBuildTool/Modes/QueryMode.cs
--- a/Engine/Source/Programs/UnrealBuildTool/Modes/QueryMode.cs
+++ b/Engine/Source/Programs/UnrealBuildTool/Modes/QueryMode.cs
@@ -7,10 +7,8 @@ using System.Collections.Immutable;
 using System.IO;
 using System.Linq;
 using System.Text.Json;
-using System.Threading;
 using System.Threading.Tasks;
 using EpicGames.Core;
-using Microsoft.CodeAnalysis.CSharp.Syntax;
 using Microsoft.Extensions.Logging;
 using OpenTracing.Util;
 using UnrealBuildBase;
@@ -34,6 +32,7 @@ namespace UnrealBuildTool
 			public List<string> ForcedIncludes { get; set; } = new();
 			public string? CompilerPath { get; set; }
 			public List<string> CompilerArgs { get; set; } = new();
+			public Dictionary<string, List<string>> CompilerAdditionalArgs { get; set; } = new();
 			public string? WindowsSdkVersion { get; set; }
 		}
 
@@ -375,6 +374,12 @@ namespace UnrealBuildTool
 
 						CppCompileEnvironment ModuleCompileEnvironment = Module.CreateModuleCompileEnvironment(CurrentTarget.Rules, BinaryCompileEnvironment, Logger);
 
+						// Remove include paths and defintiions from an environment used to get the command line args
+						CppCompileEnvironment ModuleCompileEnvironmentForArgs = new CppCompileEnvironment(ModuleCompileEnvironment);
+						ModuleCompileEnvironmentForArgs.SystemIncludePaths.Clear();
+						ModuleCompileEnvironmentForArgs.UserIncludePaths.Clear();
+						ModuleCompileEnvironmentForArgs.Definitions.Clear();
+
 						TargetIntellisenseInfo.CompileSettings Settings = new TargetIntellisenseInfo.CompileSettings();
 						if (TargetToolChain is VCToolChain TargetVCToolChain)
 						{
@@ -386,7 +391,9 @@ namespace UnrealBuildTool
 						Settings.Standard = ModuleCompileEnvironment.CppStandard.ToString();
 						Settings.ForcedIncludes = ModuleCompileEnvironment.ForceIncludeFiles.Select(x => x.ToString()).ToList();
 						Settings.CompilerPath = TargetToolChain.GetCppCompilerPath()?.ToString();
-						Settings.CompilerArgs = TargetToolChain.GetGlobalCommandLineArgs(ModuleCompileEnvironment).ToList();
+						Settings.CompilerArgs = TargetToolChain.GetGlobalCommandLineArgs(ModuleCompileEnvironmentForArgs).ToList();
+						Settings.CompilerAdditionalArgs.Add("c", TargetToolChain.GetCCommandLineArgs(ModuleCompileEnvironmentForArgs).ToList());
+						Settings.CompilerAdditionalArgs.Add("cpp", TargetToolChain.GetCPPCommandLineArgs(ModuleCompileEnvironmentForArgs).ToList());
 						Settings.WindowsSdkVersion = CurrentTarget.Rules.WindowsPlatform.WindowsSdkVersion;
 						CurrentTargetIntellisenseInfo.ModuleToCompileSettings.TryAdd(Module, Settings);
 
diff --git a/Engine/Source/Programs/UnrealBuildTool/Platform/Windows/VCToolChain.cs b/Engine/Source/Programs/UnrealBuildTool/Platform/Windows/VCToolChain.cs
--- a/Engine/Source/Programs/UnrealBuildTool/Platform/Windows/VCToolChain.cs
+++ b/Engine/Source/Programs/UnrealBuildTool/Platform/Windows/VCToolChain.cs
@@ -2223,6 +2223,20 @@ namespace UnrealBuildTool
 			return Arguments;
 		}
 
+		public override IEnumerable<string> GetCPPCommandLineArgs(CppCompileEnvironment CompileEnvironment)
+		{
+			List<string> Arguments = new();
+			AppendCLArguments_CPP(CompileEnvironment, Arguments);
+			return Arguments;
+		}
+
+		public override IEnumerable<string> GetCCommandLineArgs(CppCompileEnvironment CompileEnvironment)
+		{
+			List<string> Arguments = new();
+			AppendCLArguments_C(CompileEnvironment, Arguments);
+			return Arguments;
+		}
+
 		public override CppCompileEnvironment CreateSharedResponseFile(CppCompileEnvironment CompileEnvironment, FileReference OutResponseFile, IActionGraphBuilder Graph)
 		{
 			CppCompileEnvironment NewCompileEnvironment = new CppCompileEnvironment(CompileEnvironment);
diff --git a/Engine/Source/Programs/UnrealBuildTool/ToolChain/ClangToolChain.cs b/Engine/Source/Programs/UnrealBuildTool/ToolChain/ClangToolChain.cs
--- a/Engine/Source/Programs/UnrealBuildTool/ToolChain/ClangToolChain.cs
+++ b/Engine/Source/Programs/UnrealBuildTool/ToolChain/ClangToolChain.cs
@@ -219,6 +219,11 @@ namespace UnrealBuildTool
 
 		protected abstract ClangToolChainInfo GetToolChainInfo();
 
+		public override FileReference? GetCppCompilerPath()
+		{
+			return LazyInfo.Value.Clang;
+		}
+
 		public override void SetUpGlobalEnvironment(ReadOnlyTargetRules Target)
 		{
 			base.SetUpGlobalEnvironment(Target);
@@ -1015,6 +1020,20 @@ namespace UnrealBuildTool
 			return Arguments;
 		}
 
+		public override IEnumerable<string> GetCPPCommandLineArgs(CppCompileEnvironment CompileEnvironment)
+		{
+			List<string> Arguments = new();
+			GetCompileArguments_CPP(CompileEnvironment, Arguments);
+			return Arguments;
+		}
+
+		public override IEnumerable<string> GetCCommandLineArgs(CppCompileEnvironment CompileEnvironment)
+		{
+			List<string> Arguments = new();
+			GetCompileArguments_C(CompileEnvironment, Arguments);
+			return Arguments;
+		}
+
 		public override CppCompileEnvironment CreateSharedResponseFile(CppCompileEnvironment CompileEnvironment, FileReference OutResponseFile, IActionGraphBuilder Graph)
 		{
 			CppCompileEnvironment NewCompileEnvironment = new CppCompileEnvironment(CompileEnvironment);
diff --git a/Engine/Source/Programs/UnrealBuildTool/ToolChain/UEToolChain.cs b/Engine/Source/Programs/UnrealBuildTool/ToolChain/UEToolChain.cs
--- a/Engine/Source/Programs/UnrealBuildTool/ToolChain/UEToolChain.cs
+++ b/Engine/Source/Programs/UnrealBuildTool/ToolChain/UEToolChain.cs
@@ -167,6 +167,16 @@ namespace UnrealBuildTool
 			return Array.Empty<string>();
 		}
 
+		public virtual IEnumerable<string> GetCPPCommandLineArgs(CppCompileEnvironment CompileEnvironment)
+		{
+			return Array.Empty<string>();
+		}
+
+		public virtual IEnumerable<string> GetCCommandLineArgs(CppCompileEnvironment CompileEnvironment)
+		{
+			return Array.Empty<string>();
+		}
+
 		public virtual CppCompileEnvironment CreateSharedResponseFile(CppCompileEnvironment CompileEnvironment, FileReference OutResponseFile, IActionGraphBuilder Graph)
 		{
 			return CompileEnvironment;
