UnrealBuildTool: Add RequiresUniqueBuildEnvironment to bStrictConformanceMode

#rnx

diff --git a/Engine/Source/Programs/UnrealBuildTool/Platform/Windows/UEBuildWindows.cs b/Engine/Source/Programs/UnrealBuildTool/Platform/Windows/UEBuildWindows.cs
--- a/Engine/Source/Programs/UnrealBuildTool/Platform/Windows/UEBuildWindows.cs
+++ b/Engine/Source/Programs/UnrealBuildTool/Platform/Windows/UEBuildWindows.cs
@@ -435,6 +435,7 @@ namespace UnrealBuildTool
 		/// <summary>
 		/// Enables strict standard conformance mode (/permissive-).
 		/// </summary>
+		[RequiresUniqueBuildEnvironment]
 		[XmlConfigFile(Category = "WindowsPlatform")]
 		[CommandLine("-Strict")]
 		public bool bStrictConformanceMode
@@ -461,7 +462,7 @@ namespace UnrealBuildTool
 		public bool bStrictInlineConformance = false;
 
 		/// <summary>
-		/// Enables new preprocessor conformance (/Zc:preprocessor).
+		/// Enables new preprocessor conformance (/Zc:preprocessor). This is always enabled for C++20 modules.
 		/// </summary>
 		[XmlConfigFile(Category = "WindowsPlatform")]
 		[CommandLine("-StrictPreprocessor")]
