UnrealBuildTool: Fix force includes for querymode, outputing intermediate files to a separate directory

#rnx

diff --git a/Engine/Source/Programs/UnrealBuildTool/Configuration/UEBuildTarget.cs b/Engine/Source/Programs/UnrealBuildTool/Configuration/UEBuildTarget.cs
--- a/Engine/Source/Programs/UnrealBuildTool/Configuration/UEBuildTarget.cs
+++ b/Engine/Source/Programs/UnrealBuildTool/Configuration/UEBuildTarget.cs
@@ -1121,6 +1121,10 @@ namespace UnrealBuildTool
 		/// Static code analysis
 		/// </summary>
 		Analyze,
+		/// <summary>
+		/// Query
+		/// </summary>
+		Query,
 	}
 
 	/// <summary>
@@ -1799,6 +1803,9 @@ namespace UnrealBuildTool
 				case UnrealIntermediateEnvironment.Analyze:
 					TargetFolderName += "SA";
 					break;
+				case UnrealIntermediateEnvironment.Query:
+					TargetFolderName += "QRY";
+					break;
 			}
 			return TargetFolderName;
 		}
@@ -2344,7 +2351,7 @@ namespace UnrealBuildTool
 		/// <summary>
 		/// Builds the target, appending list of output files and returns building result.
 		/// </summary>
-		public async Task<TargetMakefile> BuildAsync(BuildConfiguration BuildConfiguration, ISourceFileWorkingSet WorkingSet, TargetDescriptor TargetDescriptor, ILogger Logger, bool bInitOnly = false, bool bGenUHTOnly = false)
+		public async Task<TargetMakefile> BuildAsync(BuildConfiguration BuildConfiguration, ISourceFileWorkingSet WorkingSet, TargetDescriptor TargetDescriptor, ILogger Logger, bool bInitOnly = false)
 		{
 			CppConfiguration CppConfiguration = GetCppConfiguration(Configuration);
 
@@ -2614,11 +2621,6 @@ namespace UnrealBuildTool
 			await VNITask;
 #endif
 
-			if (bGenUHTOnly)
-			{
-				return Makefile;
-			}
-
 			foreach (UEBuildModuleCPP Module in Modules.Values.OfType<UEBuildModuleCPP>())
 			{
 				foreach ((string Subdirectory, Action<ILogger, DirectoryReference> Func) in Module.Rules.GenerateHeaderFuncs)
@@ -3645,7 +3647,7 @@ namespace UnrealBuildTool
 		/// <param name="OriginalBinaries">The list of binaries</param>
 		/// <param name="GlobalCompileEnvironment">The compile environment. The shared PCHs will be added to the SharedPCHs list in this.</param>
 		/// <param name="Logger">Logger for output</param>
-		void FindSharedPCHs(List<UEBuildBinary> OriginalBinaries, CppCompileEnvironment GlobalCompileEnvironment, ILogger Logger)
+		public void FindSharedPCHs(List<UEBuildBinary> OriginalBinaries, CppCompileEnvironment GlobalCompileEnvironment, ILogger Logger)
 		{
 			// Find how many other shared PCH modules each module depends on, and use that to sort the shared PCHs by reverse order of size.
 			HashSet<UEBuildModuleCPP> SharedPCHModules = new HashSet<UEBuildModuleCPP>();
@@ -3695,7 +3697,7 @@ namespace UnrealBuildTool
 		/// <param name="GlobalCompileEnvironment">The compile environment. The shared PCHs will be added to the SharedPCHs list in this.</param>
 		/// <param name="Logger">Logger for output</param>
 		/// <param name="Graph">List of build actions</param>
-		void CreateSharedPCHInstances(ReadOnlyTargetRules Target, UEToolChain ToolChain, List<UEBuildBinary> OriginalBinaries, CppCompileEnvironment GlobalCompileEnvironment, IActionGraphBuilder Graph, ILogger Logger)
+		public void CreateSharedPCHInstances(ReadOnlyTargetRules Target, UEToolChain ToolChain, List<UEBuildBinary> OriginalBinaries, CppCompileEnvironment GlobalCompileEnvironment, IActionGraphBuilder Graph, ILogger Logger)
 		{
 			int NumSharedPCHs = GlobalCompileEnvironment.SharedPCHs.Count;
 			if (!Target.bUsePCHFiles || NumSharedPCHs == 0)
diff --git a/Engine/Source/Programs/UnrealBuildTool/Modes/QueryMode.cs b/Engine/Source/Programs/UnrealBuildTool/Modes/QueryMode.cs
--- a/Engine/Source/Programs/UnrealBuildTool/Modes/QueryMode.cs
+++ b/Engine/Source/Programs/UnrealBuildTool/Modes/QueryMode.cs
@@ -108,9 +108,6 @@ namespace UnrealBuildTool
 		[CommandLine("-Indented")]
 		public bool bIndented;
 
-		[CommandLine("-GenFiles")]
-		public bool bGenFiles = false;
-
 		private BuildConfiguration BuildConfiguration = new();
 
 		public override Task<int> ExecuteAsync(CommandLineArguments Arguments, ILogger Logger)
@@ -290,6 +287,7 @@ namespace UnrealBuildTool
 				RawArgs.Add(ProjectFileArg.ToString());
 			}
 			RawArgs.AddRange(Arguments.GetUnusedArguments());
+			RawArgs.Add("-NoPCHChain"); // Currently unsupported
 			CommandLineArguments Args = new CommandLineArguments(RawArgs.ToArray());
 			List<TargetDescriptor> TargetDescriptors = new();
 
@@ -311,15 +309,28 @@ namespace UnrealBuildTool
 					string MutexName = SingleInstanceMutex.GetUniqueMutexForPath("UnrealBuildTool_QueryMode_UEBuildTarget-Create", Unreal.RootDirectory.FullName);
 					using (new SingleInstanceMutex(MutexName, true))
 					{
-						CurrentTarget = UEBuildTarget.Create(TargetDescriptors[0], false, false, bUsePrecompiled, Logger);
+						CurrentTarget = UEBuildTarget.Create(TargetDescriptors[0], false, false, bUsePrecompiled, UnrealIntermediateEnvironment.Query, Logger);
 					}
 				}
 
-				if (bGenFiles)
+				UEBuildBinary? LaunchBinary = CurrentTarget.Binaries.FirstOrDefault(Binary => Binary.Modules.Any(Module => Module.Name == CurrentTarget.Rules.LaunchModuleName));
+				if (LaunchBinary == null)
+				{
+					throw new BuildException("Unable to find launch binary for target");
+				}
+
+				// Create the makefile for the target and export the module information
 				{
-					// Create the makefile for the target and export the module information
 					using ISourceFileWorkingSet WorkingSet = new EmptySourceFileWorkingSet();
-					TargetMakefile Makefile = await CurrentTarget.BuildAsync(BuildConfiguration, WorkingSet, TargetDescriptors[0], Logger, bGenUHTOnly: true);
+					TargetMakefile Makefile;
+					try
+					{
+						Makefile = await BuildMode.CreateMakefileAsync(BuildConfiguration, TargetDescriptors[0], WorkingSet, Logger);
+					}
+					finally
+					{
+						SourceFileMetadataCache.SaveAll();
+					}
 				}
 
 				TargetIntellisenseInfo CurrentTargetIntellisenseInfo = new TargetIntellisenseInfo();
@@ -335,11 +346,13 @@ namespace UnrealBuildTool
 				TargetToolChain.SetEnvironmentVariables();
 				CurrentTarget.SetupGlobalEnvironment(TargetToolChain, GlobalCompileEnvironment, GlobalLinkEnvironment);
 
-				UEBuildBinary? LaunchBinary = CurrentTarget.Binaries.FirstOrDefault(Binary => Binary.Modules.Any(Module => Module.Name == CurrentTarget.Rules.LaunchModuleName));
-
-				if (LaunchBinary == null)
+				if (CurrentTarget.Rules.bUseSharedPCHs)
 				{
-					throw new BuildException("Unable to find launch binary for target");
+					// Find all the shared PCHs.
+					CurrentTarget.FindSharedPCHs(CurrentTarget.Binaries, GlobalCompileEnvironment, Logger);
+
+					// Create all the shared PCH instances before processing the modules
+					CurrentTarget.CreateSharedPCHInstances(CurrentTarget.Rules, TargetToolChain, CurrentTarget.Binaries, GlobalCompileEnvironment, new NullActionGraphBuilder(Logger), Logger);
 				}
 
 				LaunchSettings CurrentLaunchSettings = new LaunchSettings();
@@ -372,7 +385,12 @@ namespace UnrealBuildTool
 							CurrentTargetIntellisenseInfo.DirToModule.TryAdd(Module.GeneratedCodeDirectory, Module);
 						}
 
-						CppCompileEnvironment ModuleCompileEnvironment = Module.CreateModuleCompileEnvironment(CurrentTarget.Rules, BinaryCompileEnvironment, Logger);
+						CppCompileEnvironment ModuleCompileEnvironment = Module.CreateCompileEnvironmentForIntellisense(CurrentTarget.Rules, BinaryCompileEnvironment, Logger);
+						if (ModuleCompileEnvironment.PrecompiledHeaderAction == PrecompiledHeaderAction.Include)
+						{
+							FileItem IncludeHeader = FileItem.GetItemByFileReference(ModuleCompileEnvironment.PrecompiledHeaderIncludeFilename!);
+							ModuleCompileEnvironment.ForceIncludeFiles.Insert(0, IncludeHeader);
+						}
 
 						// Remove include paths and defintiions from an environment used to get the command line args
 						CppCompileEnvironment ModuleCompileEnvironmentForArgs = new CppCompileEnvironment(ModuleCompileEnvironment);
@@ -381,12 +399,12 @@ namespace UnrealBuildTool
 						ModuleCompileEnvironmentForArgs.Definitions.Clear();
 
 						TargetIntellisenseInfo.CompileSettings Settings = new TargetIntellisenseInfo.CompileSettings();
+						Settings.IncludePaths.AddRange(ModuleCompileEnvironment.UserIncludePaths.Select(x => x.ToString()));
+						Settings.IncludePaths.AddRange(ModuleCompileEnvironment.SystemIncludePaths.Select(x => x.ToString()));
 						if (TargetToolChain is VCToolChain TargetVCToolChain)
 						{
 							Settings.IncludePaths.AddRange(TargetVCToolChain.GetVCIncludePaths().Select(x => x.ToString()));
 						}
-						Settings.IncludePaths.AddRange(ModuleCompileEnvironment.SystemIncludePaths.Select(x => x.ToString()));
-						Settings.IncludePaths.AddRange(ModuleCompileEnvironment.UserIncludePaths.Select(x => x.ToString()));
 						Settings.Defines = ModuleCompileEnvironment.Definitions;
 						Settings.Standard = ModuleCompileEnvironment.CppStandard.ToString();
 						Settings.ForcedIncludes = ModuleCompileEnvironment.ForceIncludeFiles.Select(x => x.ToString()).ToList();
@@ -407,8 +425,6 @@ namespace UnrealBuildTool
 					ModuleToCompileSettings = CurrentTargetIntellisenseInfo.ModuleToCompileSettings.ToImmutableSortedDictionary(x => x.Key.Name, x => x.Value),
 					LaunchSettings = CurrentLaunchSettings,
 				};
-				
-			
 
 				await WriteResultsAsync(Result, JsonOptions, Logger);
 				return 0;
