Use llvm-strip output file argument instead of copying file.

diff --git a/Engine/Source/Programs/UnrealBuildTool/Platform/Android/UEDeployAndroid.cs b/Engine/Source/Programs/UnrealBuildTool/Platform/Android/UEDeployAndroid.cs
--- a/Engine/Source/Programs/UnrealBuildTool/Platform/Android/UEDeployAndroid.cs
+++ b/Engine/Source/Programs/UnrealBuildTool/Platform/Android/UEDeployAndroid.cs
@@ -1297,24 +1297,10 @@ namespace UnrealBuildTool
 
 		private static void StripDebugSymbols(string SourceFileName, string TargetFileName, UnrealArch UnrealArch, ILogger Logger, bool bStripAll = false)
 		{
-			// Copy the file and remove read-only if necessary
-			File.Copy(SourceFileName, TargetFileName, true);
-			FileAttributes Attribs = File.GetAttributes(TargetFileName);
-			if (Attribs.HasFlag(FileAttributes.ReadOnly))
-			{
-				File.SetAttributes(TargetFileName, Attribs & ~FileAttributes.ReadOnly);
-			}
-
 			ProcessStartInfo StartInfo = new ProcessStartInfo();
 			StartInfo.FileName = AndroidToolChain.GetStripExecutablePath(UnrealArch).Trim('"');
-			if (bStripAll)
-			{
-				StartInfo.Arguments = "--strip-unneeded \"" + TargetFileName + "\"";
-			}
-			else
-			{
-				StartInfo.Arguments = "--strip-debug \"" + TargetFileName + "\"";
-			}
+			string StripCommand = bStripAll ? "--strip-unneeded" : "--strip-debug"; 
+			StartInfo.Arguments = $"{StripCommand} -o \"{TargetFileName}\" \"{SourceFileName}\"";
 			StartInfo.UseShellExecute = false;
 			StartInfo.CreateNoWindow = true;
 			Utils.RunLocalProcessAndLogOutput(StartInfo, Logger);
