Disable linking via box

#rnx

diff --git a/Engine/Source/Programs/UnrealBuildTool/Platform/Windows/VCToolChain.cs b/Engine/Source/Programs/UnrealBuildTool/Platform/Windows/VCToolChain.cs
--- a/Engine/Source/Programs/UnrealBuildTool/Platform/Windows/VCToolChain.cs
+++ b/Engine/Source/Programs/UnrealBuildTool/Platform/Windows/VCToolChain.cs
@@ -2603,6 +2603,7 @@ namespace UnrealBuildTool
 			// Create an action that invokes the linker.
 			Action LinkAction = Graph.CreateAction(ActionType.Link);
 			string ReadableArch = UnrealArchitectureConfig.ForPlatform(LinkEnvironment.Platform).ConvertToReadableArchitecture(LinkEnvironment.Architecture);
+			LinkAction.bCanExecuteInBox = false;
 			LinkAction.CommandDescription += $"Link [{ReadableArch}]";
 			LinkAction.WorkingDirectory = Unreal.EngineSourceDirectory;
 			if (bIsBuildingLibraryOrImportLibrary)
