UnrealBuildTool: Ensure AdditionalArguments are passed to -SingleFile compiles

#rnx

diff --git a/Engine/Source/Programs/UnrealBuildTool/Platform/Windows/VCToolChain.cs b/Engine/Source/Programs/UnrealBuildTool/Platform/Windows/VCToolChain.cs
--- a/Engine/Source/Programs/UnrealBuildTool/Platform/Windows/VCToolChain.cs
+++ b/Engine/Source/Programs/UnrealBuildTool/Platform/Windows/VCToolChain.cs
@@ -961,6 +961,11 @@ namespace UnrealBuildTool
 				AddDefinition(Arguments, "_WINDOWS");
 				AddDefinition(Arguments, "WIN32");
 			}
+
+			if (!String.IsNullOrEmpty(CompileEnvironment.AdditionalArguments))
+			{
+				Arguments.Add(CompileEnvironment.AdditionalArguments);
+			}
 		}
 
 		protected virtual void AppendCLArguments_CPP(CppCompileEnvironment CompileEnvironment, List<string> Arguments)
@@ -1469,7 +1474,7 @@ namespace UnrealBuildTool
 			if (Target.WindowsPlatform.Compiler.IsMSVC())
 			{
 				// If deterministic is enabled, MSVC does not use multiple threads
-				if (!CompileEnvironment.bDeterministic)
+				if (!CompileEnvironment.bDeterministic && !CompileEnvironment.bPreprocessOnly)
 				{
 					BaseCompileAction.Weight = Target.MSVCCompileActionWeight;
 
@@ -1693,11 +1698,6 @@ namespace UnrealBuildTool
 				CompileAction.AdditionalPrerequisiteItems.AddRange(CompileEnvironment.ForceIncludeFiles);
 				CompileAction.AdditionalPrerequisiteItems.AddRange(CompileEnvironment.AdditionalPrerequisites);
 
-				if (!String.IsNullOrEmpty(CompileEnvironment.AdditionalArguments))
-				{
-					CompileAction.Arguments.Add(CompileEnvironment.AdditionalArguments);
-				}
-
 				if (SourceFile.HasExtension(".ixx"))
 				{
 					FileItem IfcFile = FileItem.GetItemByFileReference(FileReference.Combine(GetModuleInterfaceDir(OutputDir), SourceFile.Location.ChangeExtension(".ifc").GetFileName()));
