UnrealBuildTool: Add support to set NTDDI_VERSION with WindowsPlatform.TargetWindowsMinorVersion

diff --git a/Engine/Source/Programs/UnrealBuildTool/Platform/Windows/UEBuildWindows.cs b/Engine/Source/Programs/UnrealBuildTool/Platform/Windows/UEBuildWindows.cs
--- a/Engine/Source/Programs/UnrealBuildTool/Platform/Windows/UEBuildWindows.cs
+++ b/Engine/Source/Programs/UnrealBuildTool/Platform/Windows/UEBuildWindows.cs
@@ -336,8 +336,18 @@ namespace UnrealBuildTool
 		/// <summary>
 		/// Value for the WINVER macro, defining the minimum supported Windows version.
 		/// </summary>
+		[RequiresUniqueBuildEnvironment]
+		[ConfigFile(ConfigHierarchyType.Engine, "/Script/WindowsTargetPlatform.WindowsTargetSettings", "TargetWindowsVersion")]
 		public int TargetWindowsVersion = 0x601;
 
+		/// <summary>
+		/// Value for the NTDDI_VERSION macro, defining the minimum supported Windows version.
+		/// https://learn.microsoft.com/en-us/windows/win32/winprog/using-the-windows-headers?redirectedfrom=MSDN#macros-for-conditional-declarations
+		/// </summary>
+		[RequiresUniqueBuildEnvironment]
+		[ConfigFile(ConfigHierarchyType.Engine, "/Script/WindowsTargetPlatform.WindowsTargetSettings", "TargetWindowsMinorVersion")]
+		public int? TargetWindowsMinorVersion = null;
+
 		/// <summary>
 		/// Enable PIX debugging (automatically disabled in Shipping and Test configs)
 		/// </summary>
@@ -774,6 +784,8 @@ namespace UnrealBuildTool
 
 		public int TargetWindowsVersion => Inner.TargetWindowsVersion;
 
+		public int? TargetWindowsMinorVersion => Inner.TargetWindowsMinorVersion;
+
 		public bool bPixProfilingEnabled => Inner.bPixProfilingEnabled;
 
 		public bool bUseWindowsSDK10 => Inner.bUseWindowsSDK10;
@@ -1548,6 +1560,11 @@ namespace UnrealBuildTool
 			CompileEnvironment.Definitions.Add(String.Format("_WIN32_WINNT=0x{0:X4}", Target.WindowsPlatform.TargetWindowsVersion));
 			CompileEnvironment.Definitions.Add(String.Format("WINVER=0x{0:X4}", Target.WindowsPlatform.TargetWindowsVersion));
 
+			if (Target.WindowsPlatform.TargetWindowsMinorVersion != null)
+			{
+				CompileEnvironment.Definitions.Add(String.Format("NTDDI_VERSION=0x{0:X8}", Target.WindowsPlatform.TargetWindowsMinorVersion));
+			}
+
 			CompileEnvironment.Definitions.Add("PLATFORM_WINDOWS=1");
 			CompileEnvironment.Definitions.Add("PLATFORM_MICROSOFT=1");
 
