Add option to build dependant plugins when building a foreign plugin in ubt

diff --git a/Engine/Source/Programs/UnrealBuildTool/Configuration/TargetDescriptor.cs b/Engine/Source/Programs/UnrealBuildTool/Configuration/TargetDescriptor.cs
--- a/Engine/Source/Programs/UnrealBuildTool/Configuration/TargetDescriptor.cs
+++ b/Engine/Source/Programs/UnrealBuildTool/Configuration/TargetDescriptor.cs
@@ -47,6 +47,12 @@ namespace UnrealBuildTool
 		[CommandLine("-BuildPluginAsLocal")]
 		public bool bBuildPluginAsLocal = false;
 
+		/// <summary>
+		/// When building a foreign plugin, whether to build plugins it depends on as well.
+		/// </summary>
+		[CommandLine("-BuildDependantPlugins")]
+		public bool bBuildDependantPlugins = false;
+
 		/// <summary>
 		/// Set of module names to compile.
 		/// </summary>
diff --git a/Engine/Source/Programs/UnrealBuildTool/Configuration/UEBuildTarget.cs b/Engine/Source/Programs/UnrealBuildTool/Configuration/UEBuildTarget.cs
--- a/Engine/Source/Programs/UnrealBuildTool/Configuration/UEBuildTarget.cs
+++ b/Engine/Source/Programs/UnrealBuildTool/Configuration/UEBuildTarget.cs
@@ -1600,6 +1600,11 @@ namespace UnrealBuildTool
 		/// </summary>
 		public FileReference? ForeignPlugin;
 
+		/// <summary>
+		/// When building a foreign plugin, whether to build plugins it depends on as well.
+		/// </summary>
+		public bool bBuildDependantPlugins = false;
+
 		/// <summary>
 		/// Collection of all UBT plugins (project files)
 		/// </summary>
@@ -1693,6 +1698,7 @@ namespace UnrealBuildTool
 			RulesAssembly = InRulesAssembly;
 			TargetType = Rules.Type;
 			ForeignPlugin = InDescriptor.ForeignPlugin;
+			bBuildDependantPlugins = InDescriptor.bBuildDependantPlugins;
 			bDeployAfterCompile = InRules.bDeployAfterCompile && !InRules.bDisableLinking && InDescriptor.OnlyModuleNames.Count == 0;
 
 			// now that we have the platform, we can set the intermediate path to include the platform/architecture name
@@ -2698,9 +2704,14 @@ namespace UnrealBuildTool
 			if (ForeignPlugin != null)
 			{
 				HashSet<FileItem> RetainOutputItems = new HashSet<FileItem>();
+
+				UEBuildPlugin? ForeignBuildPlugin = BuildPlugins.Find(x => x.File == ForeignPlugin);
+				
 				foreach (UEBuildPlugin Plugin in BuildPlugins)
 				{
-					if (Plugin.File == ForeignPlugin)
+					// Retain foreign plugin dependencies if it was specified.
+					bool? bIsForeignPluginDependency = ForeignBuildPlugin?.Dependencies?.Contains(Plugin);
+					if (Plugin.File == ForeignPlugin || (bBuildDependantPlugins && bIsForeignPluginDependency == true))
 					{
 						foreach (UEBuildModule Module in Plugin.Modules)
 						{
