Make ModuleManifest public to allow reading module manifest files in automation scripts

diff --git a/Engine/Source/Programs/UnrealBuildTool/System/ModuleManifest.cs b/Engine/Source/Programs/UnrealBuildTool/System/ModuleManifest.cs
--- a/Engine/Source/Programs/UnrealBuildTool/System/ModuleManifest.cs
+++ b/Engine/Source/Programs/UnrealBuildTool/System/ModuleManifest.cs
@@ -17,9 +17,16 @@ namespace UnrealBuildTool
 	/// module when making versioned builds.
 	/// </summary>
 	[Serializable]
-	class ModuleManifest
+	public class ModuleManifest
 	{
+		/// <summary>
+		/// Unique build id for the modules described in the manifest.
+		/// </summary>
 		public string BuildId;
+
+		/// <summary>
+		/// Map of module name to file name.
+		/// </summary>
 		public Dictionary<string, string> ModuleNameToFileName = new Dictionary<string, string>();
 
 		/// <summary>
