UnrealBuildTool: Fix querymode not generating UHT headers in the correct intermediate location

#rnx

diff --git a/Engine/Source/Programs/UnrealBuildTool/Modes/QueryMode.cs b/Engine/Source/Programs/UnrealBuildTool/Modes/QueryMode.cs
--- a/Engine/Source/Programs/UnrealBuildTool/Modes/QueryMode.cs
+++ b/Engine/Source/Programs/UnrealBuildTool/Modes/QueryMode.cs
@@ -297,6 +297,9 @@ namespace UnrealBuildTool
 				// TOOD: Error 
 				return 1;
 			}
+			
+			// Ensure the intermediate environment does not conflict with normal builds
+			TargetDescriptors[0].IntermediateEnvironment = UnrealIntermediateEnvironment.Query;
 
 			try
 			{
@@ -309,7 +312,7 @@ namespace UnrealBuildTool
 					string MutexName = SingleInstanceMutex.GetUniqueMutexForPath("UnrealBuildTool_QueryMode_UEBuildTarget-Create", Unreal.RootDirectory.FullName);
 					using (new SingleInstanceMutex(MutexName, true))
 					{
-						CurrentTarget = UEBuildTarget.Create(TargetDescriptors[0], false, false, bUsePrecompiled, UnrealIntermediateEnvironment.Query, Logger);
+						CurrentTarget = UEBuildTarget.Create(TargetDescriptors[0], false, false, bUsePrecompiled, TargetDescriptors[0].IntermediateEnvironment, Logger);
 					}
 				}
 
