UnrealBuildTool: Fix PVS not printing warnings

#rnx

diff --git a/Engine/Source/Programs/UnrealBuildTool/Platform/Windows/PVSToolChain.cs b/Engine/Source/Programs/UnrealBuildTool/Platform/Windows/PVSToolChain.cs
--- a/Engine/Source/Programs/UnrealBuildTool/Platform/Windows/PVSToolChain.cs
+++ b/Engine/Source/Programs/UnrealBuildTool/Platform/Windows/PVSToolChain.cs
@@ -442,7 +442,7 @@ namespace UnrealBuildTool
 										// Output the line to the log
 										if (!bFalseAlarm && Level == 1)
 										{
-											Log.WriteLine(LogEventType.Warning, LogFormatOptions.NoSeverityPrefix, "{0}({1}): warning {2}: {3}", FileName, LineNumber, WarningCode, WarningMessage);
+											Logger.LogWarning("{Path}({LineNumber}): warning {WarningCode}: {WarningMessage}", FileName, LineNumber, WarningCode, WarningMessage);
 										}
 									}
 								}
@@ -450,7 +450,7 @@ namespace UnrealBuildTool
 
 							if (!bCanParse)
 							{
-								Log.WriteLine(LogEventType.Warning, LogFormatOptions.NoSeverityPrefix, "{0}({1}): warning: Unable to parse PVS output line '{2}' (tokens=|{3}|)", InputFile, LineIdx + 1, Line, String.Join("|", Tokens));
+								Logger.LogWarning("{Path}({LineNumber}): warning: Unable to parse PVS output line '{Line}' (tokens=|{Tokens}|)", InputFile, LineIdx + 1, Line, String.Join("|", Tokens));
 							}
 						}
 					}
@@ -525,7 +525,7 @@ namespace UnrealBuildTool
 			base.GetVersionInfo(Lines);
 
 			ReadOnlyPVSTargetSettings Settings = Target.WindowsPlatform.PVS;
-			Lines.Add(String.Format("Using PVS-Studio installation at {0} with analysis mode {1} ({2})", AnalyzerFile, (uint)Settings.ModeFlags, Settings.ModeFlags.ToString()));
+			Lines.Add(String.Format("Using PVS-Studio {0} at {1} with analysis mode {2} ({3})", AnalyzerVersion, AnalyzerFile, (uint)Settings.ModeFlags, Settings.ModeFlags.ToString()));
 		}
 
 		static Version GetAnalyzerVersion(FileReference AnalyzerPath)
@@ -786,23 +786,30 @@ namespace UnrealBuildTool
 				AnalyzeAction.WorkingDirectory = Unreal.EngineSourceDirectory;
 				AnalyzeAction.CommandPath = AnalyzerFile;
 
-				StringBuilder Arguments = new StringBuilder();
-				Arguments.Append($"--source-file \"{SourceFileItem.AbsolutePath}\" --output-file \"{OutputFileLocation}\" --cfg \"{ConfigFileItem.AbsolutePath}\" --i-file=\"{PreprocessedFileItem.AbsolutePath}\" --analysis-mode {(uint)Settings.ModeFlags}");
+				List<string> Arguments = new();
+				Arguments.Add($"--source-file \"{SourceFileItem.AbsolutePath}\"");
+				Arguments.Add($"--output-file \"{OutputFileItem.AbsolutePath}\"");
+				Arguments.Add($"--cfg \"{ConfigFileItem.AbsolutePath}\"");
+				Arguments.Add($"--i-file=\"{PreprocessedFileItem.AbsolutePath}\"");
+				Arguments.Add($"--analysis-mode {(uint)Settings.ModeFlags}");
+
 				if (LicenseFile != null)
 				{
-					Arguments.Append($" --lic-file \"{LicenseFile}\"");
+					Arguments.Add($"--lic-file \"{LicenseFile}\"");
 					AnalyzeAction.PrerequisiteItems.Add(FileItem.GetItemByFileReference(LicenseFile));
 				}
-				AnalyzeAction.CommandArguments = Arguments.ToString();
+				AnalyzeAction.CommandArguments = String.Join(' ', Arguments);
 
 				AnalyzeAction.PrerequisiteItems.Add(ConfigFileItem);
 				AnalyzeAction.PrerequisiteItems.Add(PreprocessedFileItem);
 				AnalyzeAction.PrerequisiteItems.UnionWith(InputFiles); // Add the InputFiles as PrerequisiteItems so that in SingleFileCompile mode the PVSAnalyze step is not filtered out
 				AnalyzeAction.ProducedItems.Add(OutputFileItem);
 				AnalyzeAction.DeleteItems.Add(OutputFileItem); // PVS Studio will append by default, so need to delete produced items
+				AnalyzeAction.Weight = Target.MSVCCompileActionWeight * 4.0; // Very high memory usage
 				AnalyzeAction.bCanExecuteRemotely = true;
 				AnalyzeAction.bCanExecuteRemotelyWithXGE = false;
 				AnalyzeAction.bCanExecuteRemotelyWithSNDBS = false;
+				AnalyzeAction.bCanExecuteInBox = false;
 
 				Result.ObjectFiles.AddRange(AnalyzeAction.ProducedItems);
 			}
