UnrealBuildTool: Always use separate intermediate environment for analysis

#rnx

diff --git a/Engine/Source/Programs/UnrealBuildTool/Configuration/TargetRules.cs b/Engine/Source/Programs/UnrealBuildTool/Configuration/TargetRules.cs
--- a/Engine/Source/Programs/UnrealBuildTool/Configuration/TargetRules.cs
+++ b/Engine/Source/Programs/UnrealBuildTool/Configuration/TargetRules.cs
@@ -2734,16 +2734,17 @@ namespace UnrealBuildTool
 			bDeployAfterCompile = bForceSkipDeploy ? false : bDeployAfterCompile;
 
 			// Determine intermediate environment overrides based on command line flags
-			if (IntermediateEnvironment == UnrealIntermediateEnvironment.Default)
+			if (StaticAnalyzer != StaticAnalyzer.None)
+			{
+				// Always override environment for analyzing, regardless of other settings
+				IntermediateEnvironment = UnrealIntermediateEnvironment.Analyze;
+			}
+			else if (IntermediateEnvironment == UnrealIntermediateEnvironment.Default)
 			{
 				if (bIWYU)
 				{
 					IntermediateEnvironment = UnrealIntermediateEnvironment.IWYU;
 				}
-				else if (StaticAnalyzer != StaticAnalyzer.None)
-				{
-					IntermediateEnvironment = UnrealIntermediateEnvironment.Analyze;
-				}
 				else if (!bUseUnityBuild)
 				{
 					IntermediateEnvironment = UnrealIntermediateEnvironment.NonUnity;
