Disable WriteMedataAction running through Box

speculative fix for a failure

#rnx

diff --git a/Engine/Source/Programs/UnrealBuildTool/Configuration/UEBuildTarget.cs b/Engine/Source/Programs/UnrealBuildTool/Configuration/UEBuildTarget.cs
--- a/Engine/Source/Programs/UnrealBuildTool/Configuration/UEBuildTarget.cs
+++ b/Engine/Source/Programs/UnrealBuildTool/Configuration/UEBuildTarget.cs
@@ -3050,6 +3050,7 @@ namespace UnrealBuildTool
 				WriteMetadataAction.WorkingDirectory = Unreal.EngineSourceDirectory;
 				WriteMetadataAction.StatusDescription = StatusDescription;
 				WriteMetadataAction.bCanExecuteRemotely = false;
+				WriteMetadataAction.bCanExecuteInBox = false;
 				WriteMetadataAction.bUseActionHistory = false; // Different files for each target; do not want to invalidate based on this.
 				WriteMetadataAction.ProducedItems.UnionWith(ProducedItems.Select(x => FileItem.GetItemByFileReference(x)));
 				WriteMetadataAction.PrerequisiteItems.UnionWith(PrerequisiteItems);
