UnrealBuildTool: Less agressive weighting of PVS

#rnx

diff --git a/Engine/Source/Programs/UnrealBuildTool/Platform/Windows/PVSToolChain.cs b/Engine/Source/Programs/UnrealBuildTool/Platform/Windows/PVSToolChain.cs
--- a/Engine/Source/Programs/UnrealBuildTool/Platform/Windows/PVSToolChain.cs
+++ b/Engine/Source/Programs/UnrealBuildTool/Platform/Windows/PVSToolChain.cs
@@ -518,6 +518,11 @@ namespace UnrealBuildTool
 					LicenseFile = DefaultLicenseFile;
 				}
 			}
+
+			if (BuildHostPlatform.Current.IsRunningOnWine())
+			{
+				throw new BuildException("PVS-Studio is not supported with Wine.");
+			}
 		}
 
 		public override void GetVersionInfo(List<string> Lines)
@@ -805,11 +810,9 @@ namespace UnrealBuildTool
 				AnalyzeAction.PrerequisiteItems.UnionWith(InputFiles); // Add the InputFiles as PrerequisiteItems so that in SingleFileCompile mode the PVSAnalyze step is not filtered out
 				AnalyzeAction.ProducedItems.Add(OutputFileItem);
 				AnalyzeAction.DeleteItems.Add(OutputFileItem); // PVS Studio will append by default, so need to delete produced items
-				AnalyzeAction.Weight = Target.MSVCCompileActionWeight * 4.0; // Very high memory usage
-				AnalyzeAction.bCanExecuteRemotely = true;
-				AnalyzeAction.bCanExecuteRemotelyWithXGE = false;
-				AnalyzeAction.bCanExecuteRemotelyWithSNDBS = false;
-				AnalyzeAction.bCanExecuteInBox = false;
+				AnalyzeAction.Weight = Target.MSVCCompileActionWeight * 2.0; // Very high memory usage
+				AnalyzeAction.bCanExecuteRemotely = false;
+				AnalyzeAction.bCanExecuteInBox = true;
 
 				Result.ObjectFiles.AddRange(AnalyzeAction.ProducedItems);
 			}
