UnrealBuildTool: Box -> UBA. Command line args are doubled up until all usage can be updated

#rnx

diff --git a/Engine/Source/Programs/UnrealBuildTool/Configuration/BuildConfiguration.cs b/Engine/Source/Programs/UnrealBuildTool/Configuration/BuildConfiguration.cs
--- a/Engine/Source/Programs/UnrealBuildTool/Configuration/BuildConfiguration.cs
+++ b/Engine/Source/Programs/UnrealBuildTool/Configuration/BuildConfiguration.cs
@@ -40,7 +40,7 @@ namespace UnrealBuildTool
 		public bool bAllowHybridExecutor = false;
 
 		/// <summary>
-		/// Whether the experimental box executor will be used.
+		/// This property is being renamed and should not be used
 		/// </summary>
 		[XmlConfigFile]
 		[CommandLine("-Box", Value = "true")]
@@ -48,17 +48,36 @@ namespace UnrealBuildTool
 		public bool bAllowBoxExecutor = false;
 
 		/// <summary>
-		/// Whether the experimental box (local only) executor will be used.
+		/// Whether the experimental UnrealBuildAccelerator executor will be used.
+		/// </summary>
+		[CommandLine("-UBA", Value = "true")]
+		[CommandLine("-NoUBA", Value = "false")]
+		public bool bAllowUBAExecutor
+		{
+			get => bAllowBoxExecutor || bAllowUBAExecutorPrivate;
+			set => bAllowUBAExecutorPrivate = value;
+		}
+		private bool bAllowUBAExecutorPrivate = false;
+
+		/// <summary>
+		/// This property is being renamed and should not be used
 		/// </summary>
 		[XmlConfigFile]
 		[CommandLine("-BoxLocal", Value = "true")]
 		[CommandLine("-NoBoxLocal", Value = "false")]
-		public bool bAllowBoxLocalExecutor
+		public bool bAllowBoxLocalExecutor = false;
+
+		/// <summary>
+		/// Whether the experimental UnrealBuildAccelerator (local only) executor will be used.
+		/// </summary>
+		[CommandLine("-UBALocal", Value = "true")]
+		[CommandLine("-NoUBALocal", Value = "false")]
+		public bool bAllowUBALocalExecutor
 		{
-			get => bAllowBoxExecutor || bAllowBoxLocalExecutorPrivate;
-			set => bAllowBoxLocalExecutorPrivate = value;
+			get => bAllowUBAExecutor || bAllowBoxLocalExecutor || bAllowUBALocalExecutorPrivate;
+			set => bAllowUBALocalExecutorPrivate = value;
 		}
-		private bool bAllowBoxLocalExecutorPrivate = false;
+		private bool bAllowUBALocalExecutorPrivate = false;
 
 		/// <summary>
 		/// Whether XGE may be used if available, default is true.
diff --git a/Engine/Source/Programs/UnrealBuildTool/Configuration/UEBuildTarget.cs b/Engine/Source/Programs/UnrealBuildTool/Configuration/UEBuildTarget.cs
--- a/Engine/Source/Programs/UnrealBuildTool/Configuration/UEBuildTarget.cs
+++ b/Engine/Source/Programs/UnrealBuildTool/Configuration/UEBuildTarget.cs
@@ -3050,7 +3050,7 @@ namespace UnrealBuildTool
 				WriteMetadataAction.WorkingDirectory = Unreal.EngineSourceDirectory;
 				WriteMetadataAction.StatusDescription = StatusDescription;
 				WriteMetadataAction.bCanExecuteRemotely = false;
-				WriteMetadataAction.bCanExecuteInBox = false;
+				WriteMetadataAction.bCanExecuteInUBA = false;
 				WriteMetadataAction.bUseActionHistory = false; // Different files for each target; do not want to invalidate based on this.
 				WriteMetadataAction.ProducedItems.UnionWith(ProducedItems.Select(x => FileItem.GetItemByFileReference(x)));
 				WriteMetadataAction.PrerequisiteItems.UnionWith(PrerequisiteItems);
diff --git a/Engine/Source/Programs/UnrealBuildTool/Platform/Linux/LinuxToolChain.cs b/Engine/Source/Programs/UnrealBuildTool/Platform/Linux/LinuxToolChain.cs
--- a/Engine/Source/Programs/UnrealBuildTool/Platform/Linux/LinuxToolChain.cs
+++ b/Engine/Source/Programs/UnrealBuildTool/Platform/Linux/LinuxToolChain.cs
@@ -1124,7 +1124,7 @@ namespace UnrealBuildTool
 			}
 
 			// Saw a 6 hour link time potentially caused by box. Will disable for now and revisit later
-			LinkAction.bCanExecuteInBox = !LinkEnvironment.bPGOProfile && !LinkEnvironment.bPGOOptimize && !LinkEnvironment.bAllowLTCG;
+			LinkAction.bCanExecuteInUBA = !LinkEnvironment.bPGOProfile && !LinkEnvironment.bPGOOptimize && !LinkEnvironment.bAllowLTCG;
 
 			// because the logic choosing between lld and ld is somewhat messy atm (lld fails to link .DSO due to bugs), make the name of the linker clear
 			LinkAction.CommandDescription += (LinkCommandString.Contains("-fuse-ld=lld")) ? " (lld)" : " (ld)";
diff --git a/Engine/Source/Programs/UnrealBuildTool/Platform/Windows/PVSToolChain.cs b/Engine/Source/Programs/UnrealBuildTool/Platform/Windows/PVSToolChain.cs
--- a/Engine/Source/Programs/UnrealBuildTool/Platform/Windows/PVSToolChain.cs
+++ b/Engine/Source/Programs/UnrealBuildTool/Platform/Windows/PVSToolChain.cs
@@ -812,7 +812,7 @@ namespace UnrealBuildTool
 				AnalyzeAction.DeleteItems.Add(OutputFileItem); // PVS Studio will append by default, so need to delete produced items
 				AnalyzeAction.Weight = Target.MSVCCompileActionWeight * 2.0; // Very high memory usage
 				AnalyzeAction.bCanExecuteRemotely = false;
-				AnalyzeAction.bCanExecuteInBox = true;
+				AnalyzeAction.bCanExecuteInUBA = true;
 
 				Result.ObjectFiles.AddRange(AnalyzeAction.ProducedItems);
 			}
diff --git a/Engine/Source/Programs/UnrealBuildTool/Platform/Windows/VCCompileAction.cs b/Engine/Source/Programs/UnrealBuildTool/Platform/Windows/VCCompileAction.cs
--- a/Engine/Source/Programs/UnrealBuildTool/Platform/Windows/VCCompileAction.cs
+++ b/Engine/Source/Programs/UnrealBuildTool/Platform/Windows/VCCompileAction.cs
@@ -186,7 +186,7 @@ namespace UnrealBuildTool
 		public bool bCanExecuteRemotelyWithXGE { get; set; } = true;
 
 		/// <inheritdoc/>
-		public bool bCanExecuteInBox { get; set; } = true;
+		public bool bCanExecuteInUBA { get; set; } = true;
 
 		/// <inheritdoc/>
 		public bool bUseActionHistory => true;
@@ -417,7 +417,7 @@ namespace UnrealBuildTool
 			bCanExecuteRemotely = Reader.ReadBool();
 			bCanExecuteRemotelyWithSNDBS = Reader.ReadBool();
 			bCanExecuteRemotelyWithXGE = Reader.ReadBool();
-			bCanExecuteInBox = Reader.ReadBool();
+			bCanExecuteInUBA = Reader.ReadBool();
 			Architecture = UnrealArch.Parse(Reader.ReadString()!);
 			Weight = Reader.ReadDouble();
 
@@ -456,7 +456,7 @@ namespace UnrealBuildTool
 			Writer.WriteBool(bCanExecuteRemotely);
 			Writer.WriteBool(bCanExecuteRemotelyWithSNDBS);
 			Writer.WriteBool(bCanExecuteRemotelyWithXGE);
-			Writer.WriteBool(bCanExecuteInBox);
+			Writer.WriteBool(bCanExecuteInUBA);
 			Writer.WriteString(Architecture.ToString());
 			Writer.WriteDouble(Weight);
 
diff --git a/Engine/Source/Programs/UnrealBuildTool/Platform/Windows/VCToolChain.cs b/Engine/Source/Programs/UnrealBuildTool/Platform/Windows/VCToolChain.cs
--- a/Engine/Source/Programs/UnrealBuildTool/Platform/Windows/VCToolChain.cs
+++ b/Engine/Source/Programs/UnrealBuildTool/Platform/Windows/VCToolChain.cs
@@ -2651,7 +2651,7 @@ namespace UnrealBuildTool
 
 			if (LinkEnvironment.bPGOOptimize || LinkEnvironment.bPGOProfile)
 			{
-				LinkAction.bCanExecuteInBox = false; // Disabled for now. Should revisit to see why it is not working
+				LinkAction.bCanExecuteInUBA = false; // Disabled for now. Should revisit to see why it is not working
 			}
 
 			// Create link repro if requested, this argument is intentionally not added to the response file
@@ -2663,7 +2663,7 @@ namespace UnrealBuildTool
 				LinkAction.CommandArguments = $"{LinkAction.CommandArguments} /LINKREPRO:{Utils.MakePathSafeToUseWithCommandLine(LinkReproPath.FullName)}";
 
 				LinkAction.bCanExecuteRemotely = false;
-				LinkAction.bCanExecuteInBox = false;
+				LinkAction.bCanExecuteInUBA = false;
 			}
 
 			Logger.LogDebug("     Linking: {StatusDescription}", LinkAction.StatusDescription);
diff --git a/Engine/Source/Programs/UnrealBuildTool/System/Action.cs b/Engine/Source/Programs/UnrealBuildTool/System/Action.cs
--- a/Engine/Source/Programs/UnrealBuildTool/System/Action.cs
+++ b/Engine/Source/Programs/UnrealBuildTool/System/Action.cs
@@ -148,7 +148,7 @@ namespace UnrealBuildTool
 		/// <summary>
 		/// True if this action can be executed by box
 		/// </summary>
-		bool bCanExecuteInBox { get; }
+		bool bCanExecuteInUBA { get; }
 
 		/// <summary>
 		/// True if this action is using the GCC compiler.  Some build systems may be able to optimize for this case.
@@ -270,7 +270,7 @@ namespace UnrealBuildTool
 		/// <summary>
 		/// True if this action can be executed in box
 		/// </summary>
-		public bool bCanExecuteInBox { get; set; } = true;
+		public bool bCanExecuteInUBA { get; set; } = true;
 
 		/// <summary>
 		/// True if this action is using the GCC compiler.  Some build systems may be able to optimize for this case.
@@ -329,7 +329,7 @@ namespace UnrealBuildTool
 			bCanExecuteRemotely = InOther.bCanExecuteRemotely;
 			bCanExecuteRemotelyWithSNDBS = InOther.bCanExecuteRemotelyWithSNDBS;
 			bCanExecuteRemotelyWithXGE = InOther.bCanExecuteRemotelyWithXGE;
-			bCanExecuteInBox = InOther.bCanExecuteInBox;
+			bCanExecuteInUBA = InOther.bCanExecuteInUBA;
 			bIsGCCCompiler = InOther.bIsGCCCompiler;
 			bShouldOutputStatusDescription = InOther.bShouldOutputStatusDescription;
 			bProducesImportLibrary = InOther.bProducesImportLibrary;
@@ -351,7 +351,7 @@ namespace UnrealBuildTool
 			bCanExecuteRemotely = Reader.ReadBool();
 			bCanExecuteRemotelyWithSNDBS = Reader.ReadBool();
 			bCanExecuteRemotelyWithXGE = Reader.ReadBool();
-			bCanExecuteInBox = Reader.ReadBool();
+			bCanExecuteInUBA = Reader.ReadBool();
 			bIsGCCCompiler = Reader.ReadBool();
 			bShouldOutputStatusDescription = Reader.ReadBool();
 			bProducesImportLibrary = Reader.ReadBool();
@@ -381,7 +381,7 @@ namespace UnrealBuildTool
 			Writer.WriteBool(bCanExecuteRemotely);
 			Writer.WriteBool(bCanExecuteRemotelyWithSNDBS);
 			Writer.WriteBool(bCanExecuteRemotelyWithXGE);
-			Writer.WriteBool(bCanExecuteInBox);
+			Writer.WriteBool(bCanExecuteInUBA);
 			Writer.WriteBool(bIsGCCCompiler);
 			Writer.WriteBool(bShouldOutputStatusDescription);
 			Writer.WriteBool(bProducesImportLibrary);
@@ -463,9 +463,9 @@ namespace UnrealBuildTool
 			}
 
 			bool bCanExecuteInBox;
-			if (Object.TryGetBoolField("bCanExecuteInBox", out bCanExecuteInBox))
+			if (Object.TryGetBoolField("bCanExecuteInUBA", out bCanExecuteInBox))
 			{
-				Action.bCanExecuteInBox = bCanExecuteInBox;
+				Action.bCanExecuteInUBA = bCanExecuteInBox;
 			}
 
 			bool bIsGCCCompiler;
@@ -576,7 +576,7 @@ namespace UnrealBuildTool
 			Writer.WriteValue("bCanExecuteRemotely", Action.bCanExecuteRemotely);
 			Writer.WriteValue("bCanExecuteRemotelyWithSNDBS", Action.bCanExecuteRemotelyWithSNDBS);
 			Writer.WriteValue("bCanExecuteRemotelyWithXGE", Action.bCanExecuteRemotelyWithXGE);
-			Writer.WriteValue("bCanExecuteInBox", Action.bCanExecuteInBox);
+			Writer.WriteValue("bCanExecuteInUBA", Action.bCanExecuteInUBA);
 			Writer.WriteValue("bIsGCCCompiler", Action.bIsGCCCompiler);
 			Writer.WriteValue("bShouldOutputStatusDescription", Action.bShouldOutputStatusDescription);
 			Writer.WriteValue("bProducesImportLibrary", Action.bProducesImportLibrary);
@@ -698,7 +698,7 @@ namespace UnrealBuildTool
 		public bool bCanExecuteRemotely => Inner.bCanExecuteRemotely;
 		public bool bCanExecuteRemotelyWithSNDBS => Inner.bCanExecuteRemotelyWithSNDBS;
 		public bool bCanExecuteRemotelyWithXGE => Inner.bCanExecuteRemotelyWithXGE;
-		public bool bCanExecuteInBox => Inner.bCanExecuteInBox;
+		public bool bCanExecuteInUBA => Inner.bCanExecuteInUBA;
 		public bool bIsGCCCompiler => Inner.bIsGCCCompiler;
 		public bool bShouldOutputStatusDescription => Inner.bShouldOutputStatusDescription;
 		public bool bProducesImportLibrary => Inner.bProducesImportLibrary;
diff --git a/Engine/Source/Programs/UnrealBuildTool/System/ActionGraph.cs b/Engine/Source/Programs/UnrealBuildTool/System/ActionGraph.cs
--- a/Engine/Source/Programs/UnrealBuildTool/System/ActionGraph.cs
+++ b/Engine/Source/Programs/UnrealBuildTool/System/ActionGraph.cs
@@ -364,7 +364,7 @@ namespace UnrealBuildTool
 			if (ActionCount > ParallelExecutor.GetDefaultNumParallelProcesses(BuildConfiguration.MaxParallelActions, BuildConfiguration.bAllCores, Logger))
 			{
 #if __UBAEXECUTOR_AVAILABLE__
-				if (BuildConfiguration.bAllowBoxExecutor && UBAExecutor.IsAvailable(Logger))
+				if (BuildConfiguration.bAllowUBAExecutor && UBAExecutor.IsAvailable(Logger))
 				{
 					return new UBAExecutor(BuildConfiguration.MaxParallelActions, BuildConfiguration.bAllCores, BuildConfiguration.bCompactOutput, Logger, TargetDescriptors.FirstOrDefault()?.AdditionalArguments);
 				}
@@ -385,7 +385,7 @@ namespace UnrealBuildTool
 			}
 
 #if __UBAEXECUTOR_AVAILABLE__
-			if (BuildConfiguration.bAllowBoxLocalExecutor && UBALocalExecutor.IsAvailable(Logger))
+			if (BuildConfiguration.bAllowUBALocalExecutor && UBALocalExecutor.IsAvailable(Logger))
 			{
 				return new UBALocalExecutor(BuildConfiguration.MaxParallelActions, BuildConfiguration.bAllCores, BuildConfiguration.bCompactOutput, Logger, TargetDescriptors.FirstOrDefault()?.AdditionalArguments);
 			}
