UnrealBuildTool: Move common action properties for CreateRecursiveAction to be shared

#rnx

diff --git a/Engine/Source/Programs/UnrealBuildTool/Configuration/UEBuildTarget.cs b/Engine/Source/Programs/UnrealBuildTool/Configuration/UEBuildTarget.cs
--- a/Engine/Source/Programs/UnrealBuildTool/Configuration/UEBuildTarget.cs
+++ b/Engine/Source/Programs/UnrealBuildTool/Configuration/UEBuildTarget.cs
@@ -3049,8 +3049,6 @@ namespace UnrealBuildTool
 				Action WriteMetadataAction = MakefileBuilder.CreateRecursiveAction<WriteMetadataMode>(ActionType.WriteMetadata, WriteMetadataArguments.ToString());
 				WriteMetadataAction.WorkingDirectory = Unreal.EngineSourceDirectory;
 				WriteMetadataAction.StatusDescription = StatusDescription;
-				WriteMetadataAction.bCanExecuteRemotely = false;
-				WriteMetadataAction.bCanExecuteInUBA = false;
 				WriteMetadataAction.bUseActionHistory = false; // Different files for each target; do not want to invalidate based on this.
 				WriteMetadataAction.ProducedItems.UnionWith(ProducedItems.Select(x => FileItem.GetItemByFileReference(x)));
 				WriteMetadataAction.PrerequisiteItems.UnionWith(PrerequisiteItems);
diff --git a/Engine/Source/Programs/UnrealBuildTool/Platform/IOS/IOSToolChain.cs b/Engine/Source/Programs/UnrealBuildTool/Platform/IOS/IOSToolChain.cs
--- a/Engine/Source/Programs/UnrealBuildTool/Platform/IOS/IOSToolChain.cs
+++ b/Engine/Source/Programs/UnrealBuildTool/Platform/IOS/IOSToolChain.cs
@@ -1197,7 +1197,6 @@ namespace UnrealBuildTool
 				PostBuildSyncAction.ProducedItems.Add(FileItem.GetItemByFileReference(GetStagedExecutablePath(Executable.Location, Target.Name)));
 				PostBuildSyncAction.DeleteItems.UnionWith(PostBuildSyncAction.ProducedItems);
 				PostBuildSyncAction.StatusDescription = "Executing PostBuildSync";
-				PostBuildSyncAction.bCanExecuteRemotely = false;
 
 				OutputFiles.AddRange(PostBuildSyncAction.ProducedItems);
 			}
diff --git a/Engine/Source/Programs/UnrealBuildTool/Platform/Windows/VCToolChain.cs b/Engine/Source/Programs/UnrealBuildTool/Platform/Windows/VCToolChain.cs
--- a/Engine/Source/Programs/UnrealBuildTool/Platform/Windows/VCToolChain.cs
+++ b/Engine/Source/Programs/UnrealBuildTool/Platform/Windows/VCToolChain.cs
@@ -1806,8 +1806,6 @@ namespace UnrealBuildTool
 			Action ParseTimingInfoAction = Graph.CreateRecursiveAction<ParseMsvcTimingInfoMode>(ActionType.ParseTimingInfo, ParseTimingArguments);
 			ParseTimingInfoAction.WorkingDirectory = Unreal.EngineSourceDirectory;
 			ParseTimingInfoAction.StatusDescription = Path.GetFileName(TimingFile.AbsolutePath);
-			ParseTimingInfoAction.bCanExecuteRemotely = true;
-			ParseTimingInfoAction.bCanExecuteRemotelyWithSNDBS = true;
 			ParseTimingInfoAction.PrerequisiteItems.Add(SourceFile);
 			ParseTimingInfoAction.PrerequisiteItems.Add(TimingFile);
 			ParseTimingInfoAction.ProducedItems.Add(TimingJsonFile);
@@ -1860,8 +1858,6 @@ namespace UnrealBuildTool
 						Action AggregateTimingInfoAction = MakefileBuilder.CreateRecursiveAction<AggregateParsedTimingInfo>(ActionType.ParseTimingInfo, String.Join(" ", ActionArgs));
 						AggregateTimingInfoAction.WorkingDirectory = Unreal.EngineSourceDirectory;
 						AggregateTimingInfoAction.StatusDescription = $"Aggregating {TimingJsonFiles.Count} Timing File(s)";
-						AggregateTimingInfoAction.bCanExecuteRemotely = false;
-						AggregateTimingInfoAction.bCanExecuteRemotelyWithSNDBS = false;
 						AggregateTimingInfoAction.PrerequisiteItems.UnionWith(TimingJsonFiles);
 
 						FileItem AggregateOutputFile = FileItem.GetItemByFileReference(FileReference.Combine(Makefile.ProjectIntermediateDirectory, $"{Target.Name}.cta"));
@@ -1885,8 +1881,6 @@ namespace UnrealBuildTool
 						Action AggregateTimingInfoAction = MakefileBuilder.CreateRecursiveAction<AggregateClangTimingInfo>(ActionType.ParseTimingInfo, String.Join(" ", AggregateActionArgs));
 						AggregateTimingInfoAction.WorkingDirectory = Unreal.EngineSourceDirectory;
 						AggregateTimingInfoAction.StatusDescription = $"Aggregating {TimingJsonFiles.Count} Timing File(s)";
-						AggregateTimingInfoAction.bCanExecuteRemotely = false;
-						AggregateTimingInfoAction.bCanExecuteRemotelyWithSNDBS = false;
 						AggregateTimingInfoAction.PrerequisiteItems.UnionWith(TimingJsonFiles);
 
 						AggregateTimingInfoAction.ProducedItems.Add(AggregateOutputFile);
@@ -1903,8 +1897,6 @@ namespace UnrealBuildTool
 						Action ArchiveTimingInfoAction = MakefileBuilder.CreateRecursiveAction<AggregateClangTimingInfo>(ActionType.ParseTimingInfo, String.Join(" ", ArchiveActionArgs));
 						ArchiveTimingInfoAction.WorkingDirectory = Unreal.EngineSourceDirectory;
 						ArchiveTimingInfoAction.StatusDescription = $"Archiving {TimingJsonFiles.Count} Timing File(s)";
-						ArchiveTimingInfoAction.bCanExecuteRemotely = false;
-						ArchiveTimingInfoAction.bCanExecuteRemotelyWithSNDBS = false;
 						ArchiveTimingInfoAction.PrerequisiteItems.UnionWith(TimingJsonFiles);
 
 						ArchiveTimingInfoAction.ProducedItems.Add(ArchiveOutputFile);
diff --git a/Engine/Source/Programs/UnrealBuildTool/System/ActionGraphBuilder.cs b/Engine/Source/Programs/UnrealBuildTool/System/ActionGraphBuilder.cs
--- a/Engine/Source/Programs/UnrealBuildTool/System/ActionGraphBuilder.cs
+++ b/Engine/Source/Programs/UnrealBuildTool/System/ActionGraphBuilder.cs
@@ -276,6 +276,7 @@ namespace UnrealBuildTool
 			CopyAction.DeleteItems.Add(TargetFile);
 			CopyAction.StatusDescription = TargetFile.Location.GetFileName();
 			CopyAction.bCanExecuteRemotely = false;
+			CopyAction.bCanExecuteInUBA = false;
 			return CopyAction;
 		}
 
@@ -315,6 +316,8 @@ namespace UnrealBuildTool
 			NewAction.CommandPath = Unreal.DotnetPath;
 			NewAction.CommandArguments = $"\"{Unreal.UnrealBuildToolDllPath}\" -Mode={Attribute.Name} {Arguments}";
 			NewAction.CommandDescription = Attribute.Name;
+			NewAction.bCanExecuteRemotely = false;
+			NewAction.bCanExecuteInUBA = false;
 			return NewAction;
 		}
 
diff --git a/Engine/Source/Programs/UnrealBuildTool/ToolChain/AppleToolChain.cs b/Engine/Source/Programs/UnrealBuildTool/ToolChain/AppleToolChain.cs
--- a/Engine/Source/Programs/UnrealBuildTool/ToolChain/AppleToolChain.cs
+++ b/Engine/Source/Programs/UnrealBuildTool/ToolChain/AppleToolChain.cs
@@ -1175,7 +1175,6 @@ namespace UnrealBuildTool
 			PostBuildSyncAction.PrerequisiteItems.Add(Executable);
 			PostBuildSyncAction.ProducedItems.Add(GetPostBuildOutputFile(Executable.Location, Target.Name, Target.Platform));
 			PostBuildSyncAction.StatusDescription = $"Executing PostBuildSync [{Executable.Location}]";
-			PostBuildSyncAction.bCanExecuteRemotely = false;
 
 			if (Target.Platform == UnrealTargetPlatform.IOS || Target.Platform == UnrealTargetPlatform.TVOS)
 			{
diff --git a/Engine/Source/Programs/UnrealBuildTool/ToolChain/ClangToolChain.cs b/Engine/Source/Programs/UnrealBuildTool/ToolChain/ClangToolChain.cs
--- a/Engine/Source/Programs/UnrealBuildTool/ToolChain/ClangToolChain.cs
+++ b/Engine/Source/Programs/UnrealBuildTool/ToolChain/ClangToolChain.cs
@@ -266,8 +266,6 @@ namespace UnrealBuildTool
 					Action AggregateTimingInfoAction = MakefileBuilder.CreateRecursiveAction<AggregateClangTimingInfo>(ActionType.ParseTimingInfo, String.Join(" ", AggregateActionArgs));
 					AggregateTimingInfoAction.WorkingDirectory = Unreal.EngineSourceDirectory;
 					AggregateTimingInfoAction.StatusDescription = $"Aggregating {TimingJsonFiles.Count} Timing File(s)";
-					AggregateTimingInfoAction.bCanExecuteRemotely = false;
-					AggregateTimingInfoAction.bCanExecuteRemotelyWithSNDBS = false;
 					AggregateTimingInfoAction.PrerequisiteItems.UnionWith(TimingJsonFiles);
 
 					AggregateTimingInfoAction.ProducedItems.Add(AggregateOutputFile);
@@ -284,8 +282,6 @@ namespace UnrealBuildTool
 					Action ArchiveTimingInfoAction = MakefileBuilder.CreateRecursiveAction<AggregateClangTimingInfo>(ActionType.ParseTimingInfo, String.Join(" ", ArchiveActionArgs));
 					ArchiveTimingInfoAction.WorkingDirectory = Unreal.EngineSourceDirectory;
 					ArchiveTimingInfoAction.StatusDescription = $"Archiving {TimingJsonFiles.Count} Timing File(s)";
-					ArchiveTimingInfoAction.bCanExecuteRemotely = false;
-					ArchiveTimingInfoAction.bCanExecuteRemotelyWithSNDBS = false;
 					ArchiveTimingInfoAction.PrerequisiteItems.UnionWith(TimingJsonFiles);
 
 					ArchiveTimingInfoAction.ProducedItems.Add(ArchiveOutputFile);
@@ -302,6 +298,7 @@ namespace UnrealBuildTool
 						CompileScoreExtractorAction.StatusDescription = $"Extracting CompileScore";
 						CompileScoreExtractorAction.bCanExecuteRemotely = false;
 						CompileScoreExtractorAction.bCanExecuteRemotelyWithSNDBS = false;
+						CompileScoreExtractorAction.bCanExecuteInUBA = false; // TODO: Unknown if supported
 						CompileScoreExtractorAction.PrerequisiteItems.UnionWith(TimingJsonFiles);
 						CompileScoreExtractorAction.CommandPath = ScoreDataExtractor;
 						CompileScoreExtractorAction.CommandArguments = $"-clang -verbosity 0 -timelinepack 1000000 -extract -i \"{NormalizeCommandLinePath(Makefile.ProjectIntermediateDirectory)}\" -o \"{NormalizeCommandLinePath(CompileScoreOutput)}\"";
