Rollback PVS version

#rnx

diff --git a/Engine/Source/Programs/UnrealBuildTool/Platform/Windows/PVSToolChain.cs b/Engine/Source/Programs/UnrealBuildTool/Platform/Windows/PVSToolChain.cs
--- a/Engine/Source/Programs/UnrealBuildTool/Platform/Windows/PVSToolChain.cs
+++ b/Engine/Source/Programs/UnrealBuildTool/Platform/Windows/PVSToolChain.cs
@@ -810,8 +810,9 @@ namespace UnrealBuildTool
 				AnalyzeAction.PrerequisiteItems.UnionWith(InputFiles); // Add the InputFiles as PrerequisiteItems so that in SingleFileCompile mode the PVSAnalyze step is not filtered out
 				AnalyzeAction.ProducedItems.Add(OutputFileItem);
 				AnalyzeAction.DeleteItems.Add(OutputFileItem); // PVS Studio will append by default, so need to delete produced items
-				AnalyzeAction.Weight = Target.MSVCCompileActionWeight * 2.0; // Very high memory usage
-				AnalyzeAction.bCanExecuteRemotely = false;
+				AnalyzeAction.bCanExecuteRemotely = true;
+				AnalyzeAction.bCanExecuteRemotelyWithSNDBS = false;
+				AnalyzeAction.bCanExecuteRemotelyWithXGE = false;
 				AnalyzeAction.bCanExecuteInUBA = true;
 
 				Result.ObjectFiles.AddRange(AnalyzeAction.ProducedItems);
