UnrealBuildTool: Analysis warning

#rnx

diff --git a/Engine/Source/Programs/UnrealBuildTool/System/ActionGraph.cs b/Engine/Source/Programs/UnrealBuildTool/System/ActionGraph.cs
--- a/Engine/Source/Programs/UnrealBuildTool/System/ActionGraph.cs
+++ b/Engine/Source/Programs/UnrealBuildTool/System/ActionGraph.cs
@@ -364,7 +364,7 @@ namespace UnrealBuildTool
 			if (ActionCount > ParallelExecutor.GetDefaultNumParallelProcesses(BuildConfiguration.MaxParallelActions, BuildConfiguration.bAllCores, Logger))
 			{
 #if __UBAEXECUTOR_AVAILABLE__
-				if (BuildConfiguration.bAllowUBAExecutor && UBAExecutor.IsAvailable(Logger))
+				if (BuildConfiguration.bAllowUBAExecutor && UBAExecutor.IsAvailable())
 				{
 					return new UBAExecutor(BuildConfiguration.MaxParallelActions, BuildConfiguration.bAllCores, BuildConfiguration.bCompactOutput, Logger, TargetDescriptors.FirstOrDefault()?.AdditionalArguments);
 				}
@@ -385,7 +385,7 @@ namespace UnrealBuildTool
 			}
 
 #if __UBAEXECUTOR_AVAILABLE__
-			if (BuildConfiguration.bAllowUBALocalExecutor && UBALocalExecutor.IsAvailable(Logger))
+			if (BuildConfiguration.bAllowUBALocalExecutor && UBALocalExecutor.IsAvailable())
 			{
 				return new UBALocalExecutor(BuildConfiguration.MaxParallelActions, BuildConfiguration.bAllCores, BuildConfiguration.bCompactOutput, Logger, TargetDescriptors.FirstOrDefault()?.AdditionalArguments);
 			}
