Fix xcode would build UnrealLightmass for arm

diff --git a/Engine/Source/Programs/UnrealBuildTool/ProjectFiles/Xcode/XcodeProject.cs b/Engine/Source/Programs/UnrealBuildTool/ProjectFiles/Xcode/XcodeProject.cs
--- a/Engine/Source/Programs/UnrealBuildTool/ProjectFiles/Xcode/XcodeProject.cs
+++ b/Engine/Source/Programs/UnrealBuildTool/ProjectFiles/Xcode/XcodeProject.cs
@@ -1040,7 +1040,14 @@ namespace UnrealBuildTool.XcodeProjectXcconfig
 			{
 				Content.WriteLine(4, $"SUPPORTED_PLATFORMS = \"macosx iphonesimulator iphoneos appletvsimulator appletvos xros xrsimulator\";");
 				Content.WriteLine(4, $"ONLY_ACTIVE_ARCH = YES;");
+
+				if (Info.bSupportsMac)
+				{
+					string SupportedMacArchitectures = String.Join(" ", XcodeUtils.GetSupportedMacArchitectures(Info.BuildTarget, Info.ProjectTarget?.UnrealProjectFilePath).Architectures.Select(x => x.AppleName));
+					Content.WriteLine(4, $"\"VALID_ARCHS[sdk=macos*]\" = \"{SupportedMacArchitectures}\";");
+				}
 			}
+				
 			Content.WriteLine(3, "};");
 			Content.WriteLine(3, $"name = \"{Info.DisplayName}\";");
 			Content.WriteLine(2, "};");
@@ -1740,9 +1747,6 @@ namespace UnrealBuildTool.XcodeProjectXcconfig
 
 				// @todo: get a version for  games, like IOS has
 				MarketingVersion = MacToolChain.LoadEngineDisplayVersion();
-
-				string SupportedMacArchitectures = String.Join(" ", XcodeUtils.GetSupportedMacArchitectures(BuildConfig.BuildTarget, UnrealData.UProjectFileLocation).Architectures.Select(x => x.AppleName));
-				ExtraConfigLines.Add($"VALID_ARCHS = {SupportedMacArchitectures}");
 			}
 			else
 			{
