UnrealBuildTool: Add Posix platform group which includes Linux and Mac platforms

diff --git a/Engine/Source/Programs/UnrealBuildTool/Configuration/UEBuildTarget.cs b/Engine/Source/Programs/UnrealBuildTool/Configuration/UEBuildTarget.cs
--- a/Engine/Source/Programs/UnrealBuildTool/Configuration/UEBuildTarget.cs
+++ b/Engine/Source/Programs/UnrealBuildTool/Configuration/UEBuildTarget.cs
@@ -566,6 +566,11 @@ namespace UnrealBuildTool
 		/// </summary>
 		public static UnrealPlatformGroup ThirtyHz = FindOrAddByName("30Hz");
 
+		/// <summary>
+		/// POSIX-compliant platforms
+		/// </summary>
+		public static UnrealPlatformGroup Posix = FindOrAddByName("Posix");
+
 	}
 
 	/// <summary>
diff --git a/Engine/Source/Programs/UnrealBuildTool/Platform/Linux/UEBuildLinux.cs b/Engine/Source/Programs/UnrealBuildTool/Platform/Linux/UEBuildLinux.cs
--- a/Engine/Source/Programs/UnrealBuildTool/Platform/Linux/UEBuildLinux.cs
+++ b/Engine/Source/Programs/UnrealBuildTool/Platform/Linux/UEBuildLinux.cs
@@ -726,11 +726,13 @@ namespace UnrealBuildTool
 			UEBuildPlatform.RegisterPlatformWithGroup(UnrealTargetPlatform.Linux, UnrealPlatformGroup.Linux);
 			UEBuildPlatform.RegisterPlatformWithGroup(UnrealTargetPlatform.Linux, UnrealPlatformGroup.Unix);
 			UEBuildPlatform.RegisterPlatformWithGroup(UnrealTargetPlatform.Linux, UnrealPlatformGroup.Desktop);
+			UEBuildPlatform.RegisterPlatformWithGroup(UnrealTargetPlatform.Linux, UnrealPlatformGroup.Posix);
 
 			UEBuildPlatform.RegisterBuildPlatform(new LinuxPlatform(UnrealTargetPlatform.LinuxArm64, SDKArm64, Logger), Logger);
 			UEBuildPlatform.RegisterPlatformWithGroup(UnrealTargetPlatform.LinuxArm64, UnrealPlatformGroup.Linux);
 			UEBuildPlatform.RegisterPlatformWithGroup(UnrealTargetPlatform.LinuxArm64, UnrealPlatformGroup.Unix);
 			UEBuildPlatform.RegisterPlatformWithGroup(UnrealTargetPlatform.LinuxArm64, UnrealPlatformGroup.Desktop);
+			UEBuildPlatform.RegisterPlatformWithGroup(UnrealTargetPlatform.LinuxArm64, UnrealPlatformGroup.Posix);
 		}
 	}
 }
diff --git a/Engine/Source/Programs/UnrealBuildTool/Platform/Mac/UEBuildMac.cs b/Engine/Source/Programs/UnrealBuildTool/Platform/Mac/UEBuildMac.cs
--- a/Engine/Source/Programs/UnrealBuildTool/Platform/Mac/UEBuildMac.cs
+++ b/Engine/Source/Programs/UnrealBuildTool/Platform/Mac/UEBuildMac.cs
@@ -636,6 +636,7 @@ namespace UnrealBuildTool
 			UEBuildPlatform.RegisterBuildPlatform(new MacPlatform(SDK, Logger), Logger);
 			UEBuildPlatform.RegisterPlatformWithGroup(UnrealTargetPlatform.Mac, UnrealPlatformGroup.Apple);
 			UEBuildPlatform.RegisterPlatformWithGroup(UnrealTargetPlatform.Mac, UnrealPlatformGroup.Desktop);
+			UEBuildPlatform.RegisterPlatformWithGroup(UnrealTargetPlatform.Mac, UnrealPlatformGroup.Posix);
 		}
 	}
 }
