UnrealBuildTool: Rename Posix group to PosixOS

#rnx

diff --git a/Engine/Source/Programs/UnrealBuildTool/Configuration/UEBuildTarget.cs b/Engine/Source/Programs/UnrealBuildTool/Configuration/UEBuildTarget.cs
--- a/Engine/Source/Programs/UnrealBuildTool/Configuration/UEBuildTarget.cs
+++ b/Engine/Source/Programs/UnrealBuildTool/Configuration/UEBuildTarget.cs
@@ -569,7 +569,7 @@ namespace UnrealBuildTool
 		/// <summary>
 		/// POSIX-compliant platforms
 		/// </summary>
-		public static UnrealPlatformGroup Posix = FindOrAddByName("Posix");
+		public static UnrealPlatformGroup PosixOS = FindOrAddByName("PosixOS");
 
 	}
 
diff --git a/Engine/Source/Programs/UnrealBuildTool/Platform/Linux/UEBuildLinux.cs b/Engine/Source/Programs/UnrealBuildTool/Platform/Linux/UEBuildLinux.cs
--- a/Engine/Source/Programs/UnrealBuildTool/Platform/Linux/UEBuildLinux.cs
+++ b/Engine/Source/Programs/UnrealBuildTool/Platform/Linux/UEBuildLinux.cs
@@ -726,13 +726,13 @@ namespace UnrealBuildTool
 			UEBuildPlatform.RegisterPlatformWithGroup(UnrealTargetPlatform.Linux, UnrealPlatformGroup.Linux);
 			UEBuildPlatform.RegisterPlatformWithGroup(UnrealTargetPlatform.Linux, UnrealPlatformGroup.Unix);
 			UEBuildPlatform.RegisterPlatformWithGroup(UnrealTargetPlatform.Linux, UnrealPlatformGroup.Desktop);
-			UEBuildPlatform.RegisterPlatformWithGroup(UnrealTargetPlatform.Linux, UnrealPlatformGroup.Posix);
+			UEBuildPlatform.RegisterPlatformWithGroup(UnrealTargetPlatform.Linux, UnrealPlatformGroup.PosixOS);
 
 			UEBuildPlatform.RegisterBuildPlatform(new LinuxPlatform(UnrealTargetPlatform.LinuxArm64, SDKArm64, Logger), Logger);
 			UEBuildPlatform.RegisterPlatformWithGroup(UnrealTargetPlatform.LinuxArm64, UnrealPlatformGroup.Linux);
 			UEBuildPlatform.RegisterPlatformWithGroup(UnrealTargetPlatform.LinuxArm64, UnrealPlatformGroup.Unix);
 			UEBuildPlatform.RegisterPlatformWithGroup(UnrealTargetPlatform.LinuxArm64, UnrealPlatformGroup.Desktop);
-			UEBuildPlatform.RegisterPlatformWithGroup(UnrealTargetPlatform.LinuxArm64, UnrealPlatformGroup.Posix);
+			UEBuildPlatform.RegisterPlatformWithGroup(UnrealTargetPlatform.LinuxArm64, UnrealPlatformGroup.PosixOS);
 		}
 	}
 }
diff --git a/Engine/Source/Programs/UnrealBuildTool/Platform/Mac/UEBuildMac.cs b/Engine/Source/Programs/UnrealBuildTool/Platform/Mac/UEBuildMac.cs
--- a/Engine/Source/Programs/UnrealBuildTool/Platform/Mac/UEBuildMac.cs
+++ b/Engine/Source/Programs/UnrealBuildTool/Platform/Mac/UEBuildMac.cs
@@ -636,7 +636,7 @@ namespace UnrealBuildTool
 			UEBuildPlatform.RegisterBuildPlatform(new MacPlatform(SDK, Logger), Logger);
 			UEBuildPlatform.RegisterPlatformWithGroup(UnrealTargetPlatform.Mac, UnrealPlatformGroup.Apple);
 			UEBuildPlatform.RegisterPlatformWithGroup(UnrealTargetPlatform.Mac, UnrealPlatformGroup.Desktop);
-			UEBuildPlatform.RegisterPlatformWithGroup(UnrealTargetPlatform.Mac, UnrealPlatformGroup.Posix);
+			UEBuildPlatform.RegisterPlatformWithGroup(UnrealTargetPlatform.Mac, UnrealPlatformGroup.PosixOS);
 		}
 	}
 }
