UnrealBuildTool: Fix not selecting the preferred WindowsSDK if it is available when a supported version min or max is also requested

diff --git a/Engine/Source/Programs/UnrealBuildTool/Platform/Windows/MicrosoftPlatformSDK.cs b/Engine/Source/Programs/UnrealBuildTool/Platform/Windows/MicrosoftPlatformSDK.cs
--- a/Engine/Source/Programs/UnrealBuildTool/Platform/Windows/MicrosoftPlatformSDK.cs
+++ b/Engine/Source/Programs/UnrealBuildTool/Platform/Windows/MicrosoftPlatformSDK.cs
@@ -202,27 +202,25 @@ namespace UnrealBuildTool
 					throw new BuildException("Unable to find requested Windows SDK; '{0}' is an invalid version", DesiredVersion);
 				}
 			}
-			else if (MinVersion == null && MaxVersion == null)
+			else if (CachedWindowsSdkDirs!.Count > 0)
 			{
+				IEnumerable<KeyValuePair<VersionNumber, DirectoryReference>> AllowedSdkDirs = CachedWindowsSdkDirs.OrderBy(x => x.Key).Where(
+					x =>
+					(MinVersion == null || x.Key >= MinVersion) &&
+					(MaxVersion == null || x.Key <= MaxVersion));
+
 				// convert the desired version into a VersionNumber
-				VersionNumber MainVersion = VersionNumber.Parse(UEBuildPlatformSDK.GetSDKForPlatform("Win64")!.GetMainVersion());
-				if (CachedWindowsSdkDirs!.ContainsKey(MainVersion))
+				VersionNumber MainVersion = VersionNumber.Parse(GetSDKForPlatform("Win64")!.GetMainVersion());
+				if (AllowedSdkDirs.Any(x => x.Key == MainVersion))
 				{
 					WindowsSdkVersion = MainVersion;
 				}
 				// if it's not an installed version, use the highest installed version
-				else if (CachedWindowsSdkDirs.Any())
+				else if (AllowedSdkDirs.Any())
 				{
-					WindowsSdkVersion = CachedWindowsSdkDirs.OrderBy(x => x.Key).Last().Key;
+					WindowsSdkVersion = AllowedSdkDirs.Last().Key;
 				}
 			}
-			else if (CachedWindowsSdkDirs!.Count > 0)
-			{
-				WindowsSdkVersion = CachedWindowsSdkDirs.OrderBy(x => x.Key).Where(
-					x =>
-					(MinVersion == null || x.Key >= MinVersion) &&
-					(MaxVersion == null || x.Key <= MaxVersion)).Last().Key;
-			}
 
 			// Get the actual directory for this version
 			DirectoryReference? SdkDir;
