PipInstall: Generate plugin site-packages path file so that packaged python dependencies are included in pip dependency computations. Also match specific pip install utils version so that utils updates update install env. Some minor fixes for readability and to add clarity to log messages.

diff --git a/Engine/Plugins/Experimental/PythonScriptPlugin/Source/PythonScriptPlugin/Private/PipInstall.cpp b/Engine/Plugins/Experimental/PythonScriptPlugin/Source/PythonScriptPlugin/Private/PipInstall.cpp
--- a/Engine/Plugins/Experimental/PythonScriptPlugin/Source/PythonScriptPlugin/Private/PipInstall.cpp
+++ b/Engine/Plugins/Experimental/PythonScriptPlugin/Source/PythonScriptPlugin/Private/PipInstall.cpp
@@ -16,7 +16,12 @@
 
 #if WITH_PYTHON
 
+// In order to keep editor startup time fast, check directly for this utils version (make sure to match with wheel version in PythonScriptPlugin/Content/Python/Lib/wheels)
+// NOTE: This version must also be changed in PipInstallMode.cs in order to support UBT functionality
+const FString FPipInstall::PipInstallUtilsVer = TEXT("0.1.3");
+
 const FString FPipInstall::PluginsListingFilename = TEXT("pyreqs_plugins.list");
+const FString FPipInstall::PluginsSitePackageFilename = TEXT("plugin_site_package.pth");
 const FString FPipInstall::RequirementsInputFilename = TEXT("merged_requirements.in");
 const FString FPipInstall::ExtraUrlsFilename = TEXT("extra_urls.txt");
 const FString FPipInstall::ParsedRequirementsFilename = TEXT("merged_requirements.txt");
@@ -30,20 +35,43 @@ FString FPipInstall::WritePluginsListing(TArray<TSharedRef<IPlugin>>& OutPythonP
 
 	OutPythonPlugins.Empty();
 
-	TArray<FString> PythonPluginPaths;
+	// List of plugins with pip dependencies
+	TArray<FString> PipPluginPaths;
+	// List of enabled plugins' site-packages folders
+	TArray<FString> PluginSitePackagePaths;
 	for ( const TSharedRef<IPlugin>& Plugin : IPluginManager::Get().GetEnabledPlugins() )
 	{
+		const FString PythonContentPath = FPaths::ConvertRelativePathToFull(Plugin->GetContentDir() / TEXT("Python"));
+		const FString PluginPlatformSitePackagesPath = PythonContentPath / TEXT("Lib") / FPlatformMisc::GetUBTPlatform() / TEXT("site-packages");
+		const FString PluginGeneralSitePackagesPath = PythonContentPath / TEXT("Lib") / TEXT("site-packages");
+
+		// Write platform/general site-packages paths per-plugin to .pth file to account for packaged python dependencies during pip install
+		if (FPaths::DirectoryExists(PluginPlatformSitePackagesPath))
+		{
+			PluginSitePackagePaths.Add(PluginPlatformSitePackagesPath);
+		}
+
+		if (FPaths::DirectoryExists(PluginGeneralSitePackagesPath))
+		{
+			PluginSitePackagePaths.Add(PluginGeneralSitePackagesPath);
+		}
+
 		const FPluginDescriptor& PluginDesc = Plugin->GetDescriptor();
 		if (PluginDesc.CachedJson->HasTypedField(TEXT("PythonRequirements"), EJson::Array))
 		{
 			const FString PluginDescFile = FPaths::ConvertRelativePathToFull(Plugin->GetDescriptorFileName());
-			PythonPluginPaths.Add(PluginDescFile);
+			PipPluginPaths.Add(PluginDescFile);
 			OutPythonPlugins.Add(Plugin);
 		}
 	}
 
+	// Create list of plugins that may require pip install dependencies
 	const FString PyPluginsListingFile = PipInstallPath / PluginsListingFilename;
-	FFileHelper::SaveStringArrayToFile(PythonPluginPaths, *PyPluginsListingFile);
+	FFileHelper::SaveStringArrayToFile(PipPluginPaths, *PyPluginsListingFile);
+
+	// Create .pth file in PipInstall/Lib/site-packages to account for plugins with packaged dependencies
+	const FString PyPluginsSitePackageFile = PipInstallPath / TEXT("Lib") / TEXT("site-packages") / PluginsSitePackageFilename;
+	FFileHelper::SaveStringArrayToFile(PluginSitePackagePaths, *PyPluginsSitePackageFile);
 
     return PyPluginsListingFile;
 }
@@ -144,6 +172,24 @@ FString FPipInstall::ParsePluginDependencies(const FString& MergedInRequirements
 	return FPaths::ConvertRelativePathToFull(ParsedReqsFile);
 }
 
+bool FPipInstall::HasInstallLines(const TArray<FString>& RequirementLines)
+{
+	for (const FStringView Line : RequirementLines)
+	{
+		if (!Line.TrimStart().StartsWith(TCHAR('#')))
+		{
+			return true;
+		}
+	}
+
+	return false;
+}
+
+FString FPipInstall::GetPipInstallPath()
+{
+	return FPaths::ProjectIntermediateDir() / TEXT("PipInstall");
+}
+
 
 void FPipInstall::SetupPipInstallUtils(const FString& VenvInterp, FFeedbackContext* Context)
 {
@@ -164,7 +210,8 @@ void FPipInstall::SetupPipInstallUtils(const FString& VenvInterp, FFeedbackConte
 	const FString PipWheelsDir = PythonScriptDir / TEXT("Content/Python/Lib/wheels");
 	const FString InstallRequirements = PythonScriptDir / TEXT("Content/Python/PipInstallUtils/requirements.txt");
 
-	const FString Cmd = FString::Printf(TEXT("-m pip install --upgrade --no-index --find-links \"%s\" -r \"%s\" ue-pipinstall-utils"), *PipWheelsDir, *InstallRequirements);
+	const FString PipInstallReq = TEXT("ue-pipinstall-utils==") + PipInstallUtilsVer;
+	const FString Cmd = FString::Printf(TEXT("-m pip install --upgrade --no-index --find-links \"%s\" -r \"%s\" %s"), *PipWheelsDir, *InstallRequirements, *PipInstallReq);
 
 	RunPythonCmd(LOCTEXT("PipInstall.SetupPipInstallUtils", "Setting up pip install utils"), VenvInterp, Cmd, Context);
 }
@@ -172,7 +219,8 @@ void FPipInstall::SetupPipInstallUtils(const FString& VenvInterp, FFeedbackConte
 
 bool FPipInstall::CheckPipInstallUtils(const FString& VenvInterp, FFeedbackContext* Context)
 {
-	const FString Cmd = TEXT("-c \"import pkg_resources; exit(0 if pkg_resources.working_set.find(pkg_resources.Requirement.parse('ue-pipinstall-utils')) is not None else 1)\"");
+	// Verify that correct version of pip install utils is already available
+	const FString Cmd = FString::Printf(TEXT("-c \"import pkg_resources;dist=pkg_resources.working_set.find(pkg_resources.Requirement.parse('ue-pipinstall-utils'));exit(dist.version!='%s' if dist is not None else 1)\""), *PipInstallUtilsVer);
 	return (RunPythonCmd(LOCTEXT("PipInstall.CheckPipInstallUtils", "Check pip install utils installed"), VenvInterp, Cmd, Context) == 0);
 }
 
@@ -235,11 +283,6 @@ bool FPipInstall::RunLoggedSubprocess(const FText& Description, const FString& U
 }
 
 
-FString FPipInstall::GetPipInstallPath()
-{
-	return FPaths::ProjectIntermediateDir() / TEXT("PipInstall");
-}
-
 FString FPipInstall::GetPythonScriptPluginPath()
 {
 	TSharedPtr<IPlugin> PythonPlugin = IPluginManager::Get().FindPlugin("PythonScriptPlugin");
diff --git a/Engine/Plugins/Experimental/PythonScriptPlugin/Source/PythonScriptPlugin/Private/PipInstall.h b/Engine/Plugins/Experimental/PythonScriptPlugin/Source/PythonScriptPlugin/Private/PipInstall.h
--- a/Engine/Plugins/Experimental/PythonScriptPlugin/Source/PythonScriptPlugin/Private/PipInstall.h
+++ b/Engine/Plugins/Experimental/PythonScriptPlugin/Source/PythonScriptPlugin/Private/PipInstall.h
@@ -14,7 +14,9 @@ class FFeedbackContext;
 class FPipInstall
 {
 public:
+	static const FString PipInstallUtilsVer;
 	static const FString PluginsListingFilename;
+	static const FString PluginsSitePackageFilename;
 	static const FString RequirementsInputFilename;
 	static const FString ExtraUrlsFilename;
 	static const FString ParsedRequirementsFilename;
@@ -25,6 +27,8 @@ public:
 	static void SetupPipEnv(FFeedbackContext* Context, bool bForceRebuild = false);
 	static FString ParsePluginDependencies(const FString& MergedInRequirementsFile, FFeedbackContext* Context);
 
+	static bool HasInstallLines(const TArray<FString>& RequirementLines);
+
 	static FString GetPipInstallPath();
 
 private:
diff --git a/Engine/Plugins/Experimental/PythonScriptPlugin/Source/PythonScriptPlugin/Private/PythonScriptPlugin.cpp b/Engine/Plugins/Experimental/PythonScriptPlugin/Source/PythonScriptPlugin/Private/PythonScriptPlugin.cpp
--- a/Engine/Plugins/Experimental/PythonScriptPlugin/Source/PythonScriptPlugin/Private/PythonScriptPlugin.cpp
+++ b/Engine/Plugins/Experimental/PythonScriptPlugin/Source/PythonScriptPlugin/Private/PythonScriptPlugin.cpp
@@ -1187,7 +1187,7 @@ void FPythonScriptPlugin::RunPipInstaller()
 	const FString ParsedReqsFile = FPaths::ConvertRelativePathToFull(FPipInstall::GetPipInstallPath() / FPipInstall::ParsedRequirementsFilename);
 
 	TArray<FString> ParsedReqLines;
-	if (!FPaths::FileExists(ParsedReqsFile) || !FFileHelper::LoadFileToStringArray(ParsedReqLines, *ParsedReqsFile) || ParsedReqLines.IsEmpty())
+	if (!FPaths::FileExists(ParsedReqsFile) || !FFileHelper::LoadFileToStringArray(ParsedReqLines, *ParsedReqsFile) || !FPipInstall::HasInstallLines(ParsedReqLines))
 	{
 		return;
 	}
diff --git a/Engine/Source/Programs/UnrealBuildTool/Modes/PipInstallMode.cs b/Engine/Source/Programs/UnrealBuildTool/Modes/PipInstallMode.cs
--- a/Engine/Source/Programs/UnrealBuildTool/Modes/PipInstallMode.cs
+++ b/Engine/Source/Programs/UnrealBuildTool/Modes/PipInstallMode.cs
@@ -112,25 +112,10 @@
 			UnrealTargetPlatform Platform = TargetDescriptor.Platform;
 
 			PipEnv Pip = new(InstallDir, Platform, Logger, ProgressWriter.bWriteMarkup);
-
-			// Generated from enabled plugins list
-			FileReference PluginsListingFile = FileReference.Combine(InstallDir, "pyreqs_plugins.list");
-			// This is the unparsed merged requirements file
-			FileReference MergedReqsInFile = FileReference.Combine(InstallDir, "merged_requirements.in");
-			// These files are used as input for the pip installer
-			FileReference ExtraUrlsFile = FileReference.Combine(InstallDir, "extra_urls.txt");
-			FileReference MergedRequirementsFile = FileReference.Combine(InstallDir, "merged_requirements.txt");
-			//string[]? PluginsList = null;
-
 			if ((Action & (PipAction)ActionBits.GenReqs) != 0)
 			{
-				if (!DirectoryReference.Exists(InstallDir))
+				Pip.WritePluginsListing(Target, Logger);
+				if (!Pip.WritePluginDependencies())
 				{
-					DirectoryReference.CreateDirectory(InstallDir);
-				}
-
-				WritePluginsList(Target, PluginsListingFile, Logger);
-				if (!Pip.WritePluginDependencies(PluginsListingFile, ExtraUrlsFile, MergedReqsInFile))
-				{
 					return 1;
 				}
@@ -147,7 +132,7 @@
 			if ((Action & (PipAction)ActionBits.ParseReqs) != 0)
 			{
 
-				if (!Pip.ParsePluginDependencies(MergedReqsInFile, MergedRequirementsFile))
+				if (!Pip.ParsePluginDependencies())
 				{
 					return 1;
 				}
@@ -155,7 +140,7 @@
 
 			if ((Action & (PipAction)ActionBits.InstallReqs) != 0)
 			{
-				if (!Pip.InstallPluginDependencies(MergedRequirementsFile, ExtraUrlsFile))
+				if (!Pip.InstallPluginDependencies())
 				{
 					return 1;
 				}
@@ -163,7 +148,7 @@
 
 			if ((Action & (PipAction)ActionBits.ViewLicenses) != 0)
 			{
-				if (!Pip.ViewInstalledLicenses(PluginsListingFile))
+				if (!Pip.ViewInstalledLicenses())
 				{
 					return 1;
 				}
@@ -171,18 +156,83 @@
 
 			return 0;
 		}
+	}
 
-		private void WritePluginsList(UEBuildTarget Target, FileReference PluginsListingFile, ILogger Logger)
+
+	/// <summary>
+	/// PipEnv helper class for setting up a self-contained pip environment and running pip commands (particularly pip install) within that env.
+	/// </summary>
+	class PipEnv
+	{
+		// Don't bother to re-install pip install tools if this version is already installed
+		// NOTE: This version must also be changed in PipInstall.cpp in order to support editor startup process
+		private const string PipInstallUtilsVer = "0.1.3";
+
+		// Generated from enabled plugins list
+		private const string PluginsListingFilename = "pyreqs_plugins.list";
+		// List full-paths to all enabled plugins site-package dirs (general/current platform)
+		private const string PluginsSitePackageFilename = "plugin_site_package.pth";
+		// This is the unparsed merged requirements file
+		private const string MergedReqsInFilename = "merged_requirements.in";
+		// These files are used as input for the pip installer
+		private const string ExtraUrlsFilename = "extra_urls.txt";
+		private const string MergedRequirementsFilename = "merged_requirements.txt";
+
+		private UnrealTargetPlatform TargetPlatform;
+		private DirectoryReference InstallDir;
+		private FileReference PythonVenv;
+
+		private ILogger Logger;
+		private IBaseProgressLogFactory LoggerFactory;
+
+		public PipEnv(DirectoryReference InInstallDir, UnrealTargetPlatform InPlatform, ILogger InLogger, bool UseProgressWriter = false)
 		{
+			Logger = InLogger;
+			InstallDir = InInstallDir;
+			TargetPlatform = InPlatform;
+
+			LoggerFactory = (UseProgressWriter) ? new PipProgressLogCreator(Logger) : new SimpleCmdLogCreator(Logger);
+
+			if (!DirectoryReference.Exists(InstallDir))
+			{
+				DirectoryReference.CreateDirectory(InstallDir);
+			}
+
+			PythonVenv = GetVenvInterpreter(InstallDir, TargetPlatform);
+		}
+
+		public void WritePluginsListing(UEBuildTarget Target, ILogger Logger)
+		{
+			FileReference PluginsListingFile = FileReference.Combine(InstallDir, PluginsListingFilename);
+
+			DirectoryReference PipSitePackagesPath = DirectoryReference.Combine(InstallDir, "Lib", "site-packages");
+			FileReference PyPluginsSitePackageFile = FileReference.Combine(PipSitePackagesPath, PluginsSitePackageFilename);
+
 			FileReference.Delete(PluginsListingFile);
 			if (Target.EnabledPlugins == null)
 			{
 				return;
 			}
 
+			List<string> PluginsSitePackages = new List<string>();
 			List<string> PluginsList = new List<string>();
 			foreach (UEBuildPlugin Plugin in Target.EnabledPlugins)
 			{
+				DirectoryReference PythonContentPath = DirectoryReference.Combine(Plugin.Directory, "Content", "Python");
+				DirectoryReference PluginPlatformSitePackagesPath = DirectoryReference.Combine(PythonContentPath, "Lib", Target.Platform.ToString(), "site-packages");
+				DirectoryReference PluginGeneralSitePackagesPath = DirectoryReference.Combine(PythonContentPath, "Lib", "site-packages");
+
+				// Write platform/general site-packages paths per-plugin to .pth file to account for packaged python dependencies during pip install
+				if (DirectoryReference.Exists(PluginPlatformSitePackagesPath))
+				{
+					PluginsSitePackages.Add(PluginPlatformSitePackagesPath.ToString());
+				}
+
+				if (DirectoryReference.Exists(PluginGeneralSitePackagesPath))
+				{
+					PluginsSitePackages.Add(PluginPlatformSitePackagesPath.ToString());
+				}
+
 				if (!JsonObject.TryRead(Plugin.File, out JsonObject? PluginJson))
 				{
 					Logger.LogWarning("Unable to parse {PluginFile}", Plugin.File.ToString());
@@ -196,32 +246,31 @@
 				}
 			}
 
+			if (!DirectoryReference.Exists(PipSitePackagesPath))
+			{
+				DirectoryReference.CreateDirectory(PipSitePackagesPath);
+			}
+
 			FileReference.WriteAllLines(PluginsListingFile, PluginsList);
+			FileReference.WriteAllLines(PyPluginsSitePackageFile, PluginsSitePackages);
 		}
-	}
 
-
-	/// <summary>
-	/// PipEnv helper class for setting up a self-contained pip environment and running pip commands (particularly pip install) within that env.
-	/// </summary>
-	class PipEnv
-	{
-		private UnrealTargetPlatform TargetPlatform;
-		private DirectoryReference InstallDir;
-		private FileReference PythonVenv;
-
-		private ILogger Logger;
-		private IBaseProgressLogFactory LoggerFactory;
-
-		public PipEnv(DirectoryReference InInstallDir, UnrealTargetPlatform InPlatform, ILogger InLogger, bool UseProgressWriter = false)
+		public bool WritePluginDependencies()
 		{
-			Logger = InLogger;
-			InstallDir = InInstallDir;
-			TargetPlatform = InPlatform;
+			FileReference PluginsListingFile = FileReference.Combine(InstallDir, PluginsListingFilename);
+			FileReference MergedReqsInFile = FileReference.Combine(InstallDir, MergedReqsInFilename);
+			FileReference ExtraUrlsFile = FileReference.Combine(InstallDir, ExtraUrlsFilename);
 
-			LoggerFactory = (UseProgressWriter) ? new PipProgressLogCreator(Logger) : new SimpleCmdLogCreator(Logger);
+			// Make merged requirements input file
+			if (!MergeRequirements(PluginsListingFile, out List<string>? MergedRequirements, out List<string>? ExtraIndexUrls))
+			{
+				return false;
+			}
 
-			PythonVenv = GetVenvInterpreter(InstallDir, TargetPlatform);
+			FileReference.WriteAllLines(MergedReqsInFile, MergedRequirements!);
+			FileReference.WriteAllLines(ExtraUrlsFile, ExtraIndexUrls!);
+
+			return true;
 		}
 
 		public bool SetupPipEnv(FileReference? EnginePython, bool ForceRebuild = false)
@@ -255,32 +304,22 @@
 			}
 		}
 
-		public bool WritePluginDependencies(FileReference PluginsListingFile, FileReference ExtraUrlsFile, FileReference MergedReqsInFile)
+		public bool ParsePluginDependencies()
 		{
-			DirectoryReference WorkDir = DirectoryReference.FromFile(PluginsListingFile);
+			FileReference MergedReqsInFile = FileReference.Combine(InstallDir, MergedReqsInFilename);
+			FileReference MergedRequirmentsFile = FileReference.Combine(InstallDir, MergedRequirementsFilename);
 
-			// Make merged requirements input file
-			if (!MergeRequirements(PluginsListingFile, out List<string>? MergedRequirements, out List<string>? ExtraIndexUrls))
-			{
-				return false;
-			}
-
-			FileReference.WriteAllLines(MergedReqsInFile, MergedRequirements!);
-			FileReference.WriteAllLines(ExtraUrlsFile, ExtraIndexUrls!);
-
-			return true;
-		}
-
-		public bool ParsePluginDependencies(FileReference MergedReqsInFile, FileReference MergedRequirmentsFile, bool NoCheckMerged = false)
-		{
 			using (IBaseCmdProgressLogger CmdLogger = new PythonCmdLogger(Logger))
 			{
 				return (RunPythonVenv($"-m ue_parse_plugin_reqs -vv \"{MergedReqsInFile}\" \"{MergedRequirmentsFile}\"", CmdLogger) == 0);
 			}
 		}
 
-		public bool InstallPluginDependencies(FileReference MergedRequirementsFile, FileReference ExtraUrlsFile, bool OfflineOnly = false, string? ForceIndexUrl = null)
+		public bool InstallPluginDependencies(bool OfflineOnly = false, string? ForceIndexUrl = null)
 		{
+			FileReference MergedRequirementsFile = FileReference.Combine(InstallDir, MergedRequirementsFilename);
+			FileReference ExtraUrlsFile = FileReference.Combine(InstallDir, ExtraUrlsFilename);
+
 			if (!FileReference.Exists(MergedRequirementsFile))
 			{
 				return true;
@@ -297,7 +336,7 @@
 			return PipInstall(MergedRequirementsFile, ExtraUrls, OfflineOnly, ForceIndexUrl);
 		}
 
-		public bool ViewInstalledLicenses(FileReference PluginsListingFile)
+		public bool ViewInstalledLicenses()
 		{
 			// TODO: Check that install is up to date and reverse-map package to plugin requirements
 			//DirectoryReference WorkDir = DirectoryReference.FromFile(PluginsListingFile);
@@ -430,8 +469,21 @@
 			return true;
 		}
 
+		public bool CheckPipInstallUtils()
+		{
+			// Verify that correct version of pip install utils is already available
+			string Args = $"-c \"import pkg_resources;dist=pkg_resources.working_set.find(pkg_resources.Requirement.parse('ue-pipinstall-utils'));exit(dist.version!='{PipInstallUtilsVer}' if dist is not None else 1)\"";
+			using IBaseCmdProgressLogger CmdLogger = new PythonCmdLogger(Logger);
+			return (RunPythonVenv(Args, CmdLogger) == 0);
+		}
+
 		private bool SetupPipInstallUtils(IBaseCmdProgressLogger StatusLogger)
 		{
+			if (CheckPipInstallUtils())
+			{
+				return true;
+			}
+
 			Logger.LogInformation("PipInstall: Updating UE PipInstall Utilities");
 			FileReference? PythonScriptPlugin = GetPythonScriptPlugin();
 			if (PythonScriptPlugin == null)
@@ -443,7 +495,7 @@
 			DirectoryReference PipInstallUtilsDir = DirectoryReference.Combine(PythonScriptDir, "Content", "Python", "PipInstallUtils");
 			DirectoryReference PipWheelsDir = DirectoryReference.Combine(PythonScriptDir, "Content", "Python", "Lib", "wheels");
 			FileReference RequirementsFile = FileReference.Combine(PipInstallUtilsDir, "requirements.txt");
-			string Args = $"-m pip install --upgrade --no-index --find-links \"{PipWheelsDir}\" -r \"{RequirementsFile}\" ue-pipinstall-utils";
+			string Args = $"-m pip install --upgrade --no-index --find-links \"{PipWheelsDir}\" -r \"{RequirementsFile}\" ue-pipinstall-utils=={PipInstallUtilsVer}";
 			int Result = RunPythonVenv(Args, StatusLogger);
 			return (Result == 0);
 		}
@@ -750,7 +802,9 @@
 		private int StepsDone;
 		private int TotalSteps;
 
+		// Start strings to use 
 		private static readonly string[] MatchStrs = { "Requirement", "Downloading", "Using", "Installing" };
+		private readonly Dictionary<string,string> LogReplaceStrs = new();
 
 		public PipProgressLogger(ILogger InLogger, string message, int GuessSteps)
 		{
@@ -759,13 +813,15 @@
 			TotalSteps = Math.Max(GuessSteps, 1);
 
 			Writer = new ProgressWriter(message, true, Logger);
+
+			LogReplaceStrs["Installing collected packages:"] = "Installing collected python package dependencies:";
 		}
 
 		static bool CheckUpdateStr(string CheckStr)
 		{
 			foreach (string ProgressMatch in MatchStrs)
 			{
-				if (CheckStr.StartsWith(ProgressMatch, StringComparison.Ordinal))
+				if (CheckStr.StartsWith(ProgressMatch, StringComparison.InvariantCultureIgnoreCase))
 				{
 					return true;
 				}
@@ -774,6 +830,19 @@
 			return false;
 		}
 
+		string ReplaceMatchStr(string CheckStr)
+		{
+			foreach(KeyValuePair<string,string> ChkPair in LogReplaceStrs)
+			{
+				if (CheckStr.Contains(ChkPair.Key))
+				{
+					return CheckStr.Replace(ChkPair.Key, ChkPair.Value);
+				}
+			}
+
+			return CheckStr;
+		}
+
 		public void OutputData(DataReceivedEventArgs DataLine)
 		{
 			// Currently we assume only one pip command so it should be finished on (null)
@@ -784,6 +853,8 @@
 			bool ShouldUpdate = CheckUpdateStr(CheckStr);
 			if (ShouldUpdate)
 			{
+				CheckStr = ReplaceMatchStr(CheckStr);
+
 				Writer = new ProgressWriter(CheckStr, true, Logger);
 				Writer.Write(StepsDone, TotalSteps);
 
@@ -792,7 +863,7 @@
 			}
 			else
 			{
-				Logger.LogInformation("{Data}", DataLine.Data);
+				Logger.LogInformation("{CheckStr}", CheckStr);
 			}
 		}
 
