- When looking up a target by type (-TargetType=Editor) to get ActiveArchitecture on Mac, we may fail to find a target with the matching name (like Foo_NDAPlatform.Target.cs).
This will just ignore the project setting for that target, since it failed to be found, there are no settings to query

diff --git a/Engine/Source/Programs/UnrealBuildTool/Platform/Mac/UEBuildMac.cs b/Engine/Source/Programs/UnrealBuildTool/Platform/Mac/UEBuildMac.cs
--- a/Engine/Source/Programs/UnrealBuildTool/Platform/Mac/UEBuildMac.cs
+++ b/Engine/Source/Programs/UnrealBuildTool/Platform/Mac/UEBuildMac.cs
@@ -173,17 +173,24 @@ namespace UnrealBuildTool
 					RulesAsm = RulesCompiler.CreateProjectRulesAssembly(ProjectFile, Unreal.IsEngineInstalled(), false, false, Log.Logger);
 				}
 
-				// CreateTargetRules here needs to have an UnrealArchitectures object, because otherwise with 'null', it will call
-				// back to this function to get the ActiveArchitectures! in this case the arch is unimportant
-				UnrealArchitectures DummyArchitectures = new(UnrealArch.X64);
-				TargetRules? Rules = RulesAsm.CreateTargetRules(TargetName, UnrealTargetPlatform.Mac, UnrealTargetConfiguration.Development, DummyArchitectures, ProjectFile, null, Log.Logger);
-				bIsEditor = Rules.Type == TargetType.Editor;
-
-				// the projectfile passed in may be a game's uproject file that we are compiling a program in the context of, 
-				// but we still want the settings for the program
-				if (Rules.Type == TargetType.Program)
+				try
 				{
-					ProjectFile = Rules.ProjectFile;
+					// CreateTargetRules here needs to have an UnrealArchitectures object, because otherwise with 'null', it will call
+					// back to this function to get the ActiveArchitectures! in this case the arch is unimportant
+					UnrealArchitectures DummyArchitectures = new(UnrealArch.X64);
+					TargetRules? Rules = RulesAsm.CreateTargetRules(TargetName, UnrealTargetPlatform.Mac, UnrealTargetConfiguration.Development, DummyArchitectures, ProjectFile, null, Log.Logger);
+					bIsEditor = Rules.Type == TargetType.Editor;
+
+					// the projectfile passed in may be a game's uproject file that we are compiling a program in the context of, 
+					// but we still want the settings for the program
+					if (Rules.Type == TargetType.Program)
+					{
+						ProjectFile = Rules.ProjectFile;
+					}
+				}
+				catch (Exception)
+				{
+					// do nothing if it fails, assume no project
 				}
 			}
 
