Allow adding optional plugin build dependencies using a target.cs file

diff --git a/Engine/Source/Programs/UnrealBuildTool/Configuration/TargetRules.cs b/Engine/Source/Programs/UnrealBuildTool/Configuration/TargetRules.cs
--- a/Engine/Source/Programs/UnrealBuildTool/Configuration/TargetRules.cs
+++ b/Engine/Source/Programs/UnrealBuildTool/Configuration/TargetRules.cs
@@ -777,6 +777,11 @@ namespace UnrealBuildTool
 		[CommandLine("-DisablePlugin=", ListSeparator = '+')]
 		public List<string> DisablePlugins = new List<string>();
 
+		/// <summary>
+		/// Additional plugins that should be included for this target if they are found.
+		/// </summary>
+		public List<string> OptionalPlugins = new List<string>();
+
 		/// <summary>
 		/// How to treat conflicts when a disabled plugin is being enabled by another plugin referencing it
 		/// </summary>
@@ -2890,6 +2895,7 @@ namespace UnrealBuildTool
 			{
 				DisablePlugins.Add(PluginName);
 				EnablePlugins.Remove(PluginName);
+				OptionalPlugins.Remove(PluginName);
 			}
 		}
 
@@ -3178,6 +3184,8 @@ namespace UnrealBuildTool
 
 		public IEnumerable<string> DisablePlugins => Inner.DisablePlugins;
 
+		public IEnumerable<string> OptionalPlugins => Inner.OptionalPlugins;
+
 		public WarningLevel DisablePluginsConflictWarningLevel => Inner.DisablePluginsConflictWarningLevel;
 
 		public IEnumerable<string> BuildPlugins => Inner.BuildPlugins;
diff --git a/Engine/Source/Programs/UnrealBuildTool/Configuration/UEBuildTarget.cs b/Engine/Source/Programs/UnrealBuildTool/Configuration/UEBuildTarget.cs
--- a/Engine/Source/Programs/UnrealBuildTool/Configuration/UEBuildTarget.cs
+++ b/Engine/Source/Programs/UnrealBuildTool/Configuration/UEBuildTarget.cs
@@ -4343,6 +4343,17 @@ namespace UnrealBuildTool
 				}
 			}
 
+			// Configure optional plugins configured via target settings
+			foreach (string PluginName in Rules.OptionalPlugins)
+			{
+				if (ReferencedNames.Add(PluginName))
+				{
+					PluginReferenceDescriptor PluginReference = new PluginReferenceDescriptor(PluginName, null, true);
+					PluginReference.bOptional = true;
+					AddPlugin(PluginReference, "target settings", ExcludeFolders, NameToInstance, NameToInfos, Logger);
+				}
+			}
+
 			bool bAllowEnginePluginsEnabledByDefault = true;
 
 			// Find a map of plugins which are explicitly referenced in the project file
