Remove AutoRTFM suffix when AutoRTFM is enabled

diff --git a/Engine/Source/Programs/UnrealBuildTool/Platform/Linux/UEBuildLinux.cs b/Engine/Source/Programs/UnrealBuildTool/Platform/Linux/UEBuildLinux.cs
--- a/Engine/Source/Programs/UnrealBuildTool/Platform/Linux/UEBuildLinux.cs
+++ b/Engine/Source/Programs/UnrealBuildTool/Platform/Linux/UEBuildLinux.cs
@@ -238,11 +238,6 @@ namespace UnrealBuildTool
 			{
 				string? SanitizerSuffix = null;
 
-				if (Target.bUseAutoRTFMCompiler)
-				{
-					SanitizerSuffix = "AutoRTFM";
-				}
-
 				if (Target.LinuxPlatform.bEnableAddressSanitizer)
 				{
 					SanitizerSuffix = "ASan";
