- Delete the staged dir inside a .app if the source staged dir is missing - this is so we don't use some stale dir insde the .app if the user wanted to delete it (mostly useful for Programs)

diff --git a/Engine/Source/Programs/UnrealBuildTool/ProjectFiles/Xcode/XcodeProject.cs b/Engine/Source/Programs/UnrealBuildTool/ProjectFiles/Xcode/XcodeProject.cs
--- a/Engine/Source/Programs/UnrealBuildTool/ProjectFiles/Xcode/XcodeProject.cs
+++ b/Engine/Source/Programs/UnrealBuildTool/ProjectFiles/Xcode/XcodeProject.cs
@@ -1416,8 +1416,12 @@ namespace UnrealBuildTool.XcodeProjectXcconfig
 					CopyScript.AddRange(new string[]
 					{
 						"# Make sure the staged directory exists and has files in it",
-						"if [[ ! -e ${STAGED_DIR} ]]; then exit 0; fi",
-					});
+						"if [[ ! -e ${STAGED_DIR} ]]; then ",
+						"  # Make sure the target doesn't exist (so if we delete the Staged dir, it goes back to unstaged",
+						$"  rm -rf \\\"${{CONFIGURATION_BUILD_DIR}}/${{CONTENTS_FOLDER_PATH}}{SyncDestSubdir}\\\"",
+						"  exit -0",
+						"fi",
+					}); ;
 				}
 
 				// when we bring stated data into the .app, we have to skip some temp stuff that went into it
