UnrealBuildTool: Also load BuildConfiguration from %ProgramData% & %LocalAppData%

diff --git a/Engine/Source/Programs/UnrealBuildTool/System/XmlConfig.cs b/Engine/Source/Programs/UnrealBuildTool/System/XmlConfig.cs
--- a/Engine/Source/Programs/UnrealBuildTool/System/XmlConfig.cs
+++ b/Engine/Source/Programs/UnrealBuildTool/System/XmlConfig.cs
@@ -276,57 +276,63 @@ namespace UnrealBuildTool
 				ILogger Logger = Log.Logger;
 
 				// Find all the input file locations
-				List<InputFile> InputFilesFound = new List<InputFile>(4);
+				List<InputFile> InputFilesFound = new List<InputFile>(5);
+
+				// InputFile info and if a default file should be created if missing
+				List<KeyValuePair<InputFile, bool>> configs = new();
 
 				// Skip all the config files under the Engine folder if it's an installed build
 				if (!Unreal.IsEngineInstalled())
 				{
-					// Check for the config file under /Engine/Programs/NotForLicensees/UnrealBuildTool
-					FileReference NotForLicenseesConfigLocation = FileReference.Combine(Unreal.EngineDirectory, "Restricted", "NotForLicensees", "Programs", "UnrealBuildTool", "BuildConfiguration.xml");
-					if (FileReference.Exists(NotForLicenseesConfigLocation))
-					{
-						InputFilesFound.Add(new InputFile(NotForLicenseesConfigLocation, "NotForLicensees"));
-					}
-					else
-					{
-						Logger.LogDebug("No config file at {NotForLicenseesConfigLocation}", NotForLicenseesConfigLocation);
-					}
+					// Check for the engine config file under /Engine/Programs/NotForLicensees/UnrealBuildTool
+					configs.Add(new(new InputFile(FileReference.Combine(Unreal.EngineDirectory, "Restricted", "NotForLicensees", "Programs", "UnrealBuildTool", "BuildConfiguration.xml"), "Engine (NotForLicensees)"), false));
 
-					// Check for the user config file under /Engine/Saved/UnrealBuildTool
-					FileReference UserConfigLocation = FileReference.Combine(Unreal.EngineDirectory, "Saved", "UnrealBuildTool", "BuildConfiguration.xml");
-					if (!FileReference.Exists(UserConfigLocation))
-					{
-						Logger.LogDebug("Creating default config file at {UserConfigLocation}", UserConfigLocation);
-						CreateDefaultConfigFile(UserConfigLocation);
-					}
-					InputFilesFound.Add(new InputFile(UserConfigLocation, "User"));
+					// Check for the engine user config file under /Engine/Saved/UnrealBuildTool
+					configs.Add(new(new InputFile(FileReference.Combine(Unreal.EngineDirectory, "Saved", "UnrealBuildTool", "BuildConfiguration.xml"), "Engine (Saved)"), true));
 				}
 
-				// Check for the global config file under AppData/Unreal Engine/UnrealBuildTool
-				string AppDataFolder = Environment.GetFolderPath(Environment.SpecialFolder.ApplicationData);
-				if (!String.IsNullOrEmpty(AppDataFolder))
+				// Check for the global config file under ProgramData/Unreal Engine/UnrealBuildTool
+				DirectoryReference? CommonProgramsFolder = DirectoryReference.FromString(Environment.GetFolderPath(Environment.SpecialFolder.CommonApplicationData));
+				if (CommonProgramsFolder != null)
 				{
-					FileReference AppDataConfigLocation = FileReference.Combine(new DirectoryReference(AppDataFolder), "Unreal Engine", "UnrealBuildTool", "BuildConfiguration.xml");
-					if (!FileReference.Exists(AppDataConfigLocation))
-					{
-						Logger.LogDebug("Creating default config file at {AppDataConfigLocation}", AppDataConfigLocation);
-						CreateDefaultConfigFile(AppDataConfigLocation);
-					}
-					InputFilesFound.Add(new InputFile(AppDataConfigLocation, "Global (AppData)"));
+					configs.Add(new(new InputFile(FileReference.Combine(CommonProgramsFolder, "Unreal Engine", "UnrealBuildTool", "BuildConfiguration.xml"), "Global (ProgramData)"), false));
+				}
+
+				// Check for the global config file under AppData/Unreal Engine/UnrealBuildTool (Roaming)
+				DirectoryReference? AppDataFolder = DirectoryReference.FromString(Environment.GetFolderPath(Environment.SpecialFolder.ApplicationData));
+				if (AppDataFolder != null)
+				{
+					configs.Add(new(new InputFile(FileReference.Combine(AppDataFolder, "Unreal Engine", "UnrealBuildTool", "BuildConfiguration.xml"), "Global (AppData)"), true));
+				}
+
+				// Check for the global config file under LocalAppData/Unreal Engine/UnrealBuildTool
+				DirectoryReference? LocalAppDataFolder = DirectoryReference.FromString(Environment.GetFolderPath(Environment.SpecialFolder.LocalApplicationData));
+				if (LocalAppDataFolder != null)
+				{
+					configs.Add(new(new InputFile(FileReference.Combine(LocalAppDataFolder, "Unreal Engine", "UnrealBuildTool", "BuildConfiguration.xml"), "Global (LocalAppData)"), false));
 				}
 
 				// Check for the global config file under My Documents/Unreal Engine/UnrealBuildTool
-				string PersonalFolder = Environment.GetFolderPath(Environment.SpecialFolder.Personal);
-				if (!String.IsNullOrEmpty(PersonalFolder))
+				DirectoryReference? PersonalFolder = DirectoryReference.FromString(Environment.GetFolderPath(Environment.SpecialFolder.Personal));
+				if (PersonalFolder != null)
 				{
-					FileReference PersonalConfigLocation = FileReference.Combine(new DirectoryReference(PersonalFolder), "Unreal Engine", "UnrealBuildTool", "BuildConfiguration.xml");
-					if (FileReference.Exists(PersonalConfigLocation))
+					configs.Add(new(new InputFile(FileReference.Combine(PersonalFolder, "Unreal Engine", "UnrealBuildTool", "BuildConfiguration.xml"), "Global (Documents)"), false));
+				}
+
+				foreach (KeyValuePair<InputFile, bool> config in configs)
+				{
+					if (config.Value && !FileReference.Exists(config.Key.Location))
+					{
+						Logger.LogDebug("Creating default config file at {ConfigLocation}", config.Key.Location);
+						CreateDefaultConfigFile(config.Key.Location);
+					}
+					if (FileReference.Exists(config.Key.Location))
 					{
-						InputFilesFound.Add(new InputFile(PersonalConfigLocation, "Global (Documents)"));
+						InputFilesFound.Add(config.Key);
 					}
 					else
 					{
-						Logger.LogDebug("No config file at {PersonalConfigLocation}", PersonalConfigLocation);
+						Logger.LogDebug("No config file at {ConfigLocation}", config.Key.Location);
 					}
 				}
 
diff --git a/Engine/Source/Programs/UnrealBuildTool/UnrealBuildTool.csproj b/Engine/Source/Programs/UnrealBuildTool/UnrealBuildTool.csproj
--- a/Engine/Source/Programs/UnrealBuildTool/UnrealBuildTool.csproj
+++ b/Engine/Source/Programs/UnrealBuildTool/UnrealBuildTool.csproj
@@ -164,9 +164,11 @@
 
   <!-- Various BuildConfiguration files -->
   <ItemGroup>
-      <None Condition="Exists('..\..\..\Restricted\NotForLicensees\Programs\UnrealBuildTool\BuildConfiguration.xml')" Include="..\..\..\Restricted\NotForLicensees\Programs\UnrealBuildTool\BuildConfiguration.xml"><Link>BuildConfiguration\BuildConfiguration.xml (NotForLicensees)</Link></None>
-      <None Condition="Exists('..\..\..\Saved\UnrealBuildTool\BuildConfiguration.xml')" Include="..\..\..\Saved\UnrealBuildTool\BuildConfiguration.xml"><Link>BuildConfiguration\BuildConfiguration.xml (Saved)</Link></None>
-      <None Condition="Exists('$(AppData)\Unreal Engine\UnrealBuildTool\BuildConfiguration.xml')" Include="$(AppData)\Unreal Engine\UnrealBuildTool\BuildConfiguration.xml"><Link>BuildConfiguration\BuildConfiguration.xml (User AppData)</Link></None>
-      <None Condition="Exists('$(UserProfile)\Documents\Unreal Engine\UnrealBuildTool\BuildConfiguration.xml')" Include="$(UserProfile)\Documents\Unreal Engine\UnrealBuildTool\BuildConfiguration.xml"><Link>BuildConfiguration\BuildConfiguration.xml (User Documents)</Link></None>
+      <None Condition="Exists('..\..\..\Restricted\NotForLicensees\Programs\UnrealBuildTool\BuildConfiguration.xml')" Include="..\..\..\Restricted\NotForLicensees\Programs\UnrealBuildTool\BuildConfiguration.xml"><Link>BuildConfiguration\BuildConfiguration.xml (Engine NotForLicensees)</Link></None>
+      <None Condition="Exists('..\..\..\Saved\UnrealBuildTool\BuildConfiguration.xml')" Include="..\..\..\Saved\UnrealBuildTool\BuildConfiguration.xml"><Link>BuildConfiguration\BuildConfiguration.xml (Engine Saved)</Link></None>
+      <None Condition="Exists('$(ProgramData)\Unreal Engine\UnrealBuildTool\BuildConfiguration.xml')" Include="$(ProgramData)\Unreal Engine\UnrealBuildTool\BuildConfiguration.xml"><Link>BuildConfiguration\BuildConfiguration.xml (Global ProgramData)</Link></None>
+      <None Condition="Exists('$(AppData)\Unreal Engine\UnrealBuildTool\BuildConfiguration.xml')" Include="$(AppData)\Unreal Engine\UnrealBuildTool\BuildConfiguration.xml"><Link>BuildConfiguration\BuildConfiguration.xml (Global AppData)</Link></None>
+      <None Condition="Exists('$(LocalAppData)\Unreal Engine\UnrealBuildTool\BuildConfiguration.xml')" Include="$(LocalAppData)\Unreal Engine\UnrealBuildTool\BuildConfiguration.xml"><Link>BuildConfiguration\BuildConfiguration.xml (Global LocalAppData)</Link></None>
+      <None Condition="Exists('$(UserProfile)\Documents\Unreal Engine\UnrealBuildTool\BuildConfiguration.xml')" Include="$(UserProfile)\Documents\Unreal Engine\UnrealBuildTool\BuildConfiguration.xml"><Link>BuildConfiguration\BuildConfiguration.xml (Global Documents)</Link></None>
   </ItemGroup>
 </Project>
