UBT: Fix for MakeStubItems with Circular Deps.

diff --git a/Engine/Source/Programs/UnrealBuildTool/Platform/Mac/MacToolChain.cs b/Engine/Source/Programs/UnrealBuildTool/Platform/Mac/MacToolChain.cs
--- a/Engine/Source/Programs/UnrealBuildTool/Platform/Mac/MacToolChain.cs
+++ b/Engine/Source/Programs/UnrealBuildTool/Platform/Mac/MacToolChain.cs
@@ -666,7 +666,10 @@ namespace UnrealBuildTool
 
 		private FileItem MakeStubItem(LinkEnvironment LinkEnvironment, string DylibPath)
 		{
-			string IntermediateDirectory = (ProjectFile == null || LinkEnvironment.OutputFilePaths.All(x => x.IsUnderDirectory(Unreal.EngineDirectory)) ? Unreal.EngineDirectory : ProjectFile.Directory) + "/Intermediate/Mac";
+			FileReference DylibPathRef = new FileReference(DylibPath);
+			bool bIsEngineDylib = DylibPathRef.IsUnderDirectory(Unreal.EngineDirectory);
+
+			string IntermediateDirectory = (ProjectFile == null || LinkEnvironment.OutputFilePaths.All(x => x.IsUnderDirectory(Unreal.EngineDirectory)) || bIsEngineDylib ? Unreal.EngineDirectory : ProjectFile.Directory) + "/Intermediate/Mac";
 			return FileItem.GetItemByPath(Path.Combine(IntermediateDirectory, "Stubs", LinkEnvironment.Architecture.ToString(), Path.GetFileName(DylibPath)));
 		}
 
