Fixing bDontBundleLibrariesInAPK not working after first install

diff --git a/Engine/Source/Programs/UnrealBuildTool/Platform/Android/UEDeployAndroid.cs b/Engine/Source/Programs/UnrealBuildTool/Platform/Android/UEDeployAndroid.cs
--- a/Engine/Source/Programs/UnrealBuildTool/Platform/Android/UEDeployAndroid.cs
+++ b/Engine/Source/Programs/UnrealBuildTool/Platform/Android/UEDeployAndroid.cs
@@ -4837,13 +4837,17 @@ pushd %~dp0
 					{
 						SOPushScript += @$"
 %AFS% %DEVICE% push {FinalSONameStrippedRelative} ""^int/libUnreal.so""
+if ""%ERRORLEVEL%"" NEQ ""0"" (exit /b %ERRORLEVEL%)
 ";
 					}
 					else
 					{
 						SOPushScript += @$"
 %ADB% %DEVICE% push -z lz4 {FinalSONameStrippedRelative} /data/local/tmp/{FinalSONameStrippedRelative}
+if ""%ERRORLEVEL%"" NEQ ""0"" (exit /b %ERRORLEVEL%)
+%ADB% %DEVICE% shell run-as {PackageName} mkdir -p ./files
 %ADB% %DEVICE% shell run-as {PackageName} cp /data/local/tmp/{FinalSONameStrippedRelative} ./files/libUnreal.so
+if ""%ERRORLEVEL%"" NEQ ""0"" (exit /b %ERRORLEVEL%)
 %ADB% %DEVICE% shell rm /data/local/tmp/{FinalSONameStrippedRelative}
 ";
 					}
