UnrealBuildTool: Fix missing XmlConfig attributes

#rnx

diff --git a/Engine/Source/Programs/UnrealBuildTool/Configuration/BuildConfiguration.cs b/Engine/Source/Programs/UnrealBuildTool/Configuration/BuildConfiguration.cs
--- a/Engine/Source/Programs/UnrealBuildTool/Configuration/BuildConfiguration.cs
+++ b/Engine/Source/Programs/UnrealBuildTool/Configuration/BuildConfiguration.cs
@@ -50,6 +50,7 @@ namespace UnrealBuildTool
 		/// <summary>
 		/// Whether the experimental UnrealBuildAccelerator executor will be used.
 		/// </summary>
+		[XmlConfigFile]
 		[CommandLine("-UBA", Value = "true")]
 		[CommandLine("-NoUBA", Value = "false")]
 		public bool bAllowUBAExecutor
@@ -70,6 +71,7 @@ namespace UnrealBuildTool
 		/// <summary>
 		/// Whether the experimental UnrealBuildAccelerator (local only) executor will be used.
 		/// </summary>
+		[XmlConfigFile]
 		[CommandLine("-UBALocal", Value = "true")]
 		[CommandLine("-NoUBALocal", Value = "false")]
 		public bool bAllowUBALocalExecutor
