AppleExports: xros->visionos

diff --git a/Engine/Source/Programs/UnrealBuildTool/Platform/Apple/AppleExports.cs b/Engine/Source/Programs/UnrealBuildTool/Platform/Apple/AppleExports.cs
--- a/Engine/Source/Programs/UnrealBuildTool/Platform/Apple/AppleExports.cs
+++ b/Engine/Source/Programs/UnrealBuildTool/Platform/Apple/AppleExports.cs
@@ -122,7 +122,7 @@ namespace UnrealBuildTool
 			}
 			else if (Platform == UnrealTargetPlatform.VisionOS)
 			{
-				return Architectures.SingleArchitecture == UnrealArch.IOSSimulator ? "xrOS Simulator" : "xrOS";
+				return Architectures.SingleArchitecture == UnrealArch.IOSSimulator ? "visionOS Simulator" : "visionOS";
 			}
 
 			throw new BuildException($"Unknown plaform {Platform}");
