- Fix for Mac monolithic CookedEditors not getting staged data in their app. Modular cooked editors will stll be missing data

diff --git a/Engine/Source/Programs/UnrealBuildTool/ProjectFiles/Xcode/XcodeProject.cs b/Engine/Source/Programs/UnrealBuildTool/ProjectFiles/Xcode/XcodeProject.cs
--- a/Engine/Source/Programs/UnrealBuildTool/ProjectFiles/Xcode/XcodeProject.cs
+++ b/Engine/Source/Programs/UnrealBuildTool/ProjectFiles/Xcode/XcodeProject.cs
@@ -1361,19 +1361,20 @@ namespace UnrealBuildTool.XcodeProjectXcconfig
 
 			CopyScript.Add("");
 
-			// Editor just need the above script to copy executable into .app
-			if (Project.UnrealData.TargetRules.Type == TargetType.Editor)
+			if (Project.UnrealData.TargetRules.Type == TargetType.Editor && UnrealData.TargetRules.LinkType == TargetLinkType.Modular)
 			{
+				// Editor just need the above script to copy executable into .app
+
 				XcodeShellScriptBuildPhase EditorCopyScriptPhase = new("Copy Executable into .app", CopyScript, new string[] { }, new string[] { $"/dev/null" });
 				BuildPhases.Add(EditorCopyScriptPhase);
 				References.Add(EditorCopyScriptPhase);
 				return;
 			}
-
-			// rsync the Staged build into the .app, unless the UE_SKIP_STAGEDDATA_SYNC var is set to 1
-			// editor builds don't need staged content in them
-			if (TargetType != TargetType.Editor)
+			else
 			{
+				// rsync the Staged build into the .app, unless the UE_SKIP_STAGEDDATA_SYNC var is set to 1
+				// editor builds don't need staged content in them
+
 				bool bIsEngineBuild = Project.UnrealData.UProjectFileLocation == null;
 				string DefaultStageDir = bIsEngineBuild ? "" : "${UE_PROJECT_DIR}/Saved/StagedBuilds/${UE_TARGET_PLATFORM_NAME}";
 				string SyncSourceSubdir = (Platform == UnrealTargetPlatform.Mac) ? "" : "/cookeddata";
