UnrealBuildTool: Move common action properties for CreateRecursiveAction to be shared

#rnx

diff --git a/Engine/Source/Programs/UnrealBuildTool/Configuration/UEBuildTarget.cs b/Engine/Source/Programs/UnrealBuildTool/Configuration/UEBuildTarget.cs
--- a/Engine/Source/Programs/UnrealBuildTool/Configuration/UEBuildTarget.cs
+++ b/Engine/Source/Programs/UnrealBuildTool/Configuration/UEBuildTarget.cs
@@ -3052,10 +3052,7 @@ namespace UnrealBuildTool
 
 				// Create the action
 				Action WriteMetadataAction = MakefileBuilder.CreateRecursiveAction<WriteMetadataMode>(ActionType.WriteMetadata, WriteMetadataArguments.ToString());
-				WriteMetadataAction.WorkingDirectory = Unreal.EngineSourceDirectory;
 				WriteMetadataAction.StatusDescription = StatusDescription;
-				WriteMetadataAction.bCanExecuteRemotely = false;
-				WriteMetadataAction.bCanExecuteInUBA = false;
 				WriteMetadataAction.bUseActionHistory = false; // Different files for each target; do not want to invalidate based on this.
 				WriteMetadataAction.ProducedItems.UnionWith(ProducedItems.Select(x => FileItem.GetItemByFileReference(x)));
 				WriteMetadataAction.PrerequisiteItems.UnionWith(PrerequisiteItems);
diff --git a/Engine/Source/Programs/UnrealBuildTool/Platform/IOS/IOSToolChain.cs b/Engine/Source/Programs/UnrealBuildTool/Platform/IOS/IOSToolChain.cs
--- a/Engine/Source/Programs/UnrealBuildTool/Platform/IOS/IOSToolChain.cs
+++ b/Engine/Source/Programs/UnrealBuildTool/Platform/IOS/IOSToolChain.cs
@@ -1197,7 +1197,6 @@ namespace UnrealBuildTool
 				PostBuildSyncAction.ProducedItems.Add(FileItem.GetItemByFileReference(GetStagedExecutablePath(Executable.Location, Target.Name)));
 				PostBuildSyncAction.DeleteItems.UnionWith(PostBuildSyncAction.ProducedItems);
 				PostBuildSyncAction.StatusDescription = "Executing PostBuildSync";
-				PostBuildSyncAction.bCanExecuteRemotely = false;
 
 				OutputFiles.AddRange(PostBuildSyncAction.ProducedItems);
 			}
diff --git a/Engine/Source/Programs/UnrealBuildTool/Platform/Windows/VCToolChain.cs b/Engine/Source/Programs/UnrealBuildTool/Platform/Windows/VCToolChain.cs
--- a/Engine/Source/Programs/UnrealBuildTool/Platform/Windows/VCToolChain.cs
+++ b/Engine/Source/Programs/UnrealBuildTool/Platform/Windows/VCToolChain.cs
@@ -1804,7 +1804,6 @@ namespace UnrealBuildTool
 			}
 
 			Action ParseTimingInfoAction = Graph.CreateRecursiveAction<ParseMsvcTimingInfoMode>(ActionType.ParseTimingInfo, ParseTimingArguments);
-			ParseTimingInfoAction.WorkingDirectory = Unreal.EngineSourceDirectory;
 			ParseTimingInfoAction.StatusDescription = Path.GetFileName(TimingFile.AbsolutePath);
 			ParseTimingInfoAction.bCanExecuteRemotely = true;
 			ParseTimingInfoAction.bCanExecuteRemotelyWithSNDBS = true;
@@ -1858,10 +1857,7 @@ namespace UnrealBuildTool
 						};
 
 						Action AggregateTimingInfoAction = MakefileBuilder.CreateRecursiveAction<AggregateParsedTimingInfo>(ActionType.ParseTimingInfo, String.Join(" ", ActionArgs));
-						AggregateTimingInfoAction.WorkingDirectory = Unreal.EngineSourceDirectory;
 						AggregateTimingInfoAction.StatusDescription = $"Aggregating {TimingJsonFiles.Count} Timing File(s)";
-						AggregateTimingInfoAction.bCanExecuteRemotely = false;
-						AggregateTimingInfoAction.bCanExecuteRemotelyWithSNDBS = false;
 						AggregateTimingInfoAction.PrerequisiteItems.UnionWith(TimingJsonFiles);
 
 						FileItem AggregateOutputFile = FileItem.GetItemByFileReference(FileReference.Combine(Makefile.ProjectIntermediateDirectory, $"{Target.Name}.cta"));
@@ -1883,10 +1879,7 @@ namespace UnrealBuildTool
 						};
 
 						Action AggregateTimingInfoAction = MakefileBuilder.CreateRecursiveAction<AggregateClangTimingInfo>(ActionType.ParseTimingInfo, String.Join(" ", AggregateActionArgs));
-						AggregateTimingInfoAction.WorkingDirectory = Unreal.EngineSourceDirectory;
 						AggregateTimingInfoAction.StatusDescription = $"Aggregating {TimingJsonFiles.Count} Timing File(s)";
-						AggregateTimingInfoAction.bCanExecuteRemotely = false;
-						AggregateTimingInfoAction.bCanExecuteRemotelyWithSNDBS = false;
 						AggregateTimingInfoAction.PrerequisiteItems.UnionWith(TimingJsonFiles);
 
 						AggregateTimingInfoAction.ProducedItems.Add(AggregateOutputFile);
@@ -1901,10 +1894,7 @@ namespace UnrealBuildTool
 						};
 
 						Action ArchiveTimingInfoAction = MakefileBuilder.CreateRecursiveAction<AggregateClangTimingInfo>(ActionType.ParseTimingInfo, String.Join(" ", ArchiveActionArgs));
-						ArchiveTimingInfoAction.WorkingDirectory = Unreal.EngineSourceDirectory;
 						ArchiveTimingInfoAction.StatusDescription = $"Archiving {TimingJsonFiles.Count} Timing File(s)";
-						ArchiveTimingInfoAction.bCanExecuteRemotely = false;
-						ArchiveTimingInfoAction.bCanExecuteRemotelyWithSNDBS = false;
 						ArchiveTimingInfoAction.PrerequisiteItems.UnionWith(TimingJsonFiles);
 
 						ArchiveTimingInfoAction.ProducedItems.Add(ArchiveOutputFile);
@@ -1921,6 +1911,7 @@ namespace UnrealBuildTool
 							CompileScoreExtractorAction.StatusDescription = $"Extracting CompileScore";
 							CompileScoreExtractorAction.bCanExecuteRemotely = false;
 							CompileScoreExtractorAction.bCanExecuteRemotelyWithSNDBS = false;
+							CompileScoreExtractorAction.bCanExecuteInUBA = false; // TODO: Unknown if supported
 							CompileScoreExtractorAction.PrerequisiteItems.UnionWith(TimingJsonFiles);
 							CompileScoreExtractorAction.CommandPath = ScoreDataExtractor;
 							CompileScoreExtractorAction.CommandArguments = $"-clang -verbosity 0 -timelinepack 1000000 -extract -i \"{NormalizeCommandLinePath(Makefile.ProjectIntermediateDirectory)}\" -o \"{NormalizeCommandLinePath(CompileScoreOutput)}\"";
diff --git a/Engine/Source/Programs/UnrealBuildTool/System/ActionGraphBuilder.cs b/Engine/Source/Programs/UnrealBuildTool/System/ActionGraphBuilder.cs
--- a/Engine/Source/Programs/UnrealBuildTool/System/ActionGraphBuilder.cs
+++ b/Engine/Source/Programs/UnrealBuildTool/System/ActionGraphBuilder.cs
@@ -313,8 +313,12 @@ namespace UnrealBuildTool
 
 			Action NewAction = Graph.CreateAction(Type);
 			NewAction.CommandPath = Unreal.DotnetPath;
+			NewAction.WorkingDirectory = Unreal.EngineSourceDirectory;
 			NewAction.CommandArguments = $"\"{Unreal.UnrealBuildToolDllPath}\" -Mode={Attribute.Name} {Arguments}";
 			NewAction.CommandDescription = Attribute.Name;
+			NewAction.bCanExecuteRemotely = false;
+			NewAction.bCanExecuteRemotelyWithSNDBS = false;
+			NewAction.bCanExecuteInUBA = false;
 			return NewAction;
 		}
 
diff --git a/Engine/Source/Programs/UnrealBuildTool/ToolChain/AppleToolChain.cs b/Engine/Source/Programs/UnrealBuildTool/ToolChain/AppleToolChain.cs
--- a/Engine/Source/Programs/UnrealBuildTool/ToolChain/AppleToolChain.cs
+++ b/Engine/Source/Programs/UnrealBuildTool/ToolChain/AppleToolChain.cs
@@ -1170,12 +1170,9 @@ namespace UnrealBuildTool
 
 			string PostBuildSyncArguments = String.Format("-modernxcode -Input=\"{0}\" -XmlConfigCache=\"{1}\" -remoteini=\"{2}\"", PostBuildSyncFile, XmlConfig.CacheFile, UnrealBuildTool.GetRemoteIniPath());
 			Action PostBuildSyncAction = Graph.CreateRecursiveAction<ApplePostBuildSyncMode>(ActionType.CreateAppBundle, PostBuildSyncArguments);
-
-			PostBuildSyncAction.WorkingDirectory = Unreal.EngineSourceDirectory;
 			PostBuildSyncAction.PrerequisiteItems.Add(Executable);
 			PostBuildSyncAction.ProducedItems.Add(GetPostBuildOutputFile(Executable.Location, Target.Name, Target.Platform));
 			PostBuildSyncAction.StatusDescription = $"Executing PostBuildSync [{Executable.Location}]";
-			PostBuildSyncAction.bCanExecuteRemotely = false;
 
 			if (Target.Platform == UnrealTargetPlatform.IOS || Target.Platform == UnrealTargetPlatform.TVOS)
 			{
diff --git a/Engine/Source/Programs/UnrealBuildTool/ToolChain/ClangToolChain.cs b/Engine/Source/Programs/UnrealBuildTool/ToolChain/ClangToolChain.cs
--- a/Engine/Source/Programs/UnrealBuildTool/ToolChain/ClangToolChain.cs
+++ b/Engine/Source/Programs/UnrealBuildTool/ToolChain/ClangToolChain.cs
@@ -264,10 +264,7 @@ namespace UnrealBuildTool
 					};
 
 					Action AggregateTimingInfoAction = MakefileBuilder.CreateRecursiveAction<AggregateClangTimingInfo>(ActionType.ParseTimingInfo, String.Join(" ", AggregateActionArgs));
-					AggregateTimingInfoAction.WorkingDirectory = Unreal.EngineSourceDirectory;
 					AggregateTimingInfoAction.StatusDescription = $"Aggregating {TimingJsonFiles.Count} Timing File(s)";
-					AggregateTimingInfoAction.bCanExecuteRemotely = false;
-					AggregateTimingInfoAction.bCanExecuteRemotelyWithSNDBS = false;
 					AggregateTimingInfoAction.PrerequisiteItems.UnionWith(TimingJsonFiles);
 
 					AggregateTimingInfoAction.ProducedItems.Add(AggregateOutputFile);
@@ -282,10 +279,7 @@ namespace UnrealBuildTool
 					};
 
 					Action ArchiveTimingInfoAction = MakefileBuilder.CreateRecursiveAction<AggregateClangTimingInfo>(ActionType.ParseTimingInfo, String.Join(" ", ArchiveActionArgs));
-					ArchiveTimingInfoAction.WorkingDirectory = Unreal.EngineSourceDirectory;
 					ArchiveTimingInfoAction.StatusDescription = $"Archiving {TimingJsonFiles.Count} Timing File(s)";
-					ArchiveTimingInfoAction.bCanExecuteRemotely = false;
-					ArchiveTimingInfoAction.bCanExecuteRemotelyWithSNDBS = false;
 					ArchiveTimingInfoAction.PrerequisiteItems.UnionWith(TimingJsonFiles);
 
 					ArchiveTimingInfoAction.ProducedItems.Add(ArchiveOutputFile);
@@ -302,6 +296,7 @@ namespace UnrealBuildTool
 						CompileScoreExtractorAction.StatusDescription = $"Extracting CompileScore";
 						CompileScoreExtractorAction.bCanExecuteRemotely = false;
 						CompileScoreExtractorAction.bCanExecuteRemotelyWithSNDBS = false;
+						CompileScoreExtractorAction.bCanExecuteInUBA = false; // TODO: Unknown if supported
 						CompileScoreExtractorAction.PrerequisiteItems.UnionWith(TimingJsonFiles);
 						CompileScoreExtractorAction.CommandPath = ScoreDataExtractor;
 						CompileScoreExtractorAction.CommandArguments = $"-clang -verbosity 0 -timelinepack 1000000 -extract -i \"{NormalizeCommandLinePath(Makefile.ProjectIntermediateDirectory)}\" -o \"{NormalizeCommandLinePath(CompileScoreOutput)}\"";
