UnrealBuildTool: Allow reordering remote executor priority

diff --git a/Engine/Source/Programs/UnrealBuildTool/Configuration/BuildConfiguration.cs b/Engine/Source/Programs/UnrealBuildTool/Configuration/BuildConfiguration.cs
--- a/Engine/Source/Programs/UnrealBuildTool/Configuration/BuildConfiguration.cs
+++ b/Engine/Source/Programs/UnrealBuildTool/Configuration/BuildConfiguration.cs
@@ -39,6 +39,12 @@ namespace UnrealBuildTool
 		[Obsolete("HybridExecutor is no longer supported")]
 		public bool bAllowHybridExecutor = false;
 
+		/// <summary>
+		/// Priority order for remote executors (XGE, SNDBS, FASTBuild, UBAExecutor)
+		/// </summary>
+		[XmlConfigFile]
+		public string[] RemoteExecutorPriority = new string[] { "XGE", "SNDBS", "FASTBuild", "UBA" };
+
 		/// <summary>
 		/// This property is being renamed and should not be used
 		/// </summary>
diff --git a/Engine/Source/Programs/UnrealBuildTool/System/ActionGraph.cs b/Engine/Source/Programs/UnrealBuildTool/System/ActionGraph.cs
--- a/Engine/Source/Programs/UnrealBuildTool/System/ActionGraph.cs
+++ b/Engine/Source/Programs/UnrealBuildTool/System/ActionGraph.cs
@@ -356,6 +356,52 @@ namespace UnrealBuildTool
 			}
 		}
 
+		private static ActionExecutor? GetRemoteExecutorByName(string Name, BuildConfiguration BuildConfiguration, int ActionCount, List<TargetDescriptor> TargetDescriptors, ILogger Logger)
+		{
+			switch (Name)
+			{
+				case "XGE":
+					{
+						if (BuildConfiguration.bAllowXGE && XGE.IsAvailable(Logger) && ActionCount >= XGE.MinActions)
+						{
+							return new XGE(Logger);
+						}
+						return null;
+					}
+				case "SNDBS":
+					{
+						if (BuildConfiguration.bAllowSNDBS && SNDBS.IsAvailable(Logger))
+						{
+							return new SNDBS(TargetDescriptors, Logger);
+						}
+						return null;
+					}
+				case "FASTBuild":
+					{
+						if (BuildConfiguration.bAllowFASTBuild && FASTBuild.IsAvailable(Logger))
+						{
+							return new FASTBuild(BuildConfiguration.MaxParallelActions, BuildConfiguration.bAllCores, BuildConfiguration.bCompactOutput, Logger);
+						}
+						return null;
+					}
+				case "UBA":
+					{
+#if __UBAEXECUTOR_AVAILABLE__
+						if (BuildConfiguration.bAllowUBAExecutor && UBAExecutor.IsAvailable())
+						{
+							return new UBAExecutor(BuildConfiguration.MaxParallelActions, BuildConfiguration.bAllCores, BuildConfiguration.bCompactOutput, Logger, TargetDescriptors.FirstOrDefault()?.AdditionalArguments);
+						}
+#endif // #if __UBAEXECUTOR_AVAILABLE__
+						return null;
+					}
+				default:
+					{
+						Logger.LogWarning("Unknown remote executor {Name}", Name);
+						return null;
+					}
+			}
+		}
+
 		/// <summary>
 		/// Selects an ActionExecutor
 		/// </summary>
@@ -363,24 +409,13 @@ namespace UnrealBuildTool
 		{
 			if (ActionCount > ParallelExecutor.GetDefaultNumParallelProcesses(BuildConfiguration.MaxParallelActions, BuildConfiguration.bAllCores, Logger))
 			{
-#if __UBAEXECUTOR_AVAILABLE__
-				if (BuildConfiguration.bAllowUBAExecutor && UBAExecutor.IsAvailable())
-				{
-					return new UBAExecutor(BuildConfiguration.MaxParallelActions, BuildConfiguration.bAllCores, BuildConfiguration.bCompactOutput, Logger, TargetDescriptors.FirstOrDefault()?.AdditionalArguments);
-				}
-#endif // #if __UBAEXECUTOR_AVAILABLE__
-
-				if (BuildConfiguration.bAllowXGE && XGE.IsAvailable(Logger) && ActionCount >= XGE.MinActions)
-				{
-					return new XGE(Logger);
-				}
-				else if (BuildConfiguration.bAllowSNDBS && SNDBS.IsAvailable(Logger))
+				foreach (string Name in BuildConfiguration.RemoteExecutorPriority)
 				{
-					return new SNDBS(TargetDescriptors, Logger);
-				}
-				else if (BuildConfiguration.bAllowFASTBuild && FASTBuild.IsAvailable(Logger))
-				{
-					return new FASTBuild(BuildConfiguration.MaxParallelActions, BuildConfiguration.bAllCores, BuildConfiguration.bCompactOutput, Logger);
+					ActionExecutor? Executor = GetRemoteExecutorByName(Name, BuildConfiguration, ActionCount, TargetDescriptors, Logger);
+					if (Executor != null)
+					{
+						return Executor;
+					}
 				}
 			}
 
