[UBT]
* Disabled uba for linking on linux. vfork/exec is not supported right now and we will push this support  in front of us (since remote linking is disabled either way)

diff --git a/Engine/Source/Programs/UnrealBuildTool/Platform/Linux/LinuxToolChain.cs b/Engine/Source/Programs/UnrealBuildTool/Platform/Linux/LinuxToolChain.cs
--- a/Engine/Source/Programs/UnrealBuildTool/Platform/Linux/LinuxToolChain.cs
+++ b/Engine/Source/Programs/UnrealBuildTool/Platform/Linux/LinuxToolChain.cs
@@ -964,6 +964,8 @@ namespace UnrealBuildTool
 			ArchiveAction.WorkingDirectory = Unreal.EngineSourceDirectory;
 			ArchiveAction.CommandPath = Info.Archiver;
 
+			ArchiveAction.bCanExecuteInUBA = OperatingSystem.IsWindows(); // Linker on native linux uses vfork/exec which is not handled in uba right now
+
 			// this will produce a final library
 			ArchiveAction.bProducesImportLibrary = true;
 
@@ -1036,6 +1038,7 @@ namespace UnrealBuildTool
 				PostLinkAction.StatusDescription = String.Format("{0}", Path.GetFileName(Executable.AbsolutePath));
 				PostLinkAction.CommandDescription = "FixDeps";
 				PostLinkAction.bCanExecuteRemotely = false;
+				PostLinkAction.bCanExecuteInUBA = OperatingSystem.IsWindows(); // Linker on native linux uses vfork/exec which is not handled in uba right now
 				PostLinkAction.CommandArguments = ExecuteSwitch;
 
 				PostLinkAction.CommandArguments += bUseCmdExe ? " \"" : " -c '";
@@ -1126,6 +1129,9 @@ namespace UnrealBuildTool
 			// Saw a 6 hour link time potentially caused by box. Will disable for now and revisit later
 			LinkAction.bCanExecuteInUBA = !LinkEnvironment.bPGOProfile && !LinkEnvironment.bPGOOptimize && !LinkEnvironment.bAllowLTCG;
 
+			if (!OperatingSystem.IsWindows())
+				LinkAction.bCanExecuteInUBA = false; // Linker on native linux uses vfork/exec which is not handled in uba right now
+
 			// because the logic choosing between lld and ld is somewhat messy atm (lld fails to link .DSO due to bugs), make the name of the linker clear
 			LinkAction.CommandDescription += (LinkCommandString.Contains("-fuse-ld=lld")) ? " (lld)" : " (ld)";
 			LinkAction.CommandVersion = Info.ClangVersionString;
@@ -1531,6 +1537,7 @@ namespace UnrealBuildTool
 					RelinkAction.StatusDescription = LinkAction.StatusDescription;
 					RelinkAction.CommandDescription = "Relink";
 					RelinkAction.bCanExecuteRemotely = false;
+					RelinkAction.bCanExecuteInUBA = OperatingSystem.IsWindows(); // Linker on native linux uses vfork/exec which is not handled in uba right now
 					RelinkAction.ProducedItems.Clear();
 					RelinkAction.PrerequisiteItems = new SortedSet<FileItem>(LinkAction.PrerequisiteItems);
 					foreach (FileItem Dependency in EngineAndGameLibrariesFiles)
