UnrealBuildTool: Don't allow post-build scripts to run with UBA as they can run anything

#rnx

diff --git a/Engine/Source/Programs/UnrealBuildTool/Configuration/UEBuildTarget.cs b/Engine/Source/Programs/UnrealBuildTool/Configuration/UEBuildTarget.cs
--- a/Engine/Source/Programs/UnrealBuildTool/Configuration/UEBuildTarget.cs
+++ b/Engine/Source/Programs/UnrealBuildTool/Configuration/UEBuildTarget.cs
@@ -2904,6 +2904,7 @@ namespace UnrealBuildTool
 					PostBuildStepAction.bCanExecuteRemotely = false;
 					PostBuildStepAction.PrerequisiteItems.Add(FileItem.GetItemByFileReference(ReceiptFileName));
 					PostBuildStepAction.ProducedItems.Add(FileItem.GetItemByFileReference(OutputFile));
+					PostBuildStepAction.bCanExecuteInUBA = false;
 
 					Makefile.OutputItems.AddRange(PostBuildStepAction.ProducedItems);
 				}
