UnrealBuildTool: Disable Winvalid-unevaluated-string for clang 17

#rnx

diff --git a/Engine/Source/Programs/UnrealBuildTool/ToolChain/ClangWarnings.cs b/Engine/Source/Programs/UnrealBuildTool/ToolChain/ClangWarnings.cs
--- a/Engine/Source/Programs/UnrealBuildTool/ToolChain/ClangWarnings.cs
+++ b/Engine/Source/Programs/UnrealBuildTool/ToolChain/ClangWarnings.cs
@@ -1,5 +1,6 @@
 ﻿// Copyright Epic Games, Inc. All Rights Reserved.
 
+using EpicGames.UHT.Utils;
 using System.Collections.Generic;
 
 namespace UnrealBuildTool
@@ -38,9 +39,16 @@ namespace UnrealBuildTool
 			}
 			if (ClangVersion >= new VersionNumber(16))
 			{
-				Arguments.Add("-Wno-deprecated-copy");
+				Arguments.Add("-Wno-deprecated-copy");                  // https://clang.llvm.org/docs/DiagnosticsReference.html#wdeprecated-copy						// new warning for clang 16
 				Arguments.Add("-Wno-deprecated-copy-with-user-provided-copy");
 			}
+			if (ClangVersion >= new VersionNumber(17))
+			{
+				if (CompileEnvironment.CppStandard < CppStandardVersion.Latest)
+				{
+					Arguments.Add("-Wno-invalid-unevaluated-string");   // https://clang.llvm.org/docs/DiagnosticsReference.html#winvalid-unevaluated-string			// new warning for clang 17
+				}
+			}
 
 			Arguments.Add("-Wno-gnu-string-literal-operator-template"); // https://clang.llvm.org/docs/DiagnosticsReference.html#wgnu-string-literal-operator-template	// We use this feature to allow static FNames.
 			Arguments.Add("-Wno-inconsistent-missing-override");        // https://clang.llvm.org/docs/DiagnosticsReference.html#winconsistent-missing-override			// ?? no reason given
