Capture a module's low level tests from any folder - removes restriction to only modules with a "Tests" folder.

diff --git a/Engine/Source/Programs/UnrealBuildTool/Configuration/ModuleRules.cs b/Engine/Source/Programs/UnrealBuildTool/Configuration/ModuleRules.cs
--- a/Engine/Source/Programs/UnrealBuildTool/Configuration/ModuleRules.cs
+++ b/Engine/Source/Programs/UnrealBuildTool/Configuration/ModuleRules.cs
@@ -1772,8 +1772,7 @@ namespace UnrealBuildTool
 
 		/// <summary>
 		/// Prepares a module for building a low level tests executable.
-		/// If we're building a module as part of a test module chain and there's a Tests folder with low level tests, then they require the LowLevelTestsRunner dependency.
-		/// We also keep track of any Editor, Engine and other conditionally compiled dependencies.
+		/// If we're building a module as part of a test module chain, then they require the LowLevelTestsRunner dependency.
 		/// </summary>
 		internal void PrepareModuleForTests()
 		{
@@ -1785,7 +1784,7 @@ namespace UnrealBuildTool
 
 			lock (TestTargetRules)
 			{
-				if (Name != "LowLevelTestsRunner" && System.IO.Directory.Exists(TestsDirectory))
+				if (Name != "LowLevelTestsRunner")
 				{
 					if (!PrivateIncludePathModuleNames.Contains("LowLevelTestsRunner"))
 					{
