Move source files from editor projects to game projects

diff --git a/Engine/Source/Programs/UnrealBuildTool/ProjectFiles/ProjectFileGenerator.cs b/Engine/Source/Programs/UnrealBuildTool/ProjectFiles/ProjectFileGenerator.cs
--- a/Engine/Source/Programs/UnrealBuildTool/ProjectFiles/ProjectFileGenerator.cs
+++ b/Engine/Source/Programs/UnrealBuildTool/ProjectFiles/ProjectFileGenerator.cs
@@ -373,6 +373,11 @@ namespace UnrealBuildTool
 		/// </summary>
 		protected virtual bool bMakeProjectPerTarget => false;
 
+		/// <summary>
+		/// If true, this generator will add modules to multiple suitable projects (i.e. QAGame.Build.cs will be added to both QAGame and QAGameEditor), only makes sense when bMakeProjectPerTarget is true
+		/// </summary>
+		protected virtual bool bAllowMultiModuleReference => false;
+
 		/// <summary>
 		/// If true, the project generator will allow conetnt-only projects for games with no targets, when generating for a single game target with -game
 		/// </summary>
@@ -2385,6 +2390,7 @@ namespace UnrealBuildTool
 				}
 			}
 
+			List<ProjectDescriptor> AllGameDescriptors = AllGames.ConvertAll(ProjectDescriptor.FromFile);
 			HashSet<ProjectFile> ProjectsWithPlugins = new HashSet<ProjectFile>();
 			foreach (FileReference CurModuleFile in AllModuleFiles)
 			{
@@ -2418,11 +2424,28 @@ namespace UnrealBuildTool
 				if (WantProjectFileForModule)
 				{
 					DirectoryReference BaseFolder;
-					ProjectFile ProjectFile = FindProjectForModule(CurModuleFile, TargetType.Editor, AllGames, ProgramProjects, ModProjects, ModuleToAdditionalPlugin, out BaseFolder);
+					List<ProjectFile> ProjectFiles = FindProjectsForModule(CurModuleFile, AllGames, AllGameDescriptors, ProgramProjects, ModProjects, ModuleToAdditionalPlugin, out BaseFolder)!;
 
 					// Update our module map
-					ModuleToEditorProjectFileMap[CurModuleFile] = ProjectFile;
-					ProjectFile.IsGeneratedProject = true;
+					if (ProjectFiles.Count() == 1)
+					{
+						ModuleToEditorProjectFileMap[CurModuleFile] = ProjectFiles[0];
+					}
+					else
+					{
+						Debug.Assert(bAllowMultiModuleReference, "ProjectFileGenerator assert", $"Unexpected multi projects for module {CurModuleFile.GetFileName()}");
+						// e.g. QAGame module would be add to both QAGame.xcodeproj and QAGameEditor.xcodeproj, use the editor one for module map
+						ProjectFile? EditorProjectFile = ProjectFiles.FirstOrDefault(x => x.ProjectTargets.Any(x => x.TargetRules!.Type == TargetType.Editor));
+						if (EditorProjectFile != null)
+						{
+							ModuleToEditorProjectFileMap[CurModuleFile] = EditorProjectFile;
+						}
+						else
+						{
+							Logger.LogWarning("No suitable project found for {Module}", CurModuleFile.GetFileName());
+						}
+					}
+					ProjectFiles.ForEach(x => x.IsGeneratedProject = true);
 
 					// Only search subdirectories for non-external modules.  We don't want to add all of the source and header files
 					// for every third-party module, unless we were configured to do so.
@@ -2444,24 +2467,26 @@ namespace UnrealBuildTool
 						continue;
 					}
 
-					ProjectFile.AddFilesToProject(FoundFiles, BaseFolder);
-
-					// Check if there's a plugin directory here
-					if (!ProjectsWithPlugins.Contains(ProjectFile))
+					foreach (ProjectFile aProjectFile in ProjectFiles)
 					{
-						foreach (DirectoryReference PluginFolder in Unreal.GetExtensionDirs(BaseFolder, "Plugins"))
+						aProjectFile.AddFilesToProject(FoundFiles, BaseFolder);
+						// Check if there's a plugin directory here
+						if (!ProjectsWithPlugins.Contains(aProjectFile))
 						{
-							// Add all the plugin files for this project
-							foreach (FileReference PluginFileName in PluginsBase.EnumeratePlugins(PluginFolder))
+							foreach (DirectoryReference PluginFolder in Unreal.GetExtensionDirs(BaseFolder, "Plugins"))
 							{
-								if (!ModProjects.Any(x => x.BaseDir == PluginFileName.Directory))
+								// Add all the plugin files for this project
+								foreach (FileReference PluginFileName in PluginsBase.EnumeratePlugins(PluginFolder))
 								{
-									AddPluginFilesToProject(PluginFileName, BaseFolder, ProjectFile);
+									if (!ModProjects.Any(x => x.BaseDir == PluginFileName.Directory))
+									{
+										AddPluginFilesToProject(PluginFileName, BaseFolder, aProjectFile);
+									}
 								}
 							}
-						}
 
-						ProjectsWithPlugins.Add(ProjectFile);
+							ProjectsWithPlugins.Add(aProjectFile);
+						}
 					}
 				}
 			}
@@ -2527,17 +2552,18 @@ namespace UnrealBuildTool
 			return FindOrAddProject(ProjectFileName, InBaseFolder, IncludeInGeneratedProjects: true, bAlreadyExisted: out _);
 		}
 
-		private ProjectFile FindProjectForModule(FileReference CurModuleFile, TargetType TargetType, List<FileReference> AllGames, Dictionary<FileReference, ProjectFile> ProgramProjects, List<ProjectFile> ModProjects, Dictionary<FileReference, AdditionalPluginData> ModuleToAdditionalPlugin, out DirectoryReference BaseFolder)
+		private List<ProjectFile> FindProjectsForModule(FileReference CurModuleFile, List<FileReference> AllGames, List<ProjectDescriptor> AllGameDescriptors, Dictionary<FileReference, ProjectFile> ProgramProjects, List<ProjectFile> ModProjects, Dictionary<FileReference, AdditionalPluginData> ModuleToAdditionalPlugin, out DirectoryReference BaseFolder)
 		{
 			// Starting at the base directory of the module find a project which has the same directory as base, walking up the directory hierarchy until a match is found
 
+			List<ProjectFile> FoundProjects = new List<ProjectFile>();
 			DirectoryReference Path = CurModuleFile.Directory;
 			bool bIsTemporaryModule = Path.ContainsName("Intermediate", 0);
 
 			while (!Path.IsRootDirectory())
 			{
 				// Figure out which game project this target belongs to
-				foreach (FileReference Game in AllGames)
+				foreach (var (Game, GameDescriptor) in AllGames.Zip(AllGameDescriptors))
 				{
 					// the source and the actual game directory are conceptually the same
 					if (Path == Game.Directory || Path == DirectoryReference.Combine(Game.Directory, "Source"))
@@ -2550,9 +2576,12 @@ namespace UnrealBuildTool
 						{
 							if (SingleTargetName != null)
 							{
-								return FindOrAddProjectHelper(SingleTargetName, BaseFolder);
+								FoundProjects.Add(FindOrAddProjectHelper(SingleTargetName, BaseFolder));
+								return FoundProjects;
 							}
 
+							ModuleDescriptor? CurModuleDescriptor = GameDescriptor.Modules!.FirstOrDefault(x => x.Name == CurModuleFile.GetFileNameWithoutAnyExtensions());
+
 							// find the project that the module is under, and has a TargetType target (useful with bMakeProjectPerTarget)
 							foreach (KeyValuePair<FileReference, ProjectFile> Pair in ProjectFileMap)
 							{
@@ -2563,16 +2592,41 @@ namespace UnrealBuildTool
 								if ((!bIsTemporaryTarget && TargetFile.Directory.ParentDirectory == Path) ||
 									(bIsTemporaryTarget && TargetFile.Directory.ParentDirectory!.ParentDirectory == Path))
 								{
-									if (Pair.Value.ProjectTargets.Any(x => x.TargetRules!.Type == TargetType))
+									if (CurModuleDescriptor != null
+										&& new[] { ModuleHostType.Editor, ModuleHostType.EditorNoCommandlet, ModuleHostType.EditorAndProgram }.Contains(CurModuleDescriptor.Type))
+									{
+										// if this module is for editor, then only add it to the editor project
+										if (Pair.Value.ProjectTargets.Any(x => x.TargetRules!.Type == TargetType.Editor))
+										{
+											FoundProjects.Add(Pair.Value);
+										}
+									}
+									else
 									{
-										return Pair.Value;
+										// plugins and game modules add into both editor and game projects
+										if (Pair.Value.ProjectTargets.Any(x => (x.TargetRules!.Type == TargetType.Editor) || (x.TargetRules!.Type == TargetType.Game)))
+										{
+											FoundProjects.Add(Pair.Value);
+										}
 									}
 								}
 							}
+							if (FoundProjects.Count() > 0)
+							{
+								if (bAllowMultiModuleReference)
+								{
+									return FoundProjects;
+								}
+								else
+								{
+									return new List<ProjectFile>{ FoundProjects[0]};
+								}
+							}
 						}
 						else
 						{
-							return FindOrAddProjectHelper(ProjectInfo.GetFileNameWithoutExtension(), BaseFolder);
+							FoundProjects.Add(FindOrAddProjectHelper(ProjectInfo.GetFileNameWithoutExtension(), BaseFolder));
+							return FoundProjects;
 						}
 					}
 				}
@@ -2582,7 +2636,8 @@ namespace UnrealBuildTool
 					if (Path == ModProject.BaseDir)
 					{
 						BaseFolder = ModProject.BaseDir;
-						return ModProject;
+						FoundProjects.Add(ModProject);
+						return FoundProjects;
 					}
 				}
 
@@ -2594,7 +2649,8 @@ namespace UnrealBuildTool
 						if (Path == ProgramProject.BaseDir)
 						{
 							BaseFolder = ProgramProject.BaseDir;
-							return ProgramProject;
+							FoundProjects.Add(ProgramProject);
+							return FoundProjects;
 						}
 					}
 				}
@@ -2605,7 +2661,8 @@ namespace UnrealBuildTool
 					if (Path == ExtensionDir)
 					{
 						BaseFolder = Unreal.EngineDirectory;
-						return FindOrAddProjectHelper(EngineEditorProjectFileNameBase, BaseFolder);
+						FoundProjects.Add(FindOrAddProjectHelper(EngineEditorProjectFileNameBase, BaseFolder));
+						return FoundProjects;
 					}
 				}
 
@@ -2617,7 +2674,8 @@ namespace UnrealBuildTool
 			{
 				FileReference ProjectFileRef = PluginData.ReferencingProjects[0];
 				BaseFolder = ProjectFileRef.Directory;
-				return FindOrAddProjectHelper(ProjectFileRef.GetFileNameWithoutExtension(), BaseFolder);
+				FoundProjects.Add(FindOrAddProjectHelper(ProjectFileRef.GetFileNameWithoutExtension(), BaseFolder));
+				return FoundProjects;
 			}
 
 			throw new BuildException("Found a module file (" + CurModuleFile + ") that did not exist within any of the known game folders or other source locations");
diff --git a/Engine/Source/Programs/UnrealBuildTool/ProjectFiles/Xcode/XcodeProjectFileGenerator.cs b/Engine/Source/Programs/UnrealBuildTool/ProjectFiles/Xcode/XcodeProjectFileGenerator.cs
--- a/Engine/Source/Programs/UnrealBuildTool/ProjectFiles/Xcode/XcodeProjectFileGenerator.cs
+++ b/Engine/Source/Programs/UnrealBuildTool/ProjectFiles/Xcode/XcodeProjectFileGenerator.cs
@@ -28,6 +28,7 @@ namespace UnrealBuildTool
 		/// </summary>
 		protected override bool bMakeProjectPerTarget => true;
 		protected override bool bAllowContentOnlyProjects => true;
+		protected override bool bAllowMultiModuleReference => true;
 
 		public DirectoryReference? XCWorkspace;
 
