UnrealBuildTool: Correctly print that a file is being analyzed when using -SingleFile

#rnx

diff --git a/Engine/Source/Programs/UnrealBuildTool/Platform/Windows/VCCompileAction.cs b/Engine/Source/Programs/UnrealBuildTool/Platform/Windows/VCCompileAction.cs
--- a/Engine/Source/Programs/UnrealBuildTool/Platform/Windows/VCCompileAction.cs
+++ b/Engine/Source/Programs/UnrealBuildTool/Platform/Windows/VCCompileAction.cs
@@ -377,6 +377,7 @@ namespace UnrealBuildTool
 			bCanExecuteRemotelyWithSNDBS = InAction.bCanExecuteRemotelyWithSNDBS;
 			bCanExecuteRemotelyWithXGE = InAction.bCanExecuteRemotelyWithXGE;
 			Architecture = InAction.Architecture;
+			bIsAnalyzing = InAction.bIsAnalyzing;
 			Weight = InAction.Weight;
 
 			AdditionalPrerequisiteItems = new List<FileItem>(InAction.AdditionalPrerequisiteItems);
@@ -419,6 +420,7 @@ namespace UnrealBuildTool
 			bCanExecuteRemotelyWithXGE = Reader.ReadBool();
 			bCanExecuteInUBA = Reader.ReadBool();
 			Architecture = UnrealArch.Parse(Reader.ReadString()!);
+			bIsAnalyzing = Reader.ReadBool();
 			Weight = Reader.ReadDouble();
 
 			AdditionalPrerequisiteItems = Reader.ReadList(() => Reader.ReadFileItem())!;
@@ -458,6 +460,7 @@ namespace UnrealBuildTool
 			Writer.WriteBool(bCanExecuteRemotelyWithXGE);
 			Writer.WriteBool(bCanExecuteInUBA);
 			Writer.WriteString(Architecture.ToString());
+			Writer.WriteBool(bIsAnalyzing);
 			Writer.WriteDouble(Weight);
 
 			Writer.WriteList(AdditionalPrerequisiteItems, Item => Writer.WriteFileItem(Item));
diff --git a/Engine/Source/Programs/UnrealBuildTool/Platform/Windows/VCToolChain.cs b/Engine/Source/Programs/UnrealBuildTool/Platform/Windows/VCToolChain.cs
--- a/Engine/Source/Programs/UnrealBuildTool/Platform/Windows/VCToolChain.cs
+++ b/Engine/Source/Programs/UnrealBuildTool/Platform/Windows/VCToolChain.cs
@@ -1428,6 +1428,8 @@ namespace UnrealBuildTool
 
 			AppendCLArguments_Global(CompileEnvironment, BaseCompileAction.Arguments);
 
+			BaseCompileAction.bIsAnalyzing = Target.StaticAnalyzer != StaticAnalyzer.None;
+
 			// Add include paths to the argument list.
 			BaseCompileAction.IncludePaths.AddRange(CompileEnvironment.UserIncludePaths);
 			BaseCompileAction.SystemIncludePaths.AddRange(CompileEnvironment.SystemIncludePaths);
