Fix up HordeStorageArtifactCache in light of CL# 29926074
#rnx

diff --git a/Engine/Source/Programs/UnrealBuildTool/Artifacts/HordeStorageArtifactCache.cs b/Engine/Source/Programs/UnrealBuildTool/Artifacts/HordeStorageArtifactCache.cs
--- a/Engine/Source/Programs/UnrealBuildTool/Artifacts/HordeStorageArtifactCache.cs
+++ b/Engine/Source/Programs/UnrealBuildTool/Artifacts/HordeStorageArtifactCache.cs
@@ -153,14 +153,6 @@ namespace UnrealBuildTool.Artifacts
 		{
 			writer.WriteDictionary<IoHash, HordeArtifactAction>(ArtifactActions, (x) => writer.WriteIoHash(x), (x) => writer.WriteHordeArtifactAction(x));
 		}
-
-		/// <summary>
-		/// Mark the node as dirty
-		/// </summary>
-		public new void MarkAsDirty()
-		{
-			base.MarkAsDirty();
-		}
 	}
 
 	/// <summary>
@@ -490,7 +482,6 @@ namespace UnrealBuildTool.Artifacts
 					// Update the artifact action collection
 					HordeArtifactAction hordeArtifactAction = new(artifactAction);
 					node.ArtifactActions[artifactAction.ActionKey] = hordeArtifactAction;
-					node.MarkAsDirty();
 
 					// Save the artifact action file
 					await using IStorageWriter writer = _store!.CreateWriter();
