PVS: Adjust warning logger

#rnx

diff --git a/Engine/Source/Programs/UnrealBuildTool/Platform/Windows/PVSToolChain.cs b/Engine/Source/Programs/UnrealBuildTool/Platform/Windows/PVSToolChain.cs
--- a/Engine/Source/Programs/UnrealBuildTool/Platform/Windows/PVSToolChain.cs
+++ b/Engine/Source/Programs/UnrealBuildTool/Platform/Windows/PVSToolChain.cs
@@ -442,7 +442,7 @@ namespace UnrealBuildTool
 										// Output the line to the log
 										if (!bFalseAlarm && Level == 1)
 										{
-											Logger.LogWarning("{Path}({LineNumber}): warning {WarningCode}: {WarningMessage}", FileName, LineNumber, WarningCode, WarningMessage);
+											Logger.LogWarning(KnownLogEvents.Compiler, "{Path}({LineNumber}): warning {WarningCode}: {WarningMessage}", LogValue.SourceFile(new FileReference(FileName), FileName), LineNumber, WarningCode, WarningMessage);
 										}
 									}
 								}
