UnrealBuildTool: Add -VCExtendedWarningInfo and always enable when analyzing

#rnx

diff --git a/Engine/Source/Programs/UnrealBuildTool/Platform/Windows/UEBuildWindows.cs b/Engine/Source/Programs/UnrealBuildTool/Platform/Windows/UEBuildWindows.cs
--- a/Engine/Source/Programs/UnrealBuildTool/Platform/Windows/UEBuildWindows.cs
+++ b/Engine/Source/Programs/UnrealBuildTool/Platform/Windows/UEBuildWindows.cs
@@ -307,6 +307,14 @@ namespace UnrealBuildTool
 		[CommandLine("-VCFastFail")]
 		public bool bVCFastFail = false;
 
+		/// <summary>
+		/// True if /d2ExtendedWarningInfo should be passed to the compiler and /d2:-ExtendedWarningInfo to the linker
+		/// </summary>
+		[ConfigFile(ConfigHierarchyType.Engine, "/Script/WindowsTargetPlatform.WindowsTargetSettings", "bVCExtendedWarningInfo")]
+		[XmlConfigFile(Category = "WindowsPlatform")]
+		[CommandLine("-VCExtendedWarningInfo")]
+		public bool bVCExtendedWarningInfo = false;
+
 		/// <summary>
 		/// True if optimizations to reduce the size of debug information should be disabled
 		/// See https://clang.llvm.org/docs/UsersManual.html#cmdoption-fstandalone-debug for more information
@@ -794,6 +802,8 @@ namespace UnrealBuildTool
 
 		public bool bVCFastFail => Inner.bVCFastFail;
 
+		public bool bVCExtendedWarningInfo => Inner.bVCExtendedWarningInfo;
+
 		public bool bClangStandaloneDebug => Inner.bClangStandaloneDebug;
 
 		public bool bAllowClangLinker => Inner.bAllowClangLinker;
@@ -1102,6 +1112,9 @@ namespace UnrealBuildTool
 			{
 				Target.bDisableLinking = true;
 				Target.bIgnoreBuildOutputs = true;
+
+				// Enable extended warnings when analyzing
+				Target.WindowsPlatform.bVCExtendedWarningInfo = true;
 			}
 
 			// Disable PCHs for PVS studio analyzer.
diff --git a/Engine/Source/Programs/UnrealBuildTool/Platform/Windows/VCToolChain.cs b/Engine/Source/Programs/UnrealBuildTool/Platform/Windows/VCToolChain.cs
--- a/Engine/Source/Programs/UnrealBuildTool/Platform/Windows/VCToolChain.cs
+++ b/Engine/Source/Programs/UnrealBuildTool/Platform/Windows/VCToolChain.cs
@@ -595,6 +595,11 @@ namespace UnrealBuildTool
 				Arguments.Add("/fastfail");
 			}
 
+			if (Target.WindowsPlatform.Compiler.IsMSVC() && Target.WindowsPlatform.bVCExtendedWarningInfo)
+			{
+				Arguments.Add("/d2ExtendedWarningInfo");
+			}
+
 			// Address sanitizer
 			if (Target.WindowsPlatform.bEnableAddressSanitizer)
 			{
@@ -1289,6 +1294,11 @@ namespace UnrealBuildTool
 				Arguments.Add("/fastfail");
 			}
 
+			if (Target.WindowsPlatform.Compiler.IsMSVC() && Target.WindowsPlatform.bVCExtendedWarningInfo)
+			{
+				Arguments.Add("/d2:-ExtendedWarningInfo");
+			}
+
 			// Allow for PDBs larger than 4GB
 			if (Target.WindowsPlatform.PdbPageSize.HasValue)
 			{
