Extra check over .so files to see if we need to rebuild the .apk

diff --git a/Engine/Source/Programs/UnrealBuildTool/Platform/Android/UEDeployAndroid.cs b/Engine/Source/Programs/UnrealBuildTool/Platform/Android/UEDeployAndroid.cs
--- a/Engine/Source/Programs/UnrealBuildTool/Platform/Android/UEDeployAndroid.cs
+++ b/Engine/Source/Programs/UnrealBuildTool/Platform/Android/UEDeployAndroid.cs
@@ -2049,8 +2049,9 @@ namespace UnrealBuildTool
 			return CurrentSettings.ToString();
 		}
 
-		private bool CheckDependencies(UnrealArchitectures Architectures, string ProjectName, string ProjectDirectory, string UnrealBuildFilesPath, string GameBuildFilesPath, string EngineDirectory, List<string> SettingsFiles,
-			string CookFlavor, string OutputPath, bool bMakeSeparateApks, bool bPackageDataInsideApk, bool bDontBundleLibrariesInAPK)
+		private bool CheckDependencies(UnrealArchitectures Architectures, string ProjectName, string ProjectDirectory, string IntermediateAndroidPath,
+			string UnrealBuildFilesPath, string GameBuildFilesPath, string EngineDirectory, List<string> SettingsFiles, string CookFlavor,
+			string OutputPath, bool bMakeSeparateApks, bool bPackageDataInsideApk, bool bDontBundleLibrariesInAPK)
 		{
 			// check all input files (.so, java files, .ini files, etc)
 			bool bAllInputsCurrent = true;
@@ -2067,16 +2068,26 @@ namespace UnrealBuildTool
 					DestApkName = AndroidToolChain.InlineArchName(DestApkName, Arch);
 				}
 
-				// check to see if it's out of date before trying the slow make apk process (look at .so and all Engine and Project build files to be safe)
 				List<String> InputFiles = new List<string>();
+				// check to see if the libUnreal.so is out of date before trying the slow make apk process
 				if (!bDontBundleLibrariesInAPK)
 				{
 					InputFiles.Add(SourceSOName);
-					InputFiles.AddRange(Directory.EnumerateFiles(UnrealBuildFilesPath, "*.*", SearchOption.AllDirectories));
-					if (Directory.Exists(GameBuildFilesPath))
-					{
-						InputFiles.AddRange(Directory.EnumerateFiles(GameBuildFilesPath, "*.*", SearchOption.AllDirectories));
-					}
+				}
+
+				// add all files in jni and libs subfolders
+				foreach (var SoDirName in new [] {"jni", "libs"})
+				{
+					string SoDirPath = Path.Combine(IntermediateAndroidPath, ArchRemapping[GetNDKArch(Arch)], SoDirName, GetNDKArch(Arch));
+					var files = Directory.EnumerateFiles(SoDirPath, "*.*", SearchOption.AllDirectories);
+					InputFiles.AddRange(files);
+				}
+
+				// add all Engine and Project build files to be safe
+				InputFiles.AddRange(Directory.EnumerateFiles(UnrealBuildFilesPath, "*.*", SearchOption.AllDirectories));
+				if (Directory.Exists(GameBuildFilesPath))
+				{
+					InputFiles.AddRange(Directory.EnumerateFiles(GameBuildFilesPath, "*.*", SearchOption.AllDirectories));
 				}
 
 				// make sure changed java files will rebuild apk
@@ -3996,6 +4007,13 @@ namespace UnrealBuildTool
 			return GetDontBundleLibrariesInAPK(ProjectFile, ForceDontBundleLibrariesInAPK, Configuration, bIsArchive, bFromMSBuild, bIsFromUAT, bVerbose ? Logger : null);
 		}
 
+		// Architecture remapping
+		private static readonly Dictionary<string, string> ArchRemapping = new()
+		{
+			{"arm64-v8a", "arm64"},
+			{"x86_64", "x64"}
+		};
+
 		private void MakeApk(AndroidToolChain ToolChain, string ProjectName, TargetType InTargetType, string ProjectDirectory, string OutputPath, string EngineDirectory, bool bForDistribution, string CookFlavor,
 			UnrealTargetConfiguration Configuration, bool bMakeSeparateApks, bool bIncrementalPackage, bool bDisallowPackagingDataInApk, bool bDisallowExternalFilesDir, bool bSkipGradleBuild, bool bIsArchive, bool bIsFromUAT)
 		{
@@ -4305,11 +4323,6 @@ namespace UnrealBuildTool
 			string CurrentBuildSettings = GetAllBuildSettings(ToolChain, UPL!, bForDistribution, bMakeSeparateApks, bPackageDataInsideApk, bDisableVerifyOBBOnStartUp, bUseExternalFilesDir, bDontBundleLibrariesInAPK, TemplatesHashCode);
 			string BuildSettingsCacheFile = Path.Combine(IntermediateAndroidPath, "UEBuildSettings.txt");
 
-			// Architecture remapping
-			Dictionary<string, string> ArchRemapping = new Dictionary<string, string>();
-			ArchRemapping.Add("arm64-v8a", "arm64");
-			ArchRemapping.Add("x86_64", "x64");
-
 			// do we match previous build settings?
 			bool bBuildSettingsMatch = true;
 
@@ -4347,8 +4360,9 @@ namespace UnrealBuildTool
 				JavaFiles.AddRange(from t in templates select t.SourceFile);
 				JavaFiles.AddRange(from t in templates select t.DestinationFile);
 
-				bBuildSettingsMatch = CheckDependencies(Architectures, ProjectName, ProjectDirectory, UnrealBuildFilesPath, GameBuildFilesPath,
-					EngineDirectory, JavaFiles, CookFlavor, OutputPath, bMakeSeparateApks, bPackageDataInsideApk, bDontBundleLibrariesInAPK);
+				bBuildSettingsMatch = CheckDependencies(Architectures, ProjectName, ProjectDirectory, IntermediateAndroidPath,
+					UnrealBuildFilesPath, GameBuildFilesPath, EngineDirectory, JavaFiles, CookFlavor, OutputPath,
+					bMakeSeparateApks, bPackageDataInsideApk, bDontBundleLibrariesInAPK);
 
 			}
 
