UnrealBuildTool: Clean up .editorconfig warnings in UBA

#rnx

diff --git a/Engine/Source/Programs/UnrealBuildTool/Executors/Experimental/UBAExecutor.cs b/Engine/Source/Programs/UnrealBuildTool/Executors/Experimental/UBAExecutor.cs
--- a/Engine/Source/Programs/UnrealBuildTool/Executors/Experimental/UBAExecutor.cs
+++ b/Engine/Source/Programs/UnrealBuildTool/Executors/Experimental/UBAExecutor.cs
@@ -21,6 +21,7 @@ namespace UnrealBuildTool
 	/// <summary>
 	/// Configuration for Unreal Build Accelerator
 	/// </summary>
+	[System.Diagnostics.CodeAnalysis.SuppressMessage("Style", "IDE1006:Naming Styles", Justification = "UnrealBuildTool naming style")]
 	class UnrealBuildAcceleratorConfig
 	{
 		/// <summary>
@@ -29,7 +30,7 @@ namespace UnrealBuildTool
 		[XmlConfigFile(Category = "UnrealBuildAccelerator")]
 		[CommandLine("-BoxDisableRemote")]
 		[CommandLine("-UBADisableRemote")]
-		public bool bDisableRemote = false;
+		public bool bDisableRemote { get; set; } = false;
 
 		/// <summary>
 		/// When set to true, UBA will force all actions that can be built remotely to be built remotely. This will hang if there are no remote agents available
@@ -37,7 +38,7 @@ namespace UnrealBuildTool
 		[XmlConfigFile(Category = "UnrealBuildAccelerator")]
 		[CommandLine("-BoxForceRemote")]
 		[CommandLine("-UBAForceRemote")]
-		public bool bForceBuildAllRemote = false;
+		public bool bForceBuildAllRemote { get; set; } = false;
 
 		/// <summary>
 		/// When set to true, actions that fail locally with UBA will be retried without UBA.
@@ -45,7 +46,7 @@ namespace UnrealBuildTool
 		[XmlConfigFile(Category = "UnrealBuildAccelerator")]
 		[CommandLine("-BoxForcedRetry")]
 		[CommandLine("-UBAForcedRetry")]
-		public bool bForcedRetry = false;
+		public bool bForcedRetry { get; set; } = false;
 
 		/// <summary>
 		/// When set to true, all errors and warnings from UBA will be output at the appropriate severity level to the log (rather than being output as 'information' and attempting to continue regardless).
@@ -53,7 +54,7 @@ namespace UnrealBuildTool
 		[XmlConfigFile(Category = "UnrealBuildAccelerator")]
 		[CommandLine("-BoxStrict")]
 		[CommandLine("-UBAStrict")]
-		public bool bStrict = false;
+		public bool bStrict { get; set; } = false;
 
 		/// <summary>
 		/// If UBA should store cas compressed or raw
@@ -61,7 +62,7 @@ namespace UnrealBuildTool
 		[XmlConfigFile(Category = "UnrealBuildAccelerator")]
 		[CommandLine("-BoxStoreRaw")]
 		[CommandLine("-UBAStoreRaw")]
-		public bool bStoreRaw = false;
+		public bool bStoreRaw { get; set; } = false;
 
 		/// <summary>
 		/// If UBA should distribute linking to remote workers. This needs bandwidth but can be an optimization
@@ -69,7 +70,7 @@ namespace UnrealBuildTool
 		[XmlConfigFile(Category = "UnrealBuildAccelerator")]
 		[CommandLine("-BoxLinkRemote")]
 		[CommandLine("-UBALinkRemote")]
-		public bool bLinkRemote = false;
+		public bool bLinkRemote { get; set; } = false;
 
 		/// <summary>
 		/// The amount of gigabytes UBA is allowed to use to store workset and cached data. It is a good idea to have this >10gb
@@ -77,7 +78,7 @@ namespace UnrealBuildTool
 		[XmlConfigFile(Category = "UnrealBuildAccelerator")]
 		[CommandLine("-BoxStoreCapacityGb")]
 		[CommandLine("-UBAStoreCapacityGb")]
-		public int StoreCapacityGb = 40;
+		public int StoreCapacityGb { get; set; } = 40;
 
 		/// <summary>
 		/// Max number of worker threads that can handle messages from remotes. 
@@ -85,7 +86,7 @@ namespace UnrealBuildTool
 		[XmlConfigFile(Category = "UnrealBuildAccelerator")]
 		[CommandLine("-BoxMaxWorkers")]
 		[CommandLine("-UBAMaxWorkers")]
-		public int MaxWorkers = 192;
+		public int MaxWorkers { get; set; } = 192;
 
 		/// <summary>
 		/// Max size of each message sent from server to client
@@ -93,7 +94,7 @@ namespace UnrealBuildTool
 		[XmlConfigFile(Category = "UnrealBuildAccelerator")]
 		[CommandLine("-BoxSendSize")]
 		[CommandLine("-UBASendSize")]
-		public int SendSize = 256 * 1024;
+		public int SendSize { get; set; } = 256 * 1024;
 
 		/// <summary>
 		/// Which ip UBA server should listen to for connections
@@ -101,7 +102,7 @@ namespace UnrealBuildTool
 		[XmlConfigFile(Category = "UnrealBuildAccelerator")]
 		[CommandLine("-BoxHost")]
 		[CommandLine("-UBAHost")]
-		public string Host = String.Empty;
+		public string Host { get; set; } = String.Empty;
 
 		/// <summary>
 		/// Which port UBA server should listen to for connections.
@@ -109,7 +110,7 @@ namespace UnrealBuildTool
 		[XmlConfigFile(Category = "UnrealBuildAccelerator")]
 		[CommandLine("-BoxPort")]
 		[CommandLine("-UBAPort")]
-		public int Port = 1345;
+		public int Port { get; set; } = 1345;
 
 		/// <summary>
 		/// Which directory to store files for UBA.
@@ -117,7 +118,7 @@ namespace UnrealBuildTool
 		[XmlConfigFile(Category = "UnrealBuildAccelerator")]
 		[CommandLine("-BoxRootDir")]
 		[CommandLine("-UBARootDir")]
-		public string? RootDir = null;
+		public string? RootDir { get; set; } = null;
 
 		/// <summary>
 		/// Use Quic protocol instead of Tcp (experimental)
@@ -125,7 +126,7 @@ namespace UnrealBuildTool
 		[XmlConfigFile(Category = "UnrealBuildAccelerator")]
 		[CommandLine("-BoxQuic", Value = "true")]
 		[CommandLine("-UBAQuic", Value = "true")]
-		public bool bUseQuic = false;
+		public bool bUseQuic { get; set; } = false;
 
 		/// <summary>
 		/// Enable logging of UBA processes
@@ -133,7 +134,7 @@ namespace UnrealBuildTool
 		[XmlConfigFile(Category = "UnrealBuildAccelerator")]
 		[CommandLine("-BoxLog", Value = "true")]
 		[CommandLine("-UBALog", Value = "true")]
-		public bool bLogEnabled = false;
+		public bool bLogEnabled { get; set; } = false;
 
 		/// <summary>
 		/// Prints summary of UBA stats at end of build
@@ -141,7 +142,7 @@ namespace UnrealBuildTool
 		[XmlConfigFile(Category = "UnrealBuildAccelerator")]
 		[CommandLine("-BoxPrintSummary", Value = "true")]
 		[CommandLine("-UBAPrintSummary", Value = "true")]
-		public bool bPrintSummary = false;
+		public bool bPrintSummary { get; set; } = false;
 
 		/// <summary>
 		/// Launch visualizer application which shows build progress
@@ -149,7 +150,7 @@ namespace UnrealBuildTool
 		[XmlConfigFile(Category = "UnrealBuildAccelerator")]
 		[CommandLine("-BoxVisualizer", Value = "true")]
 		[CommandLine("-UBAVisualizer", Value = "true")]
-		public bool bLaunchVisualizer = false;
+		public bool bLaunchVisualizer { get; set; } = false;
 
 		/// <summary>
 		/// Resets the cas cache
@@ -157,7 +158,7 @@ namespace UnrealBuildTool
 		[XmlConfigFile(Category = "UnrealBuildAccelerator")]
 		[CommandLine("-BoxResetCas", Value = "true")]
 		[CommandLine("-UBAResetCas", Value = "true")]
-		public bool bResetCas = false;
+		public bool bResetCas { get; set; } = false;
 
 		/// <summary>
 		/// Provide custom path for trace output file
@@ -165,7 +166,7 @@ namespace UnrealBuildTool
 		[XmlConfigFile(Category = "UnrealBuildAccelerator")]
 		[CommandLine("-BoxTraceOutputFile")]
 		[CommandLine("-UBATraceOutputFile")]
-		public string TraceFile = String.Empty;
+		public string TraceFile { get; set; } = String.Empty;
 
 		/// <summary>
 		/// Add verbose details to the UBA trace
@@ -173,14 +174,14 @@ namespace UnrealBuildTool
 		[XmlConfigFile(Category = "UnrealBuildAccelerator")]
 		[CommandLine("-BoxDetailedTrace", Value = "true")]
 		[CommandLine("-UBADetailedTrace", Value = "true")]
-		public bool bDetailedTrace;
+		public bool bDetailedTrace { get; set; }
 
 		/// <summary>
 		/// Disable UBA waiting on available memory before spawning new processes
 		/// </summary>
 		[XmlConfigFile(Category = "UnrealBuildAccelerator")]
 		[CommandLine("-UBADisableWaitOnMem", Value = "true")]
-		public bool bDisableWaitOnMem;
+		public bool bDisableWaitOnMem { get; set; }
 
 		/// <summary>
 		/// Let UBA kill running processes when close to out of memory
@@ -188,7 +189,7 @@ namespace UnrealBuildTool
 		[XmlConfigFile(Category = "UnrealBuildAccelerator")]
 		[CommandLine("-BoxAllowKillOnMem", Value = "true")]
 		[CommandLine("-UBAAllowKillOnMem", Value = "true")]
-		public bool bAllowKillOnMem;
+		public bool bAllowKillOnMem { get; set; }
 
 		/// <summary>
 		/// Threshold for when executor should output logging for the process. Defaults to never
@@ -196,7 +197,7 @@ namespace UnrealBuildTool
 		[XmlConfigFile(Category = "UnrealBuildAccelerator")]
 		[CommandLine("-BoxOutputStatsThresholdMs")]
 		[CommandLine("-UBAOutputStatsThresholdMs")]
-		public int OutputStatsThresholdMs = int.MaxValue;
+		public int OutputStatsThresholdMs { get; set; } = Int32.MaxValue;
 
 		/// <summary>
 		/// Skip writing intermediate and output files to disk. Useful for validation builds where we don't need the output
@@ -204,7 +205,7 @@ namespace UnrealBuildTool
 		[XmlConfigFile(Category = "UnrealBuildAccelerator")]
 		[CommandLine("-BoxNoWrite", Value = "false")]
 		[CommandLine("-UBANoWrite", Value = "false")]
-		public bool bWriteToDisk = true;
+		public bool bWriteToDisk { get; set; } = true;
 
 		/// <summary>
 		/// Set to true to disable mimalloc and detouring of memory allocations.
@@ -212,7 +213,7 @@ namespace UnrealBuildTool
 		[XmlConfigFile(Category = "UnrealBuildAccelerator")]
 		[CommandLine("-BoxNoCustoMalloc", Value = "true")]
 		[CommandLine("-UBANoCustoMalloc", Value = "true")]
-		public bool bDisableCustomAlloc = false;
+		public bool bDisableCustomAlloc { get; set; } = false;
 
 		/// <summary>
 		/// The zone to use for UBA.
@@ -220,7 +221,7 @@ namespace UnrealBuildTool
 		[XmlConfigFile(Category = "UnrealBuildAccelerator")]
 		[CommandLine("-BoxZone=")]
 		[CommandLine("-UBAZone=")]
-		public string Zone = String.Empty;
+		public string Zone { get; set; } = String.Empty;
 
 		/// <summary>
 		/// Set to true to enable encryption when transfering files over the network.
@@ -228,7 +229,7 @@ namespace UnrealBuildTool
 		[XmlConfigFile(Category = "UnrealBuildAccelerator")]
 		[CommandLine("-BoxCrypto", Value = "true")]
 		[CommandLine("-UBACrypto", Value = "true")]
-		public bool bUseCrypto = false;
+		public bool bUseCrypto { get; set; } = false;
 	}
 
 	class UBAExecutor : ParallelExecutor
@@ -238,23 +239,33 @@ namespace UnrealBuildTool
 
 		public string Crypto { get; private set; } = String.Empty;
 		public IServer? Server { get; private set; }
-		ISessionServer? Session;
-		DirectoryReference? RootDirRef;
-		bool IsCancelled;
-		bool IsRemoteActionsAllowed = true;
-		readonly object ActionsChangedLock = new();
-		bool ActionsChanged = true;
-		uint ActionsQueuedThatCanRunRemotely = uint.MaxValue;
-		readonly ThreadedLogger ThreadedLogger;
+		ISessionServer? _session;
+		DirectoryReference? _rootDirRef;
+		bool _bIsCancelled;
+		bool _bIsRemoteActionsAllowed = true;
+		readonly object _actionsChangedLock = new();
+		bool _bActionsChanged = true;
+		uint _actionsQueuedThatCanRunRemotely = UInt32.MaxValue;
+		readonly ThreadedLogger _threadedLogger;
 
 		// Tracking for LinkedActions that failed remotely that should be retried locally
-		readonly ConcurrentDictionary<LinkedAction, bool> LocalRetryActions = new();
+		readonly ConcurrentDictionary<LinkedAction, bool> _localRetryActions = new();
 		// Tracking for LinkedActions that failed locally that should be retried without UBA
-		readonly ConcurrentDictionary<LinkedAction, bool> ForcedRetryActions = new();
+		readonly ConcurrentDictionary<LinkedAction, bool> _forcedRetryActions = new();
+
+		protected override void Dispose(bool disposing)
+		{
+			if (disposing)
+			{
+				_session?.Dispose();
+				_session = null;
+			}
+			base.Dispose(disposing);
+		}
 
 		public override string Name => "Unreal Build Accelerator";
 
-		public new static bool IsAvailable()
+		public static new bool IsAvailable()
 		{
 			return EpicGames.UBA.Utils.IsAvailable();
 		}
@@ -280,19 +291,21 @@ namespace UnrealBuildTool
 		private UBAExecutor(int maxLocalActions, bool bAllCores, bool bCompactOutput, ThreadedLogger logger)
 			: base(maxLocalActions, bAllCores, bCompactOutput, logger)
 		{
-			ThreadedLogger = logger;
+			_threadedLogger = logger;
 		}
 
 		private void PrintConfiguration()
 		{
-			ThreadedLogger.LogInformation("  Storage capacity {StoreCapacityGb}Gb", UBAConfig.StoreCapacityGb);
+			_threadedLogger.LogInformation("  Storage capacity {StoreCapacityGb}Gb", UBAConfig.StoreCapacityGb);
 		}
 
+		[System.Diagnostics.CodeAnalysis.SuppressMessage("Globalization", "CA1308:Normalize strings to uppercase", Justification = "lowercase crypto string")]
 		public static string CreateCrypto()
 		{
-			byte[] Bytes = new byte[16];
-			System.Security.Cryptography.RandomNumberGenerator.Create().GetBytes(Bytes);
-			return BitConverter.ToString(Bytes).Replace("-", "").ToLowerInvariant(); // "1234567890abcdef1234567890abcdef";
+			byte[] bytes = new byte[16];
+			using System.Security.Cryptography.RandomNumberGenerator random = System.Security.Cryptography.RandomNumberGenerator.Create();
+			random.GetBytes(bytes);
+			return BitConverter.ToString(bytes).Replace("-", "", StringComparison.OrdinalIgnoreCase).ToLowerInvariant(); // "1234567890abcdef1234567890abcdef";
 		}
 
 		public override async Task<bool> ExecuteActionsAsync(IEnumerable<LinkedAction> inputActions, Microsoft.Extensions.Logging.ILogger logger, IActionArtifactCache? actionArtifactCache)
@@ -309,50 +322,50 @@ namespace UnrealBuildTool
 
 			PrintConfiguration();
 
-			logger = ThreadedLogger;
+			logger = _threadedLogger;
 
 			if (UBAConfig.bUseCrypto)
 			{
 				Crypto = CreateCrypto();
 			}
 
-			if (!string.IsNullOrEmpty(UBAConfig.RootDir))
+			if (!String.IsNullOrEmpty(UBAConfig.RootDir))
 			{
-				RootDirRef = DirectoryReference.FromString(UBAConfig.RootDir);
+				_rootDirRef = DirectoryReference.FromString(UBAConfig.RootDir);
 			}
 			else
 			{
-				RootDirRef = DirectoryReference.FromString(Environment.GetEnvironmentVariable("UBA_ROOT") ?? Environment.GetEnvironmentVariable("BOX_ROOT"));
+				_rootDirRef = DirectoryReference.FromString(Environment.GetEnvironmentVariable("UBA_ROOT") ?? Environment.GetEnvironmentVariable("BOX_ROOT"));
 			}
 
 			using CancellationTokenSource cancellationTokenSource = new();
-			using Task<UBAHordeSession?> hordeSessionTask = !UBAConfig.bDisableRemote ? UBAHordeSession.TryCreateHordeSession(HordeConfig, this, UBAConfig.bStrict, cancellationTokenSource.Token, logger) : Task.FromResult<UBAHordeSession?>(null);
+			using Task<UBAHordeSession?> hordeSessionTask = !UBAConfig.bDisableRemote ? UBAHordeSession.TryCreateHordeSession(HordeConfig, this, UBAConfig.bStrict, logger, cancellationTokenSource.Token) : Task.FromResult<UBAHordeSession?>(null);
 
-			if (!UBAConfig.bDisableRemote && RootDirRef == null)
+			if (!UBAConfig.bDisableRemote && _rootDirRef == null)
 			{
 				DirectoryReference? hordeSharedDir = DirectoryReference.FromString(Environment.GetEnvironmentVariable("UE_HORDE_SHARED_DIR"));
 				if (hordeSharedDir != null)
 				{
-					RootDirRef = DirectoryReference.Combine(hordeSharedDir, "UbaHost");
+					_rootDirRef = DirectoryReference.Combine(hordeSharedDir, "UbaHost");
 				}
 			}
 
-			if (RootDirRef == null)
+			if (_rootDirRef == null)
 			{
 				if (OperatingSystem.IsWindows())
 				{
-					RootDirRef = DirectoryReference.Combine(DirectoryReference.GetSpecialFolder(Environment.SpecialFolder.CommonApplicationData)!, "Epic", "UnrealBuildAccelerator");
+					_rootDirRef = DirectoryReference.Combine(DirectoryReference.GetSpecialFolder(Environment.SpecialFolder.CommonApplicationData)!, "Epic", "UnrealBuildAccelerator");
 				}
 				else
 				{
-					RootDirRef = DirectoryReference.Combine(DirectoryReference.GetSpecialFolder(Environment.SpecialFolder.UserProfile)!, ".epic", "UnrealBuildAccelerator");
+					_rootDirRef = DirectoryReference.Combine(DirectoryReference.GetSpecialFolder(Environment.SpecialFolder.UserProfile)!, ".epic", "UnrealBuildAccelerator");
 				}
 			}
 
-			DirectoryReference.CreateDirectory(RootDirRef);
+			DirectoryReference.CreateDirectory(_rootDirRef);
 
 			FileReference ubaTraceFile;
-			if (!string.IsNullOrEmpty(UBAConfig.TraceFile))
+			if (!String.IsNullOrEmpty(UBAConfig.TraceFile))
 			{
 				ubaTraceFile = new FileReference(UBAConfig.TraceFile);
 			}
@@ -375,29 +388,30 @@ namespace UnrealBuildTool
 			IStorageServer? ubaStorage = null;
 			void CancelKeyPress(object? sender, ConsoleCancelEventArgs e)
 			{
-				IsCancelled = true;
-				Session?.CancelAll();
+				_bIsCancelled = true;
+				_session?.CancelAll();
 				ubaStorage?.SaveCasTable();
 			}
+			Console.CancelKeyPress += CancelKeyPress;
 
 			try
 			{
 				using EpicGames.UBA.ILogger ubaLogger = EpicGames.UBA.ILogger.CreateLogger(logger);
 				using (Server = IServer.CreateServer(UBAConfig.MaxWorkers, UBAConfig.SendSize, ubaLogger, UBAConfig.bUseQuic))
 				{
-					using IStorageServer ubaStorageServer = IStorageServer.CreateStorageServer(Server, ubaLogger, new StorageServerCreateInfo(RootDirRef, ((ulong)UBAConfig.StoreCapacityGb) * 1000 * 1000 * 1000, !UBAConfig.bStoreRaw, UBAConfig.Zone));
-					using (Session = ISessionServer.CreateSessionServer(ISessionServerCreateInfo.CreateSessionServerCreateInfo(ubaStorageServer, Server, ubaLogger, new SessionServerCreateInfo(RootDirRef, ubaTraceFile, UBAConfig.bDisableCustomAlloc, UBAConfig.bLaunchVisualizer, UBAConfig.bResetCas, UBAConfig.bWriteToDisk, UBAConfig.bDetailedTrace, !UBAConfig.bDisableWaitOnMem, UBAConfig.bAllowKillOnMem))))
+					using IStorageServer ubaStorageServer = IStorageServer.CreateStorageServer(Server, ubaLogger, new StorageServerCreateInfo(_rootDirRef, ((ulong)UBAConfig.StoreCapacityGb) * 1000 * 1000 * 1000, !UBAConfig.bStoreRaw, UBAConfig.Zone));
+					using ISessionServerCreateInfo serverCreateInfo = ISessionServerCreateInfo.CreateSessionServerCreateInfo(ubaStorageServer, Server, ubaLogger, new SessionServerCreateInfo(_rootDirRef, ubaTraceFile, UBAConfig.bDisableCustomAlloc, UBAConfig.bLaunchVisualizer, UBAConfig.bResetCas, UBAConfig.bWriteToDisk, UBAConfig.bDetailedTrace, !UBAConfig.bDisableWaitOnMem, UBAConfig.bAllowKillOnMem));
+					using (_session = ISessionServer.CreateSessionServer(serverCreateInfo))
 					{
 
 						ubaStorage = ubaStorageServer;
-						Console.CancelKeyPress += CancelKeyPress;
 
 						if (!UBAConfig.bDisableRemote)
 						{
 							Server.StartServer(UBAConfig.Host, UBAConfig.Port, Crypto);
 						}
 
-						bool success = ExecuteActionsInternal(inputActions, Session, hordeSessionTask, cancellationTokenSource, logger, actionArtifactCache);
+						bool success = ExecuteActionsInternal(inputActions, _session, hordeSessionTask, cancellationTokenSource, logger, actionArtifactCache);
 
 						if (!UBAConfig.bDisableRemote)
 						{
@@ -406,7 +420,7 @@ namespace UnrealBuildTool
 
 						if (UBAConfig.bPrintSummary)
 						{
-							Session.PrintSummary();
+							_session.PrintSummary();
 						}
 
 						return success;
@@ -415,15 +429,16 @@ namespace UnrealBuildTool
 			}
 			finally
 			{
-				if (ubaStorage != null)
-				{
-					Console.CancelKeyPress -= CancelKeyPress;
-					ubaStorage = null;
-				}
+				Console.CancelKeyPress -= CancelKeyPress;
 
 				cancellationTokenSource.Cancel();
-
-				await ThreadedLogger.FinishAsync();
+				UBAHordeSession? hordeSession = await hordeSessionTask;
+				if (hordeSession != null)
+				{
+					await hordeSession.DisposeAsync();
+				}
+				await _threadedLogger.FinishAsync();
+				_session = null;
 			}
 		}
 
@@ -438,13 +453,13 @@ namespace UnrealBuildTool
 			using ImmediateActionQueue queue = CreateActionQueue(inputActions, actionArtifactCache, logger);
 			int actionLimit = Math.Min(NumParallelProcesses, queue.TotalActions);
 			queue.CreateAutomaticRunner(action => RunActionLocal(queue, action), bUseActionWeights, actionLimit, NumParallelProcesses);
-			ImmediateActionQueueRunner _remoteRunner = queue.CreateManualRunner(action => RunActionRemote(queue, action));
+			ImmediateActionQueueRunner remoteRunner = queue.CreateManualRunner(action => RunActionRemote(queue, action));
 
 			// Setup a notification that alerts uba when an artifact has been read from the cache
-			queue.OnArtifactsRead = (Action) =>
+			queue.OnArtifactsRead = (action) =>
 			{
 				HashSet<DirectoryItem> refreshedDirectories = new();
-				foreach (FileItem output in Action.ProducedItems)
+				foreach (FileItem output in action.ProducedItems)
 				{
 					if (refreshedDirectories.Add(output.Directory))
 					{
@@ -459,20 +474,20 @@ namespace UnrealBuildTool
 			// Handle process available from remote
 			session!.RemoteProcessSlotAvailable += (sender, args) =>
 			{
-				if (!queue.TryStartOneAction(_remoteRunner) && (IsRemoteActionsAllowed && !UBAConfig.bForceBuildAllRemote))
+				if (!queue.TryStartOneAction(remoteRunner) && (_bIsRemoteActionsAllowed && !UBAConfig.bForceBuildAllRemote))
 				{
 					uint count = ActionsLeftThatCanRunRemotely(queue);
 
 					// We didn't find an action to start, let's check how many queued items are left.. if there are less than NumParallelProcesses, then disconnect
 					if (count <= NumParallelProcesses)
 					{
-						IsRemoteActionsAllowed = false;
+						_bIsRemoteActionsAllowed = false;
 						hordeRequestCancellationSource.Cancel();
 						session!.DisableRemoteExecution();
 					}
 					else
 					{
-						// Tell UBA session max number of remote processes left. Providing this information makes it possible for UBA session to start disconnecting clients
+						// Tell UBA hordeSession max number of remote processes left. Providing this information makes it possible for UBA hordeSession to start disconnecting clients
 						session!.SetMaxRemoteProcessCount(count);
 					}
 				}
@@ -501,7 +516,7 @@ namespace UnrealBuildTool
 
 					UBAHordeSession? hordeSession = await hordeSessionTask;
 
-					if (queue.IsDone || hordeSession == null || !IsRemoteActionsAllowed)
+					if (queue.IsDone || hordeSession == null || !_bIsRemoteActionsAllowed)
 					{
 						return;
 					}
@@ -553,7 +568,7 @@ namespace UnrealBuildTool
 					{
 						if (!shownNoAgentsFoundMessage)
 						{
-							logger.Log(UBAConfig.bStrict ? LogLevel.Warning : LogLevel.Information, KnownLogEvents.Systemic_Horde_Compute, ex, "No agents found matching requirements (cluster: {ClusterId}, requirements: {requirements})", ex.ClusterId, ex.Requirements);
+							logger.Log(UBAConfig.bStrict ? LogLevel.Warning : LogLevel.Information, KnownLogEvents.Systemic_Horde_Compute, ex, "No agents found matching requirements (cluster: {ClusterId}, requirements: {Requirements})", ex.ClusterId, ex.Requirements);
 							shownNoAgentsFoundMessage = true;
 						}
 					}
@@ -583,7 +598,7 @@ namespace UnrealBuildTool
 		/// </summary>
 		/// <param name="action">The action to check</param>
 		/// <returns>If this action can be run remotely</returns>
-		bool CanRunRemotely(LinkedAction action) => action.bCanExecuteInUBA && action.bCanExecuteRemotely && !LocalRetryActions.ContainsKey(action) && !ForcedRetryActions.ContainsKey(action) && (UBAConfig.bLinkRemote || action.ActionType != ActionType.Link);
+		bool CanRunRemotely(LinkedAction action) => action.bCanExecuteInUBA && action.bCanExecuteRemotely && !_localRetryActions.ContainsKey(action) && !_forcedRetryActions.ContainsKey(action) && (UBAConfig.bLinkRemote || action.ActionType != ActionType.Link);
 
 		ProcessStartInfo GetActionStartInfo(LinkedAction action, out FileItem? pchItem)
 		{
@@ -601,7 +616,7 @@ namespace UnrealBuildTool
 			};
 
 			bool usingLtcg = true; // ltcg linking checks what pch was used and it seems like it needs to be identical in other ways than timestamp
-			pchItem = action.Inner.ProducedItems.FirstOrDefault(Item => Item.Name.EndsWith(".pch"));
+			pchItem = action.Inner.ProducedItems.FirstOrDefault(item => item.Name.EndsWith(".pch", StringComparison.OrdinalIgnoreCase));
 			if (pchItem != null)
 			{
 				startInfo.Priority = System.Diagnostics.ProcessPriorityClass.AboveNormal;
@@ -625,18 +640,18 @@ namespace UnrealBuildTool
 				return null;
 			}
 
-			if (!action.bCanExecuteInUBA || ForcedRetryActions.ContainsKey(action))
+			if (!action.bCanExecuteInUBA || _forcedRetryActions.ContainsKey(action))
 			{
 				return async () =>
 				{
-					uint processId = Session!.BeginExternalProcess(action.StatusDescription + " (External)");
+					uint processId = _session!.BeginExternalProcess(action.StatusDescription + " (External)");
 					ExecuteResults result = await RunAction(action, queue.ProcessGroup, queue.CancellationToken, "(UBA disabled)");
 
 					if (result.ExitCode == 0)
 					{
-						Session!.RegisterNewFiles(action.ProducedItems.Select(x => x.Location).ToArray());
+						_session!.RegisterNewFiles(action.ProducedItems.Select(x => x.Location).ToArray());
 					}
-					Session!.EndExternalProcess(processId, (uint)result.ExitCode);
+					_session!.EndExternalProcess(processId, (uint)result.ExitCode);
 
 					ActionFinished(queue, result, action, null, null);
 				};
@@ -645,18 +660,18 @@ namespace UnrealBuildTool
 			return () =>
 			{
 				ProcessStartInfo startInfo = GetActionStartInfo(action, out FileItem? pchItem);
-				IProcess process = Session!.RunProcess(startInfo, false, null);
+				IProcess process = _session!.RunProcess(startInfo, false, null);
 				if (process.ExitCode != 0 && UBAConfig.bForcedRetry)
 				{
-					ThreadedLogger.LogWarning("{Description} {StatusDescription}: Exited with error code {exitCode}. This action will retry witout UBA", action.CommandDescription, action.StatusDescription, process.ExitCode);
-					ForcedRetryActions.AddOrUpdate(action, false, (k, v) => false);
+					_threadedLogger.LogWarning("{Description} {StatusDescription}: Exited with error code {ExitCode}. This action will retry witout UBA", action.CommandDescription, action.StatusDescription, process.ExitCode);
+					_forcedRetryActions.AddOrUpdate(action, false, (k, v) => false);
 					queue.RequeueAction(action);
 					return Task.CompletedTask;
 				}
 				TimeSpan processorTime = process.TotalProcessorTime;
 				TimeSpan executionTime = process.TotalWallTime;
 				List<string> logLines = process.LogLines;
-				logLines.RemoveAll((line) => line.StartsWith("   Creating library ") && line.EndsWith(".exp") || line.EndsWith("file(s) copied."));
+				logLines.RemoveAll((line) => line.StartsWith("   Creating library ", StringComparison.OrdinalIgnoreCase) && line.EndsWith(".exp", StringComparison.OrdinalIgnoreCase) || line.EndsWith("file(s) copied.", StringComparison.OrdinalIgnoreCase));
 				ActionFinished(queue, new ExecuteResults(logLines, process.ExitCode, executionTime, processorTime), action, pchItem, process);
 				return Task.CompletedTask;
 			};
@@ -672,7 +687,7 @@ namespace UnrealBuildTool
 			return () =>
 			{
 				ProcessStartInfo startInfo = GetActionStartInfo(action, out FileItem? pchItem);
-				Session!.RunProcessRemote(startInfo, (s, e) =>
+				_session!.RunProcessRemote(startInfo, (s, e) =>
 				{
 					if (e.ExitCode != 0 && !e.LogLines.Any())
 					{
@@ -684,7 +699,7 @@ namespace UnrealBuildTool
 					TimeSpan processorTime = e.TotalProcessorTime;
 					TimeSpan executionTime = e.TotalWallTime;
 					List<string> logLines = e.LogLines;
-					logLines.RemoveAll((line) => line.StartsWith("   Creating library ") && line.EndsWith(".exp"));
+					logLines.RemoveAll((line) => line.StartsWith("   Creating library ", StringComparison.OrdinalIgnoreCase) && line.EndsWith(".exp", StringComparison.OrdinalIgnoreCase));
 					ActionFinished(queue, new ExecuteResults(logLines, e.ExitCode, executionTime, processorTime, additionalDescription), action, pchItem, s as IProcess);
 				}, action.Weight);
 				return Task.CompletedTask;
@@ -693,20 +708,20 @@ namespace UnrealBuildTool
 
 		uint ActionsLeftThatCanRunRemotely(ImmediateActionQueue queue)
 		{
-			lock (ActionsChangedLock)
+			lock (_actionsChangedLock)
 			{
-				if (ActionsChanged)
+				if (_bActionsChanged)
 				{
-					ActionsQueuedThatCanRunRemotely = queue.GetQueuedActionsCount(CanRunRemotely);
-					ActionsChanged = false;
+					_actionsQueuedThatCanRunRemotely = queue.GetQueuedActionsCount(CanRunRemotely);
+					_bActionsChanged = false;
 				}
-				return ActionsQueuedThatCanRunRemotely;
+				return _actionsQueuedThatCanRunRemotely;
 			}
 		}
 
 		protected void ActionFinished(ImmediateActionQueue queue, ExecuteResults results, LinkedAction action, FileItem? pchItem = null, IProcess? process = null)
 		{
-			if (IsCancelled)
+			if (_bIsCancelled)
 			{
 				return;
 			}
@@ -714,31 +729,31 @@ namespace UnrealBuildTool
 			bool success = results.ExitCode == 0;
 			if (pchItem != null && process != null && success)
 			{
-				Session!.SetCustomCasKeyFromTrackedInputs(pchItem.Location, action.WorkingDirectory, process);
+				_session!.SetCustomCasKeyFromTrackedInputs(pchItem.Location, action.WorkingDirectory, process);
 			}
 
 			queue.OnActionCompleted(action, success, results);
 
-			lock (ActionsChangedLock)
+			lock (_actionsChangedLock)
 			{
-				ActionsChanged = true;
+				_bActionsChanged = true;
 			}
 		}
 
 		void RemoteActionFailedNoOutput(ImmediateActionQueue queue, LinkedAction action, int exitCode, string executingHost)
 		{
-			if (IsCancelled)
+			if (_bIsCancelled)
 			{
 				return;
 			}
 
-			ThreadedLogger.LogWarning("{Description} {StatusDescription} [RemoteExecutor: {executingHost}]: Exited with error code {exitCode} with no output. This action will retry locally", action.CommandDescription, action.StatusDescription, executingHost, exitCode);
-			LocalRetryActions.AddOrUpdate(action, false, (k, v) => false);
+			_threadedLogger.LogWarning("{Description} {StatusDescription} [RemoteExecutor: {ExecutingHost}]: Exited with error code {ExitCode} with no output. This action will retry locally", action.CommandDescription, action.StatusDescription, executingHost, exitCode);
+			_localRetryActions.AddOrUpdate(action, false, (k, v) => false);
 			queue.RequeueAction(action);
 
-			lock (ActionsChangedLock)
+			lock (_actionsChangedLock)
 			{
-				ActionsChanged = true;
+				_bActionsChanged = true;
 			}
 		}
 	}
@@ -750,7 +765,7 @@ namespace UnrealBuildTool
 	{
 		public override string Name => "Unreal Build Accelerator local";
 
-		public new static bool IsAvailable()
+		public static new bool IsAvailable()
 		{
 			return EpicGames.UBA.Utils.IsAvailable();
 		}
diff --git a/Engine/Source/Programs/UnrealBuildTool/Executors/Experimental/UBAHordeSession.cs b/Engine/Source/Programs/UnrealBuildTool/Executors/Experimental/UBAHordeSession.cs
--- a/Engine/Source/Programs/UnrealBuildTool/Executors/Experimental/UBAHordeSession.cs
+++ b/Engine/Source/Programs/UnrealBuildTool/Executors/Experimental/UBAHordeSession.cs
@@ -28,6 +28,7 @@ namespace UnrealBuildTool
 	/// <summary>
 	/// Configuration for Unreal Build Accelerator Horde session
 	/// </summary>
+	[System.Diagnostics.CodeAnalysis.SuppressMessage("Style", "IDE1006:Naming Styles", Justification = "UnrealBuildTool naming style")]
 	class UnrealBuildAcceleratorHordeConfig
 	{
 		/// <summary>
@@ -36,7 +37,7 @@ namespace UnrealBuildTool
 		[XmlConfigFile(Category = "Horde", Name = "Server")]
 		[CommandLine("-BoxHorde=")]
 		[CommandLine("-UBAHorde=")]
-		public string? HordeServer;
+		public string? HordeServer { get; set; }
 
 		/// <summary>
 		/// Uri of the Horde server
@@ -44,7 +45,7 @@ namespace UnrealBuildTool
 		[XmlConfigFile(Category = "Horde", Name = "Token")]
 		[CommandLine("-BoxHordeToken=")]
 		[CommandLine("-UBAHordeToken=")]
-		public string? HordeToken;
+		public string? HordeToken { get; set; }
 
 		/// <summary>
 		/// OIDC id for the login to use
@@ -52,7 +53,7 @@ namespace UnrealBuildTool
 		[XmlConfigFile(Category = "Horde", Name = "OidcProvider")]
 		[CommandLine("-BoxHordeOidc=")]
 		[CommandLine("-UBAHordeOidc=")]
-		public string? HordeOidcProvider;
+		public string? HordeOidcProvider { get; set; }
 
 		/// <summary>
 		/// Pool for the Horde agent to assign
@@ -60,7 +61,7 @@ namespace UnrealBuildTool
 		[XmlConfigFile(Category = "Horde", Name = "Pool")]
 		[CommandLine("-BoxHordePool=")]
 		[CommandLine("-UBAHordePool=")]
-		public string? HordePool;
+		public string? HordePool { get; set; }
 
 		/// <summary>
 		/// Requirements for the Horde agent to assign
@@ -68,7 +69,7 @@ namespace UnrealBuildTool
 		[XmlConfigFile(Category = "Horde", Name = "Requirements")]
 		[CommandLine("-BoxHordeRequirements=")]
 		[CommandLine("-UBAHordeRequirements=")]
-		public string? HordeCondition;
+		public string? HordeCondition { get; set; }
 
 		/// <summary>
 		/// Which ip UBA server should give to agents. This will invert so host listens and agents connect
@@ -76,7 +77,7 @@ namespace UnrealBuildTool
 		[XmlConfigFile(Category = "Horde", Name = "LocalHost")]
 		[CommandLine("-BoxHordeHost")]
 		[CommandLine("-UBAHordeHost")]
-		public string HordeHost = String.Empty;
+		public string HordeHost { get; set; } = String.Empty;
 
 		/// <summary>
 		/// Max cores allowed to be used by build session
@@ -84,7 +85,7 @@ namespace UnrealBuildTool
 		[XmlConfigFile(Category = "UnrealBuildAccelerator")]
 		[CommandLine("-BoxHordeMaxCores")]
 		[CommandLine("-UBAHordeMaxCores")]
-		public int HordeMaxCores = 576;
+		public int HordeMaxCores { get; set; } = 576;
 
 		/// <summary>
 		/// How long UBT should wait to ask for help. Useful in build configs where machine can delay remote work and still get same wall time results (pch dependencies etc)
@@ -92,7 +93,7 @@ namespace UnrealBuildTool
 		[XmlConfigFile(Category = "UnrealBuildAccelerator")]
 		[CommandLine("-BoxHordeDelay")]
 		[CommandLine("-UBAHordeDelay")]
-		public int HordeDelay = 0;
+		public int HordeDelay { get; set; } = 0;
 
 		/// <summary>
 		/// Allow use of Wine. Only applicable to Horde agents running Linux. Can still be ignored if Wine executable is not set on agent.
@@ -100,7 +101,7 @@ namespace UnrealBuildTool
 		[XmlConfigFile(Category = "UnrealBuildAccelerator")]
 		[CommandLine("-BoxHordeAllowWine", Value = "true")]
 		[CommandLine("-UBAHordeAllowWine", Value = "true")]
-		public bool bHordeAllowWine = true;
+		public bool bHordeAllowWine { get; set; } = true;
 
 		/// <summary>
 		/// Connection mode for agent/compute communication
@@ -109,7 +110,7 @@ namespace UnrealBuildTool
 		[XmlConfigFile(Category = "Horde", Name = "ConnectionMode")]
 		[CommandLine("-BoxHordeConnectionMode=")]
 		[CommandLine("-UBAHordeConnectionMode=")]
-		public string? HordeConnectionMode;
+		public string? HordeConnectionMode { get; set; }
 
 		/// <summary>
 		/// Sentry URL to send box data to. Optional.
@@ -117,7 +118,7 @@ namespace UnrealBuildTool
 		[XmlConfigFile(Category = "Horde")]
 		[CommandLine("-BoxSentryUrl=")]
 		[CommandLine("-UBASentryUrl=")]
-		public string? UBASentryUrl;
+		public string? UBASentryUrl { get; set; }
 
 		/// <summary>
 		/// Disable horde all together
@@ -125,7 +126,7 @@ namespace UnrealBuildTool
 		[XmlConfigFile(Category = "UnrealBuildAccelerator")]
 		[CommandLine("-BoxDisableHorde")]
 		[CommandLine("-UBADisableHorde")]
-		public bool bDisableHorde = false;
+		public bool bDisableHorde { get; set; } = false;
 	}
 
 	class UBAHordeSession : IAsyncDisposable
@@ -158,13 +159,13 @@ namespace UnrealBuildTool
 
 		public struct Worker
 		{
-			public Task BackgroundTask;
-			public int NumLogicalCores;
-			public Stopwatch StartTime;
-			public bool Started;
-			public string Ip;
-			public ConnectionMetadataPort Port;
-			public ConnectionMetadataPort ProxyPort;
+			public Task BackgroundTask { get; set; }
+			public int NumLogicalCores { get; set; }
+			public Stopwatch StartTime { get; set; }
+			public bool Started { get; set; }
+			public string Ip { get; set; }
+			public ConnectionMetadataPort Port { get; set; }
+			public ConnectionMetadataPort ProxyPort { get; set; }
 		}
 
 		readonly List<Worker> _workers = new();
@@ -191,7 +192,7 @@ namespace UnrealBuildTool
 			services.AddHttpClient<HordeHttpClient>(ConfigureHttpClient);
 			_serviceProvider = services.BuildServiceProvider();
 
-			if (connectionMode == ConnectionMode.Relay && _crypto == String.Empty)
+			if (connectionMode == ConnectionMode.Relay && String.IsNullOrEmpty(_crypto))
 			{
 				_crypto = UBAExecutor.CreateCrypto();
 			}
@@ -220,6 +221,8 @@ namespace UnrealBuildTool
 
 			await _serviceProvider.DisposeAsync();
 			_cancellationTokenSource.Dispose();
+
+			_storage.Dispose();
 		}
 
 		public async Task InitAsync(bool useSentry, CancellationToken cancellationToken)
@@ -294,7 +297,7 @@ namespace UnrealBuildTool
 			return handle.GetLocator();
 		}
 
-		public int NumLogicalCores = 0;
+		public int NumLogicalCores { get; private set; } = 0;
 
 		public async void RemoveCompleteWorkers()
 		{
@@ -379,13 +382,13 @@ namespace UnrealBuildTool
 				string computeIp = String.Empty;
 				foreach (string property in lease.Properties)
 				{
-					int equalsIdx = property.IndexOf('=');
+					int equalsIdx = property.IndexOf('=', StringComparison.OrdinalIgnoreCase);
 					StringView propertyName = new(property, 0, equalsIdx);
 					if (s_logProperties.Contains(propertyName))
 					{
 						_logger.LogInformation("  {Property}", property);
 
-						if (propertyName == ResourceLogicalCores && int.TryParse(property.AsSpan(equalsIdx + 1), out int value))
+						if (propertyName == ResourceLogicalCores && Int32.TryParse(property.AsSpan(equalsIdx + 1), out int value))
 						{
 							numLogicalCores = value;
 						}
@@ -452,7 +455,7 @@ namespace UnrealBuildTool
 			}
 		}
 
-		public static async Task<UBAHordeSession?> TryCreateHordeSession(UnrealBuildAcceleratorHordeConfig hordeConfig, UBAExecutor executor, bool bStrictErrors, CancellationToken cancellationToken, ILogger logger)
+		public static async Task<UBAHordeSession?> TryCreateHordeSession(UnrealBuildAcceleratorHordeConfig hordeConfig, UBAExecutor executor, bool bStrictErrors, ILogger logger, CancellationToken cancellationToken = default)
 		{
 			if (hordeConfig.bDisableHorde)
 			{
@@ -489,7 +492,7 @@ namespace UnrealBuildTool
 			{
 				if (String.IsNullOrEmpty(token) && hasOidcProvider)
 				{
-					token = await GetOidcBearerTokenAsync(null, oidcProvider!, cancellationToken, logger);
+					token = await GetOidcBearerTokenAsync(null, oidcProvider!, logger, cancellationToken);
 				}
 
 				AuthenticationHeaderValue? authHeader = null;
@@ -498,9 +501,9 @@ namespace UnrealBuildTool
 					authHeader = new AuthenticationHeaderValue("Bearer", token);
 				}
 
-				bool AllowWine = hordeConfig.bHordeAllowWine && OperatingSystem.IsWindows();
+				bool allowWine = hordeConfig.bHordeAllowWine && OperatingSystem.IsWindows();
 
-				UBAHordeSession session = new(executor, new Uri(server), authHeader, hordeConfig.HordePool, AllowWine, hordeConfig.HordeMaxCores, bStrictErrors, connectionMode, logger);
+				UBAHordeSession session = new(executor, new Uri(server), authHeader, hordeConfig.HordePool, allowWine, hordeConfig.HordeMaxCores, bStrictErrors, connectionMode, logger);
 				await session.InitAsync(useSentry: !String.IsNullOrEmpty(hordeConfig.UBASentryUrl), cancellationToken);
 				return session;
 			}
@@ -515,7 +518,7 @@ namespace UnrealBuildTool
 			}
 		}
 
-		static async Task<string> GetOidcBearerTokenAsync(DirectoryReference? projectDir, string oidcProvider, CancellationToken cancellationToken, ILogger logger)
+		static async Task<string> GetOidcBearerTokenAsync(DirectoryReference? projectDir, string oidcProvider, ILogger logger, CancellationToken cancellationToken = default)
 		{
 			logger.LogInformation("Performing OIDC token refresh...");
 
@@ -572,7 +575,9 @@ namespace UnrealBuildTool
 
 						string hordeHost = _owner.UBAConfig.Host;
 						if (!String.IsNullOrEmpty(_owner.HordeConfig.HordeHost))
+						{
 							hordeHost = _owner.HordeConfig.HordeHost;
+						}
 
 						bool useListen = !String.IsNullOrEmpty(_owner.HordeConfig.HordeHost);
 						List<string> arguments = new();
