Implement GetPhysicalProcessorCount for Linux so that PGO builds can work properly

#tests tested on WSL, Ubuntu 22.04 Threadripper and Ubuntu 22.04 Megastation

diff --git a/Engine/Source/Programs/UnrealBuildTool/System/Utils.cs b/Engine/Source/Programs/UnrealBuildTool/System/Utils.cs
--- a/Engine/Source/Programs/UnrealBuildTool/System/Utils.cs
+++ b/Engine/Source/Programs/UnrealBuildTool/System/Utils.cs
@@ -1073,6 +1073,16 @@ namespace UnrealBuildTool
 					}
 				}
 			}
+			else if (RuntimePlatform.IsLinux)
+			{
+				// query socket/logical core pairings.  There should not be duplicates in this list since each hyperthread
+				// will show up as it's own logical "cpu".  Including the socket number allows us to count multi-processor
+				// system cores correctly
+				string Output = RunLocalProcessAndReturnStdOut("lscpu", "-p='SOCKET,CPU'");
+				List<string> CPUs = Output.Split("\n").Where(x => !x.StartsWith("#")).ToList();
+
+				return CPUs.Count;
+			}
 			return Environment.ProcessorCount;
 		}
 
@@ -1134,6 +1144,16 @@ namespace UnrealBuildTool
 					return Value;
 				}
 			}
+			else if (RuntimePlatform.IsLinux)
+			{
+				// query socket/physical core pairings.  There will be duplicates in this if there are hyperthreads
+				// using the HashSet ensures that those duplicates are removed and we only count the first "cpu" found
+				// for each "core".  Including the socket number allows us to count multi-processor system cores correctly
+				string Output = RunLocalProcessAndReturnStdOut("lscpu", "-p='SOCKET,CORE'");
+				HashSet<string> CPUs = Output.Split("\n").Where(x => !x.StartsWith("#")).ToHashSet();
+
+				return CPUs.Count;
+			}
 
 			return -1;
 		}
