UnrealBuildTool: Set PGO_PATH_TRANSLATION for all msvc platforms and additionally add the project path

#rnx

diff --git a/Engine/Source/Programs/UnrealBuildTool/Platform/Windows/VCToolChain.cs b/Engine/Source/Programs/UnrealBuildTool/Platform/Windows/VCToolChain.cs
--- a/Engine/Source/Programs/UnrealBuildTool/Platform/Windows/VCToolChain.cs
+++ b/Engine/Source/Programs/UnrealBuildTool/Platform/Windows/VCToolChain.cs
@@ -70,6 +70,14 @@ namespace UnrealBuildTool
 		{
 			EnvVars.SetEnvironmentVariables();
 
+			// This allows PGD files generated from this build to be portable to other machines. It needs to be set in all configurations so we don't split the environment
+			List<string> pgoPathTranslations = new() { $"{Unreal.RootDirectory.FullName}=ROOT" };
+			if (Target.ProjectFile != null)
+			{
+				pgoPathTranslations.Add($"{Target.ProjectFile.Directory.FullName}=PROJ");
+			}
+			Environment.SetEnvironmentVariable("PGO_PATH_TRANSLATION", String.Join(";", pgoPathTranslations));
+
 			// Don't allow the INCLUDE environment variable to propagate. It's set by the IDE based on the IncludePath property in the project files which we
 			// add to improve Visual Studio memory usage, but we don't actually need it to set when invoking the compiler. Doing so results in it being converted
 			// into /I arguments by the CL driver, which results in errors due to the command line not fitting into the PDB debug record.
