UnrealBuildTool: ParallelExecutor use Process.WaitForExitAsync

#rnx

diff --git a/Engine/Source/Programs/UnrealBuildTool/Executors/ParallelExecutor.cs b/Engine/Source/Programs/UnrealBuildTool/Executors/ParallelExecutor.cs
--- a/Engine/Source/Programs/UnrealBuildTool/Executors/ParallelExecutor.cs
+++ b/Engine/Source/Programs/UnrealBuildTool/Executors/ParallelExecutor.cs
@@ -222,7 +222,7 @@ namespace UnrealBuildTool
 
 			CancellationToken.ThrowIfCancellationRequested();
 
-			Process.WaitForExit();
+			await Process.WaitForExitAsync(CancellationToken);
 
 			List<string> LogLines = Console.OutputEncoding.GetString(StdOutStream.GetBuffer(), 0, Convert.ToInt32(StdOutStream.Length)).Split(LineEndingSplit, StringSplitOptions.RemoveEmptyEntries).ToList();
 			int ExitCode = Process.ExitCode;
