UnrealBuildTool: DeprecationWarningLevel RequiresUniqueBuildEnvironment

#rnx

diff --git a/Engine/Source/Programs/UnrealBuildTool/Configuration/TargetRules.cs b/Engine/Source/Programs/UnrealBuildTool/Configuration/TargetRules.cs
--- a/Engine/Source/Programs/UnrealBuildTool/Configuration/TargetRules.cs
+++ b/Engine/Source/Programs/UnrealBuildTool/Configuration/TargetRules.cs
@@ -1705,17 +1705,18 @@ namespace UnrealBuildTool
 		/// <summary>
 		/// Emits compilation errors for incorrect UE_LOG format strings.
 		/// </summary>
-		[XmlConfigFile(Category = "BuildConfiguration")]		
+		[XmlConfigFile(Category = "BuildConfiguration")]
 		public bool bValidateFormatStrings
 		{
 				get { return bValidateFormatStringsPrivate ?? (DefaultBuildSettings >= BuildSettingsVersion.V5); }
 				set { bValidateFormatStringsPrivate = value; }
 		}
 		private bool? bValidateFormatStringsPrivate;
-			
+
 		/// <summary>
 		/// Level to report deprecation warnings as errors
 		/// </summary>
+		[RequiresUniqueBuildEnvironment]
 		public WarningLevel DeprecationWarningLevel
 		{
 			get => (DeprecationWarningLevelPrivate == WarningLevel.Default) ? DefaultWarningLevel : DeprecationWarningLevelPrivate;
