- Added a new TargetBuildEnvronment - UniqueIfNeeded. This will use Shared if possible, but if a property that requires unique environment, of a project has overridden an SDK version, Unique will beused instead

diff --git a/Engine/Source/Programs/UnrealBuildTool/Configuration/TargetRules.cs b/Engine/Source/Programs/UnrealBuildTool/Configuration/TargetRules.cs
--- a/Engine/Source/Programs/UnrealBuildTool/Configuration/TargetRules.cs
+++ b/Engine/Source/Programs/UnrealBuildTool/Configuration/TargetRules.cs
@@ -9,6 +9,7 @@ using System.Runtime.Serialization;
 using EpicGames.Core;
 using Microsoft.Extensions.Logging;
 using UnrealBuildBase;
+using System.Diagnostics.CodeAnalysis;
 
 namespace UnrealBuildTool
 {
@@ -84,6 +85,11 @@ namespace UnrealBuildTool
 		/// Engine binaries and intermediates are specific to this target
 		/// </summary>
 		Unique,
+
+		/// <summary>
+		/// Will switch to Unique if needed - per-project SDK is enabled, or a property that requires unique is set away from default
+		/// </summary>
+		UniqueIfNeeded,
 	}
 
 	/// <summary>
@@ -2461,6 +2467,10 @@ namespace UnrealBuildTool
 			{
 				if (BuildEnvironmentOverride.HasValue)
 				{
+					if (BuildEnvironmentOverride.Value == TargetBuildEnvironment.UniqueIfNeeded)
+					{
+						throw new BuildException($"Target {Name} had BuildEnv set to UniqueIfNeeded when querying, which means UpdateBuildEnvironmentIfNeeded wasn't called in time");
+					}
 					return BuildEnvironmentOverride.Value;
 				}
 				if (Type == TargetType.Program && ProjectFile != null && File!.IsUnderDirectory(ProjectFile.Directory))
@@ -2646,6 +2656,160 @@ namespace UnrealBuildTool
 			return LinkType == TargetLinkType.Monolithic || BuildEnvironment == TargetBuildEnvironment.Unique;
 		}
 
+		/// <summary>
+		/// Checks if a property has been set with the RequiresUniqueBuildEnvironmentAttribute, and is different from it's base
+		/// </summary>
+		/// <param name="RulesAssembly">Assembly containing the target</param>
+		/// <param name="Arguments">Commandline options that may affect the target creation</param>
+		/// <param name="PropNamesThatRequireUnique">If the target requires a unique build environment, this will contain the names of the field/property that require unique</param>
+		/// <param name="BaseTargetName">If the target requires a unique build environment, this will contain the name of the target this was based on (UnrealGame, UnrealEditor, etc)/property</param>
+		/// <returns>true if a property was set such that it requires a unique build environment</returns>
+		/// <exception cref="BuildException"></exception>
+		public bool RequiresUniqueEnvironment(RulesAssembly RulesAssembly, CommandLineArguments? Arguments, List<string> PropNamesThatRequireUnique, [NotNullWhen(true)] out string? BaseTargetName)
+		{
+			BaseTargetName = null;
+
+			TargetRules ThisRules = this;
+			// Allow disabling these checks
+			if (ThisRules.bOverrideBuildEnvironment)
+			{
+				return false;
+			}
+
+			// Get the name of the target with default settings
+			switch (ThisRules.Type)
+			{
+				case TargetType.Game:
+					BaseTargetName = "UnrealGame";
+					break;
+				case TargetType.Editor:
+					BaseTargetName = "UnrealEditor";
+					break;
+				case TargetType.Client:
+					BaseTargetName = "UnrealClient";
+					break;
+				case TargetType.Server:
+					BaseTargetName = "UnrealServer";
+					break;
+				default:
+					return false;
+			}
+
+			// Create the target rules for it
+			TargetRules BaseRules = RulesAssembly.CreateTargetRules(BaseTargetName, ThisRules.Platform, ThisRules.Configuration, ThisRules.Architectures, null, Arguments, Logger, IntermediateEnvironment: ThisRules.IntermediateEnvironment);
+
+			// Get all the configurable objects
+			object[] BaseObjects = BaseRules.GetConfigurableObjects().ToArray();
+			object[] ThisObjects = GetConfigurableObjects().ToArray();
+			if (BaseObjects.Length != ThisObjects.Length)
+			{
+				throw new BuildException("Expected same number of configurable objects from base rules object.");
+			}
+
+			// Iterate through all fields with the [SharedBuildEnvironment] attribute
+			for (int Idx = 0; Idx < BaseObjects.Length; Idx++)
+			{
+				Type ObjectType = BaseObjects[Idx].GetType();
+				foreach (FieldInfo Field in ObjectType.GetFields())
+				{
+					if (Field.GetCustomAttribute<RequiresUniqueBuildEnvironmentAttribute>() != null)
+					{
+						object? ThisValue = Field.GetValue(ThisObjects[Idx]);
+						object? BaseValue = Field.GetValue(BaseObjects[Idx]);
+						if (!CheckValuesMatch(Field.FieldType, ThisValue, BaseValue))
+						{
+							PropNamesThatRequireUnique.Add(Field.Name);
+						}
+					}
+				}
+				foreach (PropertyInfo Property in ObjectType.GetProperties())
+				{
+					if (Property.GetCustomAttribute<RequiresUniqueBuildEnvironmentAttribute>() != null)
+					{
+						object? ThisValue = Property.GetValue(ThisObjects[Idx]);
+						object? BaseValue = Property.GetValue(BaseObjects[Idx]);
+						if (!CheckValuesMatch(Property.PropertyType, ThisValue, BaseValue))
+						{
+							PropNamesThatRequireUnique.Add(Property.Name);
+						}
+					}
+				}
+			}
+
+			// if any properties require a unique build environment, return true
+			return PropNamesThatRequireUnique.Count > 0;
+		}
+
+		/// <summary>
+		/// For any target that has set BuildEnvironment = TargetBuildEnvironment.Unique, this will change the Environment to either Shared or Unique accordingly
+		/// </summary>
+		/// <param name="RulesAssembly">Assembly containing the target</param>
+		/// <param name="Arguments">Commandline options that may affect the target creation</param>
+		/// <param name="Logger">Logger</param>
+		public void UpdateBuildEnvironmentIfNeeded(RulesAssembly RulesAssembly, CommandLineArguments? Arguments, ILogger Logger)
+		{
+			// only do anything here when using UniqueIfNeeded
+			if (!BuildEnvironmentOverride.HasValue || BuildEnvironmentOverride.Value != TargetBuildEnvironment.UniqueIfNeeded)
+			{
+				return;
+			}
+
+			if (UEBuildPlatformSDK.bHasAnySDKOverride)
+			{
+				// check if any platform needs a unique environment for the SDK check
+				foreach (UnrealTargetPlatform Platform in UnrealTargetPlatform.GetValidPlatforms())
+				{
+					if (IsSDKVersionRelevant(Platform))
+					{
+						UEBuildPlatformSDK? SDK = UEBuildPlatformSDK.GetSDKForPlatform(Platform.ToString());
+						if (SDK != null && SDK.bHasSDKOverride)
+						{
+							Logger.LogInformation("Setting {Target}'s BuildEnvironment to Unique, because a project overrode the {Platform} SDK version", Name, Platform);
+							BuildEnvironment = TargetBuildEnvironment.Unique;
+							break;
+						}
+					}
+				}
+			}
+
+			// if we didn't set it above, check the properties
+			if (BuildEnvironmentOverride.Value == TargetBuildEnvironment.UniqueIfNeeded)
+			{
+				List<string> PropNames = new();
+				string? BaseTargetName;
+				if (RequiresUniqueEnvironment(RulesAssembly, Arguments, PropNames, out BaseTargetName))
+				{
+					Logger.LogInformation("Setting {Target}'s BuildEnvironment to Unique, because it had changed the values of the propertues [ {Props} ] away from the values specified in {BaseTarget}", 
+						Name, string.Join(", ", PropNames), BaseTargetName);
+					BuildEnvironment = TargetBuildEnvironment.Unique;
+				}
+				else
+				{
+					BuildEnvironment = TargetBuildEnvironment.Shared;
+				}
+			}
+		}
+
+		static bool CheckValuesMatch(Type ValueType, object? ThisValue, object? BaseValue)
+		{
+			// Check if the fields match, treating lists of strings (eg. definitions) differently to value types.
+			bool bFieldsMatch;
+			if (ThisValue == null || BaseValue == null)
+			{
+				bFieldsMatch = (ThisValue == BaseValue);
+			}
+			else if (typeof(IEnumerable<string>).IsAssignableFrom(ValueType))
+			{
+				bFieldsMatch = Enumerable.SequenceEqual((IEnumerable<string>)ThisValue, (IEnumerable<string>)BaseValue);
+			}
+			else
+			{
+				bFieldsMatch = ThisValue.Equals(BaseValue);
+			}
+
+			return bFieldsMatch;
+		}
+
 		/// <summary>
 		/// Create a TargetRules instance
 		/// </summary>
diff --git a/Engine/Source/Programs/UnrealBuildTool/Configuration/UEBuildTarget.cs b/Engine/Source/Programs/UnrealBuildTool/Configuration/UEBuildTarget.cs
--- a/Engine/Source/Programs/UnrealBuildTool/Configuration/UEBuildTarget.cs
+++ b/Engine/Source/Programs/UnrealBuildTool/Configuration/UEBuildTarget.cs
@@ -13,6 +13,7 @@ using System.Text;
 using System.Text.Json.Serialization;
 using System.Threading.Tasks;
 using EpicGames.Core;
+using EpicGames.Serialization;
 using EpicGames.UHT.Utils;
 using Microsoft.Extensions.Logging;
 using OpenTracing.Util;
@@ -1396,65 +1397,12 @@ namespace UnrealBuildTool
 		/// </summary>
 		static void ValidateSharedEnvironment(RulesAssembly RulesAssembly, string ThisTargetName, CommandLineArguments Arguments, TargetRules ThisRules, ILogger Logger)
 		{
-			// Allow disabling these checks
-			if (ThisRules.bOverrideBuildEnvironment)
+			List<string> PropNamesThatRequiredUnique = new();
+			string? BaseTargetName;
+			if (ThisRules.RequiresUniqueEnvironment(RulesAssembly, Arguments, PropNamesThatRequiredUnique, out BaseTargetName))
 			{
-				return;
-			}
-
-			// Get the name of the target with default settings
-			string BaseTargetName;
-			switch (ThisRules.Type)
-			{
-				case TargetType.Game:
-					BaseTargetName = "UnrealGame";
-					break;
-				case TargetType.Editor:
-					BaseTargetName = "UnrealEditor";
-					break;
-				case TargetType.Client:
-					BaseTargetName = "UnrealClient";
-					break;
-				case TargetType.Server:
-					BaseTargetName = "UnrealServer";
-					break;
-				default:
-					return;
-			}
-
-			// Create the target rules for it
-			TargetRules BaseRules = RulesAssembly.CreateTargetRules(BaseTargetName, ThisRules.Platform, ThisRules.Configuration, ThisRules.Architectures, null, Arguments, Logger, IntermediateEnvironment: ThisRules.IntermediateEnvironment);
-
-			// Get all the configurable objects
-			object[] BaseObjects = BaseRules.GetConfigurableObjects().ToArray();
-			object[] ThisObjects = ThisRules.GetConfigurableObjects().ToArray();
-			if (BaseObjects.Length != ThisObjects.Length)
-			{
-				throw new BuildException("Expected same number of configurable objects from base rules object.");
-			}
-
-			// Iterate through all fields with the [SharedBuildEnvironment] attribute
-			for (int Idx = 0; Idx < BaseObjects.Length; Idx++)
-			{
-				Type ObjectType = BaseObjects[Idx].GetType();
-				foreach (FieldInfo Field in ObjectType.GetFields())
-				{
-					if (Field.GetCustomAttribute<RequiresUniqueBuildEnvironmentAttribute>() != null)
-					{
-						object? ThisValue = Field.GetValue(ThisObjects[Idx]);
-						object? BaseValue = Field.GetValue(BaseObjects[Idx]);
-						CheckValuesMatch(ThisRules.GetType(), ThisTargetName, BaseTargetName, Field.Name, Field.FieldType, ThisValue, BaseValue);
-					}
-				}
-				foreach (PropertyInfo Property in ObjectType.GetProperties())
-				{
-					if (Property.GetCustomAttribute<RequiresUniqueBuildEnvironmentAttribute>() != null)
-					{
-						object? ThisValue = Property.GetValue(ThisObjects[Idx]);
-						object? BaseValue = Property.GetValue(BaseObjects[Idx]);
-						CheckValuesMatch(ThisRules.GetType(), ThisTargetName, BaseTargetName, Property.Name, Property.PropertyType, ThisValue, BaseValue);
-					}
-				}
+				throw new BuildException("{0} modifies the values of properties: [ {1} ]. This is not allowed, as {0} has build products in common with {2}.\nRemove the modified setting, change {0} to use a unique build environment by setting 'BuildEnvironment = TargetBuildEnvironment.Unique;' in the {3} constructor, or set bOverrideBuildEnvironment = true to force this setting on.", 
+					ThisTargetName, string.Join(", ", PropNamesThatRequiredUnique), BaseTargetName, ThisRules.GetType().Name);
 			}
 
 			// Make sure that we don't explicitly enable or disable any plugins through the target rules. We can't do this with the shared build environment because it requires recompiling the "Projects" engine module.
@@ -1468,33 +1416,6 @@ namespace UnrealBuildTool
 			}
 		}
 
-		/// <summary>
-		/// Check that two values match between a base and derived rules type
-		/// </summary>
-		static void CheckValuesMatch(Type RulesType, string ThisTargetName, string BaseTargetName, string FieldName, Type ValueType, object? ThisValue, object? BaseValue)
-		{
-			// Check if the fields match, treating lists of strings (eg. definitions) differently to value types.
-			bool bFieldsMatch;
-			if (ThisValue == null || BaseValue == null)
-			{
-				bFieldsMatch = (ThisValue == BaseValue);
-			}
-			else if (typeof(IEnumerable<string>).IsAssignableFrom(ValueType))
-			{
-				bFieldsMatch = Enumerable.SequenceEqual((IEnumerable<string>)ThisValue, (IEnumerable<string>)BaseValue);
-			}
-			else
-			{
-				bFieldsMatch = ThisValue.Equals(BaseValue);
-			}
-
-			// Throw an exception if they don't match
-			if (!bFieldsMatch)
-			{
-				throw new BuildException("{0} modifies the value of {1}. This is not allowed, as {0} has build products in common with {2}.\nRemove the modified setting, change {0} to use a unique build environment by setting 'BuildEnvironment = TargetBuildEnvironment.Unique;' in the {3} constructor, or set bOverrideBuildEnvironment = true to force this setting on.", ThisTargetName, FieldName, BaseTargetName, RulesType.Name);
-			}
-		}
-
 		/// <summary>
 		/// The target rules
 		/// </summary>
diff --git a/Engine/Source/Programs/UnrealBuildTool/Platform/Mac/UEBuildMac.cs b/Engine/Source/Programs/UnrealBuildTool/Platform/Mac/UEBuildMac.cs
--- a/Engine/Source/Programs/UnrealBuildTool/Platform/Mac/UEBuildMac.cs
+++ b/Engine/Source/Programs/UnrealBuildTool/Platform/Mac/UEBuildMac.cs
@@ -178,7 +178,7 @@ namespace UnrealBuildTool
 					// CreateTargetRules here needs to have an UnrealArchitectures object, because otherwise with 'null', it will call
 					// back to this function to get the ActiveArchitectures! in this case the arch is unimportant
 					UnrealArchitectures DummyArchitectures = new(UnrealArch.X64);
-					TargetRules? Rules = RulesAsm.CreateTargetRules(TargetName, UnrealTargetPlatform.Mac, UnrealTargetConfiguration.Development, DummyArchitectures, ProjectFile, null, Log.Logger);
+					TargetRules? Rules = RulesAsm.CreateTargetRules(TargetName, UnrealTargetPlatform.Mac, UnrealTargetConfiguration.Development, DummyArchitectures, ProjectFile, null, Log.Logger, bSkipValidation:true);
 					bIsEditor = Rules.Type == TargetType.Editor;
 
 					// the projectfile passed in may be a game's uproject file that we are compiling a program in the context of, 
diff --git a/Engine/Source/Programs/UnrealBuildTool/System/RulesAssembly.cs b/Engine/Source/Programs/UnrealBuildTool/System/RulesAssembly.cs
--- a/Engine/Source/Programs/UnrealBuildTool/System/RulesAssembly.cs
+++ b/Engine/Source/Programs/UnrealBuildTool/System/RulesAssembly.cs
@@ -642,10 +642,16 @@ namespace UnrealBuildTool
 				throw new CompilationResultException(CompilationResult.RulesError, "TargetRules.LinkType should be inferred from TargetType");
 			}
 
-			// Set the default value for whether to use the shared build environment
-			if (Rules.BuildEnvironment == TargetBuildEnvironment.Unique && Unreal.IsEngineInstalled())
+			if (!bSkipValidation)
 			{
-				throw new CompilationResultException(CompilationResult.RulesError, "Targets with a unique build environment cannot be built with an installed engine.");
+				// Delayed-fixup of TargetBuildEnvironment.UniqueIfNeeded
+				Rules.UpdateBuildEnvironmentIfNeeded(this, Arguments: null, Logger);
+
+				// Set the default value for whether to use the shared build environment
+				if (Rules.BuildEnvironment == TargetBuildEnvironment.Unique && Unreal.IsEngineInstalled())
+				{
+					throw new CompilationResultException(CompilationResult.RulesError, "Targets with a unique build environment cannot be built with an installed engine.");
+				}
 			}
 
 			// Automatically include CoreUObject
