[Uba]
* Added using() around ProcessHandle so it is disposed when not needed anymore

#prelight 656f55f7d51671afb825466a

diff --git a/Engine/Source/Programs/UnrealBuildTool/Executors/Experimental/UBAExecutor.cs b/Engine/Source/Programs/UnrealBuildTool/Executors/Experimental/UBAExecutor.cs
--- a/Engine/Source/Programs/UnrealBuildTool/Executors/Experimental/UBAExecutor.cs
+++ b/Engine/Source/Programs/UnrealBuildTool/Executors/Experimental/UBAExecutor.cs
@@ -660,19 +660,21 @@ namespace UnrealBuildTool
 			return () =>
 			{
 				ProcessStartInfo startInfo = GetActionStartInfo(action, out FileItem? pchItem);
-				IProcess process = _session!.RunProcess(startInfo, false, null);
-				if (process.ExitCode != 0 && UBAConfig.bForcedRetry)
+				using (IProcess process = _session!.RunProcess(startInfo, false, null))
 				{
-					_threadedLogger.LogWarning("{Description} {StatusDescription}: Exited with error code {ExitCode}. This action will retry witout UBA", action.CommandDescription, action.StatusDescription, process.ExitCode);
-					_forcedRetryActions.AddOrUpdate(action, false, (k, v) => false);
-					queue.RequeueAction(action);
-					return Task.CompletedTask;
+					if (process.ExitCode != 0 && UBAConfig.bForcedRetry)
+					{
+						_threadedLogger.LogWarning("{Description} {StatusDescription}: Exited with error code {ExitCode}. This action will retry witout UBA", action.CommandDescription, action.StatusDescription, process.ExitCode);
+						_forcedRetryActions.AddOrUpdate(action, false, (k, v) => false);
+						queue.RequeueAction(action);
+						return Task.CompletedTask;
+					}
+					TimeSpan processorTime = process.TotalProcessorTime;
+					TimeSpan executionTime = process.TotalWallTime;
+					List<string> logLines = process.LogLines;
+					logLines.RemoveAll((line) => line.StartsWith("   Creating library ", StringComparison.OrdinalIgnoreCase) && line.EndsWith(".exp", StringComparison.OrdinalIgnoreCase) || line.EndsWith("file(s) copied.", StringComparison.OrdinalIgnoreCase));
+					ActionFinished(queue, new ExecuteResults(logLines, process.ExitCode, executionTime, processorTime), action, pchItem, process);
 				}
-				TimeSpan processorTime = process.TotalProcessorTime;
-				TimeSpan executionTime = process.TotalWallTime;
-				List<string> logLines = process.LogLines;
-				logLines.RemoveAll((line) => line.StartsWith("   Creating library ", StringComparison.OrdinalIgnoreCase) && line.EndsWith(".exp", StringComparison.OrdinalIgnoreCase) || line.EndsWith("file(s) copied.", StringComparison.OrdinalIgnoreCase));
-				ActionFinished(queue, new ExecuteResults(logLines, process.ExitCode, executionTime, processorTime), action, pchItem, process);
 				return Task.CompletedTask;
 			};
 		}
