UBA: Add encryption parameter for Horde's compute client

diff --git a/Engine/Source/Programs/UnrealBuildTool/Executors/Experimental/UBAHordeSession.cs b/Engine/Source/Programs/UnrealBuildTool/Executors/Experimental/UBAHordeSession.cs
--- a/Engine/Source/Programs/UnrealBuildTool/Executors/Experimental/UBAHordeSession.cs
+++ b/Engine/Source/Programs/UnrealBuildTool/Executors/Experimental/UBAHordeSession.cs
@@ -111,6 +111,15 @@ namespace UnrealBuildTool
 		[CommandLine("-BoxHordeConnectionMode=")]
 		[CommandLine("-UBAHordeConnectionMode=")]
 		public string? HordeConnectionMode { get; set; }
+		
+		/// <summary>
+		/// Encryption to use for agent/compute communication. Note that UBA agent uses its own encryption.
+		/// <see cref="Encryption" /> for valid modes.
+		/// </summary>
+		[XmlConfigFile(Category = "Horde", Name = "Encryption")]
+		[CommandLine("-BoxHordeEncryption=")]
+		[CommandLine("-UBAHordeEncryption=")]
+		public string? HordeEncryption { get; set; }
 
 		/// <summary>
 		/// Sentry URL to send box data to. Optional.
@@ -145,6 +154,7 @@ namespace UnrealBuildTool
 		readonly int _maxCores;
 		readonly bool _strict;
 		readonly ConnectionMode? _connectionMode;
+		readonly Encryption? _encryption;
 		readonly CancellationTokenSource _cancellationTokenSource = new();
 		readonly ILogger _logger;
 
@@ -170,7 +180,7 @@ namespace UnrealBuildTool
 
 		readonly List<Worker> _workers = new();
 
-		public UBAHordeSession(UBAExecutor owner, Uri hordeUri, AuthenticationHeaderValue? authHeader, string? pool, bool allowWine, int maxCores, bool strict, ConnectionMode? connectionMode, ILogger logger)
+		public UBAHordeSession(UBAExecutor owner, Uri hordeUri, AuthenticationHeaderValue? authHeader, string? pool, bool allowWine, int maxCores, bool strict, ConnectionMode? connectionMode, Encryption? encryption, ILogger logger)
 		{
 			_owner = owner;
 			_hordeUri = hordeUri;
@@ -179,6 +189,7 @@ namespace UnrealBuildTool
 			_maxCores = maxCores;
 			_strict = strict;
 			_connectionMode = connectionMode;
+			_encryption = encryption;
 			_logger = logger;
 			_crypto = owner.Crypto;
 
@@ -362,6 +373,7 @@ namespace UnrealBuildTool
 				ConnectionMetadataRequest cmr = new ()
 				{
 					ModePreference = _connectionMode,
+					Encryption = _encryption,
 					Ports = { {UbaPortName, UbaPort}, {UbaProxyPortName, UbaProxyPort} }
 				};
 				lease = await _client.TryAssignWorkerAsync(_clusterId, requirements, requestId, cmr, workerLogger, cancellationToken);
@@ -482,12 +494,13 @@ namespace UnrealBuildTool
 			}
 
 			ConnectionMode? connectionMode = Enum.TryParse(hordeConfig.HordeConnectionMode, true, out ConnectionMode cm) ? cm : null;
+			Encryption? encryption = Enum.TryParse(hordeConfig.HordeEncryption, true, out Encryption enc) ? enc : null;
 
 			oidcProvider ??= Environment.GetEnvironmentVariable("UE_HORDE_OIDC_PROVIDER");
 
 			bool hasOidcProvider = !String.IsNullOrEmpty(oidcProvider);
-			logger.LogInformation("Horde URL: {Server}, Pool: {Pool}, Condition: {Condition}, OIDC: {OidcProvider}, Connection: {Connection}",
-				server, hordeConfig.HordePool ?? "(none)", hordeConfig.HordeCondition ?? "(none)", hasOidcProvider ? oidcProvider! : "Disabled", connectionMode?.ToString() ?? "(none)");
+			logger.LogInformation("Horde URL: {Server}, Pool: {Pool}, Condition: {Condition}, OIDC: {OidcProvider}, Connection: {Connection} HordeEncryption: {Encryption}",
+				server, hordeConfig.HordePool ?? "(none)", hordeConfig.HordeCondition ?? "(none)", hasOidcProvider ? oidcProvider! : "Disabled", connectionMode?.ToString() ?? "(none)", encryption?.ToString() ?? "(none)");
 			try
 			{
 				if (String.IsNullOrEmpty(token) && hasOidcProvider)
@@ -503,7 +516,7 @@ namespace UnrealBuildTool
 
 				bool allowWine = hordeConfig.bHordeAllowWine && OperatingSystem.IsWindows();
 
-				UBAHordeSession session = new(executor, new Uri(server), authHeader, hordeConfig.HordePool, allowWine, hordeConfig.HordeMaxCores, bStrictErrors, connectionMode, logger);
+				UBAHordeSession session = new(executor, new Uri(server), authHeader, hordeConfig.HordePool, allowWine, hordeConfig.HordeMaxCores, bStrictErrors, connectionMode, encryption, logger);
 				await session.InitAsync(useSentry: !String.IsNullOrEmpty(hordeConfig.UBASentryUrl), cancellationToken);
 				return session;
 			}
