UnrealBuildAccelerator: Remove dependency on EpicGames.Core

#rnx

diff --git a/Engine/Source/Programs/Shared/EpicGames.UBA/EpicGames.UBA.csproj b/Engine/Source/Programs/Shared/EpicGames.UBA/EpicGames.UBA.csproj
--- a/Engine/Source/Programs/Shared/EpicGames.UBA/EpicGames.UBA.csproj
+++ b/Engine/Source/Programs/Shared/EpicGames.UBA/EpicGames.UBA.csproj
@@ -37,12 +37,9 @@
   </ItemGroup>
 
   <Import Project="Library.props" />
-
-  <ItemGroup>
-    <ProjectReference Include="..\..\Shared\EpicGames.Core\EpicGames.Core.csproj" />
-  </ItemGroup>
   <ItemGroup>
     <PackageReference Include="Microsoft.CSharp" Version="4.7.0" />
+    <PackageReference Include="Microsoft.Extensions.Logging" Version="6.0.0" />
     <PackageReference Include="Microsoft.Win32.Registry" Version="5.0.0" />
     <PackageReference Include="System.Data.DataSetExtensions" Version="4.5.0" />
   </ItemGroup>
diff --git a/Engine/Source/Programs/Shared/EpicGames.UBA/IProcess.cs b/Engine/Source/Programs/Shared/EpicGames.UBA/IProcess.cs
--- a/Engine/Source/Programs/Shared/EpicGames.UBA/IProcess.cs
+++ b/Engine/Source/Programs/Shared/EpicGames.UBA/IProcess.cs
@@ -3,7 +3,6 @@
 using System;
 using System.Collections.Generic;
 using System.Diagnostics;
-using EpicGames.Core;
 
 namespace EpicGames.UBA
 {
@@ -31,12 +30,12 @@ namespace EpicGames.UBA
 		/// <summary>
 		/// The path to the application binary
 		/// </summary>
-		public FileReference Application { get; set; }
+		public string Application { get; set; }
 
 		/// <summary>
 		/// The working directory
 		/// </summary>
-		public DirectoryReference WorkingDirectory { get; set; }
+		public string WorkingDirectory { get; set; }
 
 		/// <summary>
 		/// The command line arguments
diff --git a/Engine/Source/Programs/Shared/EpicGames.UBA/ISessionServer.cs b/Engine/Source/Programs/Shared/EpicGames.UBA/ISessionServer.cs
--- a/Engine/Source/Programs/Shared/EpicGames.UBA/ISessionServer.cs
+++ b/Engine/Source/Programs/Shared/EpicGames.UBA/ISessionServer.cs
@@ -1,7 +1,6 @@
 // Copyright Epic Games, Inc. All Rights Reserved.
 
 using System;
-using EpicGames.Core;
 
 namespace EpicGames.UBA
 {
@@ -13,12 +12,12 @@ namespace EpicGames.UBA
 		/// <summary>
 		/// Root directory to store content addressable data
 		/// </summary>
-		public DirectoryReference RootDirectory { get; init; }
+		public string RootDirectory { get; init; }
 
 		/// <summary>
 		/// Path to a trace file that records the build
 		/// </summary>
-		public FileReference TraceOutputFile { get; init; }
+		public string TraceOutputFile { get; init; }
 
 		/// <summary>
 		/// If the custom allocator should be disabled
@@ -67,7 +66,7 @@ namespace EpicGames.UBA
 		/// <param name="detailedTrace">More detailed trace information</param>
 		/// <param name="allowWaitOnMem">Wait for memory before starting new processes</param>
 		/// <param name="allowKillOnMem">Kill processes when close to run out of memory</param>
-		public SessionServerCreateInfo(DirectoryReference rootDirectory, FileReference traceOutputFile, bool disableCustomAllocator, bool launchVisualizer, bool resetCas, bool writeToDisk, bool detailedTrace, bool allowWaitOnMem, bool allowKillOnMem)
+		public SessionServerCreateInfo(string rootDirectory, string traceOutputFile, bool disableCustomAllocator, bool launchVisualizer, bool resetCas, bool writeToDisk, bool detailedTrace, bool allowWaitOnMem, bool allowKillOnMem)
 		{
 			RootDirectory = rootDirectory;
 			TraceOutputFile = traceOutputFile;
@@ -190,13 +189,13 @@ namespace EpicGames.UBA
 		/// Refresh cached information about directories
 		/// </summary>
 		/// <param name="directories">The directories to refresh</param>
-		public abstract void RefreshDirectories(params DirectoryReference[] directories);
+		public abstract void RefreshDirectories(params string[] directories);
 
 		/// <summary>
 		/// Registers external files write to session caches
 		/// </summary>
 		/// <param name="files">The files to register</param>
-		public abstract void RegisterNewFiles(params FileReference[] files);
+		public abstract void RegisterNewFiles(params string[] files);
 
 		/// <summary>
 		/// Registers the start of an external process
@@ -218,7 +217,7 @@ namespace EpicGames.UBA
 		/// <param name="file">The file to track</param>
 		/// <param name="workingDirectory">The working directory</param>
 		/// <param name="process">The process to get tracked inputs from</param>
-		public abstract void SetCustomCasKeyFromTrackedInputs(FileReference file, DirectoryReference workingDirectory, IProcess process);
+		public abstract void SetCustomCasKeyFromTrackedInputs(string file, string workingDirectory, IProcess process);
 
 		/// <summary>
 		/// Cancel all processes
diff --git a/Engine/Source/Programs/Shared/EpicGames.UBA/IStorageServer.cs b/Engine/Source/Programs/Shared/EpicGames.UBA/IStorageServer.cs
--- a/Engine/Source/Programs/Shared/EpicGames.UBA/IStorageServer.cs
+++ b/Engine/Source/Programs/Shared/EpicGames.UBA/IStorageServer.cs
@@ -1,7 +1,5 @@
 // Copyright Epic Games, Inc. All Rights Reserved.
 
-using EpicGames.Core;
-
 namespace EpicGames.UBA
 {
 	/// <summary>
@@ -12,7 +10,7 @@ namespace EpicGames.UBA
 		/// <summary>
 		/// The root directory for the storage
 		/// </summary>
-		public DirectoryReference RootDirectory { get; init; }
+		public string RootDirectory { get; init; }
 
 		/// <summary>
 		/// The capacity of the storage in bytes
@@ -36,7 +34,7 @@ namespace EpicGames.UBA
 		/// <param name="capacityBytes">The capacity of the storage in bytes</param>
 		/// <param name="storeCompressed">If the storage should be stored as compressed</param>
 		/// <param name="zone">The geographical zone this machine belongs to. Can be empty</param>
-		public StorageServerCreateInfo(DirectoryReference rootDirectory, ulong capacityBytes, bool storeCompressed, string zone)
+		public StorageServerCreateInfo(string rootDirectory, ulong capacityBytes, bool storeCompressed, string zone)
 		{
 			RootDirectory = rootDirectory;
 			CapacityBytes = capacityBytes;
diff --git a/Engine/Source/Programs/Shared/EpicGames.UBA/Impl/ProcessImpl.cs b/Engine/Source/Programs/Shared/EpicGames.UBA/Impl/ProcessImpl.cs
--- a/Engine/Source/Programs/Shared/EpicGames.UBA/Impl/ProcessImpl.cs
+++ b/Engine/Source/Programs/Shared/EpicGames.UBA/Impl/ProcessImpl.cs
@@ -29,7 +29,7 @@ namespace EpicGames.UBA
 				ExitCallbackDelegate = RaiseExited;
 			}
 
-			_handle = CreateProcessStartInfo(info.Application.FullName, info.Arguments, info.WorkingDirectory.FullName, info.Description, (uint)info.Priority, info.OutputStatsThresholdMs, info.TrackInputs, info.LogFile?? String.Empty, ExitCallbackDelegate);
+			_handle = CreateProcessStartInfo(info.Application, info.Arguments, info.WorkingDirectory, info.Description, (uint)info.Priority, info.OutputStatsThresholdMs, info.TrackInputs, info.LogFile?? String.Empty, ExitCallbackDelegate);
 		}
 
 		void RaiseExited(IntPtr userData, IntPtr handle)
diff --git a/Engine/Source/Programs/Shared/EpicGames.UBA/Impl/SessionServerImpl.cs b/Engine/Source/Programs/Shared/EpicGames.UBA/Impl/SessionServerImpl.cs
--- a/Engine/Source/Programs/Shared/EpicGames.UBA/Impl/SessionServerImpl.cs
+++ b/Engine/Source/Programs/Shared/EpicGames.UBA/Impl/SessionServerImpl.cs
@@ -4,7 +4,6 @@ using System;
 using System.Collections.Concurrent;
 using System.Collections.Generic;
 using System.Runtime.InteropServices;
-using EpicGames.Core;
 
 namespace EpicGames.UBA
 {
@@ -29,7 +28,7 @@ namespace EpicGames.UBA
 			_storage = storage;
 			_client = client;
 			_logger = logger;
-			_handle = CreateSessionServerCreateInfo(_storage.GetHandle(), _client.GetHandle(), _logger.GetHandle(), info.RootDirectory.FullName, info.TraceOutputFile.FullName, info.DisableCustomAllocator, info.LaunchVisualizer, info.ResetCas, info.WriteToDisk, info.DetailedTrace, info.AllowWaitOnMem, info.AllowKillOnMem);
+			_handle = CreateSessionServerCreateInfo(_storage.GetHandle(), _client.GetHandle(), _logger.GetHandle(), info.RootDirectory, info.TraceOutputFile, info.DisableCustomAllocator, info.LaunchVisualizer, info.ResetCas, info.WriteToDisk, info.DetailedTrace, info.AllowWaitOnMem, info.AllowKillOnMem);
 		}
 
 		#region IDisposable
@@ -179,14 +178,14 @@ namespace EpicGames.UBA
 
 		public void SetMaxRemoteProcessCount(uint count) => SessionServer_SetMaxRemoteProcessCount(_handle, count);
 
-		public void RefreshDirectories(params DirectoryReference[] directories) => Array.ForEach(directories, (directory) => SessionServer_RefreshDirectory(_handle, directory.FullName));
+		public void RefreshDirectories(params string[] directories) => Array.ForEach(directories, (directory) => SessionServer_RefreshDirectory(_handle, directory));
 
-		public void RegisterNewFiles(params FileReference[] files) => Array.ForEach(files, (file) => SessionServer_RegisterNewFile(_handle, file.FullName));
+		public void RegisterNewFiles(params string[] files) => Array.ForEach(files, (file) => SessionServer_RegisterNewFile(_handle, file));
 
 		public uint BeginExternalProcess(string description) => SessionServer_BeginExternalProcess(_handle, description);
 		public void EndExternalProcess(uint id, uint exitCode) => SessionServer_EndExternalProcess(_handle, id, exitCode);
 
-		public void SetCustomCasKeyFromTrackedInputs(FileReference file, DirectoryReference workingDirectory, IProcess process) => SessionServer_SetCustomCasKeyFromTrackedInputs(_handle, process.GetHandle(), file.FullName, workingDirectory.FullName);
+		public void SetCustomCasKeyFromTrackedInputs(string file, string workingDirectory, IProcess process) => SessionServer_SetCustomCasKeyFromTrackedInputs(_handle, process.GetHandle(), file, workingDirectory);
 
 		public void PrintSummary() => SessionServer_PrintSummary(_handle);
 		
diff --git a/Engine/Source/Programs/Shared/EpicGames.UBA/Impl/StorageServerImpl.cs b/Engine/Source/Programs/Shared/EpicGames.UBA/Impl/StorageServerImpl.cs
--- a/Engine/Source/Programs/Shared/EpicGames.UBA/Impl/StorageServerImpl.cs
+++ b/Engine/Source/Programs/Shared/EpicGames.UBA/Impl/StorageServerImpl.cs
@@ -29,7 +29,7 @@ namespace EpicGames.UBA
 		{
 			_server = server;
 			_logger = logger;
-			_handle = CreateStorageServer(_server.GetHandle(), info.RootDirectory.FullName, info.CapacityBytes, info.StoreCompressed, _logger.GetHandle(), info.Zone);
+			_handle = CreateStorageServer(_server.GetHandle(), info.RootDirectory, info.CapacityBytes, info.StoreCompressed, _logger.GetHandle(), info.Zone);
 		}
 
 		#region IDisposable
diff --git a/Engine/Source/Programs/UnrealBuildTool/Executors/Experimental/UBAExecutor.cs b/Engine/Source/Programs/UnrealBuildTool/Executors/Experimental/UBAExecutor.cs
--- a/Engine/Source/Programs/UnrealBuildTool/Executors/Experimental/UBAExecutor.cs
+++ b/Engine/Source/Programs/UnrealBuildTool/Executors/Experimental/UBAExecutor.cs
@@ -399,8 +399,8 @@ namespace UnrealBuildTool
 				using EpicGames.UBA.ILogger ubaLogger = EpicGames.UBA.ILogger.CreateLogger(logger);
 				using (Server = IServer.CreateServer(UBAConfig.MaxWorkers, UBAConfig.SendSize, ubaLogger, UBAConfig.bUseQuic))
 				{
-					using IStorageServer ubaStorageServer = IStorageServer.CreateStorageServer(Server, ubaLogger, new StorageServerCreateInfo(_rootDirRef, ((ulong)UBAConfig.StoreCapacityGb) * 1000 * 1000 * 1000, !UBAConfig.bStoreRaw, UBAConfig.Zone));
-					using ISessionServerCreateInfo serverCreateInfo = ISessionServerCreateInfo.CreateSessionServerCreateInfo(ubaStorageServer, Server, ubaLogger, new SessionServerCreateInfo(_rootDirRef, ubaTraceFile, UBAConfig.bDisableCustomAlloc, UBAConfig.bLaunchVisualizer, UBAConfig.bResetCas, UBAConfig.bWriteToDisk, UBAConfig.bDetailedTrace, !UBAConfig.bDisableWaitOnMem, UBAConfig.bAllowKillOnMem));
+					using IStorageServer ubaStorageServer = IStorageServer.CreateStorageServer(Server, ubaLogger, new StorageServerCreateInfo(_rootDirRef.FullName, ((ulong)UBAConfig.StoreCapacityGb) * 1000 * 1000 * 1000, !UBAConfig.bStoreRaw, UBAConfig.Zone));
+					using ISessionServerCreateInfo serverCreateInfo = ISessionServerCreateInfo.CreateSessionServerCreateInfo(ubaStorageServer, Server, ubaLogger, new SessionServerCreateInfo(_rootDirRef.FullName, ubaTraceFile.FullName, UBAConfig.bDisableCustomAlloc, UBAConfig.bLaunchVisualizer, UBAConfig.bResetCas, UBAConfig.bWriteToDisk, UBAConfig.bDetailedTrace, !UBAConfig.bDisableWaitOnMem, UBAConfig.bAllowKillOnMem));
 					using (_session = ISessionServer.CreateSessionServer(serverCreateInfo))
 					{
 
@@ -463,7 +463,7 @@ namespace UnrealBuildTool
 				{
 					if (refreshedDirectories.Add(output.Directory))
 					{
-						session.RefreshDirectories(output.Directory.Location);
+						session.RefreshDirectories(output.Directory.FullName);
 					}
 				}
 			};
@@ -604,8 +604,8 @@ namespace UnrealBuildTool
 		{
 			ProcessStartInfo startInfo = new()
 			{
-				Application = action.CommandPath,
-				WorkingDirectory = action.WorkingDirectory,
+				Application = action.CommandPath.FullName,
+				WorkingDirectory = action.WorkingDirectory.FullName,
 				Arguments = action.CommandArguments,
 				Priority = ProcessPriority,
 				OutputStatsThresholdMs = (uint)UBAConfig.OutputStatsThresholdMs,
@@ -649,7 +649,7 @@ namespace UnrealBuildTool
 
 					if (result.ExitCode == 0)
 					{
-						_session!.RegisterNewFiles(action.ProducedItems.Select(x => x.Location).ToArray());
+						_session!.RegisterNewFiles(action.ProducedItems.Select(x => x.FullName).ToArray());
 					}
 					_session!.EndExternalProcess(processId, (uint)result.ExitCode);
 
@@ -731,7 +731,7 @@ namespace UnrealBuildTool
 			bool success = results.ExitCode == 0;
 			if (pchItem != null && process != null && success)
 			{
-				_session!.SetCustomCasKeyFromTrackedInputs(pchItem.Location, action.WorkingDirectory, process);
+				_session!.SetCustomCasKeyFromTrackedInputs(pchItem.FullName, action.WorkingDirectory.FullName, process);
 			}
 
 			queue.OnActionCompleted(action, success, results);
