UnrealBuildAccelerator: Set a local zone when remote is not in use

#rnx

diff --git a/Engine/Source/Programs/UnrealBuildTool/Executors/Experimental/UBAExecutor.cs b/Engine/Source/Programs/UnrealBuildTool/Executors/Experimental/UBAExecutor.cs
--- a/Engine/Source/Programs/UnrealBuildTool/Executors/Experimental/UBAExecutor.cs
+++ b/Engine/Source/Programs/UnrealBuildTool/Executors/Experimental/UBAExecutor.cs
@@ -777,6 +777,7 @@ namespace UnrealBuildTool
 		{
 			UBAConfig.bDisableRemote = true;
 			UBAConfig.bForceBuildAllRemote = false;
+			UBAConfig.Zone = "local";
 		}
 	}
 }
