UnrealBuildTool: Set UE_VALIDATE_INTERNAL_API to control the define UE_INTERNAL which will be set to an empty macro for Engine modules and [[deprecated]] for non-engine modules. This will create a pch variant but will allow a way to stop (or at least discourage) the use of internal code that needs to be in public headers for performance or historical reasons, as well as a way to migrate internal code out of public headers over the span of a few releases

diff --git a/Engine/Source/Programs/UnrealBuildTool/Configuration/ModuleRules.cs b/Engine/Source/Programs/UnrealBuildTool/Configuration/ModuleRules.cs
--- a/Engine/Source/Programs/UnrealBuildTool/Configuration/ModuleRules.cs
+++ b/Engine/Source/Programs/UnrealBuildTool/Configuration/ModuleRules.cs
@@ -795,6 +795,16 @@ namespace UnrealBuildTool
 		}
 		private bool? bValidateFormatStringsPrivate;
 
+		/// <summary>
+		/// Emits deprecated warnings\errors for internal API usage for non-engine modules 
+		/// </summary>
+		public bool bValidateInternalApi
+		{
+			get => bValidateInternalApiPrivate ?? !bTreatAsEngineModule;
+			set => bValidateInternalApiPrivate = value;
+		}
+		private bool? bValidateInternalApiPrivate;
+
 		/// <summary>
 		/// Which engine version's build settings to use by default. 
 		/// </summary>
diff --git a/Engine/Source/Programs/UnrealBuildTool/Configuration/UEBuildModuleCPP.cs b/Engine/Source/Programs/UnrealBuildTool/Configuration/UEBuildModuleCPP.cs
--- a/Engine/Source/Programs/UnrealBuildTool/Configuration/UEBuildModuleCPP.cs
+++ b/Engine/Source/Programs/UnrealBuildTool/Configuration/UEBuildModuleCPP.cs
@@ -1052,6 +1052,7 @@ namespace UnrealBuildTool
 			{
 				if (Definition.Contains("UE_IS_ENGINE_MODULE") ||
 					Definition.Contains("UE_VALIDATE_FORMAT_STRINGS") ||
+					Definition.Contains("UE_VALIDATE_INTERNAL_API") ||
 					Definition.Contains("DEPRECATED_FORGAME") ||
 					Definition.Contains("UE_DEPRECATED_FORGAME") ||
 					Definition.Contains("UE_ENABLE_INCLUDE_ORDER_DEPRECATED_IN_"))
@@ -1134,6 +1135,14 @@ namespace UnrealBuildTool
 					Definitions.Add($"UE_VALIDATE_FORMAT_STRINGS={(ModuleCompileEnvironment.bValidateFormatStrings ? "1" : "0")}");
 				}
 
+				// Modify definitions if we need to create a new shared pch for validating internal api usage
+				if (ModuleCompileEnvironment.bValidateInternalApi != Template.BaseCompileEnvironment.bValidateInternalApi)
+				{
+					Definitions = new List<string>(Definitions);
+					Definitions.RemoveAll(x => x.Contains("UE_VALIDATE_INTERNAL_API", StringComparison.Ordinal));
+					Definitions.Add($"UE_VALIDATE_INTERNAL_API={(ModuleCompileEnvironment.bValidateInternalApi ? "1" : "0")}");
+				}
+
 				// Create a suffix to distinguish this shared PCH variant from any others. Currently only optimized and non-optimized shared PCHs are supported.
 				string Variant = GetSuffixForSharedPCH(ModuleCompileEnvironment, Template.BaseCompileEnvironment);
 
@@ -1257,6 +1266,11 @@ namespace UnrealBuildTool
 			{
 				return false;
 			}
+
+			if (ModuleCompileEnvironment.bValidateInternalApi != CompileEnvironment.bValidateInternalApi)
+			{
+				return false;
+			}
 			return true;
 		}
 
@@ -1314,6 +1328,18 @@ namespace UnrealBuildTool
 				}
 			}
 
+			if (CompileEnvironment.bValidateInternalApi != BaseCompileEnvironment.bValidateInternalApi)
+			{
+				if (CompileEnvironment.bValidateInternalApi)
+				{
+					Variant += ".ValApi";
+				}
+				else
+				{
+					Variant += ".NoValApi";
+				}
+			}
+
 			if (CompileEnvironment.bDeterministic != BaseCompileEnvironment.bDeterministic)
 			{
 				if (CompileEnvironment.bDeterministic)
@@ -1363,6 +1389,7 @@ namespace UnrealBuildTool
 			CompileEnvironment.bUseAutoRTFMCompiler = ModuleCompileEnvironment.bUseAutoRTFMCompiler;
 			CompileEnvironment.bAllowAutoRTFMInstrumentation = ModuleCompileEnvironment.bAllowAutoRTFMInstrumentation;
 			CompileEnvironment.bValidateFormatStrings = ModuleCompileEnvironment.bValidateFormatStrings;
+			CompileEnvironment.bValidateInternalApi = ModuleCompileEnvironment.bValidateInternalApi;
 		}
 
 		/// <summary>
@@ -1612,6 +1639,10 @@ namespace UnrealBuildTool
 									Writer.AppendLine("#undef " + DeprecationDefine);
 								}
 							}
+							if (Rules.bValidateInternalApi)
+							{
+								Writer.AppendLine("#undef UE_VALIDATE_INTERNAL_API");
+							}
 
 							// Only add new definitions that are not already existing in the shared pch
 							List<string> NewDefinitions = new();
@@ -1990,6 +2021,7 @@ namespace UnrealBuildTool
 			Result.IncludeOrderVersion = Rules.IncludeOrderVersion;
 			Result.DeterministicWarningLevel = Rules.DeterministicWarningLevel;
 			Result.bValidateFormatStrings = Rules.bValidateFormatStrings;
+			Result.bValidateInternalApi = Rules.bValidateInternalApi;
 			Result.bUseAutoRTFMCompiler = Target.bUseAutoRTFMCompiler;
 
 			CompileEnvironmentDebugInfoSettings(Target, Result);
@@ -2082,6 +2114,7 @@ namespace UnrealBuildTool
 			}
 
 			Result.Definitions.Add($"UE_VALIDATE_FORMAT_STRINGS={(Rules.bValidateFormatStrings ? "1" : "0")}");
+			Result.Definitions.Add($"UE_VALIDATE_INTERNAL_API={(Rules.bValidateInternalApi ? "1" : "0")}");
 
 			Result.Definitions.AddRange(EngineIncludeOrderHelper.GetDeprecationDefines(Rules.IncludeOrderVersion));
 
@@ -2130,6 +2163,7 @@ namespace UnrealBuildTool
 			CompileEnvironment.bOptimizeCode = ShouldEnableOptimization(ModuleRules.CodeOptimization.Default, Target.Configuration, Rules.bTreatAsEngineModule, Rules.bCodeCoverage);
 			CompileEnvironment.bCodeCoverage = Rules.bCodeCoverage;
 			CompileEnvironment.bValidateFormatStrings = Rules.bValidateFormatStrings;
+			CompileEnvironment.bValidateInternalApi = Rules.bValidateInternalApi;
 
 			// Override compile environment
 			CompileEnvironment.bIsBuildingDLL = !Target.ShouldCompileMonolithic();
@@ -2146,6 +2180,7 @@ namespace UnrealBuildTool
 			}
 
 			CompileEnvironment.Definitions.Add($"UE_VALIDATE_FORMAT_STRINGS={(Rules.bValidateFormatStrings ? "1" : "0")}");
+			CompileEnvironment.Definitions.Add($"UE_VALIDATE_INTERNAL_API={(Rules.bValidateInternalApi ? "1" : "0")}");
 
 			CompileEnvironment.Definitions.AddRange(EngineIncludeOrderHelper.GetDeprecationDefines(Rules.IncludeOrderVersion));
 
diff --git a/Engine/Source/Programs/UnrealBuildTool/System/CppCompileEnvironment.cs b/Engine/Source/Programs/UnrealBuildTool/System/CppCompileEnvironment.cs
--- a/Engine/Source/Programs/UnrealBuildTool/System/CppCompileEnvironment.cs
+++ b/Engine/Source/Programs/UnrealBuildTool/System/CppCompileEnvironment.cs
@@ -682,6 +682,11 @@ namespace UnrealBuildTool
 		/// </summary>
 		public bool bValidateFormatStrings = true;
 
+		/// <summary>
+		/// Emits compilation errors for incorrect UE_LOG format strings.
+		/// </summary>
+		public bool bValidateInternalApi = false;
+
 		/// <summary>
 		/// Directory where to put crash report files for platforms that support it
 		/// </summary>
@@ -797,6 +802,7 @@ namespace UnrealBuildTool
 			DeterministicWarningLevel = Other.DeterministicWarningLevel;
 			CrashDiagnosticDirectory = Other.CrashDiagnosticDirectory;
 			bValidateFormatStrings = Other.bValidateFormatStrings;
+			bValidateInternalApi = Other.bValidateInternalApi;
 		}
 
 		public CppCompileEnvironment(CppCompileEnvironment Other, UnrealArch OverrideArchitecture)
diff --git a/Engine/Source/Runtime/Core/Public/Misc/CoreMiscDefines.h b/Engine/Source/Runtime/Core/Public/Misc/CoreMiscDefines.h
--- a/Engine/Source/Runtime/Core/Public/Misc/CoreMiscDefines.h
+++ b/Engine/Source/Runtime/Core/Public/Misc/CoreMiscDefines.h
@@ -278,6 +278,12 @@ enum EInPlace {InPlace};
 	#define UE_DEPRECATED_FORGAME(...)
 #endif
 
+#if UE_VALIDATE_INTERNAL_API
+	#define UE_INTERNAL [[deprecated("Please remove usage of this internal API before upgrading to the next release, otherwise your project will no longer compile.")]]
+#else
+	#define UE_INTERNAL
+#endif
+
 /*
  * Macro that can be defined in the target file to strip deprecated properties in objects across the engine that check against this define.
  * Can be used by project that have migrated away from using deprecated functions and data members to potentially gain back some memory and perf.
