UnrealBuildTool: Don't bundle UbaVisualizer with EpicGames.UBA

#rnx

diff --git a/Engine/Source/Programs/UnrealBuildTool/Executors/Experimental/UBAExecutor.cs b/Engine/Source/Programs/UnrealBuildTool/Executors/Experimental/UBAExecutor.cs
--- a/Engine/Source/Programs/UnrealBuildTool/Executors/Experimental/UBAExecutor.cs
+++ b/Engine/Source/Programs/UnrealBuildTool/Executors/Experimental/UBAExecutor.cs
@@ -10,7 +10,9 @@ using System;
 using System.Collections.Concurrent;
 using System.Collections.Generic;
 using System.Linq;
+using System.Management;
 using System.Runtime.InteropServices;
+using System.Runtime.Versioning;
 using System.Threading;
 using System.Threading.Tasks;
 using UnrealBuildBase;
@@ -396,11 +398,16 @@ namespace UnrealBuildTool
 
 			try
 			{
+				if (UBAConfig.bLaunchVisualizer && OperatingSystem.IsWindows())
+				{
+					_ = Task.Run(LaunchVisualizer, cancellationTokenSource.Token);
+				}
+
 				using EpicGames.UBA.ILogger ubaLogger = EpicGames.UBA.ILogger.CreateLogger(logger);
 				using (Server = IServer.CreateServer(UBAConfig.MaxWorkers, UBAConfig.SendSize, ubaLogger, UBAConfig.bUseQuic))
 				{
 					using IStorageServer ubaStorageServer = IStorageServer.CreateStorageServer(Server, ubaLogger, new StorageServerCreateInfo(_rootDirRef.FullName, ((ulong)UBAConfig.StoreCapacityGb) * 1000 * 1000 * 1000, !UBAConfig.bStoreRaw, UBAConfig.Zone));
-					using ISessionServerCreateInfo serverCreateInfo = ISessionServerCreateInfo.CreateSessionServerCreateInfo(ubaStorageServer, Server, ubaLogger, new SessionServerCreateInfo(_rootDirRef.FullName, ubaTraceFile.FullName, UBAConfig.bDisableCustomAlloc, UBAConfig.bLaunchVisualizer, UBAConfig.bResetCas, UBAConfig.bWriteToDisk, UBAConfig.bDetailedTrace, !UBAConfig.bDisableWaitOnMem, UBAConfig.bAllowKillOnMem));
+					using ISessionServerCreateInfo serverCreateInfo = ISessionServerCreateInfo.CreateSessionServerCreateInfo(ubaStorageServer, Server, ubaLogger, new SessionServerCreateInfo(_rootDirRef.FullName, ubaTraceFile.FullName, UBAConfig.bDisableCustomAlloc, false, UBAConfig.bResetCas, UBAConfig.bWriteToDisk, UBAConfig.bDetailedTrace, !UBAConfig.bDisableWaitOnMem, UBAConfig.bAllowKillOnMem));
 					using (_session = ISessionServer.CreateSessionServer(serverCreateInfo))
 					{
 
@@ -442,6 +449,36 @@ namespace UnrealBuildTool
 			}
 		}
 
+		[SupportedOSPlatform("windows")]
+		static void LaunchVisualizer()
+		{
+			FileReference VisaulizerPath = FileReference.Combine(Unreal.EngineDirectory, "Binaries", "Win64", "UnrealBuildAccelerator", RuntimeInformation.ProcessArchitecture.ToString(), "UbaVisualizer.exe");
+			if (!FileReference.Exists(VisaulizerPath))
+			{
+				return;
+			}
+
+			// Check if a listening visualizer is alread running
+			try
+			{
+
+				foreach (System.Diagnostics.Process process in System.Diagnostics.Process.GetProcessesByName(VisaulizerPath.GetFileNameWithoutAnyExtensions()))
+				{
+					using ManagementObjectSearcher searcher = new ManagementObjectSearcher($"SELECT CommandLine FROM Win32_Process WHERE ProcessId = {process.Id}");
+					using ManagementObjectCollection objects = searcher.Get();
+					string args = objects.Cast<ManagementBaseObject>().SingleOrDefault()?["CommandLine"]?.ToString() ?? "";
+					if (args.Contains("-listen", StringComparison.OrdinalIgnoreCase))
+					{
+						return;
+					}
+				}
+			}
+			catch(Exception)
+			{
+			}
+			System.Diagnostics.Process.Start(VisaulizerPath.FullName, "-listen");
+		}
+
 		public override bool VerifyOutputs => UBAConfig.bWriteToDisk;
 
 		/// <summary>
