UnrealBuildTool: Run UbaVisualizer out of the temp directory

#rnx

diff --git a/Engine/Source/Programs/UnrealBuildTool/Executors/Experimental/UBAExecutor.cs b/Engine/Source/Programs/UnrealBuildTool/Executors/Experimental/UBAExecutor.cs
--- a/Engine/Source/Programs/UnrealBuildTool/Executors/Experimental/UBAExecutor.cs
+++ b/Engine/Source/Programs/UnrealBuildTool/Executors/Experimental/UBAExecutor.cs
@@ -452,17 +452,17 @@ namespace UnrealBuildTool
 		[SupportedOSPlatform("windows")]
 		static void LaunchVisualizer()
 		{
-			FileReference VisaulizerPath = FileReference.Combine(Unreal.EngineDirectory, "Binaries", "Win64", "UnrealBuildAccelerator", RuntimeInformation.ProcessArchitecture.ToString(), "UbaVisualizer.exe");
-			if (!FileReference.Exists(VisaulizerPath))
+			FileReference visaulizerPath = FileReference.Combine(Unreal.EngineDirectory, "Binaries", "Win64", "UnrealBuildAccelerator", RuntimeInformation.ProcessArchitecture.ToString(), "UbaVisualizer.exe");
+			FileReference tempPath = FileReference.Combine(new DirectoryReference(System.IO.Path.GetTempPath()), visaulizerPath.GetFileName());
+			if (!FileReference.Exists(visaulizerPath))
 			{
 				return;
 			}
 
-			// Check if a listening visualizer is alread running
 			try
 			{
-
-				foreach (System.Diagnostics.Process process in System.Diagnostics.Process.GetProcessesByName(VisaulizerPath.GetFileNameWithoutAnyExtensions()))
+				// Check if a listening visualizer is already running
+				foreach (System.Diagnostics.Process process in System.Diagnostics.Process.GetProcessesByName(visaulizerPath.GetFileNameWithoutAnyExtensions()))
 				{
 					using ManagementObjectSearcher searcher = new ManagementObjectSearcher($"SELECT CommandLine FROM Win32_Process WHERE ProcessId = {process.Id}");
 					using ManagementObjectCollection objects = searcher.Get();
@@ -472,11 +472,18 @@ namespace UnrealBuildTool
 						return;
 					}
 				}
+				if (!FileReference.Exists(tempPath) || tempPath.ToFileInfo().LastWriteTime < visaulizerPath.ToFileInfo().LastWriteTime)
+				{
+					FileReference.Copy(visaulizerPath, tempPath, true);
+				}
+				if (FileReference.Exists(tempPath))
+				{
+					System.Diagnostics.Process.Start(tempPath.FullName, "-listen");
+				}
 			}
 			catch(Exception)
 			{
 			}
-			System.Diagnostics.Process.Start(VisaulizerPath.FullName, "-listen");
 		}
 
 		public override bool VerifyOutputs => UBAConfig.bWriteToDisk;
