Headers in /Internal/ folders don't need their includes deprecated.

diff --git a/Engine/Source/Programs/UnrealBuildTool/Modes/IWYUMode.cs b/Engine/Source/Programs/UnrealBuildTool/Modes/IWYUMode.cs
--- a/Engine/Source/Programs/UnrealBuildTool/Modes/IWYUMode.cs
+++ b/Engine/Source/Programs/UnrealBuildTool/Modes/IWYUMode.cs
@@ -1323,6 +1323,7 @@ namespace UnrealBuildTool
 
 				bool IsCpp = Info.IsCpp;
 				bool IsPrivate = IsCpp || Info.File.Contains("/Private/");
+				bool IsInternal = Info.File.Contains("/Internal/");
 
 				// If we only want to update private files we early out for non-private headers
 				if (!IsPrivate && (bUpdateOnlyPrivate || (Info.Module?.Rules.IWYUSupport == IWYUSupport.KeepPublicAsIsForNow)))
@@ -1805,7 +1806,7 @@ namespace UnrealBuildTool
 				// If file is public, in engine and we have a deprecation tag set we will 
 				// add a deprecated include scope at the end of the file (unless scope already exists, then we'll add it inside that)
 				string EngineDir = Unreal.EngineDirectory.FullName.Replace('\\', '/');
-				if (!IsPrivate && Info.File.StartsWith(EngineDir) && !String.IsNullOrEmpty(HeaderDeprecationTag))
+				if (!(IsPrivate || IsInternal) && Info.File.StartsWith(EngineDir) && !String.IsNullOrEmpty(HeaderDeprecationTag))
 				{
 					Dictionary<string, string> PrintableToFull = new();
 					foreach (IWYUIncludeEntry Seen in Info.IncludesSeenInFile)
