TargetRules can filter referenced plugins to determine if it should be ignored from the build.

diff --git a/Engine/Source/Programs/UnrealBuildTool/Configuration/TargetRules.cs b/Engine/Source/Programs/UnrealBuildTool/Configuration/TargetRules.cs
--- a/Engine/Source/Programs/UnrealBuildTool/Configuration/TargetRules.cs
+++ b/Engine/Source/Programs/UnrealBuildTool/Configuration/TargetRules.cs
@@ -2656,6 +2656,12 @@ namespace UnrealBuildTool
 			return LinkType == TargetLinkType.Monolithic || BuildEnvironment == TargetBuildEnvironment.Unique;
 		}
 
+		/// <summary>
+		/// Checks if a plugin should be programmatically allowed in the build
+		/// </summary>
+		/// <returns>true if the plugin is allowed</returns>
+		public virtual bool IgnorePluginDependency(PluginInfo ParentInfo, PluginReferenceDescriptor Descriptor) => false;
+
 		/// <summary>
 		/// Checks if a property has been set with the RequiresUniqueBuildEnvironmentAttribute, and is different from it's base
 		/// </summary>
@@ -3894,6 +3900,12 @@ namespace UnrealBuildTool
 			return Inner.OptedInModulePlatforms == null || Inner.OptedInModulePlatforms.Contains(Platform);
 		}
 
+		/// <summary>
+		/// Checks if a plugin should be programmatically allowed in the build
+		/// </summary>
+		/// <returns>true if the plugin is allowed</returns>
+		public bool IgnorePluginDependency(PluginInfo ParentInfo, PluginReferenceDescriptor ChildDescriptor) => Inner.IgnorePluginDependency(ParentInfo, ChildDescriptor);
+
 		/// <summary>
 		/// Determines if the automation tests should be compiled based on the current configuration and optional forced settings.
 		/// If either the development tests or performance tests are compiled (unless explicitly disabled), this property returns true.
diff --git a/Engine/Source/Programs/UnrealBuildTool/Configuration/UEBuildTarget.cs b/Engine/Source/Programs/UnrealBuildTool/Configuration/UEBuildTarget.cs
--- a/Engine/Source/Programs/UnrealBuildTool/Configuration/UEBuildTarget.cs
+++ b/Engine/Source/Programs/UnrealBuildTool/Configuration/UEBuildTarget.cs
@@ -4545,6 +4545,13 @@ namespace UnrealBuildTool
 				{
 					foreach (PluginReferenceDescriptor NextReference in Info.Descriptor.Plugins)
 					{
+						// Ignore any plugin dependency which is programmatically filtered out by the target
+						if (Rules.IgnorePluginDependency(Info, NextReference))
+						{
+							Logger.LogDebug("Ignoring plugin '{0}' (referenced via {1}) due to it being filtered out by the target rules.", NextReference.Name, Info.Name);
+							continue;
+						}
+
 						UEBuildPlugin? NextInstance = AddPlugin(NextReference, PluginReferenceChain, ExcludeFolders, NameToInstance, NameToInfos, Logger);
 						if (NextInstance != null)
 						{
