UnrealBuildTool: Disable -Wshadow for clang 17 due to [REDACTED_PROFILE_URL]/llvm-project/issues/71976, and note that it should not be used as a preferred version until resolved

#rnx

diff --git a/Engine/Source/Programs/UnrealBuildTool/Platform/Windows/MicrosoftPlatformSDK.Versions.cs b/Engine/Source/Programs/UnrealBuildTool/Platform/Windows/MicrosoftPlatformSDK.Versions.cs
--- a/Engine/Source/Programs/UnrealBuildTool/Platform/Windows/MicrosoftPlatformSDK.Versions.cs
+++ b/Engine/Source/Programs/UnrealBuildTool/Platform/Windows/MicrosoftPlatformSDK.Versions.cs
@@ -28,6 +28,7 @@ namespace UnrealBuildTool
 		/// </summary>
 		static readonly VersionNumberRange[] PreferredClangVersions =
 		{
+			// Clang 17: [REDACTED_PROFILE_URL]/llvm-project/issues/71976 should not be used as a preferred version until this issue is resolved and the max affected version is added to ClangWarnings.cs
 			VersionNumberRange.Parse("16.0.0", "16.999"), // VS2022 17.7.x runtime requires Clang 16
 			VersionNumberRange.Parse("15.0.0", "15.999"), // VS2022 17.5.x runtime requires Clang 15
 		};
diff --git a/Engine/Source/Programs/UnrealBuildTool/ToolChain/ClangWarnings.cs b/Engine/Source/Programs/UnrealBuildTool/ToolChain/ClangWarnings.cs
--- a/Engine/Source/Programs/UnrealBuildTool/ToolChain/ClangWarnings.cs
+++ b/Engine/Source/Programs/UnrealBuildTool/ToolChain/ClangWarnings.cs
@@ -1,6 +1,7 @@
 ﻿// Copyright Epic Games, Inc. All Rights Reserved.
 
 using EpicGames.UHT.Utils;
+using Microsoft.Extensions.Logging;
 using System.Collections.Generic;
 
 namespace UnrealBuildTool
@@ -107,10 +108,18 @@ namespace UnrealBuildTool
 				}
 			}
 
-			// https://clang.llvm.org/docs/DiagnosticsReference.html#wshadow
-			if (CompileEnvironment.ShadowVariableWarningLevel != WarningLevel.Off)
+			// Clang 17 suffers from [REDACTED_PROFILE_URL]/llvm-project/issues/71976 and should not be used as a preferred version until resolved
+			if (ClangVersion >= new VersionNumber(17))
+			{
+				Arguments.Add("-Wno-shadow");
+			}
+			else
 			{
-				Arguments.Add("-Wshadow" + ((CompileEnvironment.ShadowVariableWarningLevel == WarningLevel.Error) ? "" : " -Wno-error=shadow"));
+				// https://clang.llvm.org/docs/DiagnosticsReference.html#wshadow
+				if (CompileEnvironment.ShadowVariableWarningLevel != WarningLevel.Off)
+				{
+					Arguments.Add("-Wshadow" + ((CompileEnvironment.ShadowVariableWarningLevel == WarningLevel.Error) ? "" : " -Wno-error=shadow"));
+				}
 			}
 
 			// https://clang.llvm.org/docs/DiagnosticsReference.html#wundef
