UnrealBuildTool: Ensure FileInlineGenCPPMap is set for the module compile environment

diff --git a/Engine/Source/Programs/UnrealBuildTool/Configuration/UEBuildModuleCPP.cs b/Engine/Source/Programs/UnrealBuildTool/Configuration/UEBuildModuleCPP.cs
--- a/Engine/Source/Programs/UnrealBuildTool/Configuration/UEBuildModuleCPP.cs
+++ b/Engine/Source/Programs/UnrealBuildTool/Configuration/UEBuildModuleCPP.cs
@@ -530,6 +530,7 @@ namespace UnrealBuildTool
 							}
 						}
 					}
+					ModuleCompileEnvironment.FileInlineGenCPPMap = CompileEnvironment.FileInlineGenCPPMap;
 
 					if (Rules.bEnableNonInlinedGenCppWarnings)
 					{
