- Disabled UBA for Mac linking

diff --git a/Engine/Source/Programs/UnrealBuildTool/Platform/Mac/MacToolChain.cs b/Engine/Source/Programs/UnrealBuildTool/Platform/Mac/MacToolChain.cs
--- a/Engine/Source/Programs/UnrealBuildTool/Platform/Mac/MacToolChain.cs
+++ b/Engine/Source/Programs/UnrealBuildTool/Platform/Mac/MacToolChain.cs
@@ -880,6 +880,7 @@ namespace UnrealBuildTool
 
 			// Only execute linking on the local Mac.
 			LinkAction.bCanExecuteRemotely = false;
+			LinkAction.bCanExecuteInUBA = false;
 
 			LinkAction.StatusDescription = Path.GetFileName(OutputFile.AbsolutePath);
 
