iOS15+ now requires an entitlement key to enable GameCenter.

diff --git a/Engine/Source/Programs/UnrealBuildTool/Platform/IOS/IOSExports.cs b/Engine/Source/Programs/UnrealBuildTool/Platform/IOS/IOSExports.cs
--- a/Engine/Source/Programs/UnrealBuildTool/Platform/IOS/IOSExports.cs
+++ b/Engine/Source/Programs/UnrealBuildTool/Platform/IOS/IOSExports.cs
@@ -341,6 +341,15 @@ namespace UnrealBuildTool
 					Text.AppendLine("\t<array><string>runs-as-current-user</string></array>");
 				}
 
+				// As of iOS15, to support GameCenter, "com.apple.developer.game-center" entitlement must be set
+				bool bEnableGameCenterSupport = false;
+				PlatformGameConfig.GetBool("/Script/IOSRuntimeSettings.IOSRuntimeSettings", "bEnableGameCenterSupport", out bEnableGameCenterSupport);
+				if (bEnableGameCenterSupport)
+				{
+					Text.AppendLine("\t<key>com.apple.developer.game-center</key>");
+					Text.AppendLine("\t<true/>");
+				}
+
 				// End of entitlements
 				Text.AppendLine("</dict>");
 				Text.AppendLine("</plist>");
