Fix minor typo

#rnx

diff --git a/Engine/Source/Programs/UnrealBuildTool/Executors/Experimental/UBAExecutor.cs b/Engine/Source/Programs/UnrealBuildTool/Executors/Experimental/UBAExecutor.cs
--- a/Engine/Source/Programs/UnrealBuildTool/Executors/Experimental/UBAExecutor.cs
+++ b/Engine/Source/Programs/UnrealBuildTool/Executors/Experimental/UBAExecutor.cs
@@ -692,7 +692,7 @@ namespace UnrealBuildTool
 				{
 					if (process.ExitCode != 0 && UBAConfig.bForcedRetry)
 					{
-						_threadedLogger.LogWarning("{Description} {StatusDescription}: Exited with error code {ExitCode}. This action will retry witout UBA", action.CommandDescription, action.StatusDescription, process.ExitCode);
+						_threadedLogger.LogWarning("{Description} {StatusDescription}: Exited with error code {ExitCode}. This action will retry without UBA", action.CommandDescription, action.StatusDescription, process.ExitCode);
 						_forcedRetryActions.AddOrUpdate(action, false, (k, v) => false);
 						queue.RequeueAction(action);
 						return Task.CompletedTask;
