Fix FindProjectsForModule for special case where no game module exists

diff --git a/Engine/Source/Programs/UnrealBuildTool/ProjectFiles/ProjectFileGenerator.cs b/Engine/Source/Programs/UnrealBuildTool/ProjectFiles/ProjectFileGenerator.cs
--- a/Engine/Source/Programs/UnrealBuildTool/ProjectFiles/ProjectFileGenerator.cs
+++ b/Engine/Source/Programs/UnrealBuildTool/ProjectFiles/ProjectFileGenerator.cs
@@ -2580,7 +2580,7 @@ namespace UnrealBuildTool
 								return FoundProjects;
 							}
 
-							ModuleDescriptor? CurModuleDescriptor = GameDescriptor.Modules!.FirstOrDefault(x => x.Name == CurModuleFile.GetFileNameWithoutAnyExtensions());
+							ModuleDescriptor? CurModuleDescriptor = GameDescriptor.Modules?.FirstOrDefault(x => x.Name == CurModuleFile.GetFileNameWithoutAnyExtensions());
 
 							// find the project that the module is under, and has a TargetType target (useful with bMakeProjectPerTarget)
 							foreach (KeyValuePair<FileReference, ProjectFile> Pair in ProjectFileMap)
