UBA: Ensure the visualizer spawns detached

#rnx

diff --git a/Engine/Source/Programs/UnrealBuildTool/Executors/Experimental/UBAExecutor.cs b/Engine/Source/Programs/UnrealBuildTool/Executors/Experimental/UBAExecutor.cs
--- a/Engine/Source/Programs/UnrealBuildTool/Executors/Experimental/UBAExecutor.cs
+++ b/Engine/Source/Programs/UnrealBuildTool/Executors/Experimental/UBAExecutor.cs
@@ -478,7 +478,13 @@ namespace UnrealBuildTool
 				}
 				if (FileReference.Exists(tempPath))
 				{
-					System.Diagnostics.Process.Start(tempPath.FullName, "-listen");
+					System.Diagnostics.ProcessStartInfo psi = new(BuildHostPlatform.Current.Shell.FullName, $" /C start \"\" \"{tempPath.FullName}\" -listen")
+					{
+						WorkingDirectory = System.IO.Path.GetTempPath(),
+						WindowStyle = System.Diagnostics.ProcessWindowStyle.Hidden,
+						UseShellExecute = true,
+					};
+					System.Diagnostics.Process.Start(psi);
 				}
 			}
 			catch(Exception)
