UnrealBuildTool: Support -LoadTargets for Querying available targets that will load each Target.cs so the actual list of supported platforms and configurations can be returned

#rnx

diff --git a/Engine/Source/Programs/UnrealBuildTool/Configuration/TargetRules.cs b/Engine/Source/Programs/UnrealBuildTool/Configuration/TargetRules.cs
--- a/Engine/Source/Programs/UnrealBuildTool/Configuration/TargetRules.cs
+++ b/Engine/Source/Programs/UnrealBuildTool/Configuration/TargetRules.cs
@@ -3876,6 +3876,16 @@ namespace UnrealBuildTool
 		/// </summary>
 		public bool WithLowLevelTests => Inner.WithLowLevelTests;
 
+		/// <summary>
+		/// Get the platforms this target supports
+		/// </summary>
+		public IEnumerable<UnrealTargetPlatform> SupportedPlatforms => Inner.GetSupportedPlatforms();
+
+		/// <summary>
+		/// Get the configurations this target supports
+		/// </summary>
+		public IEnumerable<UnrealTargetConfiguration> SupportedConfigurations => Inner.GetSupportedConfigurations();
+
 		/// <summary>
 		/// Checks if current platform is part of a given platform group
 		/// </summary>
diff --git a/Engine/Source/Programs/UnrealBuildTool/Modes/QueryMode.cs b/Engine/Source/Programs/UnrealBuildTool/Modes/QueryMode.cs
--- a/Engine/Source/Programs/UnrealBuildTool/Modes/QueryMode.cs
+++ b/Engine/Source/Programs/UnrealBuildTool/Modes/QueryMode.cs
@@ -10,6 +10,7 @@ using System.Text.Json;
 using System.Threading.Tasks;
 using EpicGames.Core;
 using Microsoft.Extensions.Logging;
+using Microsoft.Extensions.Logging.Abstractions;
 using OpenTracing.Util;
 using UnrealBuildBase;
 
@@ -76,7 +77,9 @@ namespace UnrealBuildTool
 
 	internal class TargetConfigs
 	{
-		public string ProjectPath { get; set; } = "";
+		public string TargetPath { get; set; } = String.Empty;
+		public string ProjectPath { get; set; } = String.Empty;
+		public string TargetType { get; set; } = String.Empty;
 		public List<string> Platforms { get; set; } = new();
 		public List<string> Configurations { get; set; } = new();
 	}
@@ -93,6 +96,9 @@ namespace UnrealBuildTool
 		[CommandLine("-Query=")]
 		public QueryType? Query = null;
 
+		[CommandLine("-LoadTargets")]
+		public bool bLoadTargets = false;
+
 		[CommandLine("-IncludeEngineSource=")]
 		public bool bIncludeEngineSource = true;
 
@@ -213,14 +219,14 @@ namespace UnrealBuildTool
 				List<UnrealTargetConfiguration> AllowedTargetConfigurations = new List<UnrealTargetConfiguration>();
 				AllowedTargetConfigurations = Enum.GetValues(typeof(UnrealTargetConfiguration)).Cast<UnrealTargetConfiguration>().ToList();
 
-				List<string> Configurations = new();
+				List<UnrealTargetConfiguration> Configurations = new();
 				foreach (UnrealTargetConfiguration CurConfiguration in AllowedTargetConfigurations)
 				{
 					if (CurConfiguration != UnrealTargetConfiguration.Unknown)
 					{
 						if (InstalledPlatformInfo.IsValidConfiguration(CurConfiguration, EProjectType.Code))
 						{
-							Configurations.Add(CurConfiguration.ToString());
+							Configurations.Add(CurConfiguration);
 						}
 					}
 				}
@@ -228,17 +234,66 @@ namespace UnrealBuildTool
 				List<FileReference> Projects = ProjectFileArg != null ? new List<FileReference>(new[] { ProjectFileArg }) : NativeProjects.EnumerateProjectFiles(Logger).ToList();
 				List<FileReference> AllTargetFiles = ProjectFileGenerator.DiscoverTargets(Projects, Logger, null, Platforms, bIncludeEngineSource: bIncludeEngineSource, bIncludeTempTargets: false);
 
-				// TODO: Check valid configurations/platforms per target 
 				Dictionary<string, TargetConfigs> Targets = new();
 				string? DefaultTarget = null;
 				foreach (FileReference TargetFilePath in AllTargetFiles)
 				{
-					string TargetName = TargetFilePath.GetFileNameWithoutAnyExtensions();
-					FileReference? ProjectPath = Projects.FirstOrDefault(p => TargetFilePath.IsUnderDirectory(p.Directory));
-					Targets.Add(TargetName, new TargetConfigs() { ProjectPath = ProjectPath?.ToString() ?? "", Configurations = Configurations, Platforms = Platforms.Select(x => x.ToString()).ToList() });
-					if (DefaultTarget == null || TargetName == "UnrealEditor")
+					try
+					{
+						FileReference? ProjectPath = Projects.FirstOrDefault(p => TargetFilePath.IsUnderDirectory(p.Directory));
+						string TargetType = String.Empty;
+						List<UnrealTargetConfiguration> SupportedConfigurations = new(Configurations);
+						List<UnrealTargetPlatform> SupportedPlatforms = new(Platforms);
+						if (bLoadTargets)
+						{
+							List<string> RawArgs = new List<string> { TargetFilePath.GetFileNameWithoutAnyExtensions(), UnrealTargetConfiguration.Development.ToString(), UnrealTargetPlatform.Win64.ToString() };
+							if (ProjectPath != null)
+							{
+								RawArgs.Add($"-Project={ProjectPath.FullName}");
+							}
+							CommandLineArguments Args = new CommandLineArguments(RawArgs.ToArray());
+
+							List<TargetDescriptor> TargetDescriptors = new();
+							TargetDescriptor.ParseSingleCommandLine(Args, false, false, false, TargetDescriptors, NullLogger.Instance);
+
+							// Ensure the intermediate environment does not conflict with normal builds
+							TargetDescriptors[0].IntermediateEnvironment = UnrealIntermediateEnvironment.Query;
+						
+
+							UEBuildTarget CurrentTarget;
+							using (GlobalTracer.Instance.BuildSpan("UEBuildTarget.Create()").StartActive())
+							{
+								bool bUsePrecompiled = false;
+
+								// Prevent multiple conflicting processes building TargetRules at the same time
+								string MutexName = SingleInstanceMutex.GetUniqueMutexForPath("UnrealBuildTool_QueryMode_UEBuildTarget-Create", Unreal.RootDirectory.FullName);
+								using (new SingleInstanceMutex(MutexName, true))
+								{
+									CurrentTarget = UEBuildTarget.Create(TargetDescriptors[0], false, false, bUsePrecompiled, TargetDescriptors[0].IntermediateEnvironment, NullLogger.Instance);
+								}
+							}
+							TargetType = CurrentTarget.Rules.Type.ToString();
+							SupportedConfigurations = CurrentTarget.Rules.SupportedConfigurations.ToList();
+							SupportedPlatforms = CurrentTarget.Rules.SupportedPlatforms.ToList();
+						}
+
+						string TargetName = TargetFilePath.GetFileNameWithoutAnyExtensions();
+						Targets.Add(TargetName, new TargetConfigs()
+						{
+							TargetPath = TargetFilePath.FullName,
+							ProjectPath = ProjectPath?.ToString() ?? String.Empty,
+							TargetType = TargetType,
+							Configurations = SupportedConfigurations.Select(x => x.ToString()).ToList(),
+							Platforms = SupportedPlatforms.Select(x => x.ToString()).ToList()
+						});
+						if (DefaultTarget == null || TargetName == "UnrealEditor")
+						{
+							DefaultTarget = TargetName;
+						}
+					}
+					catch (Exception)
 					{
-						DefaultTarget = TargetName;
+						continue;
 					}
 				}
 
