- Don't run any shell-based tasks through UBA on Mac

diff --git a/Engine/Source/Programs/UnrealBuildTool/Executors/Experimental/UBAExecutor.cs b/Engine/Source/Programs/UnrealBuildTool/Executors/Experimental/UBAExecutor.cs
--- a/Engine/Source/Programs/UnrealBuildTool/Executors/Experimental/UBAExecutor.cs
+++ b/Engine/Source/Programs/UnrealBuildTool/Executors/Experimental/UBAExecutor.cs
@@ -692,7 +692,13 @@ namespace UnrealBuildTool
 
 			return () =>
 			{
-				bool enableDetour = action.bCanExecuteInUBA && !_forcedRetryActions.ContainsKey(action);
+				bool enableDetour = action.bCanExecuteInUBA && !_forcedRetryActions.ContainsKey(action) && 
+					// Don't let Mac run shell commands through Uba as interposing dylibs into
+					// the shell results in dyld errors about no matching architecture.
+					// The shell is used to run various commands during a build like copy/ditto.
+					//So for these actions we need to make sure UBA is not used.
+					(!System.OperatingSystem.IsMacOS() || action.CommandPath != BuildHostPlatform.Current.Shell);
+
 				ProcessStartInfo startInfo = GetActionStartInfo(action, out FileItem? pchItem);
 				using (IProcess process = _session!.RunProcess(startInfo, false, null, enableDetour))
 				{
