Fixed bUseInlining in .Target.cs, caused by CL# 28288993.

diff --git a/Engine/Source/Programs/UnrealBuildTool/Platform/Windows/VCToolChain.cs b/Engine/Source/Programs/UnrealBuildTool/Platform/Windows/VCToolChain.cs
--- a/Engine/Source/Programs/UnrealBuildTool/Platform/Windows/VCToolChain.cs
+++ b/Engine/Source/Programs/UnrealBuildTool/Platform/Windows/VCToolChain.cs
@@ -563,7 +563,7 @@ namespace UnrealBuildTool
 			// Do not allow inline method expansion if E&C support is enabled or inline expansion has been disabled, 
 			// or if we are compiling in a debug build with `clang-cl`, since this will interfere with debugging capabilities.
 			if (!CompileEnvironment.bSupportEditAndContinue && CompileEnvironment.bUseInlining
-				|| !(Target.WindowsPlatform.Compiler.IsClang() && CompileEnvironment.Configuration == CppConfiguration.Debug))
+				&& !(Target.WindowsPlatform.Compiler.IsClang() && CompileEnvironment.Configuration == CppConfiguration.Debug))
 			{
 				Arguments.Add($"/Ob{Math.Clamp(Target.WindowsPlatform.InlineFunctionExpansionLevel, 1, 3)}");
 			}
