Android also skips when the lib end matches the architecture without a dash (i.e. x64 and -x64)

diff --git a/Engine/Source/Programs/UnrealBuildTool/Platform/Android/AndroidToolChain.cs b/Engine/Source/Programs/UnrealBuildTool/Platform/Android/AndroidToolChain.cs
--- a/Engine/Source/Programs/UnrealBuildTool/Platform/Android/AndroidToolChain.cs
+++ b/Engine/Source/Programs/UnrealBuildTool/Platform/Android/AndroidToolChain.cs
@@ -930,7 +930,8 @@ namespace UnrealBuildTool
 			{
 				if (ComboName.Key != Arch)
 				{
-					if (Lib.EndsWith(ComboName.Value))
+					string ArchitectureName = ComboName.Key.ToString();
+					if (Lib.EndsWith(ComboName.Value) || Lib.EndsWith(ArchitectureName))
 					{
 						return true;
 					}
