UnrealBuildTool: Fix some seealso tags

#rnx

diff --git a/Engine/Source/Programs/UnrealBuildTool/Platform/Windows/UEBuildWindows.cs b/Engine/Source/Programs/UnrealBuildTool/Platform/Windows/UEBuildWindows.cs
--- a/Engine/Source/Programs/UnrealBuildTool/Platform/Windows/UEBuildWindows.cs
+++ b/Engine/Source/Programs/UnrealBuildTool/Platform/Windows/UEBuildWindows.cs
@@ -233,7 +233,7 @@ namespace UnrealBuildTool
 		/// <summary>
 		/// If specified along with -PGOProfile, prevent the usage of extra counters. Please note that by default /FASTGENPROFILE doesnt use extra counters
 		/// </summary>
-		/// <seealso href="link">https://learn.microsoft.com/en-us/cpp/build/reference/genprofile-fastgenprofile-generate-profiling-instrumented-build</seealso>
+		/// <seealso href="https://learn.microsoft.com/en-us/cpp/build/reference/genprofile-fastgenprofile-generate-profiling-instrumented-build">genprofile-fastgenprofile-generate-profiling-instrumented-build</seealso>
 		[XmlConfigFile(Category = "WindowsPlatform")]
 		[CommandLine("-PGONoExtraCounters")]
 		public bool bPGONoExtraCounters = false;
@@ -241,7 +241,7 @@ namespace UnrealBuildTool
 		/// <summary>
 		/// Which level to use for Inline Function Expansion when TargetRules.bUseInlining is enabled
 		/// </summary>
-		/// <seealso href="link">https://learn.microsoft.com/en-us/cpp/build/reference/ob-inline-function-expansion</seealso>
+		/// <seealso href="https://learn.microsoft.com/en-us/cpp/build/reference/ob-inline-function-expansion">ob-inline-function-expansion</seealso>
 		[ConfigFile(ConfigHierarchyType.Engine, "/Script/WindowsTargetPlatform.WindowsTargetSettings", "InlineFunctionExpansionLevel")]
 		[XmlConfigFile(Category = "WindowsPlatform")]
 		public int InlineFunctionExpansionLevel { get; set; } = 2;
