UBT: Allow toolchains to provide system include path information for Rider project generation.

#rnx

diff --git a/Engine/Source/Programs/UnrealBuildTool/ProjectFiles/PlatformProjectGenerator.cs b/Engine/Source/Programs/UnrealBuildTool/ProjectFiles/PlatformProjectGenerator.cs
--- a/Engine/Source/Programs/UnrealBuildTool/ProjectFiles/PlatformProjectGenerator.cs
+++ b/Engine/Source/Programs/UnrealBuildTool/ProjectFiles/PlatformProjectGenerator.cs
@@ -364,5 +364,15 @@ namespace UnrealBuildTool
 		public virtual void GetVisualStudioTargetsString(UnrealTargetPlatform InPlatform, List<ProjectBuildConfiguration> InPlatformConfigurations, StringBuilder ProjectFileBuilder)
 		{
 		}
+
+		/// <summary>
+		/// Get include paths to system headers (STL, SDK, ...) needed for coding assistance features in IDE
+		/// </summary>
+		/// <param name="InTarget">Target for which include paths shall be returned</param>
+		/// <returns>The list of include paths</returns>
+		public virtual IList<string> GetSystemIncludePaths(UEBuildTarget InTarget)
+		{
+			return new List<string>(0);
+		}
 	}
 }
diff --git a/Engine/Source/Programs/UnrealBuildTool/ProjectFiles/Rider/RiderProjectFile.cs b/Engine/Source/Programs/UnrealBuildTool/ProjectFiles/Rider/RiderProjectFile.cs
--- a/Engine/Source/Programs/UnrealBuildTool/ProjectFiles/Rider/RiderProjectFile.cs
+++ b/Engine/Source/Programs/UnrealBuildTool/ProjectFiles/Rider/RiderProjectFile.cs
@@ -518,6 +518,15 @@ namespace UnrealBuildTool
 			{
 				Writer.WriteValue(Path.FullName);
 			}
+			
+			PlatformProjectGenerator? ProjGenerator = PlatformProjectGenerators.GetPlatformProjectGenerator(Target.Platform, true);
+			if (ProjGenerator != null)
+			{
+				foreach (string Path in ProjGenerator.GetSystemIncludePaths(Target))
+				{
+					Writer.WriteValue(Path);
+				}
+			}
 
 			if (UEBuildPlatform.IsPlatformInGroup(Target.Platform, UnrealPlatformGroup.Windows))
 			{
