UBT: add "bVcRemoveUnreferencedComdat" ModuleRules property to control whether /Zc:inline is enabled with MSVC

diff --git a/Engine/Source/Programs/UnrealBuildTool/Configuration/ModuleRules.cs b/Engine/Source/Programs/UnrealBuildTool/Configuration/ModuleRules.cs
--- a/Engine/Source/Programs/UnrealBuildTool/Configuration/ModuleRules.cs
+++ b/Engine/Source/Programs/UnrealBuildTool/Configuration/ModuleRules.cs
@@ -841,6 +841,12 @@ namespace UnrealBuildTool
 		/// </summary>
 		public bool bUseRTTI = false;
 
+		/// <summary>
+		/// Whether to direct MSVC to remove unreferenced COMDAT functions and data.
+		/// </summary>
+		/// <seealso href="https://learn.microsoft.com/en-us/cpp/build/reference/zc-inline-remove-unreferenced-comdat">zc-inline-remove-unreferenced-comdat</seealso>
+		public bool bVcRemoveUnreferencedComdat = true;
+
 		/// <summary>
 		/// Enable code coverage compilation/linking support.
 		/// </summary>
diff --git a/Engine/Source/Programs/UnrealBuildTool/Configuration/UEBuildModuleCPP.cs b/Engine/Source/Programs/UnrealBuildTool/Configuration/UEBuildModuleCPP.cs
--- a/Engine/Source/Programs/UnrealBuildTool/Configuration/UEBuildModuleCPP.cs
+++ b/Engine/Source/Programs/UnrealBuildTool/Configuration/UEBuildModuleCPP.cs
@@ -1993,6 +1993,7 @@ namespace UnrealBuildTool
 			Result.bCodeCoverage = Target.bCodeCoverage;
 			Result.bOptimizeCode = ShouldEnableOptimization(Rules.OptimizeCode, Target.Configuration, Rules.bTreatAsEngineModule, Result.bCodeCoverage);
 			Result.bUseRTTI |= Rules.bUseRTTI;
+			Result.bVcRemoveUnreferencedComdat &= Rules.bVcRemoveUnreferencedComdat;
 			Result.bEnableBufferSecurityChecks = Rules.bEnableBufferSecurityChecks;
 			Result.MinSourceFilesForUnityBuildOverride = Rules.MinSourceFilesForUnityBuildOverride;
 			Result.MinFilesUsingPrecompiledHeaderOverride = Rules.MinFilesUsingPrecompiledHeaderOverride;
diff --git a/Engine/Source/Programs/UnrealBuildTool/Platform/Windows/VCToolChain.cs b/Engine/Source/Programs/UnrealBuildTool/Platform/Windows/VCToolChain.cs
--- a/Engine/Source/Programs/UnrealBuildTool/Platform/Windows/VCToolChain.cs
+++ b/Engine/Source/Programs/UnrealBuildTool/Platform/Windows/VCToolChain.cs
@@ -318,9 +318,16 @@ namespace UnrealBuildTool
 
 		protected virtual void AppendCLArguments_Global(CppCompileEnvironment CompileEnvironment, List<string> Arguments)
 		{
-			// Suppress generation of object code for unreferenced inline functions. Enabling this option is more standards compliant, and causes a big reduction
-			// in object file sizes (and link times) due to the amount of stuff we inline.
-			Arguments.Add("/Zc:inline");
+			if (CompileEnvironment.bVcRemoveUnreferencedComdat)
+			{
+				// Suppress generation of object code for unreferenced inline functions. Enabling this option is more standards compliant, and causes a big reduction
+				// in object file sizes (and link times) due to the amount of stuff we inline.
+				Arguments.Add("/Zc:inline");
+			}
+			else
+			{
+				Arguments.Add("/Zc:inline-");
+			}
 
 			// @todo clang: Clang on Windows doesn't respect "#pragma warning (error: ####)", and we're not passing "/WX", so warnings are not
 			// treated as errors when compiling on Windows using Clang right now.
diff --git a/Engine/Source/Programs/UnrealBuildTool/System/CppCompileEnvironment.cs b/Engine/Source/Programs/UnrealBuildTool/System/CppCompileEnvironment.cs
--- a/Engine/Source/Programs/UnrealBuildTool/System/CppCompileEnvironment.cs
+++ b/Engine/Source/Programs/UnrealBuildTool/System/CppCompileEnvironment.cs
@@ -244,6 +244,12 @@ namespace UnrealBuildTool
 		/// </summary>
 		public bool bUseRTTI = false;
 
+		/// <summary>
+		/// Whether to direct MSVC to remove unreferenced COMDAT functions and data.
+		/// </summary>
+		/// <seealso href="https://learn.microsoft.com/en-us/cpp/build/reference/zc-inline-remove-unreferenced-comdat">zc-inline-remove-unreferenced-comdat</seealso>
+		public bool bVcRemoveUnreferencedComdat = true;
+
 		/// <summary>
 		/// Use Position Independent Executable (PIE). Has an overhead cost
 		/// </summary>
@@ -724,6 +730,7 @@ namespace UnrealBuildTool
 			PrecompiledHeaderAction = Other.PrecompiledHeaderAction;
 			bUseSharedBuildEnvironment = Other.bUseSharedBuildEnvironment;
 			bUseRTTI = Other.bUseRTTI;
+			bVcRemoveUnreferencedComdat = Other.bVcRemoveUnreferencedComdat;
 			bUsePIE = Other.bUsePIE;
 			bUseStackProtection = Other.bUseStackProtection;
 			bUseInlining = Other.bUseInlining;
