[UBA]
* Added code to pre-feed remotely executed processes with known input

diff --git a/Engine/Source/Programs/Shared/EpicGames.UBA/ISessionServer.cs b/Engine/Source/Programs/Shared/EpicGames.UBA/ISessionServer.cs
--- a/Engine/Source/Programs/Shared/EpicGames.UBA/ISessionServer.cs
+++ b/Engine/Source/Programs/Shared/EpicGames.UBA/ISessionServer.cs
@@ -183,8 +183,10 @@ namespace EpicGames.UBA
 		/// <param name="info">Process start info</param>
 		/// <param name="exitedEventHandler">Optional callback when the process exits</param>
 		/// <param name="weight">Number of cores this process uses</param>
+		/// <param name="knownInputs">Optionally contains input that we know process will need. Memory block containing zero-terminated strings with an extra termination in the end.</param>
+		/// <param name="knownInputsCount">Number of strings in known inputs</param>
 		/// <returns>The remote process being run</returns>
-		public abstract IProcess RunProcessRemote(ProcessStartInfo info, IProcess.ExitedEventHandler? exitedEventHandler, double weight = 1.0);
+		public abstract IProcess RunProcessRemote(ProcessStartInfo info, IProcess.ExitedEventHandler? exitedEventHandler, double weight = 1.0, byte[]? knownInputs = null, uint knownInputsCount = 0);
 
 		/// <summary>
 		/// Refresh cached information about directories
diff --git a/Engine/Source/Programs/Shared/EpicGames.UBA/Impl/SessionServerImpl.cs b/Engine/Source/Programs/Shared/EpicGames.UBA/Impl/SessionServerImpl.cs
--- a/Engine/Source/Programs/Shared/EpicGames.UBA/Impl/SessionServerImpl.cs
+++ b/Engine/Source/Programs/Shared/EpicGames.UBA/Impl/SessionServerImpl.cs
@@ -96,7 +96,7 @@ namespace EpicGames.UBA
 		static extern IntPtr SessionServer_RunProcess(IntPtr server, IntPtr info, bool async, bool enableDetour);
 
 		[DllImport("UbaHost", CharSet = CharSet.Auto)]
-		static extern IntPtr SessionServer_RunProcessRemote(IntPtr server, IntPtr info, float weight);
+		static extern IntPtr SessionServer_RunProcessRemote(IntPtr server, IntPtr info, float weight, byte[]? knownInputs, uint knownInputsCount);
 
 		[DllImport("UbaHost", CharSet = CharSet.Auto)]
 		static extern void SessionServer_SetMaxRemoteProcessCount(IntPtr server, uint count);
@@ -165,10 +165,10 @@ namespace EpicGames.UBA
 			return process;
 		}
 
-		public IProcess RunProcessRemote(ProcessStartInfo info, IProcess.ExitedEventHandler? exitedEventHandler, double weight)
+		public IProcess RunProcessRemote(ProcessStartInfo info, IProcess.ExitedEventHandler? exitedEventHandler, double weight, byte[]? knownInputs, uint knownInputsCount)
 		{
 			IProcessStartInfo startInfo = IProcessStartInfo.CreateProcessStartInfo(info, exitedEventHandler != null);
-			IntPtr processPtr = SessionServer_RunProcessRemote(_handle, startInfo.GetHandle(), (float)weight);
+			IntPtr processPtr = SessionServer_RunProcessRemote(_handle, startInfo.GetHandle(), (float)weight, knownInputs, knownInputsCount);
 			IProcess process = IProcess.CreateProcess(processPtr, startInfo, exitedEventHandler, info.UserData);
 			_remoteProcesses.AddOrUpdate(process.Hash, process, (k, v) => process);
 			return process;
diff --git a/Engine/Source/Programs/UnrealBuildTool/Executors/Experimental/UBAExecutor.cs b/Engine/Source/Programs/UnrealBuildTool/Executors/Experimental/UBAExecutor.cs
--- a/Engine/Source/Programs/UnrealBuildTool/Executors/Experimental/UBAExecutor.cs
+++ b/Engine/Source/Programs/UnrealBuildTool/Executors/Experimental/UBAExecutor.cs
@@ -232,6 +232,13 @@ namespace UnrealBuildTool
 		[CommandLine("-BoxCrypto", Value = "true")]
 		[CommandLine("-UBACrypto", Value = "true")]
 		public bool bUseCrypto { get; set; } = false;
+
+		/// <summary>
+		/// Set to true to provide known inputs to processes that are run remote. This is an experimental feature to speed up build times when ping is higher
+		/// </summary>
+		[XmlConfigFile(Category = "UnrealBuildAccelerator")]
+		[CommandLine("-UBAUseKnownInputs", Value = "true")]
+		public bool bUseKnownInputs { get; set; } = false;
 	}
 
 	class UBAExecutor : ParallelExecutor
@@ -741,6 +748,29 @@ namespace UnrealBuildTool
 
 			return () =>
 			{
+				uint knownInputsCount = 0;
+				byte[]? knownInputs = null;
+				if (UBAConfig.bUseKnownInputs)
+				{
+					int byteCount = 0;
+					foreach (var item in action.PrerequisiteItems)
+					{
+						byteCount += (item.FullName.Length + 1) * sizeof(char);
+						++knownInputsCount;
+					}
+
+					knownInputs = new byte[byteCount + sizeof(char)];
+
+					int byteOffset = 0;
+					foreach (var item in action.PrerequisiteItems)
+					{
+						var str = item.FullName;
+						int strBytes = str.Length * sizeof(char);
+						System.Buffer.BlockCopy(str.ToCharArray(), 0, knownInputs, byteOffset, strBytes);
+						byteOffset += strBytes + sizeof(char);
+					}
+				}
+
 				ProcessStartInfo startInfo = GetActionStartInfo(action, out FileItem? pchItem);
 				_session!.RunProcessRemote(startInfo, (s, e) =>
 				{
@@ -756,7 +786,7 @@ namespace UnrealBuildTool
 					List<string> logLines = e.LogLines;
 					logLines.RemoveAll((line) => line.StartsWith("   Creating library ", StringComparison.OrdinalIgnoreCase) && line.EndsWith(".exp", StringComparison.OrdinalIgnoreCase));
 					ActionFinished(queue, new ExecuteResults(logLines, e.ExitCode, executionTime, processorTime, additionalDescription), action, pchItem, s as IProcess);
-				}, action.Weight);
+				}, action.Weight, knownInputs, knownInputsCount);
 				return Task.CompletedTask;
 			};
 		}
