[UBA]
* Fixed so UbaExecutor always write knownFiles in tchar format (char is always 2 bytes inside managed code. so if unmanaged code has char size 1, we need to convert)

diff --git a/Engine/Source/Programs/UnrealBuildTool/Executors/Experimental/UBAExecutor.cs b/Engine/Source/Programs/UnrealBuildTool/Executors/Experimental/UBAExecutor.cs
--- a/Engine/Source/Programs/UnrealBuildTool/Executors/Experimental/UBAExecutor.cs
+++ b/Engine/Source/Programs/UnrealBuildTool/Executors/Experimental/UBAExecutor.cs
@@ -752,22 +752,27 @@ namespace UnrealBuildTool
 				byte[]? knownInputs = null;
 				if (UBAConfig.bUseKnownInputs)
 				{
+					int sizeOfChar = Marshal.SizeOf<char>();
+
 					int byteCount = 0;
 					foreach (var item in action.PrerequisiteItems)
 					{
-						byteCount += (item.FullName.Length + 1) * sizeof(char);
+						byteCount += (item.FullName.Length + 1) * sizeOfChar;
 						++knownInputsCount;
 					}
 
-					knownInputs = new byte[byteCount + sizeof(char)];
+					knownInputs = new byte[byteCount + sizeOfChar];
 
 					int byteOffset = 0;
 					foreach (var item in action.PrerequisiteItems)
 					{
 						var str = item.FullName;
-						int strBytes = str.Length * sizeof(char);
-						System.Buffer.BlockCopy(str.ToCharArray(), 0, knownInputs, byteOffset, strBytes);
-						byteOffset += strBytes + sizeof(char);
+						int strBytes = str.Length * sizeOfChar;
+						if (sizeOfChar == 1) // Unmanaged size uses ascii
+							System.Buffer.BlockCopy(System.Text.Encoding.ASCII.GetBytes(str.ToCharArray()), 0, knownInputs, byteOffset, strBytes);
+						else
+							System.Buffer.BlockCopy(str.ToCharArray(), 0, knownInputs, byteOffset, strBytes);
+						byteOffset += strBytes + sizeOfChar;
 					}
 				}
 
