Support devices without lz4 compression for bDontBundleLibrariesInAPK

diff --git a/Engine/Source/Programs/UnrealBuildTool/Platform/Android/UEDeployAndroid.cs b/Engine/Source/Programs/UnrealBuildTool/Platform/Android/UEDeployAndroid.cs
--- a/Engine/Source/Programs/UnrealBuildTool/Platform/Android/UEDeployAndroid.cs
+++ b/Engine/Source/Programs/UnrealBuildTool/Platform/Android/UEDeployAndroid.cs
@@ -4939,6 +4939,7 @@ if ""%ERRORLEVEL%"" NEQ ""0"" (exit /b %ERRORLEVEL%)
 					{
 						SOPushScript += @$"
 %ADB% %DEVICE% push -z lz4 {FinalSONameStrippedRelative} /data/local/tmp/{FinalSONameStrippedRelative}
+if ""%ERRORLEVEL%"" NEQ ""0"" (%ADB% %DEVICE% push {FinalSONameStrippedRelative} /data/local/tmp/{FinalSONameStrippedRelative})
 if ""%ERRORLEVEL%"" NEQ ""0"" (exit /b %ERRORLEVEL%)
 %ADB% %DEVICE% shell run-as {PackageName} mkdir -p ./files
 %ADB% %DEVICE% shell run-as {PackageName} cp /data/local/tmp/{FinalSONameStrippedRelative} ./files/libUnreal.so
