Fix incorrect Receipt.Launch value if there are executable-type runtime dependencies.

The bug was introduced with 17546102. Previously "Launch" would get set to the first suitable BuildProduct entry; the module's binary. But after that change it can get incorrectly set to the last suitable BuildProduct entry if the target doesn't not produce a "-Cmd" suffixed output. Any module (such as Sentry's UE integration [1]) that introduces a RuntimeDependency on .exe files can result in invalid target receipts.

[1] [REDACTED_PROFILE_URL]/sentry-unreal/blob/main/plugin-dev/Source/Sentry/Sentry.Build.cs

#rnx

diff --git a/Engine/Source/Programs/UnrealBuildTool/Configuration/UEBuildTarget.cs b/Engine/Source/Programs/UnrealBuildTool/Configuration/UEBuildTarget.cs
--- a/Engine/Source/Programs/UnrealBuildTool/Configuration/UEBuildTarget.cs
+++ b/Engine/Source/Programs/UnrealBuildTool/Configuration/UEBuildTarget.cs
@@ -2124,11 +2124,11 @@ namespace UnrealBuildTool
 				{
 					if (Pair.Value == BuildProductType.Executable)
 					{
-						if (System.IO.Path.GetFileNameWithoutExtension(Pair.Key.FullName).EndsWith("-Cmd"))
+						if (System.IO.Path.GetFileNameWithoutExtension(Pair.Key.FullName).EndsWith("-Cmd") && Receipt.LaunchCmd == null)
 						{
 							Receipt.LaunchCmd = Pair.Key;
 						}
-						else
+						else if (Receipt.Launch == null)
 						{
 							Receipt.Launch = Pair.Key;
 						}
