UnrealBuildTool: Ensure the AggregateException.InnerExceptions is valid

#rnx

diff --git a/Engine/Source/Programs/UnrealBuildTool/UnrealBuildTool.cs b/Engine/Source/Programs/UnrealBuildTool/UnrealBuildTool.cs
--- a/Engine/Source/Programs/UnrealBuildTool/UnrealBuildTool.cs
+++ b/Engine/Source/Programs/UnrealBuildTool/UnrealBuildTool.cs
@@ -659,9 +659,9 @@ namespace UnrealBuildTool
 				{
 					Result = Mode.ExecuteAsync(Arguments, Logger).GetAwaiter().GetResult();
 				}
-				catch (AggregateException AggEx) when (AggEx.InnerExceptions.Count == 1)
+				catch (AggregateException AggEx) when (AggEx.InnerExceptions.Count == 1 && AggEx.InnerExceptions.FirstOrDefault() != null)
 				{
-					throw AggEx.InnerExceptions[0];
+					throw AggEx.InnerExceptions.First();
 				}
 				finally
 				{
