UnrealBuildTool: TargetRules cleanup

* Change (most) fields to properties
* Naming rule fixes for local variables
* Remove unnecessary default values
* Move ReadOnlyTargetRules to new file
* Remove RequiresUniqueBuildEnvironment from obsolete properties
* Document formatting

#rnx

diff --git a/Engine/Source/Programs/UnrealBuildTool/Configuration/ReadOnlyTargetRules.cs b/Engine/Source/Programs/UnrealBuildTool/Configuration/ReadOnlyTargetRules.cs
new file mode 100644
--- /dev/null
+++ b/Engine/Source/Programs/UnrealBuildTool/Configuration/ReadOnlyTargetRules.cs
@@ -0,0 +1,674 @@
+// Copyright Epic Games, Inc. All Rights Reserved.
+
+using System;
+using System.Collections.Generic;
+using System.IO;
+using System.Linq;
+using EpicGames.Core;
+using Microsoft.Extensions.Logging;
+using UnrealBuildBase;
+
+namespace UnrealBuildTool
+{
+	/// <summary>
+	/// Read-only wrapper around an existing TargetRules instance. This exposes target settings to modules without letting them to modify the global environment.
+	/// </summary>
+	public partial class ReadOnlyTargetRules
+	{
+		/// <summary>
+		/// The writeable TargetRules instance
+		/// </summary>
+		readonly TargetRules Inner;
+
+		/// <summary>
+		/// Constructor
+		/// </summary>
+		/// <param name="inner">The TargetRules instance to wrap around</param>
+		public ReadOnlyTargetRules(TargetRules inner)
+		{
+			Inner = inner;
+			AndroidPlatform = new ReadOnlyAndroidTargetRules(inner.AndroidPlatform);
+			IOSPlatform = new ReadOnlyIOSTargetRules(inner.IOSPlatform);
+			LinuxPlatform = new ReadOnlyLinuxTargetRules(inner.LinuxPlatform);
+			MacPlatform = new ReadOnlyMacTargetRules(inner.MacPlatform);
+			WindowsPlatform = new ReadOnlyWindowsTargetRules(inner.WindowsPlatform);
+		}
+
+		/// <summary>
+		/// Accessors for fields on the inner TargetRules instance
+		/// </summary>
+		#region Read-only accessor properties
+#pragma warning disable CS1591
+
+		public string Name => Inner.Name;
+
+		public string OriginalName => Inner.OriginalName;
+
+		internal FileReference File => Inner.File!;
+
+		internal FileReference TargetSourceFile => Inner.TargetSourceFile!;
+
+		internal IReadOnlySet<FileReference> TargetFiles => Inner.TargetFiles!;
+
+		public object? AdditionalData => Inner.AdditionalData;
+
+		public UnrealTargetPlatform Platform => Inner.Platform;
+
+		public UnrealTargetConfiguration Configuration => Inner.Configuration;
+
+		public UnrealArchitectures Architectures => Inner.Architectures;
+
+		public UnrealIntermediateEnvironment IntermediateEnvironment => Inner.IntermediateEnvironment;
+
+		public UnrealArch Architecture => Inner.Architecture;
+
+		public FileReference? ProjectFile => Inner.ProjectFile;
+
+		public ReadOnlyBuildVersion Version => Inner.Version;
+
+		public TargetType Type => Inner.Type;
+
+		/// <inheritdoc cref="TargetRules.Logger"/>
+		public ILogger Logger => Inner.Logger;
+
+		public BuildSettingsVersion DefaultBuildSettings => Inner.DefaultBuildSettings;
+
+		public EngineIncludeOrderVersion? ForcedIncludeOrder => Inner.ForcedIncludeOrder;
+
+		public EngineIncludeOrderVersion IncludeOrderVersion => Inner.IncludeOrderVersion;
+
+		internal ConfigValueTracker ConfigValueTracker => Inner.ConfigValueTracker;
+
+		public string? OutputFile => Inner.OutputFile;
+
+		public bool bUsesSteam => Inner.bUsesSteam;
+
+		public bool bUsesCEF3 => Inner.bUsesCEF3;
+
+		public bool bUsesSlate => Inner.bUsesSlate;
+
+		public bool bUseStaticCRT => Inner.bUseStaticCRT;
+
+		public bool bDebugBuildsActuallyUseDebugCRT => Inner.bDebugBuildsActuallyUseDebugCRT;
+
+		public bool bLegalToDistributeBinary => Inner.bLegalToDistributeBinary;
+
+		public UnrealTargetConfiguration UndecoratedConfiguration => Inner.UndecoratedConfiguration;
+
+		public bool bAllowHotReload => Inner.bAllowHotReload;
+
+		public bool bBuildAllModules => Inner.bBuildAllModules;
+
+		public IReadOnlyList<string> AdditionalPlugins => Inner.AdditionalPlugins;
+
+		public IReadOnlyList<string> EnablePlugins => Inner.EnablePlugins;
+
+		public IReadOnlyList<string> DisablePlugins => Inner.DisablePlugins;
+
+		public IReadOnlyList<string> OptionalPlugins => Inner.OptionalPlugins;
+
+		public WarningLevel DisablePluginsConflictWarningLevel => Inner.DisablePluginsConflictWarningLevel;
+
+		public IReadOnlyList<string> BuildPlugins => Inner.BuildPlugins;
+
+		public IReadOnlyList<string> InternalPluginDependencies => Inner.InternalPluginDependencies;
+		public bool bRuntimeDependenciesComeFromBuildPlugins => Inner.bRuntimeDependenciesComeFromBuildPlugins;
+
+		public string PakSigningKeysFile => Inner.PakSigningKeysFile;
+
+		public string SolutionDirectory => Inner.SolutionDirectory;
+
+		public string CustomConfig => Inner.CustomConfig;
+
+		public bool? bBuildInSolutionByDefault => Inner.bBuildInSolutionByDefault;
+
+		public string ExeBinariesSubFolder => Inner.ExeBinariesSubFolder;
+
+		public EGeneratedCodeVersion GeneratedCodeVersion => Inner.GeneratedCodeVersion;
+		public bool bEnableMeshEditor => Inner.bEnableMeshEditor;
+
+		public bool bUseVerseBPVM => Inner.bUseVerseBPVM;
+
+		public bool bUseAutoRTFMCompiler => Inner.bUseAutoRTFMCompiler;
+
+		[Obsolete("Deprecated in UE5.1 - No longer used as Chaos is always enabled.")]
+		public bool bCompileChaos => Inner.bCompileChaos;
+
+		[Obsolete("Deprecated in UE5.1 - No longer used as Chaos is always enabled.")]
+		public bool bUseChaos => Inner.bUseChaos;
+
+		public bool bUseChaosMemoryTracking => Inner.bUseChaosMemoryTracking;
+
+		public bool bCompileChaosVisualDebuggerSupport => Inner.bCompileChaosVisualDebuggerSupport;
+
+		public bool bUseChaosChecked => Inner.bUseChaosChecked;
+
+		[Obsolete("Deprecated in UE5.1 - No longer used in engine.")]
+		public bool bCustomSceneQueryStructure => Inner.bCustomSceneQueryStructure;
+
+		[Obsolete("Deprecated in UE5.1 - No longer used as Chaos is always enabled.")]
+		public bool bCompilePhysX => Inner.bCompilePhysX;
+
+		[Obsolete("Deprecated in UE5.1 - No longer used as Chaos is always enabled.")]
+		public bool bCompileAPEX => Inner.bCompileAPEX;
+
+		[Obsolete("Deprecated in UE5.1 - No longer used as Chaos is always enabled.")]
+		public bool bCompileNvCloth => Inner.bCompileNvCloth;
+
+		public bool bCompileICU => Inner.bCompileICU;
+
+		public bool bCompileCEF3 => Inner.bCompileCEF3;
+
+		public bool bCompileISPC => Inner.bCompileISPC;
+
+		public bool bUseGameplayDebugger => Inner.bUseGameplayDebugger;
+
+		public bool bUseGameplayDebuggerCore => Inner.bUseGameplayDebuggerCore;
+
+		public bool bUseIris => Inner.bUseIris;
+
+		public bool bCompileIntelMetricsDiscovery => Inner.bCompileIntelMetricsDiscovery;
+
+		public bool bCompilePython => Inner.bCompilePython;
+
+		public bool bBuildEditor => Inner.bBuildEditor;
+
+		public bool bBuildRequiresCookedData => Inner.bBuildRequiresCookedData;
+
+		public bool bBuildWithEditorOnlyData => Inner.bBuildWithEditorOnlyData;
+
+		public bool bBuildDeveloperTools => Inner.bBuildDeveloperTools;
+
+		public bool bBuildTargetDeveloperTools => Inner.bBuildTargetDeveloperTools;
+
+		public bool bForceBuildTargetPlatforms => Inner.bForceBuildTargetPlatforms;
+
+		public bool bForceBuildShaderFormats => Inner.bForceBuildShaderFormats;
+
+		public bool bNeedsExtraShaderFormats => Inner.bNeedsExtraShaderFormats;
+
+		public bool bCompileCustomSQLitePlatform => Inner.bCompileCustomSQLitePlatform;
+
+		public bool bUseCacheFreedOSAllocs => Inner.bUseCacheFreedOSAllocs;
+
+		public bool bCompileAgainstEngine => Inner.bCompileAgainstEngine;
+
+		public bool bCompileAgainstCoreUObject => Inner.bCompileAgainstCoreUObject;
+
+		public bool bCompileAgainstApplicationCore => Inner.bCompileAgainstApplicationCore;
+
+		public bool bEnableTrace => Inner.bEnableTrace;
+
+		public bool bCompileAgainstEditor => Inner.bCompileAgainstEditor;
+
+		public bool bCompileRecast => Inner.bCompileRecast;
+
+		public bool bCompileNavmeshSegmentLinks => Inner.bCompileNavmeshSegmentLinks;
+
+		public bool bCompileNavmeshClusterLinks => Inner.bCompileNavmeshClusterLinks;
+
+		public bool bCompileSpeedTree => Inner.bCompileSpeedTree;
+
+		public bool bForceEnableExceptions => Inner.bForceEnableExceptions;
+
+		public bool bForceEnableObjCExceptions => Inner.bForceEnableObjCExceptions;
+
+		public bool bForceEnableRTTI => Inner.bForceEnableRTTI;
+
+		public bool bEnablePIE => Inner.bEnablePIE;
+
+		public bool bEnableStackProtection => Inner.bEnableStackProtection;
+
+		public bool bUseInlining => Inner.bUseInlining;
+
+		public bool bWithServerCode => Inner.bWithServerCode;
+
+		public bool bFNameOutlineNumber => Inner.bFNameOutlineNumber;
+
+		public bool bWithPushModel => Inner.bWithPushModel;
+
+		public bool bCompileWithStatsWithoutEngine => Inner.bCompileWithStatsWithoutEngine;
+
+		public bool bCompileWithPluginSupport => Inner.bCompileWithPluginSupport;
+
+		public bool bIncludePluginsForTargetPlatforms => Inner.bIncludePluginsForTargetPlatforms;
+
+		public bool bCompileWithAccessibilitySupport => Inner.bCompileWithAccessibilitySupport;
+
+		public bool bWithPerfCounters => Inner.bWithPerfCounters;
+
+		public bool bWithLiveCoding => Inner.bWithLiveCoding;
+
+		public bool bUseDebugLiveCodingConsole => Inner.bUseDebugLiveCodingConsole;
+
+		public bool bWithDirectXMath => Inner.bWithDirectXMath;
+
+		public bool bWithFixedTimeStepSupport => Inner.bWithFixedTimeStepSupport;
+
+		public bool bUseLoggingInShipping => Inner.bUseLoggingInShipping;
+
+		public bool bUseConsoleInShipping => Inner.bUseConsoleInShipping;
+
+		public bool bLoggingToMemoryEnabled => Inner.bLoggingToMemoryEnabled;
+
+		public bool bUseLauncherChecks => Inner.bUseLauncherChecks;
+
+		public bool bUseChecksInShipping => Inner.bUseChecksInShipping;
+
+		public bool bAllowProfileGPUInTest => Inner.bAllowProfileGPUInTest;
+
+		public bool bAllowProfileGPUInShipping => Inner.bAllowProfileGPUInShipping;
+
+		public bool bTCHARIsUTF8 => Inner.bTCHARIsUTF8;
+
+		public bool bUseEstimatedUtcNow => Inner.bUseEstimatedUtcNow;
+
+		public bool bCompileFreeType => Inner.bCompileFreeType;
+
+		public bool bUseExecCommandsInShipping => Inner.bUseExecCommandsInShipping;
+
+		[Obsolete("Deprecated in UE5.1 - Please use OptimizationLevel instead.")]
+		public bool bCompileForSize => Inner.bCompileForSize;
+
+		public OptimizationMode OptimizationLevel => Inner.OptimizationLevel;
+
+		public FPSemanticsMode FPSemantics => Inner.FPSemantics;
+
+		public bool bRetainFramePointers => Inner.bRetainFramePointers;
+
+		public bool bForceCompileDevelopmentAutomationTests => Inner.bForceCompileDevelopmentAutomationTests;
+
+		public bool bForceCompilePerformanceAutomationTests => Inner.bForceCompilePerformanceAutomationTests;
+
+		public bool bForceDisableAutomationTests => Inner.bForceDisableAutomationTests;
+
+		public bool bUseXGEController => Inner.bUseXGEController;
+
+		public bool bEventDrivenLoader => Inner.bEventDrivenLoader;
+
+		public PointerMemberBehavior? NativePointerMemberBehaviorOverride => Inner.NativePointerMemberBehaviorOverride;
+
+		public bool bIWYU => Inner.bIWYU;
+
+		public bool bIncludeHeaders => Inner.bIncludeHeaders;
+
+		public bool bHeadersOnly => Inner.bHeadersOnly;
+
+		public bool bEnforceIWYU => Inner.bEnforceIWYU;
+
+		public bool bWarnAboutMonolithicHeadersIncluded => Inner.bWarnAboutMonolithicHeadersIncluded;
+
+		public bool bHasExports => Inner.bHasExports;
+
+		public bool bPrecompile => Inner.bPrecompile;
+
+		public bool bEnableOSX109Support => Inner.bEnableOSX109Support;
+
+		public bool bIsBuildingConsoleApplication => Inner.bIsBuildingConsoleApplication;
+
+		public bool bBuildAdditionalConsoleApp => Inner.bBuildAdditionalConsoleApp;
+
+		public bool bDisableSymbolCache => Inner.bDisableSymbolCache;
+
+		public bool bUseUnityBuild => Inner.bUseUnityBuild;
+
+		public bool bForceUnityBuild => Inner.bForceUnityBuild;
+
+		public IReadOnlyList<string>? DisableMergingModuleAndGeneratedFilesInUnityFiles => Inner.DisableMergingModuleAndGeneratedFilesInUnityFiles;
+
+		public bool bAdaptiveUnityDisablesOptimizations => Inner.bAdaptiveUnityDisablesOptimizations;
+
+		public bool bAdaptiveUnityDisablesPCH => Inner.bAdaptiveUnityDisablesPCH;
+
+		public bool bAdaptiveUnityDisablesPCHForProject => Inner.bAdaptiveUnityDisablesPCHForProject;
+
+		public bool bAdaptiveUnityCreatesDedicatedPCH => Inner.bAdaptiveUnityCreatesDedicatedPCH;
+
+		public bool bAdaptiveUnityEnablesEditAndContinue => Inner.bAdaptiveUnityEnablesEditAndContinue;
+
+		public bool bAdaptiveUnityCompilesHeaderFiles => Inner.bAdaptiveUnityCompilesHeaderFiles;
+
+		public int MinGameModuleSourceFilesForUnityBuild => Inner.MinGameModuleSourceFilesForUnityBuild;
+
+		public bool bValidateFormatStrings => Inner.bValidateFormatStrings;
+
+		public WarningLevel DefaultWarningLevel => Inner.DefaultWarningLevel;
+
+		public WarningLevel DeprecationWarningLevel => Inner.DeprecationWarningLevel;
+
+		public WarningLevel ShadowVariableWarningLevel => Inner.ShadowVariableWarningLevel;
+
+		public WarningLevel UnsafeTypeCastWarningLevel => Inner.UnsafeTypeCastWarningLevel;
+
+		public bool bUndefinedIdentifierErrors => Inner.bUndefinedIdentifierErrors;
+
+		public WarningLevel PCHPerformanceIssueWarningLevel => Inner.PCHPerformanceIssueWarningLevel;
+
+		public WarningLevel ModuleIncludePathWarningLevel => Inner.ModuleIncludePathWarningLevel;
+
+		public WarningLevel ModuleIncludePrivateWarningLevel => Inner.ModuleIncludePrivateWarningLevel;
+
+		public WarningLevel ModuleIncludeSubdirectoryWarningLevel => Inner.ModuleIncludeSubdirectoryWarningLevel;
+
+		public bool bWarningsAsErrors => Inner.bWarningsAsErrors;
+
+		public bool bUseFastMonoCalls => Inner.bUseFastMonoCalls;
+
+		public int NumIncludedBytesPerUnityCPP => Inner.NumIncludedBytesPerUnityCPP;
+
+		public bool bDisableModuleNumIncludedBytesPerUnityCPPOverride => Inner.bDisableModuleNumIncludedBytesPerUnityCPPOverride;
+
+		public bool bStressTestUnity => Inner.bStressTestUnity;
+
+		public bool bDetailedUnityFiles => Inner.bDetailedUnityFiles;
+
+		[Obsolete("Deprecated in UE5.4 - Replace with ReadOnlyTargetRules.DebugInfo")]
+		public bool bDisableDebugInfo => Inner.bDisableDebugInfo;
+
+		public DebugInfoMode DebugInfo => Inner.DebugInfo;
+
+		public IReadOnlySet<string> DisableDebugInfoModules => Inner.DisableDebugInfoModules;
+
+		public IReadOnlySet<string> DisableDebugInfoPlugins => Inner.DisableDebugInfoPlugins;
+
+		public DebugInfoMode DebugInfoLineTablesOnly => Inner.DebugInfoLineTablesOnly;
+
+		public IReadOnlySet<string> DebugInfoLineTablesOnlyModules => Inner.DebugInfoLineTablesOnlyModules;
+
+		public IReadOnlySet<string> DebugInfoLineTablesOnlyPlugins => Inner.DebugInfoLineTablesOnlyPlugins;
+
+		public bool bDisableDebugInfoForGeneratedCode => Inner.bDisableDebugInfoForGeneratedCode;
+
+		public bool bOmitPCDebugInfoInDevelopment => Inner.bOmitPCDebugInfoInDevelopment;
+
+		public bool bUsePDBFiles => Inner.bUsePDBFiles;
+
+		public bool bUsePCHFiles => Inner.bUsePCHFiles;
+
+		public bool bDeterministic => Inner.bDeterministic;
+		public bool bChainPCHs => Inner.bChainPCHs;
+
+		public bool bForceIncludePCHHeadersForGenCppFilesWhenPCHIsDisabled => Inner.bForceIncludePCHHeadersForGenCppFilesWhenPCHIsDisabled;
+
+		public bool bPreprocessOnly => Inner.bPreprocessOnly;
+
+		public bool bPreprocessDepends => Inner.bPreprocessDepends;
+
+		public bool bWithAssembly => Inner.bWithAssembly;
+
+		public StaticAnalyzer StaticAnalyzer => Inner.StaticAnalyzer;
+
+		public StaticAnalyzerOutputType StaticAnalyzerOutputType => Inner.StaticAnalyzerOutputType;
+
+		public StaticAnalyzerMode StaticAnalyzerMode => Inner.StaticAnalyzerMode;
+
+		public int MinFilesUsingPrecompiledHeader => Inner.MinFilesUsingPrecompiledHeader;
+
+		public bool bForcePrecompiledHeaderForGameModules => Inner.bForcePrecompiledHeaderForGameModules;
+
+		public bool bUseIncrementalLinking => Inner.bUseIncrementalLinking;
+
+		public bool bAllowLTCG => Inner.bAllowLTCG;
+
+		public bool bPreferThinLTO => Inner.bPreferThinLTO;
+
+		public bool bPGOProfile => Inner.bPGOProfile;
+
+		public bool bPGOOptimize => Inner.bPGOOptimize;
+
+		public bool bSupportEditAndContinue => Inner.bSupportEditAndContinue;
+
+		public bool bCodeCoverage => Inner.bCodeCoverage;
+
+		public bool bOmitFramePointers => Inner.bOmitFramePointers;
+
+		public bool bEnableCppModules => Inner.bEnableCppModules;
+
+		public bool bEnableCppCoroutinesForEvaluation => Inner.bEnableCppCoroutinesForEvaluation;
+
+		public bool bEnableProcessPriorityControl => Inner.bEnableProcessPriorityControl;
+
+		[Obsolete("Deprecated in UE5.3 - No longer used as MallocProfiler is removed in favor of UnrealInsights.")]
+		public bool bUseMallocProfiler => Inner.bUseMallocProfiler;
+
+		public bool bShaderCompilerWorkerTrace => Inner.bShaderCompilerWorkerTrace;
+
+		public bool bUseSharedPCHs => Inner.bUseSharedPCHs;
+
+		public bool bUseShippingPhysXLibraries => Inner.bUseShippingPhysXLibraries;
+
+		public bool bUseCheckedPhysXLibraries => Inner.bUseCheckedPhysXLibraries;
+
+		public bool bCheckLicenseViolations => Inner.bCheckLicenseViolations;
+
+		public bool bBreakBuildOnLicenseViolation => Inner.bBreakBuildOnLicenseViolation;
+
+		public bool? bUseFastPDBLinking => Inner.bUseFastPDBLinking;
+
+		public bool bCreateMapFile => Inner.bCreateMapFile;
+
+		public bool bAllowRuntimeSymbolFiles => Inner.bAllowRuntimeSymbolFiles;
+
+		public string? PackagePath => Inner.PackagePath;
+
+		public string? CrashDiagnosticDirectory => Inner.CrashDiagnosticDirectory;
+
+		public string? ThinLTOCacheDirectory => Inner.ThinLTOCacheDirectory;
+
+		public string? ThinLTOCachePruningArguments => Inner.ThinLTOCachePruningArguments;
+
+		public string? BundleVersion => Inner.BundleVersion;
+
+		public bool bDeployAfterCompile => Inner.bDeployAfterCompile;
+
+		public bool bAllowRemotelyCompiledPCHs => Inner.bAllowRemotelyCompiledPCHs;
+
+		public bool bUseHeaderUnitsForPch => Inner.bUseHeaderUnitsForPch;
+
+		public bool bCheckSystemHeadersForModification => Inner.bCheckSystemHeadersForModification;
+
+		public bool bDisableLinking => Inner.bDisableLinking;
+
+		public bool bIgnoreBuildOutputs => Inner.bIgnoreBuildOutputs;
+
+		public bool bFormalBuild => Inner.bFormalBuild;
+
+		public bool bUseAdaptiveUnityBuild => Inner.bUseAdaptiveUnityBuild;
+
+		public bool bFlushBuildDirOnRemoteMac => Inner.bFlushBuildDirOnRemoteMac;
+
+		public bool bPrintToolChainTimingInfo => Inner.bPrintToolChainTimingInfo;
+
+		public bool bParseTimingInfoForTracing => Inner.bParseTimingInfoForTracing;
+
+		public bool bPublicSymbolsByDefault => Inner.bPublicSymbolsByDefault;
+
+		public string? ToolChainName => Inner.ToolChainName;
+
+		public float MSVCCompileActionWeight => Inner.MSVCCompileActionWeight;
+
+		public float ClangCompileActionWeight => Inner.ClangCompileActionWeight;
+
+		public bool bLegacyPublicIncludePaths => Inner.bLegacyPublicIncludePaths;
+
+		public bool bLegacyParentIncludePaths => Inner.bLegacyParentIncludePaths;
+
+		public CppStandardVersion CppStandardEngine => Inner.CppStandardEngine;
+
+		public CppStandardVersion CppStandard => Inner.CppStandard;
+
+		public CStandardVersion CStandard => Inner.CStandard;
+
+		public MinimumCpuArchitectureX64 MinCpuArchX64 => Inner.MinCpuArchX64;
+
+		internal bool bNoManifestChanges => Inner.bNoManifestChanges;
+
+		public string? BuildVersion => Inner.BuildVersion;
+
+		public TargetLinkType LinkType => Inner.LinkType;
+
+		public IReadOnlyList<string> GlobalDefinitions => Inner.GlobalDefinitions.AsReadOnly();
+
+		public IReadOnlyList<string> ProjectDefinitions => Inner.ProjectDefinitions.AsReadOnly();
+
+		public string? LaunchModuleName => Inner.LaunchModuleName;
+
+		public string? ExportPublicHeader => Inner.ExportPublicHeader;
+
+		public IReadOnlyList<string> ExtraModuleNames => Inner.ExtraModuleNames.AsReadOnly();
+
+		public IReadOnlyList<FileReference> ManifestFileNames => Inner.ManifestFileNames.AsReadOnly();
+
+		public IReadOnlyList<FileReference> DependencyListFileNames => Inner.DependencyListFileNames.AsReadOnly();
+
+		public TargetBuildEnvironment BuildEnvironment => Inner.BuildEnvironment;
+
+		public bool bOverrideBuildEnvironment => Inner.bOverrideBuildEnvironment;
+
+		public IReadOnlyList<TargetInfo> PreBuildTargets => Inner.PreBuildTargets;
+
+		public IReadOnlyList<string> PreBuildSteps => Inner.PreBuildSteps;
+
+		public IReadOnlyList<string> PostBuildSteps => Inner.PostBuildSteps;
+
+		public IReadOnlyList<string> AdditionalBuildProducts => Inner.AdditionalBuildProducts;
+
+		public string? AdditionalCompilerArguments => Inner.AdditionalCompilerArguments;
+
+		public string? AdditionalLinkerArguments => Inner.AdditionalLinkerArguments;
+
+		public double MemoryPerActionGB => Inner.MemoryPerActionGB;
+
+		public string? GeneratedProjectName => Inner.GeneratedProjectName;
+
+		public ReadOnlyAndroidTargetRules AndroidPlatform { get; init; }
+
+		public ReadOnlyLinuxTargetRules LinuxPlatform { get; init; }
+
+		public ReadOnlyIOSTargetRules IOSPlatform { get; init; }
+
+		public ReadOnlyMacTargetRules MacPlatform { get; init; }
+
+		public ReadOnlyWindowsTargetRules WindowsPlatform { get; init; }
+
+		public bool bShouldCompileAsDLL => Inner.bShouldCompileAsDLL;
+
+		public bool bGenerateProjectFiles => Inner.bGenerateProjectFiles;
+
+		public bool bIsEngineInstalled => Inner.bIsEngineInstalled;
+
+		public IReadOnlyList<string>? DisableUnityBuildForModules => Inner.DisableUnityBuildForModules;
+
+		public IReadOnlyList<string>? EnableOptimizeCodeForModules => Inner.EnableOptimizeCodeForModules;
+
+		public IReadOnlyList<string>? DisableOptimizeCodeForModules => Inner.DisableOptimizeCodeForModules;
+
+		public IReadOnlyList<string>? OptimizeForSizeModules => Inner.OptimizeForSizeModules;
+
+		public IReadOnlyList<string>? OptimizeForSizeAndSpeedModules => Inner.OptimizeForSizeAndSpeedModules;
+
+		public IReadOnlyList<UnrealTargetPlatform>? OptedInModulePlatforms => Inner.OptedInModulePlatforms;
+
+		public TestTargetRules? InnerTestTargetRules => Inner as TestTargetRules;
+
+#pragma warning restore C1591
+		#endregion
+
+		/// <summary>
+		/// Provide access to the RelativeEnginePath property for code referencing ModuleRules.BuildConfiguration.
+		/// </summary>
+		public string RelativeEnginePath => Unreal.EngineDirectory.MakeRelativeTo(DirectoryReference.GetCurrentDirectory()) + Path.DirectorySeparatorChar;
+
+		/// <summary>
+		/// Provide access to the UEThirdPartySourceDirectory property for code referencing ModuleRules.UEBuildConfiguration.
+		/// </summary>
+		public string UEThirdPartySourceDirectory => "ThirdParty/";
+
+		/// <summary>
+		/// Provide access to the UEThirdPartyBinariesDirectory property for code referencing ModuleRules.UEBuildConfiguration.
+		/// </summary>
+		public string UEThirdPartyBinariesDirectory => "../Binaries/ThirdParty/";
+
+		/// <summary>
+		/// Whether this is a low level tests target.
+		/// </summary>
+		public bool IsTestTarget => Inner.IsTestTarget;
+
+		/// <summary>
+		/// Whether this is a test target explicitly defined.
+		/// Explicitley defined test targets always inherit from TestTargetRules and define their own tests.
+		/// Implicit test targets are created from existing targets when building with -Mode=Test and they include tests from all dependencies.
+		/// </summary>
+		public bool ExplicitTestsTarget => Inner.ExplicitTestsTarget;
+
+		/// <summary>
+		/// Controls the value of WITH_LOW_LEVEL_TESTS that dictates whether module-specific low level tests are compiled in or not.
+		/// </summary>
+		public bool WithLowLevelTests => Inner.WithLowLevelTests;
+
+		/// <summary>
+		/// Get the platforms this target supports
+		/// </summary>
+		public IEnumerable<UnrealTargetPlatform> SupportedPlatforms => Inner.GetSupportedPlatforms();
+
+		/// <summary>
+		/// Get the configurations this target supports
+		/// </summary>
+		public IEnumerable<UnrealTargetConfiguration> SupportedConfigurations => Inner.GetSupportedConfigurations();
+
+		/// <summary>
+		/// Checks if current platform is part of a given platform group
+		/// </summary>
+		/// <param name="group">The platform group to check</param>
+		/// <returns>True if current platform is part of a platform group</returns>
+		public bool IsInPlatformGroup(UnrealPlatformGroup group) => UEBuildPlatform.IsPlatformInGroup(Platform, group);
+
+		/// <summary>
+		/// Gets diagnostic messages about default settings which have changed in newer versions of the engine
+		/// </summary>
+		/// <param name="diagnostics">List of messages to be appended to</param>
+		internal void GetBuildSettingsInfo(List<string> diagnostics) => Inner.GetBuildSettingsInfo(diagnostics);
+
+		public bool IsPlatformOptedIn(UnrealTargetPlatform platform) => Inner.OptedInModulePlatforms == null || Inner.OptedInModulePlatforms.Contains(platform);
+
+		/// <summary>
+		/// Checks if a plugin should be programmatically allowed in the build
+		/// </summary>
+		/// <returns>true if the plugin is allowed</returns>
+		public bool IgnorePluginDependency(PluginInfo parentInfo, PluginReferenceDescriptor childDescriptor) => Inner.IgnorePluginDependency(parentInfo, childDescriptor);
+
+		/// <summary>
+		/// Determines if the automation tests should be compiled based on the current configuration and optional forced settings.
+		/// If either the development tests or performance tests are compiled (unless explicitly disabled), this property returns true.
+		/// </summary>
+		private bool? _bWithAutomationTestsPrivate = null;
+		public bool WithAutomationTests
+		{
+			get
+			{
+				if (!_bWithAutomationTestsPrivate.HasValue)
+				{
+					bool bCompileDevTests = Configuration != UnrealTargetConfiguration.Shipping;
+					bool bCompilePerfTests = bCompileDevTests;
+
+					if (bForceCompileDevelopmentAutomationTests)
+					{
+						bCompileDevTests = true;
+					}
+					if (bForceCompilePerformanceAutomationTests)
+					{
+						bCompilePerfTests = true;
+					}
+					if (bForceDisableAutomationTests)
+					{
+						bCompileDevTests = bCompilePerfTests = false;
+					}
+
+					_bWithAutomationTestsPrivate = bCompileDevTests || bCompilePerfTests;
+				}
+
+				return _bWithAutomationTestsPrivate.Value;
+			}
+		}
+	}
+}
\ No newline at end of file
diff --git a/Engine/Source/Programs/UnrealBuildTool/Configuration/TargetRules.cs b/Engine/Source/Programs/UnrealBuildTool/Configuration/TargetRules.cs
--- a/Engine/Source/Programs/UnrealBuildTool/Configuration/TargetRules.cs
+++ b/Engine/Source/Programs/UnrealBuildTool/Configuration/TargetRules.cs
@@ -2,6 +2,7 @@
 
 using System;
 using System.Collections.Generic;
+using System.Diagnostics.CodeAnalysis;
 using System.IO;
 using System.Linq;
 using System.Reflection;
@@ -9,7 +10,6 @@ using System.Runtime.Serialization;
 using EpicGames.Core;
 using Microsoft.Extensions.Logging;
 using UnrealBuildBase;
-using System.Diagnostics.CodeAnalysis;
 
 namespace UnrealBuildTool
 {
@@ -87,7 +87,7 @@ namespace UnrealBuildTool
 		Unique,
 
 		/// <summary>
-		/// Will switch to Unique if needed - per-project SDK is enabled, or a property that requires unique is set away from default
+		/// Will switch to Unique if needed - per-project sdk is enabled, or a property that requires unique is set away from default
 		/// </summary>
 		UniqueIfNeeded,
 	}
@@ -393,11 +393,11 @@ namespace UnrealBuildTool
 	/// <summary>
 	/// Utility class for EngineIncludeOrderVersion defines
 	/// </summary>
-	public class EngineIncludeOrderHelper
+	public static class EngineIncludeOrderHelper
 	{
-		private static string GetDeprecationDefine(EngineIncludeOrderVersion InVersion)
+		private static string GetDeprecationDefine(EngineIncludeOrderVersion inVersion)
 		{
-			switch (InVersion)
+			switch (inVersion)
 			{
 				default:
 				case EngineIncludeOrderVersion.Unreal5_4:
@@ -411,9 +411,9 @@ namespace UnrealBuildTool
 			}
 		}
 
-		private static string GetDeprecationDefine(EngineIncludeOrderVersion InTestVersion, EngineIncludeOrderVersion InVersion)
+		private static string GetDeprecationDefine(EngineIncludeOrderVersion inTestVersion, EngineIncludeOrderVersion inVersion)
 		{
-			return GetDeprecationDefine(InTestVersion) + "=" + (InVersion < InTestVersion ? "1" : "0");
+			return GetDeprecationDefine(inTestVersion) + "=" + (inVersion < inTestVersion ? "1" : "0");
 		}
 
 		/// <summary>
@@ -435,17 +435,17 @@ namespace UnrealBuildTool
 		/// <summary>
 		/// Get a list of every deprecation define and their value for the specified engine include order.
 		/// </summary>
-		/// <param name="InVersion"></param>
+		/// <param name="inVersion"></param>
 		/// <returns></returns>
-		public static List<string> GetDeprecationDefines(EngineIncludeOrderVersion InVersion)
+		public static List<string> GetDeprecationDefines(EngineIncludeOrderVersion inVersion)
 		{
 			return new List<string>()
 			{
 #pragma warning disable CS0618 // Type or member is obsolete
-				GetDeprecationDefine(EngineIncludeOrderVersion.Unreal5_2, InVersion),
+				GetDeprecationDefine(EngineIncludeOrderVersion.Unreal5_2, inVersion),
 #pragma warning restore CS0618 // Type or member is obsolete
-				GetDeprecationDefine(EngineIncludeOrderVersion.Unreal5_3, InVersion),
-				GetDeprecationDefine(EngineIncludeOrderVersion.Unreal5_4, InVersion),
+				GetDeprecationDefine(EngineIncludeOrderVersion.Unreal5_3, inVersion),
+				GetDeprecationDefine(EngineIncludeOrderVersion.Unreal5_4, inVersion),
 			};
 		}
 
@@ -463,15 +463,15 @@ namespace UnrealBuildTool
 	/// Attribute used to mark fields which must match between targets in the shared build environment
 	/// </summary>
 	[AttributeUsage(AttributeTargets.Field | AttributeTargets.Property, AllowMultiple = false)]
-	class RequiresUniqueBuildEnvironmentAttribute : Attribute
+	sealed class RequiresUniqueBuildEnvironmentAttribute : Attribute
 	{
 	}
 
 	/// <summary>
 	/// Attribute used to mark configurable sub-objects
 	/// </summary>
-	[AttributeUsage(AttributeTargets.Field, AllowMultiple = false)]
-	class ConfigSubObjectAttribute : Attribute
+	[AttributeUsage(AttributeTargets.Field | AttributeTargets.Property, AllowMultiple = false)]
+	sealed class ConfigSubObjectAttribute : Attribute
 	{
 	}
 
@@ -516,22 +516,14 @@ namespace UnrealBuildTool
 		/// </summary>
 		public string Name
 		{
-			get
-			{
-				if (!String.IsNullOrEmpty(NameOverride))
-				{
-					return NameOverride;
-				}
-
-				return DefaultName;
-			}
+			get => !String.IsNullOrEmpty(NameOverride) ? NameOverride : DefaultName;
 			set => NameOverride = value;
 		}
 
 		/// <summary>
 		/// If the Name of this target has been overriden
 		/// </summary>
-		public bool IsNameOverriden() { return !String.IsNullOrEmpty(NameOverride); }
+		public bool IsNameOverriden() => !String.IsNullOrEmpty(NameOverride);
 
 		/// <summary>
 		/// Return the original target name without overrides or adornments
@@ -553,7 +545,7 @@ namespace UnrealBuildTool
 		/// <summary>
 		/// Override this boolean flag in inheriting classes for low level test targets.
 		/// </summary>
-		protected bool bIsTestTargetOverride;
+		protected bool bIsTestTargetOverride { get; set; }
 
 		/// <summary>
 		/// Whether this is a test target explicitly defined.
@@ -569,7 +561,7 @@ namespace UnrealBuildTool
 		/// <summary>
 		/// Override the value of WithLowLevelTests by setting this to true in inherited classes.
 		/// </summary>
-		protected bool bWithLowLevelTestsOverride;
+		protected bool bWithLowLevelTestsOverride { get; set; }
 
 		/// <summary>
 		/// File containing the general type for this target (not including platform/group)
@@ -602,17 +594,17 @@ namespace UnrealBuildTool
 		/// <summary>
 		/// Platform that this target is being built for.
 		/// </summary>
-		public readonly UnrealTargetPlatform Platform;
+		public UnrealTargetPlatform Platform { get; init; }
 
 		/// <summary>
 		/// The configuration being built.
 		/// </summary>
-		public readonly UnrealTargetConfiguration Configuration;
+		public UnrealTargetConfiguration Configuration { get; init; }
 
 		/// <summary>
 		/// Architecture that the target is being built for (or an empty string for the default).
 		/// </summary>
-		public readonly UnrealArchitectures Architectures;
+		public UnrealArchitectures Architectures { get; init; }
 
 		/// <summary>
 		/// Gets the Architecture in the normal case where there is a single Architecture in Architectures
@@ -623,22 +615,22 @@ namespace UnrealBuildTool
 		/// <summary>
 		/// Intermediate environment. Determines if the intermediates end up in a different folder than normal.
 		/// </summary>
-		public UnrealIntermediateEnvironment IntermediateEnvironment;
+		public UnrealIntermediateEnvironment IntermediateEnvironment { get; init; }
 
 		/// <summary>
 		/// Path to the project file for the project containing this target.
 		/// </summary>
-		public readonly FileReference? ProjectFile;
+		public FileReference? ProjectFile { get; init; }
 
 		/// <summary>
 		/// The current build version
 		/// </summary>
-		public readonly ReadOnlyBuildVersion Version;
+		public ReadOnlyBuildVersion Version { get; init; }
 
 		/// <summary>
 		/// The type of target.
 		/// </summary>
-		public global::UnrealBuildTool.TargetType Type = global::UnrealBuildTool.TargetType.Game;
+		public global::UnrealBuildTool.TargetType Type { get; set; } = global::UnrealBuildTool.TargetType.Game;
 
 		/// <summary>
 		/// Specifies the engine version to maintain backwards-compatible default build settings with (eg. DefaultSettingsVersion.Release_4_23, DefaultSettingsVersion.Release_4_24). Specify DefaultSettingsVersion.Latest to always
@@ -655,21 +647,14 @@ namespace UnrealBuildTool
 		/// Force the include order to a specific version. Overrides any Target and Module rules.
 		/// </summary>
 		[CommandLine("-ForceIncludeOrder=")]
-		public EngineIncludeOrderVersion? ForcedIncludeOrder = null;
+		public EngineIncludeOrderVersion? ForcedIncludeOrder { get; set; }
 
 		/// <summary>
 		/// What version of include order to use when compiling this target. Can be overridden via -ForceIncludeOrder on the command line. ModuleRules.IncludeOrderVersion takes precedence.
 		/// </summary>
 		public EngineIncludeOrderVersion IncludeOrderVersion
 		{
-			get
-			{
-				if (ForcedIncludeOrder != null)
-				{
-					return ForcedIncludeOrder.Value;
-				}
-				return IncludeOrderVersionPrivate ?? EngineIncludeOrderVersion.Oldest;
-			}
+			get => ForcedIncludeOrder ?? IncludeOrderVersionPrivate ?? EngineIncludeOrderVersion.Oldest;
 			set => IncludeOrderVersionPrivate = value;
 		}
 		private EngineIncludeOrderVersion? IncludeOrderVersionPrivate;
@@ -678,7 +663,7 @@ namespace UnrealBuildTool
 		/// Path to the output file for the main executable, relative to the Engine or project directory.
 		/// This setting is only typically useful for non-UE programs, since the engine uses paths relative to the executable to find other known folders (eg. Content).
 		/// </summary>
-		public string? OutputFile;
+		public string? OutputFile { get; set; }
 
 		/// <summary>
 		/// Tracks a list of config values read while constructing this target
@@ -688,24 +673,24 @@ namespace UnrealBuildTool
 		/// <summary>
 		/// Whether the target uses Steam.
 		/// </summary>
-		public bool bUsesSteam;
+		public bool bUsesSteam { get; set; }
 
 		/// <summary>
 		/// Whether the target uses CEF3.
 		/// </summary>
-		public bool bUsesCEF3;
+		public bool bUsesCEF3 { get; set; }
 
 		/// <summary>
 		/// Whether the project uses visual Slate UI (as opposed to the low level windowing/messaging, which is always available).
 		/// </summary>
-		public bool bUsesSlate = true;
+		public bool bUsesSlate { get; set; } = true;
 
 		/// <summary>
 		/// Forces linking against the static CRT. This is not fully supported across the engine due to the need for allocator implementations to be shared (for example), and TPS
 		/// libraries to be consistent with each other, but can be used for utility programs.
 		/// </summary>
 		[RequiresUniqueBuildEnvironment]
-		public bool bUseStaticCRT = false;
+		public bool bUseStaticCRT { get; set; }
 
 		/// <summary>
 		/// Enables the debug C++ runtime (CRT) for debug builds. By default we always use the release runtime, since the debug
@@ -716,20 +701,20 @@ namespace UnrealBuildTool
 		/// </summary>
 		[RequiresUniqueBuildEnvironment]
 		[XmlConfigFile(Category = "BuildConfiguration")]
-		public bool bDebugBuildsActuallyUseDebugCRT = false;
+		public bool bDebugBuildsActuallyUseDebugCRT { get; set; }
 
 		/// <summary>
 		/// Whether the output from this target can be publicly distributed, even if it has dependencies on modules that are in folders
 		/// with special restrictions (eg. CarefullyRedist, NotForLicensees, NoRedist).
 		/// </summary>
 		[XmlConfigFile(Category = "BuildConfiguration")]
-		public bool bLegalToDistributeBinary = false;
+		public bool bLegalToDistributeBinary { get; set; }
 
 		/// <summary>
 		/// Specifies the configuration whose binaries do not require a "-Platform-Configuration" suffix.
 		/// </summary>
 		[RequiresUniqueBuildEnvironment]
-		public UnrealTargetConfiguration UndecoratedConfiguration = UnrealTargetConfiguration.Development;
+		public UnrealTargetConfiguration UndecoratedConfiguration { get; set; } = UnrealTargetConfiguration.Development;
 
 		/// <summary>
 		/// Whether this target supports hot reload
@@ -745,81 +730,83 @@ namespace UnrealBuildTool
 		/// Build all the modules that are valid for this target type. Used for CIS and making installed engine builds.
 		/// </summary>
 		[CommandLine("-AllModules")]
-		public bool bBuildAllModules = false;
+		public bool bBuildAllModules { get; set; }
 
 		/// <summary>
 		/// Set this to reference a VSTest run settings file from generated projects.
 		/// </summary>
-		public FileReference? VSTestRunSettingsFile;
+		public FileReference? VSTestRunSettingsFile { get; set; }
 
 		/// <summary>
 		/// Additional plugins that are built for this target type but not enabled.
 		/// </summary>
 		[CommandLine("-BuildPlugin=", ListSeparator = '+')]
-		public List<string> BuildPlugins = new List<string>();
+		public List<string> BuildPlugins = new();
 
 		/// <summary>
 		/// If this is true, then the BuildPlugins list will be used to populate RuntimeDependencies, rather than EnablePlugins
 		/// </summary>
-		public bool bRuntimeDependenciesComeFromBuildPlugins = false;
+		public bool bRuntimeDependenciesComeFromBuildPlugins { get; set; }
 
 		/// <summary>
 		/// A list of additional plugins which need to be included in this target. This allows referencing non-optional plugin modules
 		/// which cannot be disabled, and allows building against specific modules in program targets which do not fit the categories
 		/// in ModuleHostType.
 		/// </summary>
-		public List<string> AdditionalPlugins = new List<string>();
+		[CommandLine("-AdditionalPlugins=", ListSeparator = '+')]
+		public List<string> AdditionalPlugins = new();
 
 		/// <summary>
 		/// Additional plugins that should be included for this target.
 		/// </summary>
 		[CommandLine("-EnablePlugin=", ListSeparator = '+')]
-		public List<string> EnablePlugins = new List<string>();
+		public List<string> EnablePlugins = new();
 
 		/// <summary>
 		/// List of plugins to be disabled for this target. Note that the project file may still reference them, so they should be marked
 		/// as optional to avoid failing to find them at runtime.
 		/// </summary>
 		[CommandLine("-DisablePlugin=", ListSeparator = '+')]
-		public List<string> DisablePlugins = new List<string>();
+		public List<string> DisablePlugins = new();
 
 		/// <summary>
 		/// Additional plugins that should be included for this target if they are found.
 		/// </summary>
-		public List<string> OptionalPlugins = new List<string>();
+		[CommandLine("-OptionalPlugins=", ListSeparator = '+')]
+		public List<string> OptionalPlugins = new();
 
 		/// <summary>
 		/// How to treat conflicts when a disabled plugin is being enabled by another plugin referencing it
 		/// </summary>
-		public WarningLevel DisablePluginsConflictWarningLevel = WarningLevel.Default;
+		public WarningLevel DisablePluginsConflictWarningLevel { get; set; } = WarningLevel.Default;
 
 		/// <summary>
 		/// A list of Plugin names that are allowed to exist as dependencies without being defined in the uplugin descriptor
 		/// </summary>
-		public List<string> InternalPluginDependencies = new List<string>();
+		public List<string> InternalPluginDependencies = new();
 
 		/// <summary>
 		/// Path to the set of pak signing keys to embed in the executable.
 		/// </summary>
-		public string PakSigningKeysFile = "";
+		public string PakSigningKeysFile { get; set; } = String.Empty;
 
 		/// <summary>
 		/// Allows a Program Target to specify it's own solution folder path.
 		/// </summary>
-		public string SolutionDirectory = String.Empty;
+		public string SolutionDirectory { get; set; } = String.Empty;
 
 		/// <summary>
 		/// Whether the target should be included in the default solution build configuration
 		/// Setting this to false will skip building when running in the IDE
 		/// </summary>
-		public bool? bBuildInSolutionByDefault = null;
+		public bool? bBuildInSolutionByDefault { get; set; }
 
 		/// <summary>
 		/// Whether this target should be compiled as a DLL.  Requires LinkType to be set to TargetLinkType.Monolithic.
 		/// </summary>
 		[RequiresUniqueBuildEnvironment]
 		[CommandLine("-CompileAsDll")]
-		public bool bShouldCompileAsDLL = false;
+		public bool bShouldCompileAsDLL { get; set; }
 
 		/// <summary>
 		/// Extra subdirectory to load config files out of, for making multiple types of builds with the same platform
@@ -827,24 +814,24 @@ namespace UnrealBuildTool
 		/// </summary>
 		[RequiresUniqueBuildEnvironment]
 		[CommandLine("-CustomConfig")]
-		public string CustomConfig = String.Empty;
+		public string CustomConfig { get; set; } = String.Empty;
 
 		/// <summary>
 		/// Subfolder to place executables in, relative to the default location.
 		/// </summary>
 		[RequiresUniqueBuildEnvironment]
-		public string ExeBinariesSubFolder = String.Empty;
+		public string ExeBinariesSubFolder { get; set; } = String.Empty;
 
 		/// <summary>
 		/// Allow target module to override UHT code generation version.
 		/// </summary>
-		public EGeneratedCodeVersion GeneratedCodeVersion = EGeneratedCodeVersion.None;
+		public EGeneratedCodeVersion GeneratedCodeVersion { get; set; } = EGeneratedCodeVersion.None;
 
 		/// <summary>
 		/// Whether to enable the mesh editor.
 		/// </summary>
 		[RequiresUniqueBuildEnvironment]
-		public bool bEnableMeshEditor = false;
+		public bool bEnableMeshEditor { get; set; }
 
 		/// <summary>
 		/// Whether to use the BPVM to run Verse.
@@ -853,7 +840,7 @@ namespace UnrealBuildTool
 		[CommandLine("-NoUseVerseBPVM", Value = "false")]
 		[CommandLine("-UseVerseBPVM", Value = "true")]
 		[XmlConfigFile(Category = "BuildConfiguration")]
-		public bool bUseVerseBPVM = true;
+		public bool bUseVerseBPVM { get; set; } = true;
 
 		/// <summary>
 		/// Whether to use the AutoRTFM Clang compiler.
@@ -862,37 +849,31 @@ namespace UnrealBuildTool
 		[CommandLine("-NoUseAutoRTFM", Value = "false")]
 		[CommandLine("-UseAutoRTFM", Value = "true")]
 		[XmlConfigFile(Category = "BuildConfiguration")]
-		public bool bUseAutoRTFMCompiler = false;
+		public bool bUseAutoRTFMCompiler { get; set; }
 
 		/// <summary>
 		/// Whether to compile the Chaos physics plugin.
 		/// </summary>
-		[RequiresUniqueBuildEnvironment]
-		[CommandLine("-NoCompileChaos", Value = "false")]
-		[CommandLine("-CompileChaos", Value = "true")]
 		[Obsolete("Deprecated in UE5.1 - No longer used as Chaos is always enabled.")]
-		public bool bCompileChaos = true;
+		public bool bCompileChaos { get; set; }
 
 		/// <summary>
 		/// Whether to use the Chaos physics interface. This overrides the physx flags to disable APEX and NvCloth
 		/// </summary>
-		[RequiresUniqueBuildEnvironment]
-		[CommandLine("-NoUseChaos", Value = "false")]
-		[CommandLine("-UseChaos", Value = "true")]
 		[Obsolete("Deprecated in UE5.1 - No longer used as Chaos is always enabled.")]
-		public bool bUseChaos = true;
+		public bool bUseChaos { get; set; }
 
 		/// <summary>
 		/// Whether to compile in checked chaos features for debugging
 		/// </summary>
 		[RequiresUniqueBuildEnvironment]
-		public bool bUseChaosChecked = false;
+		public bool bUseChaosChecked { get; set; }
 
 		/// <summary>
 		/// Whether to compile in chaos memory tracking features
 		/// </summary>
 		[RequiresUniqueBuildEnvironment]
-		public bool bUseChaosMemoryTracking = false;
+		public bool bUseChaosMemoryTracking { get; set; }
 
 		/// <summary>
 		/// Whether to compile in Chaos Visual Debugger (CVD) support features to record the state of the physics simulation
@@ -900,68 +881,63 @@ namespace UnrealBuildTool
 		[XmlConfigFile(Name = "bCompileChaosVisualDebuggerSupport")]
 		[CommandLine("-CompileChaosVisualDebuggerSupport")]
 		[RequiresUniqueBuildEnvironment]
-		public bool bCompileChaosVisualDebuggerSupport = true;
+		public bool bCompileChaosVisualDebuggerSupport { get; set; } = true;
 
 		/// <summary>
 		/// Whether scene query acceleration is done by UE. The physx scene query structure is still created, but we do not use it.
 		/// </summary>
-		[RequiresUniqueBuildEnvironment]
 		[Obsolete("Deprecated in UE5.1 - No longer used in engine.")]
-		public bool bCustomSceneQueryStructure = false;
+		public bool bCustomSceneQueryStructure { get; set; }
 
 		/// <summary>
 		/// Whether to include PhysX support.
 		/// </summary>
-		[RequiresUniqueBuildEnvironment]
 		[Obsolete("Deprecated in UE5.1 - No longer used as Chaos is always enabled.")]
-		public bool bCompilePhysX = false;
+		public bool bCompilePhysX { get; set; }
 
 		/// <summary>
 		/// Whether to include PhysX APEX support.
 		/// </summary>
-		[RequiresUniqueBuildEnvironment]
-		[ConfigFile(ConfigHierarchyType.Engine, "/Script/BuildSettings.BuildSettings", "bCompileApex")]
 		[Obsolete("Deprecated in UE5.1 - No longer used as Chaos is always enabled.")]
-		public bool bCompileAPEX = false;
+		public bool bCompileAPEX { get; set; }
 
 		/// <summary>
 		/// Whether to include NvCloth.
 		/// </summary>
-		[RequiresUniqueBuildEnvironment]
 		[Obsolete("Deprecated in UE5.1 - No longer used as Chaos is always enabled.")]
-		public bool bCompileNvCloth = false;
+		public bool bCompileNvCloth { get; set; }
 
 		/// <summary>
 		/// Whether to include ICU unicode/i18n support in Core.
 		/// </summary>
 		[RequiresUniqueBuildEnvironment]
 		[ConfigFile(ConfigHierarchyType.Engine, "/Script/BuildSettings.BuildSettings", "bCompileICU")]
-		public bool bCompileICU = true;
+		public bool bCompileICU { get; set; } = true;
 
 		/// <summary>
 		/// Whether to compile CEF3 support.
 		/// </summary>
 		[RequiresUniqueBuildEnvironment]
 		[ConfigFile(ConfigHierarchyType.Engine, "/Script/BuildSettings.BuildSettings", "bCompileCEF3")]
-		public bool bCompileCEF3 = true;
+		public bool bCompileCEF3 { get; set; } = true;
 
 		/// <summary>
 		/// Whether to compile using ISPC.
 		/// </summary>
 		[RequiresUniqueBuildEnvironment]
-		public bool bCompileISPC = false;
+		public bool bCompileISPC { get; set; }
 
 		/// <summary>
 		/// Whether to compile IntelMetricsDiscovery.
 		/// </summary>
 		[RequiresUniqueBuildEnvironment]
-		public bool bCompileIntelMetricsDiscovery = true;
+		public bool bCompileIntelMetricsDiscovery { get; set; } = true;
 
 		/// <summary>
 		/// Whether to compile in python support
 		/// </summary>
 		[RequiresUniqueBuildEnvironment]
-		public bool bCompilePython = true;
+		public bool bCompilePython { get; set; } = true;
 
 		/// <summary>
 		/// Whether to compile with WITH_GAMEPLAY_DEBUGGER enabled with all Engine's default gameplay debugger categories.
@@ -1016,7 +992,7 @@ namespace UnrealBuildTool
 		[RequiresUniqueBuildEnvironment]
 		[CommandLine("-NoUseIris", Value = "false")]
 		[CommandLine("-UseIris", Value = "true")]
-		public bool bUseIris = true;
+		public bool bUseIris { get; set; } = true;
 
 		/// <summary>
 		/// Whether we are compiling editor code or not. Prefer the more explicit bCompileAgainstEditor instead.
@@ -1024,7 +1000,8 @@ namespace UnrealBuildTool
 		public bool bBuildEditor
 		{
 			get => (Type == TargetType.Editor || bCompileAgainstEditor);
-			set => Logger.LogWarning("Setting {Type}.bBuildEditor is deprecated. Set {Type}.Type instead.", GetType().Name);
+			[Obsolete("Deprecated, replace with TargetRules.Type")]
+			set => Logger.LogWarning("Setting {Type}.bBuildEditor is deprecated. Set {Type}.Type instead.", GetType().Name, GetType().Name);
 		}
 
 		/// <summary>
@@ -1082,17 +1059,17 @@ namespace UnrealBuildTool
 		/// <summary>
 		/// Whether to force compiling the target platform modules, even if they wouldn't normally be built.
 		/// </summary>
-		public bool bForceBuildTargetPlatforms = false;
+		public bool bForceBuildTargetPlatforms { get; set; }
 
 		/// <summary>
 		/// Whether to force compiling shader format modules, even if they wouldn't normally be built.
 		/// </summary>
-		public bool bForceBuildShaderFormats = false;
+		public bool bForceBuildShaderFormats { get; set; }
 
 		/// <summary>
 		/// Override for including extra shader formats
 		/// </summary>
-		public bool? bNeedsExtraShaderFormatsOverride;
+		public bool? bNeedsExtraShaderFormatsOverride { get; set; }
 
 		/// <summary>
 		/// Whether we should include any extra shader formats. By default this is only enabled for Program and Editor targets.
@@ -1106,25 +1083,21 @@ namespace UnrealBuildTool
 		/// <summary>
 		/// Checks the above flags to see if target SDKs are likely to be compiled in
 		/// </summary>
-		private bool bUsesTargetSDKs
-		{
-			get => bNeedsExtraShaderFormats || bBuildDeveloperTools || bBuildTargetDeveloperTools || bForceBuildShaderFormats || bForceBuildTargetPlatforms;
-		}
-
+		private bool bUsesTargetSDKs => bNeedsExtraShaderFormats || bBuildDeveloperTools || bBuildTargetDeveloperTools || bForceBuildShaderFormats || bForceBuildTargetPlatforms;
 
 		/// <summary>
 		/// Whether we should compile SQLite using the custom "Unreal" platform (true), or using the native platform (false).
 		/// </summary>
 		[RequiresUniqueBuildEnvironment]
 		[ConfigFile(ConfigHierarchyType.Engine, "/Script/BuildSettings.BuildSettings", "bCompileCustomSQLitePlatform")]
-		public bool bCompileCustomSQLitePlatform = true;
+		public bool bCompileCustomSQLitePlatform { get; set; } = true;
 
 		/// <summary>
 		/// Whether to utilize cache freed OS allocs with MallocBinned
 		/// </summary>
 		[RequiresUniqueBuildEnvironment]
 		[ConfigFile(ConfigHierarchyType.Engine, "/Script/BuildSettings.BuildSettings", "bUseCacheFreedOSAllocs")]
-		public bool bUseCacheFreedOSAllocs = true;
+		public bool bUseCacheFreedOSAllocs { get; set; } = true;
 
 		/// <summary>
 		/// Enabled for all builds that include the engine project.  Disabled only when building standalone apps that only link with Core.
@@ -1165,7 +1138,7 @@ namespace UnrealBuildTool
 		[RequiresUniqueBuildEnvironment]
 		public virtual bool bEnableTrace
 		{
-			get 
+			get
 			{
 				if (bEnableTracePrivate.HasValue)
 				{
@@ -1175,7 +1148,7 @@ namespace UnrealBuildTool
 				{
 					return false;
 				}
-				if (UEBuildPlatform.IsPlatformInGroup(Platform, UnrealPlatformGroup.Apple) 
+				if (UEBuildPlatform.IsPlatformInGroup(Platform, UnrealPlatformGroup.Apple)
 				|| UEBuildPlatform.IsPlatformInGroup(Platform, UnrealPlatformGroup.Unix)
 				|| UEBuildPlatform.IsPlatformInGroup(Platform, UnrealPlatformGroup.Android)
 				|| UEBuildPlatform.IsPlatformInGroup(Platform, UnrealPlatformGroup.Windows))
@@ -1183,7 +1156,7 @@ namespace UnrealBuildTool
 					return true;
 				}
 				return false;
-			}  
+			}
 			set => bEnableTracePrivate = value;
 		}
 		private bool? bEnableTracePrivate;
@@ -1209,25 +1182,25 @@ namespace UnrealBuildTool
 		/// </summary>
 		[RequiresUniqueBuildEnvironment]
 		[ConfigFile(ConfigHierarchyType.Engine, "/Script/BuildSettings.BuildSettings", "bCompileRecast")]
-		public bool bCompileRecast = true;
+		public bool bCompileRecast { get; set; } = true;
 
 		/// <summary>
 		/// Whether to compile with navmesh segment links.
 		/// </summary>
 		[RequiresUniqueBuildEnvironment]
-		public bool bCompileNavmeshSegmentLinks = true;
+		public bool bCompileNavmeshSegmentLinks { get; set; } = true;
 
 		/// <summary>
 		/// Whether to compile with navmesh cluster links.
 		/// </summary>
 		[RequiresUniqueBuildEnvironment]
-		public bool bCompileNavmeshClusterLinks = true;
+		public bool bCompileNavmeshClusterLinks { get; set; } = true;
 
 		/// <summary>
 		/// Whether to compile SpeedTree support.
 		/// </summary>
 		[ConfigFile(ConfigHierarchyType.Engine, "/Script/BuildSettings.BuildSettings", "bCompileSpeedTree")]
-		bool? bOverrideCompileSpeedTree;
+		bool? bOverrideCompileSpeedTree { get; set; }
 
 		/// <summary>
 		/// Whether we should compile in support for SpeedTree or not.
@@ -1243,39 +1216,39 @@ namespace UnrealBuildTool
 		/// Enable exceptions for all modules.
 		/// </summary>
 		[RequiresUniqueBuildEnvironment]
-		public bool bForceEnableExceptions = false;
+		public bool bForceEnableExceptions { get; set; }
 
 		/// <summary>
 		/// Enable inlining for all modules.
 		/// </summary>
 		[RequiresUniqueBuildEnvironment]
 		[XmlConfigFile(Category = "BuildConfiguration")]
-		public bool bUseInlining = true;
+		public bool bUseInlining { get; set; } = true;
 
 		/// <summary>
 		/// Enable exceptions for all modules.
 		/// </summary>
 		[RequiresUniqueBuildEnvironment]
-		public bool bForceEnableObjCExceptions = false;
+		public bool bForceEnableObjCExceptions { get; set; }
 
 		/// <summary>
 		/// Enable RTTI for all modules.
 		/// </summary>
 		[CommandLine("-rtti")]
 		[RequiresUniqueBuildEnvironment]
-		public bool bForceEnableRTTI = false;
+		public bool bForceEnableRTTI { get; set; }
 
 		/// <summary>
 		/// Enable Position Independent Executable (PIE). Has an overhead cost
 		/// </summary>
 		[CommandLine("-pie")]
-		public bool bEnablePIE = false;
+		public bool bEnablePIE { get; set; }
 
 		/// <summary>
 		/// Enable Stack Protection. Has an overhead cost
 		/// </summary>
 		[CommandLine("-stack-protect")]
-		public bool bEnableStackProtection = false;
+		public bool bEnableStackProtection { get; set; }
 
 		/// <summary>
 		/// Compile server-only code.
@@ -1318,14 +1291,14 @@ namespace UnrealBuildTool
 		/// Whether to include stats support even without the engine.
 		/// </summary>
 		[RequiresUniqueBuildEnvironment]
-		public bool bCompileWithStatsWithoutEngine = false;
+		public bool bCompileWithStatsWithoutEngine { get; set; }
 
 		/// <summary>
 		/// Whether to include plugin support.
 		/// </summary>
 		[RequiresUniqueBuildEnvironment]
 		[ConfigFile(ConfigHierarchyType.Engine, "/Script/BuildSettings.BuildSettings", "bCompileWithPluginSupport")]
-		public bool bCompileWithPluginSupport = false;
+		public bool bCompileWithPluginSupport { get; set; }
 
 		/// <summary>
 		/// Whether to allow plugins which support all target platforms.
@@ -1342,7 +1315,7 @@ namespace UnrealBuildTool
 		/// Whether to allow accessibility code in both Slate and the OS layer.
 		/// </summary>
 		[RequiresUniqueBuildEnvironment]
-		public bool bCompileWithAccessibilitySupport = true;
+		public bool bCompileWithAccessibilitySupport { get; set; } = true;
 
 		/// <summary>
 		/// Whether to include PerfCounters support.
@@ -1373,147 +1346,145 @@ namespace UnrealBuildTool
 		/// </summary>
 		[RequiresUniqueBuildEnvironment]
 		[XmlConfigFile(Category = "BuildConfiguration")]
-		public bool bUseDebugLiveCodingConsole = false;
+		public bool bUseDebugLiveCodingConsole { get; set; }
 
 		/// <summary>
 		/// Whether to enable support for DirectX Math
 		/// LWC_TODO: No longer supported. Needs removing.
 		/// </summary>
 		[RequiresUniqueBuildEnvironment]
-		public bool bWithDirectXMath = false;
+		public bool bWithDirectXMath { get; set; }
 
 		/// <summary>
 		/// Whether to enable the support for FixedTimeStep in the engine
 		/// </summary>
 		[RequiresUniqueBuildEnvironment]
-		public bool bWithFixedTimeStepSupport = true;
+		public bool bWithFixedTimeStepSupport { get; set; } = true;
 
 		/// <summary>
 		/// Whether to turn on logging for test/shipping builds.
 		/// </summary>
 		[RequiresUniqueBuildEnvironment]
-		public bool bUseLoggingInShipping = false;
+		public bool bUseLoggingInShipping { get; set; }
 
 		/// <summary>
 		/// Whether to turn on console for shipping builds.
 		/// </summary>
 		[RequiresUniqueBuildEnvironment]
-		public bool bUseConsoleInShipping = false;
+		public bool bUseConsoleInShipping { get; set; }
 
 		/// <summary>
 		/// Whether to turn on logging to memory for test/shipping builds.
 		/// </summary>
 		[RequiresUniqueBuildEnvironment]
-		public bool bLoggingToMemoryEnabled;
+		public bool bLoggingToMemoryEnabled { get; set; }
 
 		/// <summary>
 		/// Whether to check that the process was launched through an external launcher.
 		/// </summary>
-		public bool bUseLauncherChecks = false;
+		public bool bUseLauncherChecks { get; set; }
 
 		/// <summary>
 		/// Whether to turn on checks (asserts) for test/shipping builds.
 		/// </summary>
 		[RequiresUniqueBuildEnvironment]
-		public bool bUseChecksInShipping = false;
+		public bool bUseChecksInShipping { get; set; }
 
 		/// <summary>
 		/// Whether to turn on GPU markers for Test builds.
 		/// </summary>
 		[RequiresUniqueBuildEnvironment]
-		public bool bAllowProfileGPUInTest = false;
+		public bool bAllowProfileGPUInTest { get; set; }
 
 		/// <summary>
 		/// Whether to turn on GPU markers for Shipping builds.
 		/// </summary>
 		[RequiresUniqueBuildEnvironment]
-		public bool bAllowProfileGPUInShipping = false;
+		public bool bAllowProfileGPUInShipping { get; set; }
 
 		/// <summary>
 		/// Whether to turn on UTF-8 mode, mapping TCHAR to UTF8CHAR.
 		/// </summary>
 		[RequiresUniqueBuildEnvironment]
-		public bool bTCHARIsUTF8 = false;
+		public bool bTCHARIsUTF8 { get; set; }
 
 		/// <summary>
 		/// Whether to use the EstimatedUtcNow or PlatformUtcNow.  EstimatedUtcNow is appropriate in
 		/// cases where PlatformUtcNow can be slow.
 		/// </summary>
 		[RequiresUniqueBuildEnvironment]
-		public bool bUseEstimatedUtcNow = false;
+		public bool bUseEstimatedUtcNow { get; set; }
 
 		/// <summary>
 		/// True if we need FreeType support.
 		/// </summary>
 		[RequiresUniqueBuildEnvironment]
 		[ConfigFile(ConfigHierarchyType.Engine, "/Script/BuildSettings.BuildSettings", "bCompileFreeType")]
-		public bool bCompileFreeType = true;
+		public bool bCompileFreeType { get; set; } = true;
 
 		/// <summary>
 		/// Whether to turn allow exec commands for shipping builds.
 		/// </summary>
 		[RequiresUniqueBuildEnvironment]
-		public bool bUseExecCommandsInShipping = true;
+		public bool bUseExecCommandsInShipping { get; set; } = true;
 
 		/// <summary>
 		/// True if we want to favor optimizing size over speed.
 		/// </summary>
-		[RequiresUniqueBuildEnvironment]
 		[Obsolete("Deprecated in UE5.1 - Please use OptimizationLevel instead.")]
-		[ConfigFile(ConfigHierarchyType.Engine, "/Script/BuildSettings.BuildSettings", "bCompileForSize")]
-		public bool bCompileForSize = false;
+		public bool bCompileForSize { get; set; }
 
 		/// <summary>
 		/// Allows to fine tune optimizations level for speed and\or code size
 		/// </summary>
 		[RequiresUniqueBuildEnvironment]
 		[CommandLine("-OptimizationLevel=")]
-		public OptimizationMode OptimizationLevel = OptimizationMode.Speed;
+		public OptimizationMode OptimizationLevel { get; set; } = OptimizationMode.Speed;
 
 		/// <summary>
 		/// Allows setting the FP semantics.
 		/// </summary>
 		[RequiresUniqueBuildEnvironment]
 		[CommandLine("-FPSemantics=")]
-		public FPSemanticsMode FPSemantics = FPSemanticsMode.Default;
+		public FPSemanticsMode FPSemantics { get; set; } = FPSemanticsMode.Default;
 
 		/// <summary>
 		/// Whether to compile development automation tests.
 		/// </summary>
-		public bool bForceCompileDevelopmentAutomationTests = false;
+		public bool bForceCompileDevelopmentAutomationTests { get; set; }
 
 		/// <summary>
 		/// Whether to compile performance automation tests.
 		/// </summary>
-		public bool bForceCompilePerformanceAutomationTests = false;
+		public bool bForceCompilePerformanceAutomationTests { get; set; }
 
 		/// <summary>
 		/// Whether to override the defaults for automation tests (Debug/Development configs)
 		/// </summary>
-		public bool bForceDisableAutomationTests = false;
+		public bool bForceDisableAutomationTests { get; set; }
 
 		/// <summary>
 		/// If true, event driven loader will be used in cooked builds. @todoio This needs to be replaced by a runtime solution after async loading refactor.
 		/// </summary>
 		[RequiresUniqueBuildEnvironment]
-		public bool bEventDrivenLoader;
+		public bool bEventDrivenLoader { get; set; }
 
 		/// <summary>
 		/// Used to override the behavior controlling whether UCLASSes and USTRUCTs are allowed to have native pointer members, if disallowed they will be a UHT error and should be substituted with TObjectPtr members instead.
 		/// </summary>
-		public PointerMemberBehavior? NativePointerMemberBehaviorOverride = null;
+		public PointerMemberBehavior? NativePointerMemberBehaviorOverride { get; set; }
 
 		/// <summary>
 		/// Whether the XGE controller worker and modules should be included in the engine build.
 		/// These are required for distributed shader compilation using the XGE interception interface.
 		/// </summary>
 		[XmlConfigFile(Category = "BuildConfiguration")]
-		public bool bUseXGEController = true;
+		public bool bUseXGEController { get; set; } = true;
 
 		/// <summary>
 		/// Enables "include what you use" by default for modules in this target. Changes the default PCH mode for any module in this project to PCHUsageMode.UseExplicitOrSharedPCHs.
 		/// </summary>
-		public bool bIWYU { get => IntermediateEnvironment == UnrealIntermediateEnvironment.IWYU; }
+		public bool bIWYU => IntermediateEnvironment == UnrealIntermediateEnvironment.IWYU;
 
 		/// <summary>
 		/// Adds header files in included modules to the build.
@@ -1530,17 +1501,17 @@ namespace UnrealBuildTool
 		/// When used with -IncludeHeaders, only header files will be compiled.
 		/// </summary>
 		[CommandLine("-HeadersOnly")]
-		public bool bHeadersOnly = false;
+		public bool bHeadersOnly { get; set; }
 
 		/// <summary>
 		/// Enforce "include what you use" rules; warns if monolithic headers (Engine.h, UnrealEd.h, etc...) are used, and checks that source files include their matching header first.
 		/// </summary>
-		public bool bEnforceIWYU = true;
+		public bool bEnforceIWYU { get; set; } = true;
 
 		/// <summary>
 		/// Emit a warning when an old-style monolithic header is included while compiling this target.
 		/// </summary>
-		public bool bWarnAboutMonolithicHeadersIncluded = false;
+		public bool bWarnAboutMonolithicHeadersIncluded { get; set; }
 
 		/// <summary>
 		/// Whether the final executable should export symbols.
@@ -1556,17 +1527,17 @@ namespace UnrealBuildTool
 		/// Make static libraries for all engine modules as intermediates for this target.
 		/// </summary>
 		[CommandLine("-Precompile")]
-		public bool bPrecompile = false;
+		public bool bPrecompile { get; set; }
 
 		/// <summary>
 		/// Whether we should compile with support for OS X 10.9 Mavericks. Used for some tools that we need to be compatible with this version of OS X.
 		/// </summary>
-		public bool bEnableOSX109Support = false;
+		public bool bEnableOSX109Support { get; set; }
 
 		/// <summary>
 		/// True if this is a console application that's being built.
 		/// </summary>
-		public bool bIsBuildingConsoleApplication = false;
+		public bool bIsBuildingConsoleApplication { get; set; }
 
 		/// <summary>
 		/// If true, creates an additional console application. Hack for Windows, where it's not possible to conditionally inherit a parent's console Window depending on how
@@ -1582,7 +1553,7 @@ namespace UnrealBuildTool
 		/// <summary>
 		/// True if debug symbols that are cached for some platforms should not be created.
 		/// </summary>
-		public bool bDisableSymbolCache = true;
+		public bool bDisableSymbolCache { get; set; } = true;
 
 		/// <summary>
 		/// Whether to unify C++ code into larger files for faster compilation.
@@ -1596,22 +1567,20 @@ namespace UnrealBuildTool
 		/// <summary>
 		/// Whether to unify C++ code into larger files for faster compilation.
 		/// </summary>
-		bool? bUseUnityBuildOverride = null;
+		bool? bUseUnityBuildOverride { get; set; }
 
 		/// <summary>
 		/// Whether to force C++ source files to be combined into larger files for faster compilation.
 		/// </summary>
 		[CommandLine("-ForceUnity")]
 		[XmlConfigFile(Category = "BuildConfiguration")]
-		public bool bForceUnityBuild = false;
+		public bool bForceUnityBuild { get; set; }
 
 		/// <summary>
 		/// Whether to merge module and generated unity files for faster compilation.
 		/// </summary>
 		[Obsolete("Deprecated in UE5.3; use DisableMergingModuleAndGeneratedFilesInUnityFiles instead")]
-		[CommandLine("-DisableMergingUnityFiles", Value = "false")]
-		[XmlConfigFile(Category = "BuildConfiguration")]
-		public bool bMergeModuleAndGeneratedUnityFiles = true;
+		public bool bMergeModuleAndGeneratedUnityFiles { get; set; }
 
 		/// <summary>
 		/// List of modules that disables merging module and generated cpp files in the same unity files.
@@ -1626,25 +1595,25 @@ namespace UnrealBuildTool
 		/// </summary>
 		[CommandLine("-DisableAdaptiveUnity", Value = "false")]
 		[XmlConfigFile(Category = "BuildConfiguration")]
-		public bool bUseAdaptiveUnityBuild = true;
+		public bool bUseAdaptiveUnityBuild { get; set; } = true;
 
 		/// <summary>
 		/// Disable optimization for files that are in the adaptive non-unity working set.
 		/// </summary>
 		[XmlConfigFile(Category = "BuildConfiguration")]
-		public bool bAdaptiveUnityDisablesOptimizations = false;
+		public bool bAdaptiveUnityDisablesOptimizations { get; set; }
 
 		/// <summary>
 		/// Disables force-included PCHs for files that are in the adaptive non-unity working set.
 		/// </summary>
 		[XmlConfigFile(Category = "BuildConfiguration")]
-		public bool bAdaptiveUnityDisablesPCH = false;
+		public bool bAdaptiveUnityDisablesPCH { get; set; }
 
 		/// <summary>
 		/// Backing storage for bAdaptiveUnityDisablesProjectPCH.
 		/// </summary>
 		[XmlConfigFile(Category = "BuildConfiguration")]
-		bool? bAdaptiveUnityDisablesProjectPCHForProjectPrivate;
+		public bool? bAdaptiveUnityDisablesProjectPCHForProjectPrivate { get; set; }
 
 		/// <summary>
 		/// Whether to disable force-included PCHs for project source files in the adaptive non-unity working set. Defaults to bAdaptiveUnityDisablesPCH;
@@ -1659,19 +1628,19 @@ namespace UnrealBuildTool
 		/// Creates a dedicated PCH for each source file in the working set, allowing faster iteration on cpp-only changes.
 		/// </summary>
 		[XmlConfigFile(Category = "BuildConfiguration")]
-		public bool bAdaptiveUnityCreatesDedicatedPCH = false;
+		public bool bAdaptiveUnityCreatesDedicatedPCH { get; set; }
 
 		/// <summary>
 		/// Creates a dedicated PCH for each source file in the working set, allowing faster iteration on cpp-only changes.
 		/// </summary>
 		[XmlConfigFile(Category = "BuildConfiguration")]
-		public bool bAdaptiveUnityEnablesEditAndContinue = false;
+		public bool bAdaptiveUnityEnablesEditAndContinue { get; set; }
 
 		/// <summary>
 		/// Creates a dedicated source file for each header file in the working set to detect missing includes in headers.
 		/// </summary>
 		[XmlConfigFile(Category = "BuildConfiguration")]
-		public bool bAdaptiveUnityCompilesHeaderFiles = false;
+		public bool bAdaptiveUnityCompilesHeaderFiles { get; set; }
 
 		/// <summary>
 		/// The number of source files in a game module before unity build will be activated for that module.  This
@@ -1679,7 +1648,7 @@ namespace UnrealBuildTool
 		/// rebuild times.  This setting can be overridden by the bFasterWithoutUnity option in a module's Build.cs file.
 		/// </summary>
 		[XmlConfigFile(Category = "BuildConfiguration")]
-		public int MinGameModuleSourceFilesForUnityBuild = 32;
+		public int MinGameModuleSourceFilesForUnityBuild { get; set; } = 32;
 
 		/// <summary>
 		/// Default treatment of uncategorized warnings
@@ -1695,7 +1664,6 @@ namespace UnrealBuildTool
 		[XmlConfigFile(Category = "BuildConfiguration", Name = nameof(DefaultWarningLevel))]
 		private WarningLevel DefaultWarningLevelPrivate;
 
-
 		/// <summary>
 		/// Require TObjectPtr for FReferenceCollector API's. (Needed for compatibility with incremental GC.)
 		/// </summary>
@@ -1703,10 +1671,10 @@ namespace UnrealBuildTool
 		[XmlConfigFile(Category = "BuildConfiguration")]
 		public bool bRequireObjectPtrForAddReferencedObjects
 		{
-				get { return bRequireObjectPtrForAddReferencedObjectsPrivate ?? true; }
-				set { bRequireObjectPtrForAddReferencedObjectsPrivate = value; }
+			get => bRequireObjectPtrForAddReferencedObjectsPrivate ?? true;
+			set => bRequireObjectPtrForAddReferencedObjectsPrivate = value;
 		}
-		private bool? bRequireObjectPtrForAddReferencedObjectsPrivate;		
+		private bool? bRequireObjectPtrForAddReferencedObjectsPrivate;
 
 		/// <summary>
 		/// Emits compilation errors for incorrect UE_LOG format strings.
@@ -1714,8 +1682,8 @@ namespace UnrealBuildTool
 		[XmlConfigFile(Category = "BuildConfiguration")]
 		public bool bValidateFormatStrings
 		{
-				get { return bValidateFormatStringsPrivate ?? (DefaultBuildSettings >= BuildSettingsVersion.V5); }
-				set { bValidateFormatStringsPrivate = value; }
+			get => bValidateFormatStringsPrivate ?? (DefaultBuildSettings >= BuildSettingsVersion.V5);
+			set => bValidateFormatStringsPrivate = value;
 		}
 		private bool? bValidateFormatStringsPrivate;
 
@@ -1738,7 +1706,7 @@ namespace UnrealBuildTool
 		/// </summary>
 		[CommandLine("-ShadowVariableErrors", Value = nameof(WarningLevel.Error))]
 		[RequiresUniqueBuildEnvironment]
-		public WarningLevel ShadowVariableWarningLevel = WarningLevel.Warning;
+		public WarningLevel ShadowVariableWarningLevel { get; set; } = WarningLevel.Warning;
 
 		/// <summary>
 		/// Whether to enable all warnings as errors. UE enables most warnings as errors already, but disables a few (such as deprecation warnings).
@@ -1746,49 +1714,49 @@ namespace UnrealBuildTool
 		[XmlConfigFile(Category = "BuildConfiguration")]
 		[CommandLine("-WarningsAsErrors")]
 		[RequiresUniqueBuildEnvironment]
-		public bool bWarningsAsErrors = false;
+		public bool bWarningsAsErrors { get; set; }
 
 		/// <summary>
 		/// Indicates what warning/error level to treat unsafe type casts as on platforms that support it (e.g., double->float or int64->int32)
 		/// </summary>
 		[XmlConfigFile(Category = "BuildConfiguration")]
 		[RequiresUniqueBuildEnvironment]
-		public WarningLevel UnsafeTypeCastWarningLevel = WarningLevel.Off;
+		public WarningLevel UnsafeTypeCastWarningLevel { get; set; } = WarningLevel.Off;
 
 		/// <summary>
 		/// Forces the use of undefined identifiers in conditional expressions to be treated as errors.
 		/// </summary>
 		[XmlConfigFile(Category = "BuildConfiguration")]
 		[RequiresUniqueBuildEnvironment]
-		public bool bUndefinedIdentifierErrors = true;
+		public bool bUndefinedIdentifierErrors { get; set; } = true;
 
 		/// <summary>
 		/// Indicates what warning/error level to treat potential PCH performance issues.
 		/// </summary>
 		[XmlConfigFile(Category = "BuildConfiguration")]
 		[CommandLine("-PCHPerformanceIssueWarningLevel=")]
-		public WarningLevel PCHPerformanceIssueWarningLevel = WarningLevel.Off;
+		public WarningLevel PCHPerformanceIssueWarningLevel { get; set; } = WarningLevel.Off;
 
 		/// <summary>
 		/// How to treat general module include path validation messages
 		/// </summary>
 		[XmlConfigFile(Category = "BuildConfiguration")]
 		[CommandLine("-ModuleIncludePathWarningLevel=")]
-		public WarningLevel ModuleIncludePathWarningLevel = WarningLevel.Off;
+		public WarningLevel ModuleIncludePathWarningLevel { get; set; } = WarningLevel.Off;
 
 		/// <summary>
 		/// How to treat private module include path validation messages, where a module is adding an include path that exposes private headers
 		/// </summary>
 		[XmlConfigFile(Category = "BuildConfiguration")]
 		[CommandLine("-ModuleIncludePrivateWarningLevel=")]
-		public WarningLevel ModuleIncludePrivateWarningLevel = WarningLevel.Off;
+		public WarningLevel ModuleIncludePrivateWarningLevel { get; set; } = WarningLevel.Off;
 
 		/// <summary>
 		/// How to treat unnecessary module sub-directory include path validation messages
 		/// </summary>
 		[XmlConfigFile(Category = "BuildConfiguration")]
 		[CommandLine("-ModuleIncludeSubdirectoryWarningLevel=")]
-		public WarningLevel ModuleIncludeSubdirectoryWarningLevel = WarningLevel.Off;
+		public WarningLevel ModuleIncludeSubdirectoryWarningLevel { get; set; } = WarningLevel.Off;
 
 		/// <summary>
 		/// Forces frame pointers to be retained this is usually required when you want reliable callstacks e.g. mallocframeprofiler
@@ -1807,7 +1775,7 @@ namespace UnrealBuildTool
 		[XmlConfigFile(Category = "BuildConfiguration", Name = "bRetainFramePointers")]
 		[CommandLine("-RetainFramePointers", Value = "true")]
 		[CommandLine("-NoRetainFramePointers", Value = "false")]
-		public bool? bRetainFramePointersOverride = null;
+		public bool? bRetainFramePointersOverride { get; set; }
 
 		/// <summary>
 		/// New Monolithic Graphics drivers have optional "fast calls" replacing various D3d functions
@@ -1815,28 +1783,28 @@ namespace UnrealBuildTool
 		[CommandLine("-FastMonoCalls", Value = "true")]
 		[CommandLine("-NoFastMonoCalls", Value = "false")]
 		[XmlConfigFile(Category = "BuildConfiguration")]
-		public bool bUseFastMonoCalls = true;
+		public bool bUseFastMonoCalls { get; set; } = true;
 
 		/// <summary>
 		/// An approximate number of bytes of C++ code to target for inclusion in a single unified C++ file.
 		/// </summary>
 		[CommandLine("-BytesPerUnityCPP")]
 		[XmlConfigFile(Category = "BuildConfiguration")]
-		public int NumIncludedBytesPerUnityCPP = 384 * 1024;
+		public int NumIncludedBytesPerUnityCPP { get; set; } = 384 * 1024;
 
 		/// <summary>
 		/// Disables overrides that are set by the module
 		/// </summary>
 		[CommandLine("-DisableModuleNumIncludedBytesPerUnityCPPOverride", Value = "true")]
 		[XmlConfigFile(Category = "BuildConfiguration")]
-		public bool bDisableModuleNumIncludedBytesPerUnityCPPOverride = false;
+		public bool bDisableModuleNumIncludedBytesPerUnityCPPOverride { get; set; }
 
 		/// <summary>
 		/// Whether to stress test the C++ unity build robustness by including all C++ files files in a project from a single unified file.
 		/// </summary>
 		[CommandLine("-StressTestUnity")]
 		[XmlConfigFile(Category = "BuildConfiguration")]
-		public bool bStressTestUnity = false;
+		public bool bStressTestUnity { get; set; }
 
 		/// <summary>
 		/// Whether to add additional information to the unity files, such as '_of_X' in the file name. Not recommended.
@@ -1844,13 +1812,13 @@ namespace UnrealBuildTool
 		[CommandLine("-DetailedUnityFiles", Value = "true")]
 		[CommandLine("-NoDetailedUnityFiles", Value = "false")]
 		[XmlConfigFile(Category = "BuildConfiguration")]
-		public bool bDetailedUnityFiles = false;
+		public bool bDetailedUnityFiles { get; set; }
 
 		/// <summary>
 		/// Whether to force debug info to be generated.
 		/// </summary>
 		[CommandLine("-ForceDebugInfo")]
-		public bool bForceDebugInfo = false;
+		public bool bForceDebugInfo { get; set; }
 
 		/// <summary>
 		/// Whether to globally disable debug info generation; Obsolete, please use TargetRules.DebugInfoMode instead
@@ -1907,13 +1875,13 @@ namespace UnrealBuildTool
 		/// Whether to disable debug info generation for generated files. This improves link times and reduces pdb size for modules that have a lot of generated glue code.
 		/// </summary>
 		[XmlConfigFile(Category = "BuildConfiguration")]
-		public bool bDisableDebugInfoForGeneratedCode = false;
+		public bool bDisableDebugInfoForGeneratedCode { get; set; }
 
 		/// <summary>
 		/// Whether to disable debug info on PC/Mac in development builds (for faster developer iteration, as link times are extremely fast with debug info disabled).
 		/// </summary>
 		[XmlConfigFile(Category = "BuildConfiguration")]
-		public bool bOmitPCDebugInfoInDevelopment = false;
+		public bool bOmitPCDebugInfoInDevelopment { get; set; }
 
 		/// <summary>
 		/// Whether PDB files should be used for Visual C++ builds.
@@ -1921,14 +1889,14 @@ namespace UnrealBuildTool
 		[CommandLine("-NoPDB", Value = "false")]
 		[CommandLine("-UsePDB", Value = "true")]
 		[XmlConfigFile(Category = "BuildConfiguration")]
-		public bool bUsePDBFiles = false;
+		public bool bUsePDBFiles { get; set; }
 
 		/// <summary>
 		/// Whether PCH files should be used.
 		/// </summary>
 		[CommandLine("-NoPCH", Value = "false")]
 		[XmlConfigFile(Category = "BuildConfiguration")]
-		public bool bUsePCHFiles = true;
+		public bool bUsePCHFiles { get; set; } = true;
 
 		/// <summary>
 		/// Set flags require for deterministic compiling and linking.
@@ -1950,46 +1918,46 @@ namespace UnrealBuildTool
 		/// </summary>
 		[CommandLine("-NoPCHChain", Value = "false")]
 		[XmlConfigFile(Category = "BuildConfiguration")]
-		public bool bChainPCHs = true;
+		public bool bChainPCHs { get; set; } = true;
 
 		/// <summary>
 		/// Whether PCH headers should be force included for gen.cpp files when PCH is disabled.
 		/// </summary>
 		[XmlConfigFile(Category = "BuildConfiguration")]
-		public bool bForceIncludePCHHeadersForGenCppFilesWhenPCHIsDisabled = true;
+		public bool bForceIncludePCHHeadersForGenCppFilesWhenPCHIsDisabled { get; set; } = true;
 
 		/// <summary>
 		/// Whether to just preprocess source files for this target, and skip compilation
 		/// </summary>
 		[CommandLine("-Preprocess")]
-		public bool bPreprocessOnly = false;
+		public bool bPreprocessOnly { get; set; }
 
 		/// <summary>
 		/// Generate dependency files by preprocessing. This is only recommended when distributing builds as it adds additional overhead.
 		/// </summary>
 		[CommandLine("-PreprocessDepends")]
 		[XmlConfigFile(Category = "BuildConfiguration")]
-		public bool bPreprocessDepends = false;
+		public bool bPreprocessDepends { get; set; }
 
 		/// <summary>
 		/// Whether to generate assembly data while compiling this target. Works exclusively on MSVC for now.
 		/// </summary>
 		[CommandLine("-WithAssembly")]
-		public bool bWithAssembly = false;
+		public bool bWithAssembly { get; set; }
 
 		/// <summary>
 		/// Whether static code analysis should be enabled.
 		/// </summary>
 		[CommandLine("-StaticAnalyzer")]
 		[XmlConfigFile(Category = "BuildConfiguration")]
-		public StaticAnalyzer StaticAnalyzer = StaticAnalyzer.None;
+		public StaticAnalyzer StaticAnalyzer { get; set; } = StaticAnalyzer.None;
 
 		/// <summary>
 		/// The output type to use for the static analyzer. This is only supported for Clang.
 		/// </summary>
 		[CommandLine("-StaticAnalyzerOutputType")]
 		[XmlConfigFile(Category = "BuildConfiguration")]
-		public StaticAnalyzerOutputType StaticAnalyzerOutputType = StaticAnalyzerOutputType.Text;
+		public StaticAnalyzerOutputType StaticAnalyzerOutputType { get; set; } = StaticAnalyzerOutputType.Text;
 
 		/// <summary>
 		/// The mode to use for the static analyzer. This is only supported for Clang.
@@ -1997,13 +1965,13 @@ namespace UnrealBuildTool
 		/// </summary>
 		[CommandLine("-StaticAnalyzerMode")]
 		[XmlConfigFile(Category = "BuildConfiguration")]
-		public StaticAnalyzerMode StaticAnalyzerMode = StaticAnalyzerMode.Deep;
+		public StaticAnalyzerMode StaticAnalyzerMode { get; set; } = StaticAnalyzerMode.Deep;
 
 		/// <summary>
 		/// The minimum number of files that must use a pre-compiled header before it will be created and used.
 		/// </summary>
 		[XmlConfigFile(Category = "BuildConfiguration")]
-		public int MinFilesUsingPrecompiledHeader = 6;
+		public int MinFilesUsingPrecompiledHeader { get; set; } = 6;
 
 		/// <summary>
 		/// When enabled, a precompiled header is always generated for game modules, even if there are only a few source files
@@ -2012,7 +1980,7 @@ namespace UnrealBuildTool
 		/// a module's Build.cs file.
 		/// </summary>
 		[XmlConfigFile(Category = "BuildConfiguration")]
-		public bool bForcePrecompiledHeaderForGameModules = true;
+		public bool bForcePrecompiledHeaderForGameModules { get; set; } = true;
 
 		/// <summary>
 		/// Whether to use incremental linking or not. Incremental linking can yield faster iteration times when making small changes.
@@ -2021,14 +1989,14 @@ namespace UnrealBuildTool
 		[CommandLine("-IncrementalLinking")]
 		[CommandLine("-NoIncrementalLinking", Value = "false")]
 		[XmlConfigFile(Category = "BuildConfiguration")]
-		public bool bUseIncrementalLinking = false;
+		public bool bUseIncrementalLinking { get; set; }
 
 		/// <summary>
 		/// Whether to allow the use of link time code generation (LTCG).
 		/// </summary>
 		[CommandLine("-LTCG")]
 		[XmlConfigFile(Category = "BuildConfiguration")]
-		public bool bAllowLTCG = false;
+		public bool bAllowLTCG { get; set; }
 
 		/// <summary>
 		/// When Link Time Code Generation (LTCG) is enabled, whether to 
@@ -2036,14 +2004,14 @@ namespace UnrealBuildTool
 		/// </summary>
 		[CommandLine("-ThinLTO")]
 		[XmlConfigFile(Category = "BuildConfiguration")]
-		public bool bPreferThinLTO = false;
+		public bool bPreferThinLTO { get; set; }
 
 		/// <summary>
 		/// Directory where to put the ThinLTO cache on supported platforms.
 		/// </summary>
 		[CommandLine("-ThinLTOCacheDirectory")]
 		[XmlConfigFile(Category = "BuildConfiguration")]
-		public string? ThinLTOCacheDirectory = null;
+		public string? ThinLTOCacheDirectory { get; set; }
 
 		/// <summary>
 		/// Arguments that will be applied to prune the ThinLTO cache on supported platforms.
@@ -2051,73 +2019,72 @@ namespace UnrealBuildTool
 		/// </summary>
 		[CommandLine("-ThinLTOCachePruningArguments")]
 		[XmlConfigFile(Category = "BuildConfiguration")]
-		public string? ThinLTOCachePruningArguments = null;
+		public string? ThinLTOCachePruningArguments { get; set; }
 
 		/// <summary>
 		/// Whether to enable Profile Guided Optimization (PGO) instrumentation in this build.
 		/// </summary>
 		[CommandLine("-PGOProfile", Value = "true")]
 		[XmlConfigFile(Category = "BuildConfiguration")]
-		public bool bPGOProfile = false;
+		public bool bPGOProfile { get; set; }
 
 		/// <summary>
 		/// Whether to optimize this build with Profile Guided Optimization (PGO).
 		/// </summary>
 		[CommandLine("-PGOOptimize", Value = "true")]
 		[XmlConfigFile(Category = "BuildConfiguration")]
-		public bool bPGOOptimize = false;
+		public bool bPGOOptimize { get; set; }
 
 		/// <summary>
 		/// Whether the target requires code coverage compilation and linking.
 		/// </summary>
 		[CommandLine("-CodeCoverage", Value = "true")]
 		[XmlConfigFile(Category = "BuildConfiguration")]
-		public bool bCodeCoverage;
+		public bool bCodeCoverage { get; set; }
 
 		/// <summary>
 		/// Whether to support edit and continue.
 		/// </summary>
 		[CommandLine("-SupportEditAndContinue")]
 		[XmlConfigFile(Category = "BuildConfiguration")]
-		public bool bSupportEditAndContinue = false;
+		public bool bSupportEditAndContinue { get; set; }
 
 		/// <summary>
 		/// Whether to omit frame pointers or not. Disabling is useful for e.g. memory profiling on the PC.
 		/// </summary>
 		[XmlConfigFile(Category = "BuildConfiguration")]
-		public bool bOmitFramePointers = true;
+		public bool bOmitFramePointers { get; set; } = true;
 
 		/// <summary>
 		/// Whether to enable support for C++20 modules
 		/// </summary>
-		public bool bEnableCppModules = false;
+		public bool bEnableCppModules { get; set; }
 
 		/// <summary>
 		/// Whether to enable support for C++20 coroutines
 		/// This option is provided to facilitate evaluation of the feature.
 		/// Expect the name of the option to change. Use of coroutines in with UE is untested and unsupported.
 		/// </summary>
-		public bool bEnableCppCoroutinesForEvaluation = false;
+		public bool bEnableCppCoroutinesForEvaluation { get; set; }
 
 		/// <summary>
 		/// Whether to enable engine's ability to set process priority on runtime.
 		/// This option requires some environment setup on Linux, that's why it's disabled by default.
 		/// Project has to opt-in this feature as it has to guarantee correct setup.
 		/// </summary>
-		public bool bEnableProcessPriorityControl = false;
+		public bool bEnableProcessPriorityControl { get; set; }
 
 		/// <summary>
 		/// If true, then enable memory profiling in the build (defines USE_MALLOC_PROFILER=1 and forces bOmitFramePointers=false).
 		/// </summary>
 		[Obsolete("Deprecated in UE5.3 - No longer used as MallocProfiler is removed in favor of UnrealInsights.")]
-		[XmlConfigFile(Category = "BuildConfiguration")]
-		public bool bUseMallocProfiler = false;
+		public bool bUseMallocProfiler { get; set; }
 
 		/// <summary>
 		/// If true, then enable Unreal Insights (utrace) profiling in the build for the Shader Compiler Worker (defines USE_SHADER_COMPILER_WORKER_TRACE=1).
 		/// </summary>
 		[XmlConfigFile(Category = "BuildConfiguration")]
-		public bool bShaderCompilerWorkerTrace = false;
+		public bool bShaderCompilerWorkerTrace { get; set; }
 
 		/// <summary>
 		/// Enables "Shared PCHs", a feature which significantly speeds up compile times by attempting to
@@ -2125,52 +2092,52 @@ namespace UnrealBuildTool
 		/// </summary>
 		[CommandLine("-NoSharedPCH", Value = "false")]
 		[XmlConfigFile(Category = "BuildConfiguration")]
-		public bool bUseSharedPCHs = true;
+		public bool bUseSharedPCHs { get; set; } = true;
 
 		/// <summary>
 		/// True if Development and Release builds should use the release configuration of PhysX/APEX.
 		/// </summary>
 		[XmlConfigFile(Category = "BuildConfiguration")]
-		public bool bUseShippingPhysXLibraries = false;
+		public bool bUseShippingPhysXLibraries { get; set; }
 
 		/// <summary>
 		/// True if Development and Release builds should use the checked configuration of PhysX/APEX. if bUseShippingPhysXLibraries is true this is ignored.
 		/// </summary>
 		[XmlConfigFile(Category = "BuildConfiguration")]
-		public bool bUseCheckedPhysXLibraries = false;
+		public bool bUseCheckedPhysXLibraries { get; set; }
 
 		/// <summary>
 		/// Tells the UBT to check if module currently being built is violating EULA.
 		/// </summary>
 		[XmlConfigFile(Category = "BuildConfiguration")]
-		public bool bCheckLicenseViolations = true;
+		public bool bCheckLicenseViolations { get; set; } = true;
 
 		/// <summary>
 		/// Tells the UBT to break build if module currently being built is violating EULA.
 		/// </summary>
 		[XmlConfigFile(Category = "BuildConfiguration")]
-		public bool bBreakBuildOnLicenseViolation = true;
+		public bool bBreakBuildOnLicenseViolation { get; set; } = true;
 
 		/// <summary>
 		/// Whether to use the :FASTLINK option when building with /DEBUG to create local PDBs on Windows. Fast, but currently seems to have problems finding symbols in the debugger.
 		/// </summary>
 		[CommandLine("-FastPDB")]
 		[XmlConfigFile(Category = "BuildConfiguration")]
-		public bool? bUseFastPDBLinking;
+		public bool? bUseFastPDBLinking { get; set; }
 
 		/// <summary>
 		/// Outputs a map file as part of the build.
 		/// </summary>
 		[CommandLine("-MapFile")]
 		[XmlConfigFile(Category = "BuildConfiguration")]
-		public bool bCreateMapFile = false;
+		public bool bCreateMapFile { get; set; }
 
 		/// <summary>
 		/// True if runtime symbols files should be generated as a post build step for some platforms.
 		/// These files are used by the engine to resolve symbol names of callstack backtraces in logs.
 		/// </summary>
 		[XmlConfigFile(Category = "BuildConfiguration")]
-		public bool bAllowRuntimeSymbolFiles = true;
+		public bool bAllowRuntimeSymbolFiles { get; set; } = true;
 
 		/// <summary>
 		/// Package full path (directory + filename) where to store input files used at link time 
@@ -2178,60 +2145,60 @@ namespace UnrealBuildTool
 		/// </summary>
 		[CommandLine("-PackagePath")]
 		[XmlConfigFile(Category = "BuildConfiguration")]
-		public string? PackagePath = null;
+		public string? PackagePath { get; set; }
 
 		/// <summary>
 		/// Directory where to put crash report files for platforms that support it
 		/// </summary>
 		[CommandLine("-CrashDiagnosticDirectory")]
 		[XmlConfigFile(Category = "BuildConfiguration")]
-		public string? CrashDiagnosticDirectory = null;
+		public string? CrashDiagnosticDirectory { get; set; }
 
 		/// <summary>
 		/// Bundle version for Mac apps.
 		/// </summary>
 		[CommandLine("-BundleVersion")]
-		public string? BundleVersion = null;
+		public string? BundleVersion { get; set; }
 
 		/// <summary>
 		/// Whether to deploy the executable after compilation on platforms that require deployment.
 		/// </summary>
 		[CommandLine("-Deploy")]
-		public bool bDeployAfterCompile = false;
+		public bool bDeployAfterCompile { get; set; }
 
 		/// <summary>
 		/// Whether to force skipping deployment for platforms that require deployment by default.
 		/// </summary>
 		[CommandLine("-SkipDeploy")]
-		private bool bForceSkipDeploy = false;
+		private bool bForceSkipDeploy { get; set; }
 
 		/// <summary>
 		/// When enabled, allows XGE to compile pre-compiled header files on remote machines.  Otherwise, PCHs are always generated locally.
 		/// </summary>
-		public bool bAllowRemotelyCompiledPCHs = false;
+		public bool bAllowRemotelyCompiledPCHs { get; set; }
 
 		/// <summary>
 		/// Will replace pch with ifc and use header units instead. This is an experimental and msvc-only feature
 		/// </summary>
 		[CommandLine("-HeaderUnits")]
-		public bool bUseHeaderUnitsForPch = false;
+		public bool bUseHeaderUnitsForPch { get; set; }
 
 		/// <summary>
 		/// Whether headers in system paths should be checked for modification when determining outdated actions.
 		/// </summary>
 		[XmlConfigFile(Category = "BuildConfiguration")]
-		public bool bCheckSystemHeadersForModification;
+		public bool bCheckSystemHeadersForModification { get; set; }
 
 		/// <summary>
 		/// Whether to disable linking for this target.
 		/// </summary>
 		[CommandLine("-NoLink")]
-		public bool bDisableLinking = false;
+		public bool bDisableLinking { get; set; }
 
 		/// <summary>
 		/// Whether to ignore tracking build outputs for this target.
 		/// </summary>
-		public bool bIgnoreBuildOutputs = false;
+		public bool bIgnoreBuildOutputs { get; set; }
 
 		/// <summary>
 		/// Indicates that this is a formal build, intended for distribution. This flag is automatically set to true when Build.version has a changelist set and is a promoted build.
@@ -2239,76 +2206,74 @@ namespace UnrealBuildTool
 		/// By default, we only compile the resource once to reduce build times.
 		/// </summary>
 		[CommandLine("-Formal")]
-		public bool bFormalBuild = false;
+		public bool bFormalBuild { get; set; }
 
 		/// <summary>
 		/// Whether to clean Builds directory on a remote Mac before building.
 		/// </summary>
 		[CommandLine("-FlushMac")]
 		[XmlConfigFile(Category = "BuildConfiguration")]
-		public bool bFlushBuildDirOnRemoteMac = false;
+		public bool bFlushBuildDirOnRemoteMac { get; set; }
 
 		/// <summary>
 		/// Whether to write detailed timing info from the compiler and linker.
 		/// </summary>
 		[CommandLine("-Timing")]
 		[XmlConfigFile(Category = "BuildConfiguration")]
-		public bool bPrintToolChainTimingInfo = false;
+		public bool bPrintToolChainTimingInfo { get; set; }
 
 		/// <summary>
 		/// Whether to parse timing data into a tracing file compatible with chrome://tracing.
 		/// </summary>
 		[CommandLine("-Tracing")]
 		[XmlConfigFile(Category = "BuildConfiguration")]
-		public bool bParseTimingInfoForTracing = false;
+		public bool bParseTimingInfoForTracing { get; set; }
 
 		/// <summary>
 		/// Whether to expose all symbols as public by default on POSIX platforms
 		/// </summary>
 		[CommandLine("-PublicSymbolsByDefault")]
 		[XmlConfigFile(Category = "BuildConfiguration")]
-		public bool bPublicSymbolsByDefault = false;
+		public bool bPublicSymbolsByDefault { get; set; }
 
 		/// <summary>
 		/// Disable supports for inlining gen.cpps
 		/// </summary>
 		[Obsolete("Deprecated in UE5.3; the engine code relies on inlining generated cpp files to compile")]
-		[XmlConfigFile(Name = "bDisableInliningGenCpps")]
-		[CommandLine("-DisableInliningGenCpps")]
-		public bool bDisableInliningGenCpps = false;
+		public bool bDisableInliningGenCpps { get; set; }
 
 		/// <summary>
 		/// Allows overriding the toolchain to be created for this target. This must match the name of a class declared in the UnrealBuildTool assembly.
 		/// </summary>
 		[CommandLine("-ToolChain")]
-		public string? ToolChainName = null;
+		public string? ToolChainName { get; set; }
 
 		/// <summary>
 		/// The weight(cpu/memory utilization) of a MSVC compile action
 		/// </summary>
 		[XmlConfigFile(Category = "BuildConfiguration")]
-		public float MSVCCompileActionWeight = 1.0f;
+		public float MSVCCompileActionWeight { get; set; } = 1.0f;
 
 		/// <summary>
 		/// The weight(cpu/memory utilization) of a clang compile action
 		/// </summary>
 		[XmlConfigFile(Category = "BuildConfiguration")]
-		public float ClangCompileActionWeight = 1.0f;
+		public float ClangCompileActionWeight { get; set; } = 1.0f;
 
 		/// <summary>
 		/// Whether to allow engine configuration to determine if we can load unverified certificates.
 		/// </summary>
-		public bool bDisableUnverifiedCertificates = false;
+		public bool bDisableUnverifiedCertificates { get; set; }
 
 		/// <summary>
 		/// Whether to load generated ini files in cooked build, (GameUserSettings.ini loaded either way)
 		/// </summary>
-		public bool bAllowGeneratedIniWhenCooked = true;
+		public bool bAllowGeneratedIniWhenCooked { get; set; } = true;
 
 		/// <summary>
 		/// Whether to load non-ufs ini files in cooked build, (GameUserSettings.ini loaded either way)
 		/// </summary>
-		public bool bAllowNonUFSIniWhenCooked = true;
+		public bool bAllowNonUFSIniWhenCooked { get; set; } = true;
 
 		/// <summary>
 		/// Add all the public folders as include paths for the compile environment.
@@ -2361,7 +2326,7 @@ namespace UnrealBuildTool
 		[RequiresUniqueBuildEnvironment]
 		[CommandLine("-CStd")]
 		[XmlConfigFile(Category = "BuildConfiguration")]
-		public CStandardVersion CStandard = CStandardVersion.Default;
+		public CStandardVersion CStandard { get; set; } = CStandardVersion.Default;
 
 		/// <summary>
 		/// Direct the compiler to generate AVX instructions wherever SSE or AVX intrinsics are used, on the x64 platforms that support it.
@@ -2370,19 +2335,19 @@ namespace UnrealBuildTool
 		[RequiresUniqueBuildEnvironment]
 		[CommandLine("-MinCpuArchX64")]
 		[XmlConfigFile(Category = "BuildConfiguration")]
-		public MinimumCpuArchitectureX64 MinCpuArchX64 = MinimumCpuArchitectureX64.Default;
+		public MinimumCpuArchitectureX64 MinCpuArchX64 { get; set; } = MinimumCpuArchitectureX64.Default;
 
 		/// <summary>
 		/// Do not allow manifest changes when building this target. Used to cause earlier errors when building multiple targets with a shared build environment.
 		/// </summary>
 		[CommandLine("-NoManifestChanges")]
-		internal bool bNoManifestChanges = false;
+		internal bool bNoManifestChanges { get; set; }
 
 		/// <summary>
 		/// The build version string
 		/// </summary>
 		[CommandLine("-BuildVersion")]
-		public string? BuildVersion;
+		public string? BuildVersion { get; set; }
 
 		/// <summary>
 		/// Specifies how to link modules in this target (monolithic or modular). This is currently protected for backwards compatibility. Call the GetLinkType() accessor
@@ -2407,13 +2372,13 @@ namespace UnrealBuildTool
 		/// </summary>
 		[RequiresUniqueBuildEnvironment]
 		[CommandLine("-Define:")]
-		public List<string> GlobalDefinitions = new List<string>();
+		public List<string> GlobalDefinitions { get; } = new();
 
 		/// <summary>
 		/// Macros to define across all macros in the project.
 		/// </summary>
 		[CommandLine("-ProjectDefine:")]
-		public List<string> ProjectDefinitions = new List<string>();
+		public List<string> ProjectDefinitions { get; } = new();
 
 		/// <summary>
 		/// Specifies the name of the launch module. For modular builds, this is the module that is compiled into the target's executable.
@@ -2427,29 +2392,29 @@ namespace UnrealBuildTool
 		/// <summary>
 		/// Backing storage for the LaunchModuleName property.
 		/// </summary>
-		private string? LaunchModuleNamePrivate;
+		private string? LaunchModuleNamePrivate { get; set; }
 
 		/// <summary>
 		/// Specifies the path to write a header containing public definitions for this target. Useful when building a DLL to be consumed by external build processes.
 		/// </summary>
-		public string? ExportPublicHeader;
+		public string? ExportPublicHeader { get; set; }
 
 		/// <summary>
 		/// List of additional modules to be compiled into the target.
 		/// </summary>
-		public List<string> ExtraModuleNames = new List<string>();
+		public List<string> ExtraModuleNames { get; } = new();
 
 		/// <summary>
 		/// Path to a manifest to output for this target
 		/// </summary>
 		[CommandLine("-Manifest")]
-		public List<FileReference> ManifestFileNames = new List<FileReference>();
+		public List<FileReference> ManifestFileNames { get; } = new();
 
 		/// <summary>
 		/// Path to a list of dependencies for this target, when precompiling
 		/// </summary>
 		[CommandLine("-DependencyList")]
-		public List<FileReference> DependencyListFileNames = new List<FileReference>();
+		public List<FileReference> DependencyListFileNames { get; } = new();
 
 		/// <summary>
 		/// Backing storage for the BuildEnvironment property
@@ -2493,51 +2458,51 @@ namespace UnrealBuildTool
 		/// Whether to ignore violations to the shared build environment (eg. editor targets modifying definitions)
 		/// </summary>
 		[CommandLine("-OverrideBuildEnvironment")]
-		public bool bOverrideBuildEnvironment = false;
+		public bool bOverrideBuildEnvironment { get; set; }
 
 		/// <summary>
 		/// Specifies a list of targets which should be built before this target is built.
 		/// </summary>
-		public List<TargetInfo> PreBuildTargets = new List<TargetInfo>();
+		public List<TargetInfo> PreBuildTargets { get; } = new();
 
 		/// <summary>
 		/// Specifies a list of steps which should be executed before this target is built, in the context of the host platform's shell.
 		/// The following variables will be expanded before execution:
 		/// $(EngineDir), $(ProjectDir), $(TargetName), $(TargetPlatform), $(TargetConfiguration), $(TargetType), $(ProjectFile).
 		/// </summary>
-		public List<string> PreBuildSteps = new List<string>();
+		public List<string> PreBuildSteps { get; } = new();
 
 		/// <summary>
 		/// Specifies a list of steps which should be executed after this target is built, in the context of the host platform's shell.
 		/// The following variables will be expanded before execution:
 		/// $(EngineDir), $(ProjectDir), $(TargetName), $(TargetPlatform), $(TargetConfiguration), $(TargetType), $(ProjectFile).
 		/// </summary>
-		public List<string> PostBuildSteps = new List<string>();
+		public List<string> PostBuildSteps { get; } = new();
 
 		/// <summary>
 		/// Specifies additional build products produced as part of this target.
 		/// </summary>
-		public List<string> AdditionalBuildProducts = new List<string>();
+		public List<string> AdditionalBuildProducts { get; } = new();
 
 		/// <summary>
 		/// Additional arguments to pass to the compiler
 		/// </summary>
 		[RequiresUniqueBuildEnvironment]
 		[CommandLine("-CompilerArguments=")]
-		public string? AdditionalCompilerArguments;
+		public string? AdditionalCompilerArguments { get; set; }
 
 		/// <summary>
 		/// Additional arguments to pass to the linker
 		/// </summary>
 		[RequiresUniqueBuildEnvironment]
 		[CommandLine("-LinkerArguments=")]
-		public string? AdditionalLinkerArguments;
+		public string? AdditionalLinkerArguments { get; set; }
 
 		/// <summary>
 		/// Max amount of memory that each compile action may require. Used by ParallelExecutor to decide the maximum 
 		/// number of parallel actions to start at one time.
 		/// </summary>
-		public double MemoryPerActionGB = 0.0;
+		public double MemoryPerActionGB { get; set; } = 0.0;
 
 		/// <summary>
 		/// List of modules to disable unity builds for
@@ -2574,7 +2539,7 @@ namespace UnrealBuildTool
 		/// <summary>
 		/// When generating project files, specifies the name of the project file to use when there are multiple targets of the same type.
 		/// </summary>
-		public string? GeneratedProjectName;
+		public string? GeneratedProjectName { get; set; }
 
 		/// <summary>
 		/// If this is non-null, then any platforms NOT listed will not be allowed to have modules in their directories be created
@@ -2585,69 +2550,69 @@ namespace UnrealBuildTool
 		/// Android-specific target settings.
 		/// </summary>
 		[ConfigSubObject]
-		public AndroidTargetRules AndroidPlatform = new AndroidTargetRules();
+		public AndroidTargetRules AndroidPlatform { get; init; } = new();
 
 		/// <summary>
 		/// IOS-specific target settings.
 		/// </summary>
 		[ConfigSubObject]
-		public IOSTargetRules IOSPlatform = new IOSTargetRules();
+		public IOSTargetRules IOSPlatform { get; init; } = new();
 
 		/// <summary>
 		/// Linux-specific target settings.
 		/// </summary>
 		[ConfigSubObject]
-		public LinuxTargetRules LinuxPlatform = new LinuxTargetRules();
+		public LinuxTargetRules LinuxPlatform { get; init; } = new();
 
 		/// <summary>
 		/// Mac-specific target settings.
 		/// </summary>
 		[ConfigSubObject]
-		public MacTargetRules MacPlatform = new MacTargetRules();
+		public MacTargetRules MacPlatform { get; init; } = new();
 
 		/// <summary>
 		/// Windows-specific target settings.
 		/// </summary>
 		[ConfigSubObject]
-		public WindowsTargetRules WindowsPlatform; // Requires 'this' parameter; initialized in constructor
+		public WindowsTargetRules WindowsPlatform { get; init; } // Requires 'this' parameter; initialized in constructor
 
 		/// <summary>
-		/// Checks if the target needs to care about SDK version - if not, we can share the target between projects that are on differnet SDK versions.
-		/// For instance, ShaderCompileWorker is very tied to, say, console SDK versions, but UnrealPak doesn't have SDK things compiled in
+		/// Checks if the target needs to care about sdk version - if not, we can share the target between projects that are on differnet sdk versions.
+		/// For instance, ShaderCompileWorker is very tied to, say, console sdk versions, but UnrealPak doesn't have sdk things compiled in
 		/// </summary>
 		/// <returns></returns>
-		public bool IsSDKVersionRelevant(UnrealTargetPlatform SDKPlatform)
+		public bool IsSDKVersionRelevant(UnrealTargetPlatform sdkPlatform)
 		{
-			// we always care bout the SDK of the platform we are compiling for (Platform is the platform this target is being built for)
-			if (this.Platform == SDKPlatform)
+			// we always care bout the sdk of the platform we are compiling for (Platform is the platform this target is being built for)
+			if (Platform == sdkPlatform)
 			{
 				return true;
 			}
 
-			// at this point, we are checking against an SDK that we are not compiling for, ie a target platform sdk, so check if the Target has overriden it
+			// at this point, we are checking against an sdk that we are not compiling for, ie a target platform sdk, so check if the Target has overriden it
 			if (bAreTargetSDKVersionsRelevantOverride != null)
 			{
 				return bAreTargetSDKVersionsRelevantOverride.Value;
 			}
 
-			// developer tools are what pull in SDK versions for other platforms, so Programs that don't use dev tools don't care about versions
+			// developer tools are what pull in sdk versions for other platforms, so Programs that don't use dev tools don't care about versions
 			if (Type == global::UnrealBuildTool.TargetType.Program || Type == global::UnrealBuildTool.TargetType.Editor)
 			{
-				return this.bUsesTargetSDKs;
+				return bUsesTargetSDKs;
 			}
 
-			// games/client/servers are not expected to contain other platform SDK bits in them, so we don't need to care
+			// games/client/servers are not expected to contain other platform sdk bits in them, so we don't need to care
 			return false;
 		}
 
 		/// <summary>
 		/// Allow a target to override whether or not AreSDKVersionsRelevant uses default logic
 		/// </summary>
-		protected bool? bAreTargetSDKVersionsRelevantOverride = null;
+		protected bool? bAreTargetSDKVersionsRelevantOverride { get; set; }
 
 		/// <summary>
-		/// Checks if the target can use non-standard SDK versions. For instance, the editor is usually compiled as a modular build shared between projects,
-		/// so if one projet uses SDK version 1, and another uses 2, then the generated engine modules will use the most recently built SDK, which could break
+		/// Checks if the target can use non-standard sdk versions. For instance, the editor is usually compiled as a modular build shared between projects,
+		/// so if one projet uses sdk version 1, and another uses 2, then the generated engine modules will use the most recently built sdk, which could break
 		/// </summary>
 		/// <returns></returns>
 		public bool AllowsPerProjectSDKVersion()
@@ -2660,99 +2625,94 @@ namespace UnrealBuildTool
 		/// Checks if a plugin should be programmatically allowed in the build
 		/// </summary>
 		/// <returns>true if the plugin is allowed</returns>
-		public virtual bool IgnorePluginDependency(PluginInfo ParentInfo, PluginReferenceDescriptor Descriptor) => false;
+		public virtual bool IgnorePluginDependency(PluginInfo parentInfo, PluginReferenceDescriptor descriptor) => false;
 
 		/// <summary>
 		/// Checks if a property has been set with the RequiresUniqueBuildEnvironmentAttribute, and is different from it's base
 		/// </summary>
-		/// <param name="RulesAssembly">Assembly containing the target</param>
-		/// <param name="Arguments">Commandline options that may affect the target creation</param>
-		/// <param name="PropNamesThatRequireUnique">If the target requires a unique build environment, this will contain the names of the field/property that require unique</param>
-		/// <param name="BaseTargetName">If the target requires a unique build environment, this will contain the name of the target this was based on (UnrealGame, UnrealEditor, etc)/property</param>
+		/// <param name="rulesAssembly">Assembly containing the target</param>
+		/// <param name="arguments">Commandline options that may affect the target creation</param>
+		/// <param name="propNamesThatRequireUnique">If the target requires a unique build environment, this will contain the names of the field/property that require unique</param>
+		/// <param name="baseTargetName">If the target requires a unique build environment, this will contain the name of the target this was based on (UnrealGame, UnrealEditor, etc)/property</param>
 		/// <returns>true if a property was set such that it requires a unique build environment</returns>
 		/// <exception cref="BuildException"></exception>
-		public bool RequiresUniqueEnvironment(RulesAssembly RulesAssembly, CommandLineArguments? Arguments, List<string> PropNamesThatRequireUnique, [NotNullWhen(true)] out string? BaseTargetName)
+		public bool RequiresUniqueEnvironment(RulesAssembly rulesAssembly, CommandLineArguments? arguments, List<string> propNamesThatRequireUnique, [NotNullWhen(true)] out string? baseTargetName)
 		{
-			BaseTargetName = null;
+			baseTargetName = null;
 
-			TargetRules ThisRules = this;
+			TargetRules thisRules = this;
 			// Allow disabling these checks
-			if (ThisRules.bOverrideBuildEnvironment)
+			if (thisRules.bOverrideBuildEnvironment)
 			{
 				return false;
 			}
 
 			// Get the name of the target with default settings
-			switch (ThisRules.Type)
+			switch (thisRules.Type)
 			{
 				case TargetType.Game:
-					BaseTargetName = "UnrealGame";
+					baseTargetName = "UnrealGame";
 					break;
 				case TargetType.Editor:
-					BaseTargetName = "UnrealEditor";
+					baseTargetName = "UnrealEditor";
 					break;
 				case TargetType.Client:
-					BaseTargetName = "UnrealClient";
+					baseTargetName = "UnrealClient";
 					break;
 				case TargetType.Server:
-					BaseTargetName = "UnrealServer";
+					baseTargetName = "UnrealServer";
 					break;
 				default:
 					return false;
 			}
 
 			// Create the target rules for it
-			TargetRules BaseRules = RulesAssembly.CreateTargetRules(BaseTargetName, ThisRules.Platform, ThisRules.Configuration, ThisRules.Architectures, null, Arguments, Logger, IntermediateEnvironment: ThisRules.IntermediateEnvironment);
+			TargetRules baseRules = rulesAssembly.CreateTargetRules(baseTargetName, thisRules.Platform, thisRules.Configuration, thisRules.Architectures, null, arguments, Logger, IntermediateEnvironment: thisRules.IntermediateEnvironment);
 
 			// Get all the configurable objects
-			object[] BaseObjects = BaseRules.GetConfigurableObjects().ToArray();
-			object[] ThisObjects = GetConfigurableObjects().ToArray();
-			if (BaseObjects.Length != ThisObjects.Length)
+			object[] baseObjects = baseRules.GetConfigurableObjects().ToArray();
+			object[] thisObjects = GetConfigurableObjects().ToArray();
+			if (baseObjects.Length != thisObjects.Length)
 			{
 				throw new BuildException("Expected same number of configurable objects from base rules object.");
 			}
 
 			// Iterate through all fields with the [SharedBuildEnvironment] attribute
-			for (int Idx = 0; Idx < BaseObjects.Length; Idx++)
+			for (int idx = 0; idx < baseObjects.Length; idx++)
 			{
-				Type ObjectType = BaseObjects[Idx].GetType();
-				foreach (FieldInfo Field in ObjectType.GetFields())
+				Type objectType = baseObjects[idx].GetType();
+				foreach (FieldInfo field in objectType.GetFields().Where(field => field.GetCustomAttribute<RequiresUniqueBuildEnvironmentAttribute>() != null))
 				{
-					if (Field.GetCustomAttribute<RequiresUniqueBuildEnvironmentAttribute>() != null)
+					object? thisValue = field.GetValue(thisObjects[idx]);
+					object? baseValue = field.GetValue(baseObjects[idx]);
+					if (!CheckValuesMatch(field.FieldType, thisValue, baseValue))
 					{
-						object? ThisValue = Field.GetValue(ThisObjects[Idx]);
-						object? BaseValue = Field.GetValue(BaseObjects[Idx]);
-						if (!CheckValuesMatch(Field.FieldType, ThisValue, BaseValue))
-						{
-							PropNamesThatRequireUnique.Add(Field.Name);
-						}
+						propNamesThatRequireUnique.Add(field.Name);
 					}
 				}
-				foreach (PropertyInfo Property in ObjectType.GetProperties())
+
+				foreach (PropertyInfo property in objectType.GetProperties().Where(property => property.GetCustomAttribute<RequiresUniqueBuildEnvironmentAttribute>() != null))
 				{
-					if (Property.GetCustomAttribute<RequiresUniqueBuildEnvironmentAttribute>() != null)
+					object? thisValue = property.GetValue(thisObjects[idx]);
+					object? baseValue = property.GetValue(baseObjects[idx]);
+					if (!CheckValuesMatch(property.PropertyType, thisValue, baseValue))
 					{
-						object? ThisValue = Property.GetValue(ThisObjects[Idx]);
-						object? BaseValue = Property.GetValue(BaseObjects[Idx]);
-						if (!CheckValuesMatch(Property.PropertyType, ThisValue, BaseValue))
-						{
-							PropNamesThatRequireUnique.Add(Property.Name);
-						}
+						propNamesThatRequireUnique.Add(property.Name);
 					}
 				}
 			}
 
 			// if any properties require a unique build environment, return true
-			return PropNamesThatRequireUnique.Count > 0;
+			return propNamesThatRequireUnique.Count > 0;
 		}
 
 		/// <summary>
 		/// For any target that has set BuildEnvironment = TargetBuildEnvironment.Unique, this will change the Environment to either Shared or Unique accordingly
 		/// </summary>
-		/// <param name="RulesAssembly">Assembly containing the target</param>
-		/// <param name="Arguments">Commandline options that may affect the target creation</param>
-		/// <param name="Logger">Logger</param>
-		public void UpdateBuildEnvironmentIfNeeded(RulesAssembly RulesAssembly, CommandLineArguments? Arguments, ILogger Logger)
+		/// <param name="rulesAssembly">Assembly containing the target</param>
+		/// <param name="arguments">Commandline options that may affect the target creation</param>
+		/// <param name="logger">Logger</param>
+		public void UpdateBuildEnvironmentIfNeeded(RulesAssembly rulesAssembly, CommandLineArguments? arguments, ILogger logger)
 		{
 			// only do anything here when using UniqueIfNeeded
 			if (!BuildEnvironmentOverride.HasValue || BuildEnvironmentOverride.Value != TargetBuildEnvironment.UniqueIfNeeded)
@@ -2762,15 +2722,15 @@ namespace UnrealBuildTool
 
 			if (UEBuildPlatformSDK.bHasAnySDKOverride)
 			{
-				// check if any platform needs a unique environment for the SDK check
-				foreach (UnrealTargetPlatform Platform in UnrealTargetPlatform.GetValidPlatforms())
+				// check if any platform needs a unique environment for the sdk check
+				foreach (UnrealTargetPlatform platform in UnrealTargetPlatform.GetValidPlatforms())
 				{
-					if (IsSDKVersionRelevant(Platform))
+					if (IsSDKVersionRelevant(platform))
 					{
-						UEBuildPlatformSDK? SDK = UEBuildPlatformSDK.GetSDKForPlatform(Platform.ToString());
-						if (SDK != null && SDK.bHasSDKOverride)
+						UEBuildPlatformSDK? sdk = UEBuildPlatformSDK.GetSDKForPlatform(platform.ToString());
+						if (sdk != null && sdk.bHasSDKOverride)
 						{
-							Logger.LogInformation("Setting {Target}'s BuildEnvironment to Unique, because a project overrode the {Platform} SDK version", Name, Platform);
+							logger.LogInformation("Setting {Target}'s BuildEnvironment to Unique, because a project overrode the {Platform} sdk version", Name, platform);
 							BuildEnvironment = TargetBuildEnvironment.Unique;
 							break;
 						}
@@ -2781,12 +2741,12 @@ namespace UnrealBuildTool
 			// if we didn't set it above, check the properties
 			if (BuildEnvironmentOverride.Value == TargetBuildEnvironment.UniqueIfNeeded)
 			{
-				List<string> PropNames = new();
-				string? BaseTargetName;
-				if (RequiresUniqueEnvironment(RulesAssembly, Arguments, PropNames, out BaseTargetName))
+				List<string> propNames = new();
+				string? baseTargetName;
+				if (RequiresUniqueEnvironment(rulesAssembly, arguments, propNames, out baseTargetName))
 				{
-					Logger.LogInformation("Setting {Target}'s BuildEnvironment to Unique, because it had changed the values of the propertues [ {Props} ] away from the values specified in {BaseTarget}", 
-						Name, string.Join(", ", PropNames), BaseTargetName);
+					logger.LogInformation("Setting {Target}'s BuildEnvironment to Unique, because it had changed the values of the propertues [ {Props} ] away from the values specified in {BaseTarget}",
+						Name, String.Join(", ", propNames), baseTargetName);
 					BuildEnvironment = TargetBuildEnvironment.Unique;
 				}
 				else
@@ -2796,21 +2756,21 @@ namespace UnrealBuildTool
 			}
 		}
 
-		static bool CheckValuesMatch(Type ValueType, object? ThisValue, object? BaseValue)
+		static bool CheckValuesMatch(Type valueType, object? thisValue, object? baseValue)
 		{
 			// Check if the fields match, treating lists of strings (eg. definitions) differently to value types.
 			bool bFieldsMatch;
-			if (ThisValue == null || BaseValue == null)
+			if (thisValue == null || baseValue == null)
 			{
-				bFieldsMatch = (ThisValue == BaseValue);
+				bFieldsMatch = (thisValue == baseValue);
 			}
-			else if (typeof(IEnumerable<string>).IsAssignableFrom(ValueType))
+			else if (typeof(IEnumerable<string>).IsAssignableFrom(valueType))
 			{
-				bFieldsMatch = Enumerable.SequenceEqual((IEnumerable<string>)ThisValue, (IEnumerable<string>)BaseValue);
+				bFieldsMatch = Enumerable.SequenceEqual((IEnumerable<string>)thisValue, (IEnumerable<string>)baseValue);
 			}
 			else
 			{
-				bFieldsMatch = ThisValue.Equals(BaseValue);
+				bFieldsMatch = thisValue.Equals(baseValue);
 			}
 
 			return bFieldsMatch;
@@ -2819,68 +2779,65 @@ namespace UnrealBuildTool
 		/// <summary>
 		/// Create a TargetRules instance
 		/// </summary>
-		/// <param name="RulesType">Type to create</param>
-		/// <param name="TargetInfo">Target info</param>
-		/// <param name="BaseFile">Path to the file for the rules assembly</param>
-		/// <param name="PlatformFile">Path to the platform specific rules file</param>
-		/// <param name="TargetFiles">Path to all possible rules file</param>
-		/// <param name="DefaultBuildSettings"></param>
-		/// <param name="Logger">Logger for the new target rules</param>
+		/// <param name="rulesType">Type to create</param>
+		/// <param name="targetInfo">Target info</param>
+		/// <param name="baseFile">Path to the file for the rules assembly</param>
+		/// <param name="platformFile">Path to the platform specific rules file</param>
+		/// <param name="targetFiles">Path to all possible rules file</param>
+		/// <param name="defaultBuildSettings"></param>
+		/// <param name="logger">Logger for the new target rules</param>
 		/// <returns>Target instance</returns>
-		public static TargetRules Create(Type RulesType, TargetInfo TargetInfo, FileReference? BaseFile, FileReference? PlatformFile, IEnumerable<FileReference>? TargetFiles, BuildSettingsVersion? DefaultBuildSettings, ILogger Logger)
+		public static TargetRules Create(Type rulesType, TargetInfo targetInfo, FileReference? baseFile, FileReference? platformFile, IEnumerable<FileReference>? targetFiles, BuildSettingsVersion? defaultBuildSettings, ILogger logger)
 		{
-			TargetRules Rules = (TargetRules)FormatterServices.GetUninitializedObject(RulesType);
-			if (DefaultBuildSettings.HasValue)
+			TargetRules rules = (TargetRules)FormatterServices.GetUninitializedObject(rulesType);
+			if (defaultBuildSettings.HasValue)
 			{
-				Rules.DefaultBuildSettings = DefaultBuildSettings.Value;
+				rules.DefaultBuildSettings = defaultBuildSettings.Value;
 			}
 
 			// The base target file name: this affects where the resulting build product is created so the platform/group is not desired in this case.
-			Rules.File = BaseFile;
+			rules.File = baseFile;
 
 			// The platform/group-specific target file name
-			Rules.TargetSourceFile = PlatformFile;
+			rules.TargetSourceFile = platformFile;
 
 			// All target files for this target that could cause invalidation
-			Rules.TargetFiles = TargetFiles?.ToHashSet() ?? new();
+			rules.TargetFiles = targetFiles?.ToHashSet() ?? new();
 
 			// Initialize the logger
-			Rules.Logger = Logger;
+			rules.Logger = logger;
 
 			// Find the constructor
-			ConstructorInfo? Constructor = RulesType.GetConstructor(new Type[] { typeof(TargetInfo) });
-			if (Constructor == null)
-			{
-				throw new CompilationResultException(CompilationResult.RulesError, "No constructor found on {TargetName} which takes an argument of type TargetInfo.", RulesType.Name);
-			}
+			ConstructorInfo? constructor = rulesType.GetConstructor(new Type[] { typeof(TargetInfo) })
+				?? throw new CompilationResultException(CompilationResult.RulesError, "No constructor found on {TargetName} which takes an argument of type TargetInfo.", rulesType.Name);
 
 			// Invoke the regular constructor
 			try
 			{
-				Constructor.Invoke(Rules, new object[] { TargetInfo });
+				constructor.Invoke(rules, new object[] { targetInfo });
 			}
-			catch (Exception Ex)
+			catch (Exception ex)
 			{
-				throw new CompilationResultException(CompilationResult.RulesError, Ex, "Unable to instantiate instance of '{TargetName}' object type from compiled assembly '{AssemblyPath}'.  Unreal Build Tool creates an instance of your module's 'Rules' object in order to find out about your module's requirements.  The CLR exception details may provide more information: {ExceptionMessage}",
-                    RulesType.Name, Path.GetFileNameWithoutExtension(RulesType.Assembly?.Location) ?? "Unknown Assembly", Ex.ToString());
+				throw new CompilationResultException(CompilationResult.RulesError, ex, "Unable to instantiate instance of '{TargetName}' object type from compiled assembly '{AssemblyPath}'.  Unreal Build Tool creates an instance of your module's 'Rules' object in order to find out about your module's requirements.  The CLR exception details may provide more information: {ExceptionMessage}",
+					rulesType.Name, Path.GetFileNameWithoutExtension(rulesType.Assembly?.Location) ?? "Unknown Assembly", ex.ToString());
 			}
 
-			return Rules;
+			return rules;
 		}
 
 		/// <summary>
 		/// Constructor.
 		/// </summary>
-		/// <param name="Target">Information about the target being built</param>
-		public TargetRules(TargetInfo Target)
+		/// <param name="target">Information about the target being built</param>
+		public TargetRules(TargetInfo target)
 		{
-			DefaultName = Target.Name;
-			Platform = Target.Platform;
-			Configuration = Target.Configuration;
-			Architectures = Target.Architectures;
-			IntermediateEnvironment = Target.IntermediateEnvironment;
-			ProjectFile = Target.ProjectFile;
-			Version = Target.Version;
+			DefaultName = target.Name;
+			Platform = target.Platform;
+			Configuration = target.Configuration;
+			Architectures = target.Architectures;
+			IntermediateEnvironment = target.IntermediateEnvironment;
+			ProjectFile = target.ProjectFile;
+			Version = target.Version;
 			WindowsPlatform = new WindowsTargetRules(this);
 
 			// Make sure the logger was initialized by the caller
@@ -2890,24 +2847,21 @@ namespace UnrealBuildTool
 			}
 
 			// Read settings from config files
-			Dictionary<ConfigDependencyKey, IReadOnlyList<string>?> ConfigValues = new Dictionary<ConfigDependencyKey, IReadOnlyList<string>?>();
-			foreach (object ConfigurableObject in GetConfigurableObjects())
+			Dictionary<ConfigDependencyKey, IReadOnlyList<string>?> configValues = new Dictionary<ConfigDependencyKey, IReadOnlyList<string>?>();
+			foreach (object configurableObject in GetConfigurableObjects())
 			{
-				ConfigCache.ReadSettings(DirectoryReference.FromFile(ProjectFile), Platform, ConfigurableObject, ConfigValues, Target.Arguments);
-				XmlConfig.ApplyTo(ConfigurableObject);
-				if (Target.Arguments != null)
-				{
-					Target.Arguments.ApplyTo(ConfigurableObject);
-				}
+				ConfigCache.ReadSettings(DirectoryReference.FromFile(ProjectFile), Platform, configurableObject, configValues, target.Arguments);
+				XmlConfig.ApplyTo(configurableObject);
+				target.Arguments?.ApplyTo(configurableObject);
 			}
-			ConfigValueTracker = new ConfigValueTracker(ConfigValues);
+			ConfigValueTracker = new ConfigValueTracker(configValues);
 
 			// If we've got a changelist set, set that we're making a formal build
 			bFormalBuild = bFormalBuild || (Version.Changelist != 0 && Version.IsPromotedBuild);
 
 			// Allow the build platform to set defaults for this target
 			UEBuildPlatform.GetBuildPlatform(Platform).ResetTarget(this);
-			bDeployAfterCompile = bForceSkipDeploy ? false : bDeployAfterCompile;
+			bDeployAfterCompile = !bForceSkipDeploy && bDeployAfterCompile;
 
 			// Determine intermediate environment overrides based on command line flags
 			if (StaticAnalyzer != StaticAnalyzer.None)
@@ -2930,39 +2884,39 @@ namespace UnrealBuildTool
 			// Set the default build version
 			if (String.IsNullOrEmpty(BuildVersion))
 			{
-				if (String.IsNullOrEmpty(Target.Version.BuildVersionString))
+				if (String.IsNullOrEmpty(target.Version.BuildVersionString))
 				{
-					BuildVersion = String.Format("{0}-CL-{1}", Target.Version.BranchName, Target.Version.Changelist);
+					BuildVersion = String.Format("{0}-CL-{1}", target.Version.BranchName, target.Version.Changelist);
 				}
 				else
 				{
-					BuildVersion = Target.Version.BuildVersionString;
+					BuildVersion = target.Version.BuildVersionString;
 				}
 			}
 
 			// Get the directory to use for crypto settings. We can build engine targets (eg. UHT) with 
 			// a project file, but we can't use that to determine crypto settings without triggering
 			// constant rebuilds of UHT.
-			DirectoryReference? CryptoSettingsDir = DirectoryReference.FromFile(ProjectFile);
-			if (CryptoSettingsDir != null && File != null && !File.IsUnderDirectory(CryptoSettingsDir))
+			DirectoryReference? cryptoSettingsDir = DirectoryReference.FromFile(ProjectFile);
+			if (cryptoSettingsDir != null && File != null && !File.IsUnderDirectory(cryptoSettingsDir))
 			{
-				CryptoSettingsDir = null;
+				cryptoSettingsDir = null;
 			}
 
 			// Setup macros for signing and encryption keys
-			EncryptionAndSigning.CryptoSettings CryptoSettings = EncryptionAndSigning.ParseCryptoSettings(CryptoSettingsDir, Platform, Logger);
-			if (CryptoSettings.IsAnyEncryptionEnabled())
+			EncryptionAndSigning.CryptoSettings cryptoSettings = EncryptionAndSigning.ParseCryptoSettings(cryptoSettingsDir, Platform, Logger);
+			if (cryptoSettings.IsAnyEncryptionEnabled())
 			{
-				ProjectDefinitions.Add(String.Format("IMPLEMENT_ENCRYPTION_KEY_REGISTRATION()=UE_REGISTER_ENCRYPTION_KEY({0})", FormatHexBytes(CryptoSettings.EncryptionKey!.Key!)));
+				ProjectDefinitions.Add(String.Format("IMPLEMENT_ENCRYPTION_KEY_REGISTRATION()=UE_REGISTER_ENCRYPTION_KEY({0})", FormatHexBytes(cryptoSettings.EncryptionKey!.Key!)));
 			}
 			else
 			{
 				ProjectDefinitions.Add("IMPLEMENT_ENCRYPTION_KEY_REGISTRATION()=");
 			}
 
-			if (CryptoSettings.IsPakSigningEnabled())
+			if (cryptoSettings.IsPakSigningEnabled())
 			{
-				ProjectDefinitions.Add(String.Format("IMPLEMENT_SIGNING_KEY_REGISTRATION()=UE_REGISTER_SIGNING_KEY(UE_LIST_ARGUMENT({0}), UE_LIST_ARGUMENT({1}))", FormatHexBytes(CryptoSettings.SigningKey!.PublicKey.Exponent!), FormatHexBytes(CryptoSettings.SigningKey.PublicKey.Modulus!)));
+				ProjectDefinitions.Add(String.Format("IMPLEMENT_SIGNING_KEY_REGISTRATION()=UE_REGISTER_SIGNING_KEY(UE_LIST_ARGUMENT({0}), UE_LIST_ARGUMENT({1}))", FormatHexBytes(cryptoSettings.SigningKey!.PublicKey.Exponent!), FormatHexBytes(cryptoSettings.SigningKey.PublicKey.Modulus!)));
 			}
 			else
 			{
@@ -2973,11 +2927,11 @@ namespace UnrealBuildTool
 		/// <summary>
 		/// Formats an array of bytes as a sequence of values
 		/// </summary>
-		/// <param name="Data">The data to convert into a string</param>
+		/// <param name="data">The data to convert into a string</param>
 		/// <returns>List of hexadecimal bytes</returns>
-		private static string FormatHexBytes(byte[] Data)
+		private static string FormatHexBytes(byte[] data)
 		{
-			return String.Join(",", Data.Select(x => String.Format("0x{0:X2}", x)));
+			return String.Join(",", data.Select(x => String.Format("0x{0:X2}", x)));
 		}
 
 		/// <summary>
@@ -3037,36 +2991,36 @@ namespace UnrealBuildTool
 			// this setting can be used to compile out the shader compiler if that is important
 			//GlobalDefinitions.Add("UE_ALLOW_SHADER_COMPILING=0");
 
-			ConfigHierarchy ProjectGameIni = ConfigCache.ReadHierarchy(ConfigHierarchyType.Game, ProjectFile?.Directory, Platform);
-			List<string>? DisabledPlugins;
-			List<string> AllDisabledPlugins = new List<string>();
+			ConfigHierarchy projectGameIni = ConfigCache.ReadHierarchy(ConfigHierarchyType.Game, ProjectFile?.Directory, Platform);
+			List<string>? disabledPlugins;
+			List<string> allDisabledPlugins = new List<string>();
 
-			if (ProjectGameIni.GetArray("CookedEditorSettings", "DisabledPlugins", out DisabledPlugins))
+			if (projectGameIni.GetArray("CookedEditorSettings", "DisabledPlugins", out disabledPlugins))
 			{
-				AllDisabledPlugins.AddRange(DisabledPlugins);
+				allDisabledPlugins.AddRange(disabledPlugins);
 			}
-			if (ProjectGameIni.GetArray("CookedEditorSettings" + (bIsCookedCooker ? "_CookedCooker" : "_CookedEditor"), "DisabledPlugins", out DisabledPlugins))
+			if (projectGameIni.GetArray("CookedEditorSettings" + (bIsCookedCooker ? "_CookedCooker" : "_CookedEditor"), "DisabledPlugins", out disabledPlugins))
 			{
-				AllDisabledPlugins.AddRange(DisabledPlugins);
+				allDisabledPlugins.AddRange(disabledPlugins);
 			}
 			if (Configuration == UnrealTargetConfiguration.Shipping)
 			{
-				if (ProjectGameIni.GetArray("CookedEditorSettings", "DisabledPluginsInShipping", out DisabledPlugins))
+				if (projectGameIni.GetArray("CookedEditorSettings", "DisabledPluginsInShipping", out disabledPlugins))
 				{
-					AllDisabledPlugins.AddRange(DisabledPlugins);
+					allDisabledPlugins.AddRange(disabledPlugins);
 				}
-				if (ProjectGameIni.GetArray("CookedEditorSettings" + (bIsCookedCooker ? "_CookedCooker" : "_CookedEditor"), "DisabledPluginsInShipping", out DisabledPlugins))
+				if (projectGameIni.GetArray("CookedEditorSettings" + (bIsCookedCooker ? "_CookedCooker" : "_CookedEditor"), "DisabledPluginsInShipping", out disabledPlugins))
 				{
-					AllDisabledPlugins.AddRange(DisabledPlugins);
+					allDisabledPlugins.AddRange(disabledPlugins);
 				}
 			}
 
 			// disable them, and remove them from Enabled in case they were there
-			foreach (string PluginName in AllDisabledPlugins)
+			foreach (string pluginName in allDisabledPlugins)
 			{
-				DisablePlugins.Add(PluginName);
-				EnablePlugins.Remove(PluginName);
-				OptionalPlugins.Remove(PluginName);
+				DisablePlugins.Add(pluginName);
+				EnablePlugins.Remove(pluginName);
+				OptionalPlugins.Remove(pluginName);
 			}
 		}
 
@@ -3077,12 +3031,12 @@ namespace UnrealBuildTool
 		internal UnrealTargetPlatform[] GetSupportedPlatforms()
 		{
 			// Take the SupportedPlatformsAttribute from the first type in the inheritance chain that supports it
-			for (Type? CurrentType = GetType(); CurrentType != null; CurrentType = CurrentType.BaseType)
+			for (Type? currentType = GetType(); currentType != null; currentType = currentType.BaseType)
 			{
-				object[] Attributes = CurrentType.GetCustomAttributes(typeof(SupportedPlatformsAttribute), false);
-				if (Attributes.Length > 0)
+				object[] attributes = currentType.GetCustomAttributes(typeof(SupportedPlatformsAttribute), false);
+				if (attributes.Length > 0)
 				{
-					return Attributes.OfType<SupportedPlatformsAttribute>().SelectMany(x => x.Platforms).Distinct().ToArray();
+					return attributes.OfType<SupportedPlatformsAttribute>().SelectMany(x => x.Platforms).Distinct().ToArray();
 				}
 			}
 
@@ -3108,12 +3062,12 @@ namespace UnrealBuildTool
 		internal UnrealTargetConfiguration[] GetSupportedConfigurations()
 		{
 			// Otherwise take the SupportedConfigurationsAttribute from the first type in the inheritance chain that supports it
-			for (Type? CurrentType = GetType(); CurrentType != null; CurrentType = CurrentType.BaseType)
+			for (Type? currentType = GetType(); currentType != null; currentType = currentType.BaseType)
 			{
-				object[] Attributes = CurrentType.GetCustomAttributes(typeof(SupportedConfigurationsAttribute), false);
-				if (Attributes.Length > 0)
+				object[] attributes = currentType.GetCustomAttributes(typeof(SupportedConfigurationsAttribute), false);
+				if (attributes.Length > 0)
 				{
-					return Attributes.OfType<SupportedConfigurationsAttribute>().SelectMany(x => x.Configurations).Distinct().ToArray();
+					return attributes.OfType<SupportedConfigurationsAttribute>().SelectMany(x => x.Configurations).Distinct().ToArray();
 				}
 			}
 
@@ -3136,23 +3090,28 @@ namespace UnrealBuildTool
 		{
 			yield return this;
 
-			foreach (FieldInfo Field in GetType().GetFields(BindingFlags.Public | BindingFlags.Instance))
+			foreach (FieldInfo field in GetType().GetFields(BindingFlags.Public | BindingFlags.Instance).Where(field => field.GetCustomAttribute<ConfigSubObjectAttribute>() != null))
 			{
-				if (Field.GetCustomAttribute<ConfigSubObjectAttribute>() != null)
+				object? value = field.GetValue(this);
+				if (value != null)
 				{
-					object? Value = Field.GetValue(this);
-					if (Value != null)
-					{
-						yield return Value;
-					}
+					yield return value;
+				}
+			}
+			foreach (PropertyInfo property in GetType().GetProperties(BindingFlags.Public | BindingFlags.Instance).Where(property => property.GetCustomAttribute<ConfigSubObjectAttribute>() != null))
+			{
+				object? value = property.GetValue(this);
+				if (value != null)
+				{
+					yield return value;
 				}
 			}
-			foreach (UnrealTargetPlatform Platform in UnrealTargetPlatform.GetValidPlatforms())
+			foreach (UnrealTargetPlatform platform in UnrealTargetPlatform.GetValidPlatforms())
 			{
-				UEBuildPlatform? BuildPlatform;
-				if (UEBuildPlatform.TryGetBuildPlatform(Platform, out BuildPlatform))
+				UEBuildPlatform? buildPlatform;
+				if (UEBuildPlatform.TryGetBuildPlatform(platform, out buildPlatform))
 				{
-					yield return BuildPlatform.ArchitectureConfig;
+					yield return buildPlatform.ArchitectureConfig;
 				}
 			}
 		}
@@ -3182,68 +3141,65 @@ namespace UnrealBuildTool
 		/// <summary>
 		/// Gets diagnostic messages about default settings which have changed in newer versions of the engine
 		/// </summary>
-		/// <param name="Diagnostics">List of diagnostic messages</param>
-		internal void GetBuildSettingsInfo(List<string> Diagnostics)
+		/// <param name="diagnostics">List of diagnostic messages</param>
+		internal void GetBuildSettingsInfo(List<string> diagnostics)
 		{
 			// Resolve BuildSettingsVersion.Latest to the version it's assigned to
-			BuildSettingsVersion LatestVersion = BuildSettingsVersion.Latest;
-			foreach (BuildSettingsVersion Value in Enum.GetValues(typeof(BuildSettingsVersion)))
+			BuildSettingsVersion latestVersion = BuildSettingsVersion.Latest;
+			foreach (BuildSettingsVersion value in Enum.GetValues(typeof(BuildSettingsVersion)))
 			{
-				if ((int)Value == (int)BuildSettingsVersion.Latest)
+				if ((int)value == (int)BuildSettingsVersion.Latest)
 				{
-					LatestVersion = Value;
+					latestVersion = value;
 				}
 			}
 
-			if (DefaultBuildSettings < LatestVersion)
+			if (DefaultBuildSettings < latestVersion)
 			{
-				Diagnostics.Add("[Upgrade]");
-				Diagnostics.Add("[Upgrade] Using backward-compatible build settings. The latest version of UE sets the following values by default, which may require code changes:");
+				diagnostics.Add("[Upgrade]");
+				diagnostics.Add("[Upgrade] Using backward-compatible build settings. The latest version of UE sets the following values by default, which may require code changes:");
 
-				List<Tuple<string, string>> ModifiedSettings = new List<Tuple<string, string>>();
-				if (BuildSettingsVersion.V2 <= LatestVersion && DefaultBuildSettings < BuildSettingsVersion.V2)
+				List<Tuple<string, string>> modifiedSettings = new List<Tuple<string, string>>();
+				if (BuildSettingsVersion.V2 <= latestVersion && DefaultBuildSettings < BuildSettingsVersion.V2)
 				{
-					ModifiedSettings.Add(Tuple.Create(String.Format("{0} = false", nameof(bLegacyPublicIncludePaths)), "Omits subfolders from public include paths to reduce compiler command line length. (Previously: true)."));
-					ModifiedSettings.Add(Tuple.Create(String.Format("{0} = WarningLevel.Error", nameof(ShadowVariableWarningLevel)), "Treats shadowed variable warnings as errors. (Previously: WarningLevel.Warning)."));
-					ModifiedSettings.Add(Tuple.Create(String.Format("{0} = PCHUsageMode.UseExplicitOrSharedPCHs", nameof(ModuleRules.PCHUsage)), "Set in build.cs files to enables IWYU-style PCH model. See https://docs.unrealengine.com/en-US/Programming/BuildTools/UnrealBuildTool/IWYU/index.html. (Previously: PCHUsageMode.UseSharedPCHs)."));
+					modifiedSettings.Add(Tuple.Create(String.Format("{0} = false", nameof(bLegacyPublicIncludePaths)), "Omits subfolders from public include paths to reduce compiler command line length. (Previously: true)."));
+					modifiedSettings.Add(Tuple.Create(String.Format("{0} = WarningLevel.Error", nameof(ShadowVariableWarningLevel)), "Treats shadowed variable warnings as errors. (Previously: WarningLevel.Warning)."));
+					modifiedSettings.Add(Tuple.Create(String.Format("{0} = PCHUsageMode.UseExplicitOrSharedPCHs", nameof(ModuleRules.PCHUsage)), "Set in build.cs files to enables IWYU-style PCH model. See https://docs.unrealengine.com/en-US/Programming/BuildTools/UnrealBuildTool/IWYU/index.html. (Previously: PCHUsageMode.UseSharedPCHs)."));
 				}
 
-				if (BuildSettingsVersion.V3 <= LatestVersion && DefaultBuildSettings < BuildSettingsVersion.V3)
+				if (BuildSettingsVersion.V3 <= latestVersion && DefaultBuildSettings < BuildSettingsVersion.V3)
 				{
-					ModifiedSettings.Add(Tuple.Create(String.Format("{0} = false", nameof(bLegacyParentIncludePaths)), "Omits module parent folders from include paths to reduce compiler command line length. (Previously: true)."));
+					modifiedSettings.Add(Tuple.Create(String.Format("{0} = false", nameof(bLegacyParentIncludePaths)), "Omits module parent folders from include paths to reduce compiler command line length. (Previously: true)."));
 				}
 
-				if (BuildSettingsVersion.V4 <= LatestVersion && DefaultBuildSettings < BuildSettingsVersion.V4)
+				if (BuildSettingsVersion.V4 <= latestVersion && DefaultBuildSettings < BuildSettingsVersion.V4)
 				{
-					ModifiedSettings.Add(Tuple.Create(String.Format("{0} = CppStandardVersion.Default", nameof(CppStandard)), "Updates C++ Standard to C++20 (Previously: CppStandardVersion.Cpp17)."));
-					ModifiedSettings.Add(Tuple.Create(String.Format("{0}.{1} = true", nameof(WindowsPlatform), nameof(WindowsPlatform.bStrictConformanceMode)), "Updates MSVC strict conformance mode to true (Previously: false)."));
+					modifiedSettings.Add(Tuple.Create(String.Format("{0} = CppStandardVersion.Default", nameof(CppStandard)), "Updates C++ Standard to C++20 (Previously: CppStandardVersion.Cpp17)."));
+					modifiedSettings.Add(Tuple.Create(String.Format("{0}.{1} = true", nameof(WindowsPlatform), nameof(WindowsPlatform.bStrictConformanceMode)), "Updates MSVC strict conformance mode to true (Previously: false)."));
 				}
 
-				if (BuildSettingsVersion.V5 <= LatestVersion && DefaultBuildSettings < BuildSettingsVersion.V5)
+				if (BuildSettingsVersion.V5 <= latestVersion && DefaultBuildSettings < BuildSettingsVersion.V5)
 				{
-					ModifiedSettings.Add(Tuple.Create(String.Format("{0} = true", nameof(bValidateFormatStrings)), "Enables compile-time validation of strings+args passed to UE_LOG. (Previously: false)."));
+					modifiedSettings.Add(Tuple.Create(String.Format("{0} = true", nameof(bValidateFormatStrings)), "Enables compile-time validation of strings+args passed to UE_LOG. (Previously: false)."));
 				}
 
-				if (ModifiedSettings.Count > 0)
+				if (modifiedSettings.Count > 0)
 				{
-					string FormatString = String.Format("[Upgrade]     {{0,-{0}}}   => {{1}}", ModifiedSettings.Max(x => x.Item1.Length));
-					foreach (Tuple<string, string> ModifiedSetting in ModifiedSettings)
-					{
-						Diagnostics.Add(String.Format(FormatString, ModifiedSetting.Item1, ModifiedSetting.Item2));
-					}
+					string formatString = String.Format("[Upgrade]     {{0,-{0}}}   => {{1}}", modifiedSettings.Max(x => x.Item1.Length));
+					diagnostics.AddRange(modifiedSettings.Select(modifiedSetting => String.Format(formatString, modifiedSetting.Item1, modifiedSetting.Item2)));
 				}
-				Diagnostics.Add(String.Format("[Upgrade] Suppress this message by setting 'DefaultBuildSettings = BuildSettingsVersion.{0};' in {1}, and explicitly overriding settings that differ from the new defaults.", LatestVersion, File!.GetFileName()));
-				Diagnostics.Add("[Upgrade]");
+				diagnostics.Add(String.Format("[Upgrade] Suppress this message by setting 'DefaultBuildSettings = BuildSettingsVersion.{0};' in {1}, and explicitly overriding settings that differ from the new defaults.", latestVersion, File!.GetFileName()));
+				diagnostics.Add("[Upgrade]");
 			}
 
 			if (IncludeOrderVersion <= (EngineIncludeOrderVersion)(EngineIncludeOrderVersion.Latest - 1) && ForcedIncludeOrder == null)
 			{
-				Diagnostics.Add("[Upgrade]");
-				Diagnostics.Add("[Upgrade] Using backward-compatible include order. The latest version of UE has changed the order of includes, which may require code changes. The current setting is:");
-				Diagnostics.Add(String.Format("[Upgrade]     IncludeOrderVersion = EngineIncludeOrderVersion.{0}", IncludeOrderVersion));
-				Diagnostics.Add(String.Format("[Upgrade] Suppress this message by setting 'IncludeOrderVersion = EngineIncludeOrderVersion.{0};' in {1}.", EngineIncludeOrderVersion.Latest, File!.GetFileName()));
-				Diagnostics.Add("[Upgrade] Alternatively you can set this to 'EngineIncludeOrderVersion.Latest' to always use the latest include order. This will potentially cause compile errors when integrating new versions of the engine.");
-				Diagnostics.Add("[Upgrade]");
+				diagnostics.Add("[Upgrade]");
+				diagnostics.Add("[Upgrade] Using backward-compatible include order. The latest version of UE has changed the order of includes, which may require code changes. The current setting is:");
+				diagnostics.Add(String.Format("[Upgrade]     IncludeOrderVersion = EngineIncludeOrderVersion.{0}", IncludeOrderVersion));
+				diagnostics.Add(String.Format("[Upgrade] Suppress this message by setting 'IncludeOrderVersion = EngineIncludeOrderVersion.{0};' in {1}.", EngineIncludeOrderVersion.Latest, File!.GetFileName()));
+				diagnostics.Add("[Upgrade] Alternatively you can set this to 'EngineIncludeOrderVersion.Latest' to always use the latest include order. This will potentially cause compile errors when integrating new versions of the engine.");
+				diagnostics.Add("[Upgrade]");
 			}
 
 			Logger.LogDebug("Using EngineIncludeOrderVersion.{Version} for target {Target}", IncludeOrderVersion, File!.GetFileName());
@@ -3254,700 +3210,9 @@ namespace UnrealBuildTool
 		/// </summary>
 		public void PrintBuildSettingsInfoWarnings()
 		{
-			List<string> Diagnostics = new();
-			GetBuildSettingsInfo(Diagnostics);
-			Diagnostics.ForEach(x => Logger.LogWarning("Warning: {Message}", x));
-		}
-	}
-
-	/// <summary>
-	/// Read-only wrapper around an existing TargetRules instance. This exposes target settings to modules without letting them to modify the global environment.
-	/// </summary>
-	public partial class ReadOnlyTargetRules
-	{
-		/// <summary>
-		/// The writeable TargetRules instance
-		/// </summary>
-		TargetRules Inner;
-
-		/// <summary>
-		/// Constructor
-		/// </summary>
-		/// <param name="Inner">The TargetRules instance to wrap around</param>
-		public ReadOnlyTargetRules(TargetRules Inner)
-		{
-			this.Inner = Inner;
-			AndroidPlatform = new ReadOnlyAndroidTargetRules(Inner.AndroidPlatform);
-			IOSPlatform = new ReadOnlyIOSTargetRules(Inner.IOSPlatform);
-			LinuxPlatform = new ReadOnlyLinuxTargetRules(Inner.LinuxPlatform);
-			MacPlatform = new ReadOnlyMacTargetRules(Inner.MacPlatform);
-			WindowsPlatform = new ReadOnlyWindowsTargetRules(Inner.WindowsPlatform);
-		}
-
-		/// <summary>
-		/// Accessors for fields on the inner TargetRules instance
-		/// </summary>
-		#region Read-only accessor properties
-#pragma warning disable CS1591
-
-		public string Name => Inner.Name;
-
-		public string OriginalName => Inner.OriginalName;
-
-		internal FileReference File => Inner.File!;
-
-		internal FileReference TargetSourceFile => Inner.TargetSourceFile!;
-
-		internal IReadOnlySet<FileReference> TargetFiles => Inner.TargetFiles!;
-
-		public object? AdditionalData => Inner.AdditionalData;
-
-		public UnrealTargetPlatform Platform => Inner.Platform;
-
-		public UnrealTargetConfiguration Configuration => Inner.Configuration;
-
-		public UnrealArchitectures Architectures => Inner.Architectures;
-
-		public UnrealIntermediateEnvironment IntermediateEnvironment => Inner.IntermediateEnvironment;
-
-		public UnrealArch Architecture => Inner.Architecture;
-
-		public FileReference? ProjectFile => Inner.ProjectFile;
-
-		public ReadOnlyBuildVersion Version => Inner.Version;
-
-		public TargetType Type => Inner.Type;
-
-		/// <inheritdoc cref="TargetRules.Logger"/>
-		public ILogger Logger => Inner.Logger;
-
-		public BuildSettingsVersion DefaultBuildSettings => Inner.DefaultBuildSettings;
-
-		public EngineIncludeOrderVersion? ForcedIncludeOrder => Inner.ForcedIncludeOrder;
-
-		public EngineIncludeOrderVersion IncludeOrderVersion => Inner.IncludeOrderVersion;
-
-		internal ConfigValueTracker ConfigValueTracker => Inner.ConfigValueTracker;
-
-		public string? OutputFile => Inner.OutputFile;
-
-		public bool bUsesSteam => Inner.bUsesSteam;
-
-		public bool bUsesCEF3 => Inner.bUsesCEF3;
-
-		public bool bUsesSlate => Inner.bUsesSlate;
-
-		public bool bUseStaticCRT => Inner.bUseStaticCRT;
-
-		public bool bDebugBuildsActuallyUseDebugCRT => Inner.bDebugBuildsActuallyUseDebugCRT;
-
-		public bool bLegalToDistributeBinary => Inner.bLegalToDistributeBinary;
-
-		public UnrealTargetConfiguration UndecoratedConfiguration => Inner.UndecoratedConfiguration;
-
-		public bool bAllowHotReload => Inner.bAllowHotReload;
-
-		public bool bBuildAllModules => Inner.bBuildAllModules;
-
-		public IEnumerable<string> AdditionalPlugins => Inner.AdditionalPlugins;
-
-		public IEnumerable<string> EnablePlugins => Inner.EnablePlugins;
-
-		public IEnumerable<string> DisablePlugins => Inner.DisablePlugins;
-
-		public IEnumerable<string> OptionalPlugins => Inner.OptionalPlugins;
-
-		public WarningLevel DisablePluginsConflictWarningLevel => Inner.DisablePluginsConflictWarningLevel;
-
-		public IEnumerable<string> BuildPlugins => Inner.BuildPlugins;
-
-		public IEnumerable<string> InternalPluginDependencies => Inner.InternalPluginDependencies;
-		public bool bRuntimeDependenciesComeFromBuildPlugins => Inner.bRuntimeDependenciesComeFromBuildPlugins;
-
-		public string PakSigningKeysFile => Inner.PakSigningKeysFile;
-
-		public string SolutionDirectory => Inner.SolutionDirectory;
-
-		public string CustomConfig => Inner.CustomConfig;
-
-		public bool? bBuildInSolutionByDefault => Inner.bBuildInSolutionByDefault;
-
-		public string ExeBinariesSubFolder => Inner.ExeBinariesSubFolder;
-
-		public EGeneratedCodeVersion GeneratedCodeVersion => Inner.GeneratedCodeVersion;
-		public bool bEnableMeshEditor => Inner.bEnableMeshEditor;
-
-		public bool bUseVerseBPVM => Inner.bUseVerseBPVM;
-
-		public bool bUseAutoRTFMCompiler => Inner.bUseAutoRTFMCompiler;
-
-		[Obsolete("Deprecated in UE5.1 - No longer used as Chaos is always enabled.")]
-		public bool bCompileChaos => Inner.bCompileChaos;
-
-		[Obsolete("Deprecated in UE5.1 - No longer used as Chaos is always enabled.")]
-		public bool bUseChaos => Inner.bUseChaos;
-
-		public bool bUseChaosMemoryTracking => Inner.bUseChaosMemoryTracking;
-
-		public bool bCompileChaosVisualDebuggerSupport => Inner.bCompileChaosVisualDebuggerSupport;
-
-		public bool bUseChaosChecked => Inner.bUseChaosChecked;
-
-		[Obsolete("Deprecated in UE5.1 - No longer used in engine.")]
-		public bool bCustomSceneQueryStructure => Inner.bCustomSceneQueryStructure;
-
-		[Obsolete("Deprecated in UE5.1 - No longer used as Chaos is always enabled.")]
-		public bool bCompilePhysX => Inner.bCompilePhysX;
-
-		[Obsolete("Deprecated in UE5.1 - No longer used as Chaos is always enabled.")]
-		public bool bCompileAPEX => Inner.bCompileAPEX;
-
-		[Obsolete("Deprecated in UE5.1 - No longer used as Chaos is always enabled.")]
-		public bool bCompileNvCloth => Inner.bCompileNvCloth;
-
-		public bool bCompileICU => Inner.bCompileICU;
-
-		public bool bCompileCEF3 => Inner.bCompileCEF3;
-
-		public bool bCompileISPC => Inner.bCompileISPC;
-
-		public bool bUseGameplayDebugger => Inner.bUseGameplayDebugger;
-
-		public bool bUseGameplayDebuggerCore => Inner.bUseGameplayDebuggerCore;
-
-		public bool bUseIris => Inner.bUseIris;
-
-		public bool bCompileIntelMetricsDiscovery => Inner.bCompileIntelMetricsDiscovery;
-
-		public bool bCompilePython => Inner.bCompilePython;
-
-		public bool bBuildEditor => Inner.bBuildEditor;
-
-		public bool bBuildRequiresCookedData => Inner.bBuildRequiresCookedData;
-
-		public bool bBuildWithEditorOnlyData => Inner.bBuildWithEditorOnlyData;
-
-		public bool bBuildDeveloperTools => Inner.bBuildDeveloperTools;
-
-		public bool bBuildTargetDeveloperTools => Inner.bBuildTargetDeveloperTools;
-
-		public bool bForceBuildTargetPlatforms => Inner.bForceBuildTargetPlatforms;
-
-		public bool bForceBuildShaderFormats => Inner.bForceBuildShaderFormats;
-
-		public bool bNeedsExtraShaderFormats => Inner.bNeedsExtraShaderFormats;
-
-		public bool bCompileCustomSQLitePlatform => Inner.bCompileCustomSQLitePlatform;
-
-		public bool bUseCacheFreedOSAllocs => Inner.bUseCacheFreedOSAllocs;
-
-		public bool bCompileAgainstEngine => Inner.bCompileAgainstEngine;
-
-		public bool bCompileAgainstCoreUObject => Inner.bCompileAgainstCoreUObject;
-
-		public bool bCompileAgainstApplicationCore => Inner.bCompileAgainstApplicationCore;
-
-		public bool bEnableTrace => Inner.bEnableTrace;
-
-		public bool bCompileAgainstEditor => Inner.bCompileAgainstEditor;
-
-		public bool bCompileRecast => Inner.bCompileRecast;
-
-		public bool bCompileNavmeshSegmentLinks => Inner.bCompileNavmeshSegmentLinks;
-
-		public bool bCompileNavmeshClusterLinks => Inner.bCompileNavmeshClusterLinks;
-
-		public bool bCompileSpeedTree => Inner.bCompileSpeedTree;
-
-		public bool bForceEnableExceptions => Inner.bForceEnableExceptions;
-
-		public bool bForceEnableObjCExceptions => Inner.bForceEnableObjCExceptions;
-
-		public bool bForceEnableRTTI => Inner.bForceEnableRTTI;
-
-		public bool bEnablePIE => Inner.bEnablePIE;
-
-		public bool bEnableStackProtection => Inner.bEnableStackProtection;
-
-		public bool bUseInlining => Inner.bUseInlining;
-
-		public bool bWithServerCode => Inner.bWithServerCode;
-
-		public bool bFNameOutlineNumber => Inner.bFNameOutlineNumber;
-
-		public bool bWithPushModel => Inner.bWithPushModel;
-
-		public bool bCompileWithStatsWithoutEngine => Inner.bCompileWithStatsWithoutEngine;
-
-		public bool bCompileWithPluginSupport => Inner.bCompileWithPluginSupport;
-
-		public bool bIncludePluginsForTargetPlatforms => Inner.bIncludePluginsForTargetPlatforms;
-
-		public bool bCompileWithAccessibilitySupport => Inner.bCompileWithAccessibilitySupport;
-
-		public bool bWithPerfCounters => Inner.bWithPerfCounters;
-
-		public bool bWithLiveCoding => Inner.bWithLiveCoding;
-
-		public bool bUseDebugLiveCodingConsole => Inner.bUseDebugLiveCodingConsole;
-
-		public bool bWithDirectXMath => Inner.bWithDirectXMath;
-
-		public bool bWithFixedTimeStepSupport => Inner.bWithFixedTimeStepSupport;
-
-		public bool bUseLoggingInShipping => Inner.bUseLoggingInShipping;
-
-		public bool bUseConsoleInShipping => Inner.bUseConsoleInShipping;
-
-		public bool bLoggingToMemoryEnabled => Inner.bLoggingToMemoryEnabled;
-
-		public bool bUseLauncherChecks => Inner.bUseLauncherChecks;
-
-		public bool bUseChecksInShipping => Inner.bUseChecksInShipping;
-
-		public bool bAllowProfileGPUInTest => Inner.bAllowProfileGPUInTest;
-
-		public bool bAllowProfileGPUInShipping => Inner.bAllowProfileGPUInShipping;
-
-		public bool bTCHARIsUTF8 => Inner.bTCHARIsUTF8;
-
-		public bool bUseEstimatedUtcNow => Inner.bUseEstimatedUtcNow;
-
-		public bool bCompileFreeType => Inner.bCompileFreeType;
-
-		public bool bUseExecCommandsInShipping => Inner.bUseExecCommandsInShipping;
-
-		[Obsolete("Deprecated in UE5.1 - Please use OptimizationLevel instead.")]
-		public bool bCompileForSize => Inner.bCompileForSize;
-
-		public OptimizationMode OptimizationLevel => Inner.OptimizationLevel;
-
-		public FPSemanticsMode FPSemantics => Inner.FPSemantics;
-
-		public bool bRetainFramePointers => Inner.bRetainFramePointers;
-
-		public bool bForceCompileDevelopmentAutomationTests => Inner.bForceCompileDevelopmentAutomationTests;
-
-		public bool bForceCompilePerformanceAutomationTests => Inner.bForceCompilePerformanceAutomationTests;
-
-		public bool bForceDisableAutomationTests => Inner.bForceDisableAutomationTests;
-
-		public bool bUseXGEController => Inner.bUseXGEController;
-
-		public bool bEventDrivenLoader => Inner.bEventDrivenLoader;
-
-		public PointerMemberBehavior? NativePointerMemberBehaviorOverride => Inner.NativePointerMemberBehaviorOverride;
-
-		public bool bIWYU => Inner.bIWYU;
-
-		public bool bIncludeHeaders => Inner.bIncludeHeaders;
-
-		public bool bHeadersOnly => Inner.bHeadersOnly;
-
-		public bool bEnforceIWYU => Inner.bEnforceIWYU;
-
-		public bool bWarnAboutMonolithicHeadersIncluded => Inner.bWarnAboutMonolithicHeadersIncluded;
-
-		public bool bHasExports => Inner.bHasExports;
-
-		public bool bPrecompile => Inner.bPrecompile;
-
-		public bool bEnableOSX109Support => Inner.bEnableOSX109Support;
-
-		public bool bIsBuildingConsoleApplication => Inner.bIsBuildingConsoleApplication;
-
-		public bool bBuildAdditionalConsoleApp => Inner.bBuildAdditionalConsoleApp;
-
-		public bool bDisableSymbolCache => Inner.bDisableSymbolCache;
-
-		public bool bUseUnityBuild => Inner.bUseUnityBuild;
-
-		public bool bForceUnityBuild => Inner.bForceUnityBuild;
-
-		public IReadOnlyList<string>? DisableMergingModuleAndGeneratedFilesInUnityFiles => Inner.DisableMergingModuleAndGeneratedFilesInUnityFiles;
-
-		public bool bAdaptiveUnityDisablesOptimizations => Inner.bAdaptiveUnityDisablesOptimizations;
-
-		public bool bAdaptiveUnityDisablesPCH => Inner.bAdaptiveUnityDisablesPCH;
-
-		public bool bAdaptiveUnityDisablesPCHForProject => Inner.bAdaptiveUnityDisablesPCHForProject;
-
-		public bool bAdaptiveUnityCreatesDedicatedPCH => Inner.bAdaptiveUnityCreatesDedicatedPCH;
-
-		public bool bAdaptiveUnityEnablesEditAndContinue => Inner.bAdaptiveUnityEnablesEditAndContinue;
-
-		public bool bAdaptiveUnityCompilesHeaderFiles => Inner.bAdaptiveUnityCompilesHeaderFiles;
-
-		public int MinGameModuleSourceFilesForUnityBuild => Inner.MinGameModuleSourceFilesForUnityBuild;
-
-		public bool bValidateFormatStrings => Inner.bValidateFormatStrings;
-
-		public WarningLevel DefaultWarningLevel => Inner.DefaultWarningLevel;
-
-		public WarningLevel DeprecationWarningLevel => Inner.DeprecationWarningLevel;
-
-		public WarningLevel ShadowVariableWarningLevel => Inner.ShadowVariableWarningLevel;
-
-		public WarningLevel UnsafeTypeCastWarningLevel => Inner.UnsafeTypeCastWarningLevel;
-
-		public bool bUndefinedIdentifierErrors => Inner.bUndefinedIdentifierErrors;
-
-		public WarningLevel PCHPerformanceIssueWarningLevel => Inner.PCHPerformanceIssueWarningLevel;
-
-		public WarningLevel ModuleIncludePathWarningLevel => Inner.ModuleIncludePathWarningLevel;
-
-		public WarningLevel ModuleIncludePrivateWarningLevel => Inner.ModuleIncludePrivateWarningLevel;
-
-		public WarningLevel ModuleIncludeSubdirectoryWarningLevel => Inner.ModuleIncludeSubdirectoryWarningLevel;
-
-		public bool bWarningsAsErrors => Inner.bWarningsAsErrors;
-
-		public bool bUseFastMonoCalls => Inner.bUseFastMonoCalls;
-
-		public int NumIncludedBytesPerUnityCPP => Inner.NumIncludedBytesPerUnityCPP;
-
-		public bool bDisableModuleNumIncludedBytesPerUnityCPPOverride => Inner.bDisableModuleNumIncludedBytesPerUnityCPPOverride;
-
-		public bool bStressTestUnity => Inner.bStressTestUnity;
-
-		public bool bDetailedUnityFiles => Inner.bDetailedUnityFiles;
-
-		[Obsolete("Deprecated in UE5.4 - Replace with ReadOnlyTargetRules.DebugInfo")]
-		public bool bDisableDebugInfo => Inner.bDisableDebugInfo;
-
-		public DebugInfoMode DebugInfo => Inner.DebugInfo;
-
-		public IReadOnlySet<string> DisableDebugInfoModules => Inner.DisableDebugInfoModules;
-
-		public IReadOnlySet<string> DisableDebugInfoPlugins => Inner.DisableDebugInfoPlugins;
-
-		public DebugInfoMode DebugInfoLineTablesOnly => Inner.DebugInfoLineTablesOnly;
-
-		public IReadOnlySet<string> DebugInfoLineTablesOnlyModules => Inner.DebugInfoLineTablesOnlyModules;
-
-		public IReadOnlySet<string> DebugInfoLineTablesOnlyPlugins => Inner.DebugInfoLineTablesOnlyPlugins;
-
-		public bool bDisableDebugInfoForGeneratedCode => Inner.bDisableDebugInfoForGeneratedCode;
-
-		public bool bOmitPCDebugInfoInDevelopment => Inner.bOmitPCDebugInfoInDevelopment;
-
-		public bool bUsePDBFiles => Inner.bUsePDBFiles;
-
-		public bool bUsePCHFiles => Inner.bUsePCHFiles;
-
-		public bool bDeterministic => Inner.bDeterministic;
-		public bool bChainPCHs => Inner.bChainPCHs;
-
-		public bool bForceIncludePCHHeadersForGenCppFilesWhenPCHIsDisabled => Inner.bForceIncludePCHHeadersForGenCppFilesWhenPCHIsDisabled;
-
-		public bool bPreprocessOnly => Inner.bPreprocessOnly;
-
-		public bool bPreprocessDepends => Inner.bPreprocessDepends;
-
-		public bool bWithAssembly => Inner.bWithAssembly;
-
-		public StaticAnalyzer StaticAnalyzer => Inner.StaticAnalyzer;
-
-		public StaticAnalyzerOutputType StaticAnalyzerOutputType => Inner.StaticAnalyzerOutputType;
-
-		public StaticAnalyzerMode StaticAnalyzerMode => Inner.StaticAnalyzerMode;
-
-		public int MinFilesUsingPrecompiledHeader => Inner.MinFilesUsingPrecompiledHeader;
-
-		public bool bForcePrecompiledHeaderForGameModules => Inner.bForcePrecompiledHeaderForGameModules;
-
-		public bool bUseIncrementalLinking => Inner.bUseIncrementalLinking;
-
-		public bool bAllowLTCG => Inner.bAllowLTCG;
-
-		public bool bPreferThinLTO => Inner.bPreferThinLTO;
-
-		public bool bPGOProfile => Inner.bPGOProfile;
-
-		public bool bPGOOptimize => Inner.bPGOOptimize;
-
-		public bool bSupportEditAndContinue => Inner.bSupportEditAndContinue;
-
-		public bool bCodeCoverage => Inner.bCodeCoverage;
-
-		public bool bOmitFramePointers => Inner.bOmitFramePointers;
-
-		public bool bEnableCppModules => Inner.bEnableCppModules;
-
-		public bool bEnableCppCoroutinesForEvaluation => Inner.bEnableCppCoroutinesForEvaluation;
-
-		public bool bEnableProcessPriorityControl => Inner.bEnableProcessPriorityControl;
-
-		[Obsolete("Deprecated in UE5.3 - No longer used as MallocProfiler is removed in favor of UnrealInsights.")]
-		public bool bUseMallocProfiler => Inner.bUseMallocProfiler;
-
-		public bool bShaderCompilerWorkerTrace => Inner.bShaderCompilerWorkerTrace;
-
-		public bool bUseSharedPCHs => Inner.bUseSharedPCHs;
-
-		public bool bUseShippingPhysXLibraries => Inner.bUseShippingPhysXLibraries;
-
-		public bool bUseCheckedPhysXLibraries => Inner.bUseCheckedPhysXLibraries;
-
-		public bool bCheckLicenseViolations => Inner.bCheckLicenseViolations;
-
-		public bool bBreakBuildOnLicenseViolation => Inner.bBreakBuildOnLicenseViolation;
-
-		public bool? bUseFastPDBLinking => Inner.bUseFastPDBLinking;
-
-		public bool bCreateMapFile => Inner.bCreateMapFile;
-
-		public bool bAllowRuntimeSymbolFiles => Inner.bAllowRuntimeSymbolFiles;
-
-		public string? PackagePath => Inner.PackagePath;
-
-		public string? CrashDiagnosticDirectory => Inner.CrashDiagnosticDirectory;
-
-		public string? ThinLTOCacheDirectory => Inner.ThinLTOCacheDirectory;
-
-		public string? ThinLTOCachePruningArguments => Inner.ThinLTOCachePruningArguments;
-
-		public string? BundleVersion => Inner.BundleVersion;
-
-		public bool bDeployAfterCompile => Inner.bDeployAfterCompile;
-
-		public bool bAllowRemotelyCompiledPCHs => Inner.bAllowRemotelyCompiledPCHs;
-
-		public bool bUseHeaderUnitsForPch => Inner.bUseHeaderUnitsForPch;
-
-		public bool bCheckSystemHeadersForModification => Inner.bCheckSystemHeadersForModification;
-
-		public bool bDisableLinking => Inner.bDisableLinking;
-
-		public bool bIgnoreBuildOutputs => Inner.bIgnoreBuildOutputs;
-
-		public bool bFormalBuild => Inner.bFormalBuild;
-
-		public bool bUseAdaptiveUnityBuild => Inner.bUseAdaptiveUnityBuild;
-
-		public bool bFlushBuildDirOnRemoteMac => Inner.bFlushBuildDirOnRemoteMac;
-
-		public bool bPrintToolChainTimingInfo => Inner.bPrintToolChainTimingInfo;
-
-		public bool bParseTimingInfoForTracing => Inner.bParseTimingInfoForTracing;
-
-		public bool bPublicSymbolsByDefault => Inner.bPublicSymbolsByDefault;
-
-		public string? ToolChainName => Inner.ToolChainName;
-
-		public float MSVCCompileActionWeight => Inner.MSVCCompileActionWeight;
-
-		public float ClangCompileActionWeight => Inner.ClangCompileActionWeight;
-
-		public bool bLegacyPublicIncludePaths => Inner.bLegacyPublicIncludePaths;
-
-		public bool bLegacyParentIncludePaths => Inner.bLegacyParentIncludePaths;
-
-		public CppStandardVersion CppStandardEngine => Inner.CppStandardEngine;
-
-		public CppStandardVersion CppStandard => Inner.CppStandard;
-
-		public CStandardVersion CStandard => Inner.CStandard;
-
-		public MinimumCpuArchitectureX64 MinCpuArchX64 => Inner.MinCpuArchX64;
-
-		internal bool bNoManifestChanges => Inner.bNoManifestChanges;
-
-		public string? BuildVersion => Inner.BuildVersion;
-
-		public TargetLinkType LinkType => Inner.LinkType;
-
-		public IReadOnlyList<string> GlobalDefinitions => Inner.GlobalDefinitions.AsReadOnly();
-
-		public IReadOnlyList<string> ProjectDefinitions => Inner.ProjectDefinitions.AsReadOnly();
-
-		public string? LaunchModuleName => Inner.LaunchModuleName;
-
-		public string? ExportPublicHeader => Inner.ExportPublicHeader;
-
-		public IReadOnlyList<string> ExtraModuleNames => Inner.ExtraModuleNames.AsReadOnly();
-
-		public IReadOnlyList<FileReference> ManifestFileNames => Inner.ManifestFileNames.AsReadOnly();
-
-		public IReadOnlyList<FileReference> DependencyListFileNames => Inner.DependencyListFileNames.AsReadOnly();
-
-		public TargetBuildEnvironment BuildEnvironment => Inner.BuildEnvironment;
-
-		public bool bOverrideBuildEnvironment => Inner.bOverrideBuildEnvironment;
-
-		public IReadOnlyList<TargetInfo> PreBuildTargets => Inner.PreBuildTargets;
-
-		public IReadOnlyList<string> PreBuildSteps => Inner.PreBuildSteps;
-
-		public IReadOnlyList<string> PostBuildSteps => Inner.PostBuildSteps;
-
-		public IReadOnlyList<string> AdditionalBuildProducts => Inner.AdditionalBuildProducts;
-
-		public string? AdditionalCompilerArguments => Inner.AdditionalCompilerArguments;
-
-		public string? AdditionalLinkerArguments => Inner.AdditionalLinkerArguments;
-
-		public double MemoryPerActionGB => Inner.MemoryPerActionGB;
-
-		public string? GeneratedProjectName => Inner.GeneratedProjectName;
-
-		public ReadOnlyAndroidTargetRules AndroidPlatform
-		{
-			get;
-			private set;
-		}
-
-		public ReadOnlyLinuxTargetRules LinuxPlatform
-		{
-			get;
-			private set;
-		}
-
-		public ReadOnlyIOSTargetRules IOSPlatform
-		{
-			get;
-			private set;
-		}
-
-		public ReadOnlyMacTargetRules MacPlatform
-		{
-			get;
-			private set;
-		}
-
-		public ReadOnlyWindowsTargetRules WindowsPlatform
-		{
-			get;
-			private set;
-		}
-
-		public bool bShouldCompileAsDLL => Inner.bShouldCompileAsDLL;
-
-		public bool bGenerateProjectFiles => Inner.bGenerateProjectFiles;
-
-		public bool bIsEngineInstalled => Inner.bIsEngineInstalled;
-
-		public IReadOnlyList<string>? DisableUnityBuildForModules => Inner.DisableUnityBuildForModules;
-
-		public IReadOnlyList<string>? EnableOptimizeCodeForModules => Inner.EnableOptimizeCodeForModules;
-
-		public IReadOnlyList<string>? DisableOptimizeCodeForModules => Inner.DisableOptimizeCodeForModules;
-
-		public IReadOnlyList<string>? OptimizeForSizeModules => Inner.OptimizeForSizeModules;
-
-		public IReadOnlyList<string>? OptimizeForSizeAndSpeedModules => Inner.OptimizeForSizeAndSpeedModules;
-
-		public IReadOnlyList<UnrealTargetPlatform>? OptedInModulePlatforms => Inner.OptedInModulePlatforms;
-
-		public TestTargetRules? InnerTestTargetRules => Inner as TestTargetRules;
-
-#pragma warning restore C1591
-		#endregion
-
-		/// <summary>
-		/// Provide access to the RelativeEnginePath property for code referencing ModuleRules.BuildConfiguration.
-		/// </summary>
-		public string RelativeEnginePath => Unreal.EngineDirectory.MakeRelativeTo(DirectoryReference.GetCurrentDirectory()) + Path.DirectorySeparatorChar;
-
-		/// <summary>
-		/// Provide access to the UEThirdPartySourceDirectory property for code referencing ModuleRules.UEBuildConfiguration.
-		/// </summary>
-		public string UEThirdPartySourceDirectory => "ThirdParty/";
-
-		/// <summary>
-		/// Provide access to the UEThirdPartyBinariesDirectory property for code referencing ModuleRules.UEBuildConfiguration.
-		/// </summary>
-		public string UEThirdPartyBinariesDirectory => "../Binaries/ThirdParty/";
-
-		/// <summary>
-		/// Whether this is a low level tests target.
-		/// </summary>
-		public bool IsTestTarget => Inner.IsTestTarget;
-
-		/// <summary>
-		/// Whether this is a test target explicitly defined.
-		/// Explicitley defined test targets always inherit from TestTargetRules and define their own tests.
-		/// Implicit test targets are created from existing targets when building with -Mode=Test and they include tests from all dependencies.
-		/// </summary>
-		public bool ExplicitTestsTarget => Inner.ExplicitTestsTarget;
-
-		/// <summary>
-		/// Controls the value of WITH_LOW_LEVEL_TESTS that dictates whether module-specific low level tests are compiled in or not.
-		/// </summary>
-		public bool WithLowLevelTests => Inner.WithLowLevelTests;
-
-		/// <summary>
-		/// Get the platforms this target supports
-		/// </summary>
-		public IEnumerable<UnrealTargetPlatform> SupportedPlatforms => Inner.GetSupportedPlatforms();
-
-		/// <summary>
-		/// Get the configurations this target supports
-		/// </summary>
-		public IEnumerable<UnrealTargetConfiguration> SupportedConfigurations => Inner.GetSupportedConfigurations();
-
-		/// <summary>
-		/// Checks if current platform is part of a given platform group
-		/// </summary>
-		/// <param name="Group">The platform group to check</param>
-		/// <returns>True if current platform is part of a platform group</returns>
-		public bool IsInPlatformGroup(UnrealPlatformGroup Group)
-		{
-			return UEBuildPlatform.IsPlatformInGroup(Platform, Group);
-		}
-
-		/// <summary>
-		/// Gets diagnostic messages about default settings which have changed in newer versions of the engine
-		/// </summary>
-		/// <param name="Diagnostics">List of messages to be appended to</param>
-		internal void GetBuildSettingsInfo(List<string> Diagnostics)
-		{
-			Inner.GetBuildSettingsInfo(Diagnostics);
-		}
-
-		public bool IsPlatformOptedIn(UnrealTargetPlatform Platform)
-		{
-			return Inner.OptedInModulePlatforms == null || Inner.OptedInModulePlatforms.Contains(Platform);
-		}
-
-		/// <summary>
-		/// Checks if a plugin should be programmatically allowed in the build
-		/// </summary>
-		/// <returns>true if the plugin is allowed</returns>
-		public bool IgnorePluginDependency(PluginInfo ParentInfo, PluginReferenceDescriptor ChildDescriptor) => Inner.IgnorePluginDependency(ParentInfo, ChildDescriptor);
-
-		/// <summary>
-		/// Determines if the automation tests should be compiled based on the current configuration and optional forced settings.
-		/// If either the development tests or performance tests are compiled (unless explicitly disabled), this property returns true.
-		/// </summary>
-		private bool? bWithAutomationTestsPrivate = null;
-		public bool WithAutomationTests
-		{
-			get
-			{
-				if (!bWithAutomationTestsPrivate.HasValue)
-				{
-					bool bCompileDevTests = Configuration != UnrealTargetConfiguration.Shipping;
-					bool bCompilePerfTests = bCompileDevTests;
-
-					if (bForceCompileDevelopmentAutomationTests)
-					{
-						bCompileDevTests = true;
-					}
-					if (bForceCompilePerformanceAutomationTests)
-					{
-						bCompilePerfTests = true;
-					}
-					if (bForceDisableAutomationTests)
-					{
-						bCompileDevTests = bCompilePerfTests = false;
-					}
-
-					bWithAutomationTestsPrivate = bCompileDevTests || bCompilePerfTests;
-				}
-
-				return bWithAutomationTestsPrivate.Value;
-			}
+			List<string> diagnostics = new();
+			GetBuildSettingsInfo(diagnostics);
+			diagnostics.ForEach(x => Logger.LogWarning("Warning: {Message}", x));
 		}
 	}
 }
diff --git a/Engine/Source/Programs/UnrealBuildTool/Configuration/TestTargetRules.cs b/Engine/Source/Programs/UnrealBuildTool/Configuration/TestTargetRules.cs
--- a/Engine/Source/Programs/UnrealBuildTool/Configuration/TestTargetRules.cs
+++ b/Engine/Source/Programs/UnrealBuildTool/Configuration/TestTargetRules.cs
@@ -18,27 +18,27 @@ namespace UnrealBuildTool
 		/// Configuration mapping to control build graph test metadata generation during project files generation.
 		/// </summary>
 		[ConfigFile(ConfigHierarchyType.Engine, "LowLevelTestsSettings")]
-		public bool bUpdateBuildGraphPropertiesFile = false;
+		public bool bUpdateBuildGraphPropertiesFile { get; set; }
 
 		/// <summary>
 		/// Keeps track if low level tests executable must build with the Editor.
 		/// </summary>
-		public static bool bTestsRequireEditor = false;
+		public static bool bTestsRequireEditor { get; set; }
 
 		/// <summary>
 		/// Keeps track if low level tests executable must build with the Engine.
 		/// </summary>
-		public static bool bTestsRequireEngine = false;
+		public static bool bTestsRequireEngine { get; set; }
 
 		/// <summary>
 		/// Keeps track if low level tests executable must build with the ApplicationCore.
 		/// </summary>
-		public static bool bTestsRequireApplicationCore = false;
+		public static bool bTestsRequireApplicationCore { get; set; }
 
 		/// <summary>
 		/// Keeps track if low level tests executable must build with the CoreUObject.
 		/// </summary>
-		public static bool bTestsRequireCoreUObject = false;
+		public static bool bTestsRequireCoreUObject { get; set; }
 
 		/// <summary>
 		/// Associated tested target of this test target, if defined.
@@ -58,7 +58,7 @@ namespace UnrealBuildTool
 		/// <summary>
 		/// If set to true, it will not compile against ApplicationCore even if "ApplicationCore" is in the dependency graph.
 		/// </summary>
-		public bool bNeverCompileAgainstApplicationCore = false;
+		public bool bNeverCompileAgainstApplicationCore { get; set; }
 
 		/// <summary>
 		/// Test target override for bCompileAgainstCoreUObject.
@@ -73,7 +73,7 @@ namespace UnrealBuildTool
 		/// <summary>
 		/// If set to true, it will not compile against CoreUObject even if "CoreUObject" is in the dependency graph.
 		/// </summary>
-		public bool bNeverCompileAgainstCoreUObject = false;
+		public bool bNeverCompileAgainstCoreUObject { get; set; }
 
 		/// <summary>
 		/// Test target override for bCompileAgainstEngine.
@@ -88,7 +88,7 @@ namespace UnrealBuildTool
 		/// <summary>
 		/// If set to true, it will not compile against engine even if "Engine" is in the dependency graph.
 		/// </summary>
-		public bool bNeverCompileAgainstEngine = false;
+		public bool bNeverCompileAgainstEngine { get; set; }
 
 		/// <summary>
 		/// Test target override for bCompileAgainstEditor.
@@ -103,7 +103,7 @@ namespace UnrealBuildTool
 		/// <summary>
 		/// If set to true, it will not compile against editor even if "UnrealEd" is in the dependency graph.
 		/// </summary>
-		public bool bNeverCompileAgainstEditor = false;
+		public bool bNeverCompileAgainstEditor { get; set; }
 
 		/// <summary>
 		/// Whether to stub the platform file.
@@ -119,7 +119,7 @@ namespace UnrealBuildTool
 				GlobalDefinitions.Add($"UE_LLT_USE_PLATFORM_FILE_STUB={Convert.ToInt32(bUsePlatformFileStubPrivate)}");
 			}
 		}
-		private bool bUsePlatformFileStubPrivate = false;
+		private bool bUsePlatformFileStubPrivate { get; set; }
 
 		/// <summary>
 		/// Whether to mock engine default instances for materials, AI controller etc.
@@ -135,15 +135,15 @@ namespace UnrealBuildTool
 				GlobalDefinitions.Add($"UE_LLT_WITH_MOCK_ENGINE_DEFAULTS={Convert.ToInt32(bMockEngineDefaultsPrivate)}");
 			}
 		}
-		private bool bMockEngineDefaultsPrivate = false;
+		private bool bMockEngineDefaultsPrivate { get; set; }
 
 		/// <summary>
 		/// Constructor that explicit targets can inherit from.
 		/// </summary>
-		/// <param name="Target"></param>
-		public TestTargetRules(TargetInfo Target) : base(Target)
+		/// <param name="target"></param>
+		public TestTargetRules(TargetInfo target) : base(target)
 		{
-			SetupCommonProperties(Target);
+			SetupCommonProperties(target);
 
 			ExeBinariesSubFolder = LaunchModuleName = Name + (ExplicitTestsTarget ? String.Empty : "Tests");
 
@@ -152,10 +152,10 @@ namespace UnrealBuildTool
 				bBuildInSolutionByDefault = true;
 				if (ProjectFile != null)
 				{
-					DirectoryReference SamplesDirectory = DirectoryReference.Combine(Unreal.RootDirectory, "Samples");
+					DirectoryReference samplesDirectory = DirectoryReference.Combine(Unreal.RootDirectory, "Samples");
 					if (!ProjectFile.IsUnderDirectory(Unreal.EngineDirectory))
 					{
-						if (ProjectFile.IsUnderDirectory(SamplesDirectory))
+						if (ProjectFile.IsUnderDirectory(samplesDirectory))
 						{
 							SolutionDirectory = Path.Combine("Samples", ProjectFile.Directory.GetDirectoryName(), "LowLevelTests");
 						}
@@ -179,26 +179,26 @@ namespace UnrealBuildTool
 				bNeverCompileAgainstEngine = true;
 			}
 
-			if (Target.Platform == UnrealTargetPlatform.Win64)
+			if (target.Platform == UnrealTargetPlatform.Win64)
 			{
-				string OutputName = "$(TargetName)";
-				if (Target.Configuration != UndecoratedConfiguration)
+				string outputName = "$(TargetName)";
+				if (target.Configuration != UndecoratedConfiguration)
 				{
-					OutputName = OutputName + "-" + Target.Platform + "-" + Target.Configuration;
+					outputName = outputName + "-" + target.Platform + "-" + target.Configuration;
 				}
 				if (File != null)
 				{
-					DirectoryReference OutputDirectory;
+					DirectoryReference outputDirectory;
 					if (ProjectFile != null && File.IsUnderDirectory(ProjectFile.Directory))
 					{
-						OutputDirectory = UEBuildTarget.GetOutputDirectoryForExecutable(ProjectFile.Directory, File);
+						outputDirectory = UEBuildTarget.GetOutputDirectoryForExecutable(ProjectFile.Directory, File);
 					}
 					else
 					{
-						OutputDirectory = UEBuildTarget.GetOutputDirectoryForExecutable(Unreal.EngineDirectory, File);
+						outputDirectory = UEBuildTarget.GetOutputDirectoryForExecutable(Unreal.EngineDirectory, File);
 					}
-					FileReference OutputTestFile = FileReference.Combine(OutputDirectory, "Binaries", Target.Platform.ToString(), ExeBinariesSubFolder, $"{OutputName}.exe.is_unreal_test");
-					PostBuildSteps.Add("echo > " + OutputTestFile.FullName);
+					FileReference outputTestFile = FileReference.Combine(outputDirectory, "Binaries", target.Platform.ToString(), ExeBinariesSubFolder, $"{outputName}.exe.is_unreal_test");
+					PostBuildSteps.Add("echo > " + outputTestFile.FullName);
 				}
 			}
 		}
@@ -206,30 +206,26 @@ namespace UnrealBuildTool
 		/// <summary>
 		/// Constructs a valid TestTargetRules instance from an existent TargetRules instance.
 		/// </summary>
-		public static TestTargetRules Create(TargetRules Rules, TargetInfo TargetInfo)
+		public static TestTargetRules Create(TargetRules rules, TargetInfo targetInfo)
 		{
-			Type TestTargetRulesType = typeof(TestTargetRules);
-			TestTargetRules TestRules = (TestTargetRules)FormatterServices.GetUninitializedObject(TestTargetRulesType);
+			Type testTargetRulesType = typeof(TestTargetRules);
+			TestTargetRules testRules = (TestTargetRules)FormatterServices.GetUninitializedObject(testTargetRulesType);
 
 			// Initialize the logger before calling the constructor
-			TestRules.Logger = Rules.Logger;
-
-			ConstructorInfo? Constructor = TestTargetRulesType.GetConstructor(new Type[] { typeof(TargetRules), typeof(TargetInfo) });
-			if (Constructor == null)
-			{
-				throw new BuildException("No constructor found on {0} which takes first argument of type TargetRules and second of type TargetInfo.", TestTargetRulesType.Name);
-			}
+			testRules.Logger = rules.Logger;
 
+			ConstructorInfo? constructor = testTargetRulesType.GetConstructor(new Type[] { typeof(TargetRules), typeof(TargetInfo) })
+				?? throw new BuildException("No constructor found on {0} which takes first argument of type TargetRules and second of type TargetInfo.", testTargetRulesType.Name);
 			try
 			{
-				Constructor.Invoke(TestRules, new object[] { Rules, TargetInfo });
+				constructor.Invoke(testRules, new object[] { rules, targetInfo });
 			}
-			catch (Exception Ex)
+			catch (Exception ex)
 			{
-				throw new BuildException(Ex, "Unable to instantiate instance of '{0}' object type from compiled assembly '{1}'.  Unreal Build Tool creates an instance of your module's 'Rules' object in order to find out about your module's requirements.  The CLR exception details may provide more information:  {2}", TestTargetRulesType.Name, Path.GetFileNameWithoutExtension(TestTargetRulesType.Assembly?.Location), Ex.ToString());
+				throw new BuildException(ex, "Unable to instantiate instance of '{0}' object type from compiled assembly '{1}'.  Unreal Build Tool creates an instance of your module's 'Rules' object in order to find out about your module's requirements.  The CLR exception details may provide more information:  {2}", testTargetRulesType.Name, Path.GetFileNameWithoutExtension(testTargetRulesType.Assembly?.Location), ex.ToString());
 			}
 
-			return TestRules;
+			return testRules;
 		}
 
 		/// <summary>
@@ -237,33 +233,34 @@ namespace UnrealBuildTool
 		/// TestTargetRules is setup as a program and is linked monolithically.
 		/// It removes a lot of default compilation behavior in order to produce a minimal test environment.
 		/// </summary>
-		public TestTargetRules(TargetRules TestedTarget, TargetInfo Target) : base(Target)
+		public TestTargetRules(TargetRules testedTarget, TargetInfo target) : base(target)
 		{
-			if (TestedTarget is TestTargetRules)
+			if (testedTarget is TestTargetRules)
 			{
 				throw new BuildException("TestedTarget can't be of type TestTargetRules.");
 			}
 
-			this.TestedTarget = TestedTarget;
+			TestedTarget = testedTarget;
 
-			TargetFiles = TestedTarget.TargetFiles;
+			TargetFiles = testedTarget.TargetFiles;
 
-			ExeBinariesSubFolder = Name = TestedTarget.Name + "Tests";
-			TargetSourceFile = File = TestedTarget.File;
-			if (TestedTarget.LaunchModuleName != null)
+			ExeBinariesSubFolder = Name = testedTarget.Name + "Tests";
+			TargetSourceFile = File = testedTarget.File;
+			if (testedTarget.LaunchModuleName != null)
 			{
-				LaunchModuleName = TestedTarget.LaunchModuleName + "Tests";
+				LaunchModuleName = testedTarget.LaunchModuleName + "Tests";
 			}
 
-			ManifestFileNames = TestedTarget.ManifestFileNames;
+			ManifestFileNames.Clear();
+			ManifestFileNames.AddRange(testedTarget.ManifestFileNames);
 
-			WindowsPlatform = TestedTarget.WindowsPlatform;
+			WindowsPlatform = testedTarget.WindowsPlatform;
 
-			SetupCommonProperties(Target);
-			SetupImplicitTestProperties(TestedTarget);
+			SetupCommonProperties(target);
+			SetupImplicitTestProperties(testedTarget);
 		}
 
-		private void SetupCommonProperties(TargetInfo Target)
+		private void SetupCommonProperties(TargetInfo target)
 		{
 			IncludeOrderVersion = EngineIncludeOrderVersion.Latest;
 
@@ -278,7 +275,7 @@ namespace UnrealBuildTool
 
 			bBuildInSolutionByDefault = false;
 
-			bDeployAfterCompile = Target.Platform != UnrealTargetPlatform.Android;
+			bDeployAfterCompile = target.Platform != UnrealTargetPlatform.Android;
 			bIsBuildingConsoleApplication = true;
 
 			// Disabling default true flags that aren't necessary for tests
@@ -315,7 +312,7 @@ namespace UnrealBuildTool
 			bBuildDeveloperTools = false;
 
 			// Useful for debugging test failures
-			if (Target.Configuration == UnrealTargetConfiguration.Debug)
+			if (target.Configuration == UnrealTargetConfiguration.Debug)
 			{
 				bDebugBuildsActuallyUseDebugCRT = true;
 			}
@@ -333,9 +330,9 @@ namespace UnrealBuildTool
 			GlobalDefinitions.Add("UE_LLT_WITH_MOCK_ENGINE_DEFAULTS=0");
 
 			// Platform specific setup
-			if (Target.Platform == UnrealTargetPlatform.Android)
+			if (target.Platform == UnrealTargetPlatform.Android)
 			{
-				UndecoratedConfiguration = Target.Configuration;
+				UndecoratedConfiguration = target.Configuration;
 
 				GlobalDefinitions.Add("USE_ANDROID_INPUT=0");
 				GlobalDefinitions.Add("USE_ANDROID_OPENGL=0");
@@ -344,22 +341,22 @@ namespace UnrealBuildTool
 
 				// Workaround for a linker bug when building LowLevelTests for Android
 				// TODO: This should be written to the intermediate directory, somewhere else not when TargetRules are being created
-				FileReference VersionScriptFile = new FileReference(Path.GetTempPath() + $"LLTWorkaroundScrip-{Name}.ldscript");
-				FileReference.WriteAllTextIfDifferent(VersionScriptFile, "{ local: *; };");
-				AdditionalLinkerArguments = " -Wl,--version-script=\"" + VersionScriptFile.FullName + "\"";
+				FileReference versionScriptFile = new FileReference(Path.GetTempPath() + $"LLTWorkaroundScrip-{Name}.ldscript");
+				FileReference.WriteAllTextIfDifferent(versionScriptFile, "{ local: *; };");
+				AdditionalLinkerArguments = " -Wl,--version-script=\"" + versionScriptFile.FullName + "\"";
 			}
-			else if (Target.Platform == UnrealTargetPlatform.IOS)
+			else if (target.Platform == UnrealTargetPlatform.IOS)
 			{
 				bIsBuildingConsoleApplication = false;
 			}
 		}
 
-		private void SetupImplicitTestProperties(TargetRules TestedTarget)
+		private void SetupImplicitTestProperties(TargetRules testedTarget)
 		{
-			bool IsEditorTestedTarget = (TestedTarget.Type == TargetType.Editor);
-			bBuildWithEditorOnlyData = bCompileAgainstEditor = IsEditorTestedTarget;
-			bBuildDeveloperTools = bCompileAgainstEngine && IsEditorTestedTarget;
-			
+			bool isEditorTestedTarget = (testedTarget.Type == TargetType.Editor);
+			bBuildWithEditorOnlyData = bCompileAgainstEditor = isEditorTestedTarget;
+			bBuildDeveloperTools = bCompileAgainstEngine && isEditorTestedTarget;
+
 			bUsePlatformFileStub = bCompileAgainstEngine;
 			bMockEngineDefaults = bCompileAgainstEngine;
 			bCompileWithPluginSupport = bCompileAgainstEngine;
diff --git a/Engine/Source/Programs/UnrealBuildTool/System/RulesAssembly.cs b/Engine/Source/Programs/UnrealBuildTool/System/RulesAssembly.cs
--- a/Engine/Source/Programs/UnrealBuildTool/System/RulesAssembly.cs
+++ b/Engine/Source/Programs/UnrealBuildTool/System/RulesAssembly.cs
@@ -645,7 +645,7 @@ namespace UnrealBuildTool
 			if (!bSkipValidation)
 			{
 				// Delayed-fixup of TargetBuildEnvironment.UniqueIfNeeded
-				Rules.UpdateBuildEnvironmentIfNeeded(this, Arguments: null, Logger);
+				Rules.UpdateBuildEnvironmentIfNeeded(this, arguments: null, Logger);
 
 				// Set the default value for whether to use the shared build environment
 				if (Rules.BuildEnvironment == TargetBuildEnvironment.Unique && Unreal.IsEngineInstalled())
