UnrealBuildTool: -Mode=GenerateClangDatabase should use clang by defualt as the compiler

diff --git a/Engine/Source/Programs/UnrealBuildTool/Modes/GenerateClangDatabase.cs b/Engine/Source/Programs/UnrealBuildTool/Modes/GenerateClangDatabase.cs
--- a/Engine/Source/Programs/UnrealBuildTool/Modes/GenerateClangDatabase.cs
+++ b/Engine/Source/Programs/UnrealBuildTool/Modes/GenerateClangDatabase.cs
@@ -155,6 +155,11 @@ namespace UnrealBuildTool
 					TargetDescriptor.bUseUnityBuild = false;
 					TargetDescriptor.IntermediateEnvironment = UnrealIntermediateEnvironment.GenerateClangDatabase;
 					TargetDescriptor.AdditionalArguments = TargetDescriptor.AdditionalArguments.Append(new string[] { "-NoPCH" });
+					// Default the compiler to clang
+					if (!TargetDescriptor.AdditionalArguments.Any(x => x.StartsWith("-Compiler=", StringComparison.OrdinalIgnoreCase)))
+					{
+						TargetDescriptor.AdditionalArguments = TargetDescriptor.AdditionalArguments.Append(new string[] { "-Compiler=Clang" });
+					}
 					if (bNoUseVerseBPVM)
 					{
 						TargetDescriptor.AdditionalArguments = TargetDescriptor.AdditionalArguments.Append(new string[] { "-NoUseVerseBPVM" });
