UnrealBuildTool: Split out circular dependency module check out of module creation and into a ValidateModule function that can continue to recurse until a circular dependency is found or the graph is exhaused so the entire graph will be checked and so messages will be emitted for every build of a target even if the makefile is up to date

#rnx

diff --git a/Engine/Source/Programs/UnrealBuildTool/Configuration/UEBuildModule.cs b/Engine/Source/Programs/UnrealBuildTool/Configuration/UEBuildModule.cs
--- a/Engine/Source/Programs/UnrealBuildTool/Configuration/UEBuildModule.cs
+++ b/Engine/Source/Programs/UnrealBuildTool/Configuration/UEBuildModule.cs
@@ -163,5 +163,10 @@
 		public List<UEBuildModule>? DynamicallyLoadedModules;
 
 		/// <summary>
+		/// Set of modules that have referenced this module
+		/// </summary>
+		protected HashSet<UEBuildModule>? ReferenceStackParentModules;
+
+		/// <summary>
 		/// Set of all allowed restricted folder references
 		/// </summary>
@@ -1104,11 +1109,7 @@
 			// Set the reference chain for anything referenced by this module
 			string NextReferenceChain = String.Format("{0} -> {1}", ReferenceChain, ThisRefName);
 
-			// We need to check for cycles if this module has already been created and its name is in the stack.
-			bool CheckForCycles = PrivateIncludePathModules != null && ReferenceStack.Contains(this);
-
-			// Add us to the reference stack - note do this before checking for cycles as we allow them if an element in the stack
-			// declares the next element at leading to a dependency cycle.
+			// Add us to the reference stack
 			ReferenceStack.Add(this);
 
 			// Check if this module is invalid for the current target, based on Module attributes
@@ -1117,6 +1118,88 @@
 				Logger.LogWarning("Warning: Referenced module '{Module}' does not support {InvalidReason} via {ReferenceChain}", ThisRefName, InvalidReason, NextReferenceChain);
 			}
 
+			// Create all the referenced modules. This path can be recursive, so we check against PublicDependencyModules to ensure we don't recurse through the 
+			// same module twice (it produces better errors if something fails).
+			if (PublicDependencyModules == null)
+			{
+				// Log dependencies if required
+				if (Logger.IsEnabled((LogLevel)LogEventType.VeryVerbose))
+				{
+					Logger.LogTrace("Module {Name} dependencies:", Name);
+					LogDependencyNameList("Public:", Rules.PublicDependencyModuleNames, Logger);
+					LogDependencyNameList("Private:", Rules.PrivateDependencyModuleNames, Logger);
+					LogDependencyNameList("Dynamic:", Rules.DynamicallyLoadedModuleNames, Logger);
+					LogDependencyNameList("Public Include Paths:", Rules.PublicIncludePathModuleNames, Logger);
+					LogDependencyNameList("Private Include Paths:", Rules.PrivateIncludePathModuleNames, Logger);
+				}
+
+				// Create all the dependency modules - pass through the reference stack so we can check for cycles
+				RecursivelyCreateModulesByName(Rules.PublicDependencyModuleNames, ref PublicDependencyModules, ref bDependsOnVerse, CreateModule, NextReferenceChain, ReferenceStack, Logger);
+				if (Rules.Target.IsTestTarget)
+				{
+					// Move the test runner dependency to last position to give it the opportunity to build its special dependencies such as CoreUObject, ApplicationCore etc.
+					MoveTestsRunnerDependencyToLastPosition();
+				}
+				RecursivelyCreateModulesByName(Rules.PrivateDependencyModuleNames, ref PrivateDependencyModules, ref bDependsOnVerse, CreateModule, NextReferenceChain, ReferenceStack, Logger);
+
+				// Recursively create all the public include path modules
+				HashSet<string> PublicIncludePathModuleNames = new();
+				PublicIncludePathModuleNames.UnionWith(Rules.PublicIncludePathModuleNames!);
+				PublicIncludePathModuleNames.UnionWith(Rules.PublicDependencyModuleNames!);
+				RecursivelyCreateIncludePathModulesByName(PublicIncludePathModuleNames, ref PublicIncludePathModules, ref bDependsOnVerse, CreateModule, NextReferenceChain);
+
+				// Create the private include path modules
+				RecursivelyCreateIncludePathModulesByName(Rules.PrivateIncludePathModuleNames, ref PrivateIncludePathModules, ref bDependsOnVerse, CreateModule, NextReferenceChain);
+
+				// Dynamic loads aren't considered a reference chain so start with an empty stack
+				RecursivelyCreateModulesByName(Rules.DynamicallyLoadedModuleNames, ref DynamicallyLoadedModules, ref bDependsOnVerse, CreateModule, NextReferenceChain, new List<UEBuildModule>(), Logger);
+			}
+
+			// pop us off the current stack
+			ReferenceStack.RemoveAt(ReferenceStack.Count - 1);
+		}
+
+		/// <summary>
+		/// Public entry point to validate a module
+		/// </summary>
+		/// <param name="ReferenceChain"></param>
+		/// <param name="Logger"></param>
+		public bool ValidateModule(string ReferenceChain, ILogger Logger)
+		{
+			List<UEBuildModule> ReferenceStack = new List<UEBuildModule>();
+			RecursivelyCheckForCircularReferences(ReferenceChain, ReferenceStack, Logger);
+
+			return false; // No fatal errors
+		}
+
+		/// <summary>
+		/// Recursively check a module for circular dependencies
+		/// </summary>
+		/// <param name="ReferenceChain">Chain of references before reaching this module</param>
+		/// <param name="ReferenceStack">Stack of module dependencies that led to this module</param>
+		/// <param name="Logger">Logger for output</param>
+		protected void RecursivelyCheckForCircularReferences(string ReferenceChain, List<UEBuildModule> ReferenceStack, ILogger Logger)
+		{
+			// Name of this reference
+			string ThisRefName = (RulesFile == null) ? Name : RulesFile.GetFileName();
+
+			// Set the reference chain for anything referenced by this module
+			string NextReferenceChain = String.Format("{0} -> {1}", ReferenceChain, ThisRefName);
+
+			// We need to check for cycles if this module has already been created and its name is in the stack.
+			bool CheckForCycles = PrivateIncludePathModules != null && ReferenceStack.Contains(this);
+
+			UEBuildModule? PreviousModule = ReferenceStack.LastOrDefault();
+			ReferenceStackParentModules ??= new();
+			if (PreviousModule != null && !ReferenceStackParentModules.Add(PreviousModule))
+			{
+				return;
+			}
+
+			// Add us to the reference stack - note do this before checking for cycles as we allow them if an element in the stack
+			// declares the next element at leading to a dependency cycle.
+			ReferenceStack.Add(this);
+
 			if (CheckForCycles)
 			{
 				// We are FeatureModuleA and we know we're already in the list so there must be a circular dependency like this:
@@ -1129,7 +1212,7 @@
 				string? GuiltyModule = null;
 				string? VictimModule = null;
 
-				for (int i = 0; i < ReferenceStack.Count() - 1; i++)
+				for (int i = 0; i < ReferenceStack.Count - 1; i++)
 				{
 					UEBuildModule ReferringModule = ReferenceStack.ElementAt(i);
 					UEBuildModule TargetModule = ReferenceStack.ElementAt(i + 1);
@@ -1146,56 +1229,23 @@
 				if (String.IsNullOrEmpty(GuiltyModule))
 				{
 					string CycleChain = String.Join(" -> ", ReferenceStack);
-					Logger.LogError("Circular dependency on {Name} detected.\n" +
+					Logger.LogError("Circular dependency on {Name} detected:\n" +
 						"\tFull Route: {FullRoute}\n" +
-						"\tCycled Route: is {CycleRoot}.\n" +
+						"\tCycled Route: is {CycleRoute}\n" +
 						"Break this loop by moving dependencies into a separate module or using Private/PublicIncludePathModuleNames to reference declarations\n",
 						ThisRefName, NextReferenceChain, CycleChain);
-
 				}
 				else
 				{
 					Logger.LogDebug("Found circular reference to {ThisRefName}, but {GuiltyModule} declares a cycle on {VictimModule} which breaks the chain", ThisRefName, GuiltyModule, VictimModule);
 				}
 			}
-
-			// Create all the referenced modules. This path can be recursive, so we check against PublicDependencyModules to ensure we don't recurse through the 
-			// same module twice (it produces better errors if something fails).
-			if (PublicDependencyModules == null)
+			else
 			{
-				// Log dependencies if required
-				if (Logger.IsEnabled((LogLevel)LogEventType.VeryVerbose))
-				{
-					Logger.LogTrace("Module {0} dependencies:", Name);
-					LogDependencyNameList("Public:", Rules.PublicDependencyModuleNames, Logger);
-					LogDependencyNameList("Private:", Rules.PrivateDependencyModuleNames, Logger);
-					LogDependencyNameList("Dynamic:", Rules.DynamicallyLoadedModuleNames, Logger);
-					LogDependencyNameList("Public Include Paths:", Rules.PublicIncludePathModuleNames, Logger);
-					LogDependencyNameList("Private Include Paths:", Rules.PrivateIncludePathModuleNames, Logger);
-				}
-
-				// Create all the dependency modules - pass through the reference stack so we can check for cycles
-				RecursivelyCreateModulesByName(Rules.PublicDependencyModuleNames, ref PublicDependencyModules, ref bDependsOnVerse, CreateModule, NextReferenceChain, ReferenceStack, Logger);
-				if (Rules.Target.IsTestTarget)
-				{
-					// Move the test runner dependency to last position to give it the opportunity to build its special dependencies such as CoreUObject, ApplicationCore etc.
-					MoveTestsRunnerDependencyToLastPosition();
-				}
-				RecursivelyCreateModulesByName(Rules.PrivateDependencyModuleNames, ref PrivateDependencyModules, ref bDependsOnVerse, CreateModule, NextReferenceChain, ReferenceStack, Logger);
-
-				// Recursively create all the public include path modules
-				HashSet<string> PublicIncludePathModuleNames = new();
-				PublicIncludePathModuleNames.UnionWith(Rules.PublicIncludePathModuleNames!);
-				PublicIncludePathModuleNames.UnionWith(Rules.PublicDependencyModuleNames!);
-				RecursivelyCreateIncludePathModulesByName(PublicIncludePathModuleNames, ref PublicIncludePathModules, ref bDependsOnVerse, CreateModule, NextReferenceChain);
-
-				// Create the private include path modules
-				RecursivelyCreateIncludePathModulesByName(Rules.PrivateIncludePathModuleNames, ref PrivateIncludePathModules, ref bDependsOnVerse, CreateModule, NextReferenceChain);
-
-				// Dynamic loads aren't considered a reference chain so start with an empty stack
-				RecursivelyCreateModulesByName(Rules.DynamicallyLoadedModuleNames, ref DynamicallyLoadedModules, ref bDependsOnVerse, CreateModule, NextReferenceChain, new List<UEBuildModule>(), Logger);
+				// Continue checking for circular dependencies until a cycle is found or the graph is exhausted
+				PublicDependencyModules?.ForEach(x => x.RecursivelyCheckForCircularReferences(NextReferenceChain, ReferenceStack, Logger));
+				PrivateDependencyModules?.ForEach(x => x.RecursivelyCheckForCircularReferences(NextReferenceChain, ReferenceStack, Logger));
 			}
-
 			// pop us off the current stack
 			ReferenceStack.RemoveAt(ReferenceStack.Count - 1);
 		}
diff --git a/Engine/Source/Programs/UnrealBuildTool/Configuration/UEBuildTarget.cs b/Engine/Source/Programs/UnrealBuildTool/Configuration/UEBuildTarget.cs
--- a/Engine/Source/Programs/UnrealBuildTool/Configuration/UEBuildTarget.cs
+++ b/Engine/Source/Programs/UnrealBuildTool/Configuration/UEBuildTarget.cs
@@ -3380,6 +3380,12 @@ namespace UnrealBuildTool
 					AnyErrors |= UEBuildPlatform.GetBuildPlatform(Platform).ValidateModuleIncludePaths(Module, Rules, Modules.Values);
 				}
 
+				// Intentionally not ordered by name to maintain the reference include order
+				foreach (UEBuildModule Module in Modules.Values)
+				{
+					Module.ValidateModule("Target", Logger);
+				}
+
 				if (AnyErrors)
 				{
 					throw new BuildException("Errors validating modules.");
