UnrealBuildTool: ModuleRules cleanup pt. 1

#rnx

diff --git a/Engine/Source/Programs/UnrealBuildTool/Configuration/ModuleRules.cs b/Engine/Source/Programs/UnrealBuildTool/Configuration/ModuleRules.cs
--- a/Engine/Source/Programs/UnrealBuildTool/Configuration/ModuleRules.cs
+++ b/Engine/Source/Programs/UnrealBuildTool/Configuration/ModuleRules.cs
@@ -274,40 +274,40 @@ namespace UnrealBuildTool
 			/// <summary>
 			/// The file that should be staged. Should use $(EngineDir) and $(ProjectDir) variables as a root, so that the target can be relocated to different machines.
 			/// </summary>
-			public string Path;
+			public string Path { get; init; }
 
 			/// <summary>
 			/// The initial location for this file. It will be copied to Path at build time, ready for staging.
 			/// </summary>
-			public string? SourcePath;
+			public string? SourcePath { get; init; }
 
 			/// <summary>
 			/// How to stage this file.
 			/// </summary>
-			public StagedFileType Type;
+			public StagedFileType Type { get; init; }
 
 			/// <summary>
 			/// Constructor
 			/// </summary>
-			/// <param name="InPath">Path to the runtime dependency</param>
-			/// <param name="InType">How to stage the given path</param>
-			public RuntimeDependency(string InPath, StagedFileType InType = StagedFileType.NonUFS)
+			/// <param name="inPath">Path to the runtime dependency</param>
+			/// <param name="inType">How to stage the given path</param>
+			public RuntimeDependency(string inPath, StagedFileType inType = StagedFileType.NonUFS)
 			{
-				Path = InPath;
-				Type = InType;
+				Path = inPath;
+				Type = inType;
 			}
 
 			/// <summary>
 			/// Constructor
 			/// </summary>
-			/// <param name="InPath">Path to the runtime dependency</param>
-			/// <param name="InSourcePath">Source path for the file in the working tree</param>
-			/// <param name="InType">How to stage the given path</param>
-			public RuntimeDependency(string InPath, string InSourcePath, StagedFileType InType = StagedFileType.NonUFS)
+			/// <param name="inPath">Path to the runtime dependency</param>
+			/// <param name="inSourcePath">Source path for the file in the working tree</param>
+			/// <param name="inType">How to stage the given path</param>
+			public RuntimeDependency(string inPath, string inSourcePath, StagedFileType inType = StagedFileType.NonUFS)
 			{
-				Path = InPath;
-				SourcePath = InSourcePath;
-				Type = InType;
+				Path = inPath;
+				SourcePath = inSourcePath;
+				Type = inType;
 			}
 		}
 
@@ -320,7 +320,12 @@ namespace UnrealBuildTool
 			/// <summary>
 			/// Inner list of runtime dependencies
 			/// </summary>
-			internal List<RuntimeDependency> Inner = new List<RuntimeDependency>();
+			readonly List<RuntimeDependency> _inner = new();
+
+			/// <summary>
+			/// Readonly access of inner list of runtime dependencies
+			/// </summary>
+			internal IReadOnlyList<RuntimeDependency> Inner => _inner.AsReadOnly();
 
 			/// <summary>
 			/// Default constructor
@@ -332,31 +337,31 @@ namespace UnrealBuildTool
 			/// <summary>
 			/// Add a runtime dependency to the list
 			/// </summary>
-			/// <param name="InPath">Path to the runtime dependency. May include wildcards.</param>
-			public void Add(string InPath)
+			/// <param name="inPath">Path to the runtime dependency. May include wildcards.</param>
+			public void Add(string inPath)
 			{
-				Inner.Add(new RuntimeDependency(InPath));
+				_inner.Add(new RuntimeDependency(inPath));
 			}
 
 			/// <summary>
 			/// Add a runtime dependency to the list
 			/// </summary>
-			/// <param name="InPath">Path to the runtime dependency. May include wildcards.</param>
-			/// <param name="InType">How to stage this file</param>
-			public void Add(string InPath, StagedFileType InType)
+			/// <param name="inPath">Path to the runtime dependency. May include wildcards.</param>
+			/// <param name="inType">How to stage this file</param>
+			public void Add(string inPath, StagedFileType inType)
 			{
-				Inner.Add(new RuntimeDependency(InPath, InType));
+				_inner.Add(new RuntimeDependency(inPath, inType));
 			}
 
 			/// <summary>
 			/// Add a runtime dependency to the list
 			/// </summary>
-			/// <param name="InPath">Path to the runtime dependency. May include wildcards.</param>
-			/// <param name="InSourcePath">Source path for the file to be added as a dependency. May include wildcards.</param>
-			/// <param name="InType">How to stage this file</param>
-			public void Add(string InPath, string InSourcePath, StagedFileType InType = StagedFileType.NonUFS)
+			/// <param name="inPath">Path to the runtime dependency. May include wildcards.</param>
+			/// <param name="inSourcePath">Source path for the file to be added as a dependency. May include wildcards.</param>
+			/// <param name="inType">How to stage this file</param>
+			public void Add(string inPath, string inSourcePath, StagedFileType inType = StagedFileType.NonUFS)
 			{
-				Inner.Add(new RuntimeDependency(InPath, InSourcePath, InType));
+				_inner.Add(new RuntimeDependency(inPath, inSourcePath, inType));
 			}
 		}
 
@@ -369,7 +374,12 @@ namespace UnrealBuildTool
 			/// <summary>
 			/// Inner list of runtime dependencies
 			/// </summary>
-			internal List<ReceiptProperty> Inner = new List<ReceiptProperty>();
+			readonly List<ReceiptProperty> _inner = new();
+
+			/// <summary>
+			/// Readonly access of inner list of runtime dependencies
+			/// </summary>
+			internal IReadOnlyList<ReceiptProperty> Inner => _inner.AsReadOnly();
 
 			/// <summary>
 			/// Default constructor
@@ -381,11 +391,21 @@ namespace UnrealBuildTool
 			/// <summary>
 			/// Add a receipt property to the list
 			/// </summary>
-			/// <param name="Name">Name of the property</param>
-			/// <param name="Value">Value for the property</param>
-			public void Add(string Name, string Value)
+			/// <param name="name">Name of the property</param>
+			/// <param name="value">Value for the property</param>
+			public void Add(string name, string value)
 			{
-				Inner.Add(new ReceiptProperty(Name, Value));
+				_inner.Add(new ReceiptProperty(name, value));
+			}
+
+			/// <summary>
+			/// Remove recepit properties from the list
+			/// </summary>
+			/// <param name="match">the maatcher predicate</param>
+			/// <returns>the number of items removed</returns>
+			public int RemoveAll(Predicate<ReceiptProperty> match)
+			{
+				return _inner.RemoveAll(match);
 			}
 		}
 
@@ -397,17 +417,17 @@ namespace UnrealBuildTool
 			/// <summary>
 			/// Name of the framework
 			/// </summary>
-			internal string Name;
+			internal string Name { get; init; }
 
 			/// <summary>
 			/// Specifies the path to a zip file that contains it or where the framework is located on disk
 			/// </summary>
-			internal string Path;
+			internal string Path { get; init; }
 
 			/// <summary>
 			/// 
 			/// </summary>
-			internal string? CopyBundledAssets = null;
+			internal string? CopyBundledAssets { get; init; } = null;
 
 			/// <summary>
 			/// How to handle linking and copying the framework
@@ -433,42 +453,39 @@ namespace UnrealBuildTool
 			/// <summary>
 			/// How to treat the framework during linking and creating the .app
 			/// </summary>
-			internal FrameworkMode Mode;
+			internal FrameworkMode Mode { get; init; }
 
 			/// <summary>
 			/// Constructor
 			/// </summary>
-			/// <param name="Name">Name of the framework</param>
-			/// <param name="Path">Path to a zip file containing the framework or a framework on disk</param>
-			/// <param name="CopyBundledAssets"></param>
+			/// <param name="name">Name of the framework</param>
+			/// <param name="path">Path to a zip file containing the framework or a framework on disk</param>
+			/// <param name="copyBundledAssets"></param>
 			/// <param name="bCopyFramework">Copy the framework to the target's Framework directory</param>
-			public Framework(string Name, string Path, string? CopyBundledAssets = null, bool bCopyFramework = false)
-				: this(Name, Path, bCopyFramework ? FrameworkMode.LinkAndCopy : FrameworkMode.Link, CopyBundledAssets)
+			public Framework(string name, string path, string? copyBundledAssets = null, bool bCopyFramework = false)
+				: this(name, path, bCopyFramework ? FrameworkMode.LinkAndCopy : FrameworkMode.Link, copyBundledAssets)
 			{
 			}
 
 			/// <summary>
 			/// Constructor
 			/// </summary>
-			/// <param name="Name">Name of the framework</param>
-			/// <param name="Path">Path to a zip file containing the framework or a framework on disk</param>
-			/// <param name="Mode">How to treat the framework during linking and creating the .app</param>
-			/// <param name="CopyBundledAssets"></param>
-			public Framework(string Name, string Path, FrameworkMode Mode, string? CopyBundledAssets = null)
+			/// <param name="name">Name of the framework</param>
+			/// <param name="path">Path to a zip file containing the framework or a framework on disk</param>
+			/// <param name="mode">How to treat the framework during linking and creating the .app</param>
+			/// <param name="copyBundledAssets"></param>
+			public Framework(string name, string path, FrameworkMode mode, string? copyBundledAssets = null)
 			{
-				this.Name = Name;
-				this.Path = Path;
-				this.Mode = Mode;
-				this.CopyBundledAssets = CopyBundledAssets;
+				Name = name;
+				Path = path;
+				Mode = mode;
+				CopyBundledAssets = copyBundledAssets;
 			}
 
 			/// <summary>
 			/// Specifies if the file is a zip file
 			/// </summary>
-			public bool IsZipFile()
-			{
-				return Path.EndsWith(".zip");
-			}
+			public bool IsZipFile() => Path.EndsWith(".zip", StringComparison.OrdinalIgnoreCase);
 		}
 
 		/// <summary>
@@ -479,28 +496,28 @@ namespace UnrealBuildTool
 			/// <summary>
 			/// 
 			/// </summary>
-			public string? ResourcePath = null;
+			public string? ResourcePath { get; init; } = null;
 
 			/// <summary>
 			/// 
 			/// </summary>
-			public string? BundleContentsSubdir = null;
+			public string? BundleContentsSubdir { get; init; } = null;
 
 			/// <summary>
 			/// 
 			/// </summary>
-			public bool bShouldLog = true;
+			public bool bShouldLog { get; init; } = true;
 
 			/// <summary>
 			/// Constructor
 			/// </summary>
-			/// <param name="ResourcePath"></param>
-			/// <param name="BundleContentsSubdir"></param>
+			/// <param name="resourcePath"></param>
+			/// <param name="bundleContentsSubdir"></param>
 			/// <param name="bShouldLog"></param>
-			public BundleResource(string ResourcePath, string BundleContentsSubdir = "Resources", bool bShouldLog = true)
+			public BundleResource(string resourcePath, string bundleContentsSubdir = "Resources", bool bShouldLog = true)
 			{
-				this.ResourcePath = ResourcePath;
-				this.BundleContentsSubdir = BundleContentsSubdir;
+				ResourcePath = resourcePath;
+				BundleContentsSubdir = bundleContentsSubdir;
 				this.bShouldLog = bShouldLog;
 			}
 		}
@@ -513,40 +530,36 @@ namespace UnrealBuildTool
 			/// <summary>
 			/// Name of the type library
 			/// </summary>
-			public string FileName;
+			public string FileName { get; init; }
 
 			/// <summary>
 			/// Additional attributes for the #import directive
 			/// </summary>
-			public string Attributes;
+			public string Attributes { get; init; }
 
 			/// <summary>
 			/// Name of the output header
 			/// </summary>
-			public string Header;
+			public string Header { get; init; }
 
 			/// <summary>
 			/// Constructor
 			/// </summary>
-			/// <param name="FileName">Name of the type library. Follows the same conventions as the filename parameter in the MSVC #import directive.</param>
-			/// <param name="Attributes">Additional attributes for the import directive</param>
-			/// <param name="Header">Name of the output header</param>
-			public TypeLibrary(string FileName, string Attributes, string Header)
+			/// <param name="fileName">Name of the type library. Follows the same conventions as the filename parameter in the MSVC #import directive.</param>
+			/// <param name="attributes">Additional attributes for the import directive</param>
+			/// <param name="header">Name of the output header</param>
+			public TypeLibrary(string fileName, string attributes, string header)
 			{
-				this.FileName = FileName;
-				this.Attributes = Attributes;
-				this.Header = Header;
+				FileName = fileName;
+				Attributes = attributes;
+				Header = header;
 			}
 		}
 
 		/// <summary>
 		/// Name of this module
 		/// </summary>
-		public string Name
-		{
-			get;
-			internal set;
-		}
+		public string Name { get; internal set; }
 
 		/// <summary>
 		/// File containing this module
diff --git a/Engine/Source/Programs/UnrealBuildTool/Configuration/TestModuleRules.cs b/Engine/Source/Programs/UnrealBuildTool/Configuration/TestModuleRules.cs
--- a/Engine/Source/Programs/UnrealBuildTool/Configuration/TestModuleRules.cs
+++ b/Engine/Source/Programs/UnrealBuildTool/Configuration/TestModuleRules.cs
@@ -145,7 +145,7 @@ namespace UnrealBuildTool
 		/// </summary>
 		protected void SetResourcesFolder(string ResourcesRelativeFolder)
 		{
-			AdditionalPropertiesForReceipt.Inner.RemoveAll(Prop => Prop.Name == "ResourcesFolder");
+			AdditionalPropertiesForReceipt.RemoveAll(Prop => Prop.Name == "ResourcesFolder");
 
 			foreach (DirectoryReference Directory in GetAllModuleDirectories())
 			{
