UnrealBuildTool ModuleRulesContext.cs cleanup

#rnx

diff --git a/Engine/Source/Programs/UnrealBuildTool/Configuration/ModuleRulesContext.cs b/Engine/Source/Programs/UnrealBuildTool/Configuration/ModuleRulesContext.cs
--- a/Engine/Source/Programs/UnrealBuildTool/Configuration/ModuleRulesContext.cs
+++ b/Engine/Source/Programs/UnrealBuildTool/Configuration/ModuleRulesContext.cs
@@ -7,72 +7,72 @@ namespace UnrealBuildTool
 	/// <summary>
 	/// Stores information about where a module rules object came from, and how it can be used. 
 	/// </summary>
-	class ModuleRulesContext
+	internal class ModuleRulesContext
 	{
 		/// <summary>
 		/// The scope for this module. Used to validate references to other modules.
 		/// </summary>
-		public RulesScope Scope;
+		public RulesScope Scope { get; set; }
 
 		/// <summary>
 		/// The default directory for output files
 		/// </summary>
-		public DirectoryReference DefaultOutputBaseDir;
+		public DirectoryReference DefaultOutputBaseDir { get; set; }
 
 		/// <summary>
 		/// The plugin that this module belongs to
 		/// </summary>
-		public PluginInfo? Plugin;
+		public PluginInfo? Plugin { get; set; }
 
 		/// <summary>
 		/// Whether this module should be included in the default hot reload set
 		/// </summary>
-		public bool bCanHotReload;
+		public bool bCanHotReload { get; set; }
 
 		/// <summary>
 		/// Whether this module should be compiled with optimization disabled in DebugGame configurations (ie. whether it's a game module).
 		/// </summary>
-		public bool bCanBuildDebugGame;
+		public bool bCanBuildDebugGame { get; set; }
 
 		/// <summary>
 		/// Whether this module can be used for generating shared PCHs
 		/// </summary>
-		public bool bCanUseForSharedPCH;
+		public bool bCanUseForSharedPCH { get; set; }
 
 		/// <summary>
 		/// Whether to treat this module as a game module for UHT ordering
 		/// </summary>
-		public bool bClassifyAsGameModuleForUHT;
+		public bool bClassifyAsGameModuleForUHT { get; set; }
 
 		/// <summary>
 		/// The default module type for UnrealHeaderTool. Do not use this for inferring other things about the module.
 		/// </summary>
-		public UHTModuleType? DefaultUHTModuleType;
+		public UHTModuleType? DefaultUHTModuleType { get; set; }
 
 		/// <summary>
 		/// Constructor
 		/// </summary>
-		public ModuleRulesContext(RulesScope Scope, DirectoryReference DefaultOutputBaseDir)
+		public ModuleRulesContext(RulesScope scope, DirectoryReference defaultOutputBaseDir)
 		{
-			this.Scope = Scope;
-			this.DefaultOutputBaseDir = DefaultOutputBaseDir;
+			Scope = scope;
+			DefaultOutputBaseDir = defaultOutputBaseDir;
 			bCanUseForSharedPCH = true;
 		}
 
 		/// <summary>
 		/// Copy constructor
 		/// </summary>
-		/// <param name="Other">The context to copy from</param>
-		public ModuleRulesContext(ModuleRulesContext Other)
+		/// <param name="other">The context to copy from</param>
+		public ModuleRulesContext(ModuleRulesContext other)
 		{
-			Scope = Other.Scope;
-			DefaultOutputBaseDir = Other.DefaultOutputBaseDir;
-			Plugin = Other.Plugin;
-			bCanHotReload = Other.bCanHotReload;
-			bCanBuildDebugGame = Other.bCanBuildDebugGame;
-			bCanUseForSharedPCH = Other.bCanUseForSharedPCH;
-			bClassifyAsGameModuleForUHT = Other.bClassifyAsGameModuleForUHT;
-			DefaultUHTModuleType = Other.DefaultUHTModuleType;
+			Scope = other.Scope;
+			DefaultOutputBaseDir = other.DefaultOutputBaseDir;
+			Plugin = other.Plugin;
+			bCanHotReload = other.bCanHotReload;
+			bCanBuildDebugGame = other.bCanBuildDebugGame;
+			bCanUseForSharedPCH = other.bCanUseForSharedPCH;
+			bClassifyAsGameModuleForUHT = other.bClassifyAsGameModuleForUHT;
+			DefaultUHTModuleType = other.DefaultUHTModuleType;
 		}
 	}
 }
