UnrealBuildTool: Mark an obsolete TargetRule as Deprecated

#rnx

diff --git a/Engine/Source/Programs/UnrealBuildTool/Configuration/TargetRules.cs b/Engine/Source/Programs/UnrealBuildTool/Configuration/TargetRules.cs
--- a/Engine/Source/Programs/UnrealBuildTool/Configuration/TargetRules.cs
+++ b/Engine/Source/Programs/UnrealBuildTool/Configuration/TargetRules.cs
@@ -1823,7 +1823,7 @@ namespace UnrealBuildTool
 		/// <summary>
 		/// Whether to globally disable debug info generation; Obsolete, please use TargetRules.DebugInfoMode instead
 		/// </summary>
-		[XmlConfigFile(Category = "BuildConfiguration")]
+		[XmlConfigFile(Category = "BuildConfiguration", Deprecated = true)]
 		[Obsolete("Deprecated in UE5.4 - Replace with TargetRules.DebugInfo")]
 		public bool bDisableDebugInfo
 		{
