UnrealBuildTool: cleanup BuildConfiguration

#rnx

diff --git a/Engine/Source/Programs/UnrealBuildTool/Configuration/BuildConfiguration.cs b/Engine/Source/Programs/UnrealBuildTool/Configuration/BuildConfiguration.cs
--- a/Engine/Source/Programs/UnrealBuildTool/Configuration/BuildConfiguration.cs
+++ b/Engine/Source/Programs/UnrealBuildTool/Configuration/BuildConfiguration.cs
@@ -8,7 +8,7 @@ namespace UnrealBuildTool
 	/// <summary>
 	/// Global settings for building. Should not contain any target-specific settings.
 	/// </summary>
-	class BuildConfiguration
+	internal sealed class BuildConfiguration
 	{
 		/// <summary>
 		/// Whether to ignore import library files that are out of date when building targets. Set this to true to improve iteration time.
@@ -17,33 +17,33 @@ namespace UnrealBuildTool
 		/// in which case the target would automatically be rebuilt.
 		/// </summary>
 		[XmlConfigFile]
-		public bool bIgnoreOutdatedImportLibraries = true;
+		public bool bIgnoreOutdatedImportLibraries { get; set; } = true;
 
 		/// <summary>
 		/// Use existing static libraries for all engine modules in this target.
 		/// </summary>
 		[CommandLine("-UsePrecompiled")]
-		public bool bUsePrecompiled = false;
+		public bool bUsePrecompiled { get; set; } = false;
 
 		/// <summary>
 		/// Whether debug info should be written to the console.
 		/// </summary>
 		[XmlConfigFile]
 		[CommandLine("-PrintDebugInfo", Value = "true")]
-		public bool bPrintDebugInfo = false;
+		public bool bPrintDebugInfo { get; set; } = false;
 
 		/// <summary>
 		/// Whether the hybrid executor will be used (a remote executor and local executor). No longer supported.
 		/// </summary>
-		[XmlConfigFile]
+		[XmlConfigFile(Deprecated = true)]
 		[Obsolete("HybridExecutor is no longer supported")]
-		public bool bAllowHybridExecutor = false;
+		public bool bAllowHybridExecutor { get; set; } = false;
 
 		/// <summary>
 		/// Priority order for remote executors (XGE, SNDBS, FASTBuild, UBAExecutor)
 		/// </summary>
 		[XmlConfigFile]
-		public string[] RemoteExecutorPriority = new string[] { "XGE", "SNDBS", "FASTBuild", "UBA" };
+		public string[] RemoteExecutorPriority { get; set; } = new string[] { "XGE", "SNDBS", "FASTBuild", "UBA" };
 
 		/// <summary>
 		/// This property is being renamed and should not be used
@@ -51,7 +51,8 @@ namespace UnrealBuildTool
 		[XmlConfigFile]
 		[CommandLine("-Box", Value = "true")]
 		[CommandLine("-NoBox", Value = "false")]
-		public bool bAllowBoxExecutor = false;
+		[Obsolete("Replace with bAllowUBAExecutor")]
+		public bool bAllowBoxExecutor { get => bAllowUBAExecutor; set => bAllowUBAExecutor = value; }
 
 		/// <summary>
 		/// Whether the experimental UnrealBuildAccelerator executor will be used.
@@ -59,12 +60,7 @@ namespace UnrealBuildTool
 		[XmlConfigFile]
 		[CommandLine("-UBA", Value = "true")]
 		[CommandLine("-NoUBA", Value = "false")]
-		public bool bAllowUBAExecutor
-		{
-			get => bAllowBoxExecutor || bAllowUBAExecutorPrivate;
-			set => bAllowUBAExecutorPrivate = value;
-		}
-		private bool bAllowUBAExecutorPrivate = false;
+		public bool bAllowUBAExecutor { get; set; } = false;
 
 		/// <summary>
 		/// This property is being renamed and should not be used
@@ -72,7 +68,8 @@ namespace UnrealBuildTool
 		[XmlConfigFile]
 		[CommandLine("-BoxLocal", Value = "true")]
 		[CommandLine("-NoBoxLocal", Value = "false")]
-		public bool bAllowBoxLocalExecutor = false;
+		[Obsolete("Replace with bAllowUBALocalExecutor")]
+		public bool bAllowBoxLocalExecutor { get => bAllowUBALocalExecutor; set => bAllowUBALocalExecutor = value; }
 
 		/// <summary>
 		/// Whether the experimental UnrealBuildAccelerator (local only) executor will be used.
@@ -82,31 +79,31 @@ namespace UnrealBuildTool
 		[CommandLine("-NoUBALocal", Value = "false")]
 		public bool bAllowUBALocalExecutor
 		{
-			get => bAllowUBAExecutor || bAllowBoxLocalExecutor || bAllowUBALocalExecutorPrivate;
-			set => bAllowUBALocalExecutorPrivate = value;
+			get => bAllowUBAExecutor || _bAllowUBALocalExecutorPrivate;
+			set => _bAllowUBALocalExecutorPrivate = value;
 		}
-		private bool bAllowUBALocalExecutorPrivate = false;
+		private bool _bAllowUBALocalExecutorPrivate = false;
 
 		/// <summary>
 		/// Whether XGE may be used if available, default is true.
 		/// </summary>
 		[XmlConfigFile]
 		[CommandLine("-NoXGE", Value = "false")]
-		public bool bAllowXGE = true;
+		public bool bAllowXGE { get; set; } = true;
 
 		/// <summary>
 		/// Whether FASTBuild may be used if availabe, default is true.
 		/// </summary>
 		[XmlConfigFile]
 		[CommandLine("-NoFASTBuild", Value = "false")]
-		public bool bAllowFASTBuild = true;
+		public bool bAllowFASTBuild { get; set; } = true;
 
 		/// <summary>
 		/// Whether SN-DBS may be used if availabe, default is true.
 		/// </summary>
 		[XmlConfigFile]
 		[CommandLine("-NoSNDBS", Value = "false")]
-		public bool bAllowSNDBS = true;
+		public bool bAllowSNDBS { get; set; } = true;
 
 		/// <summary>
 		/// Enables support for very fast iterative builds by caching target data. Turning this on causes Unreal Build Tool to emit
@@ -131,7 +128,7 @@ namespace UnrealBuildTool
 		/// </summary>
 		[XmlConfigFile]
 		[CommandLine("-NoUBTMakefiles", Value = "false")]
-		public bool bUseUBTMakefiles = true;
+		public bool bUseUBTMakefiles { get; set; } = true;
 
 		/// <summary>
 		/// Number of actions that can be executed in parallel. If 0 then code will pick a default based
@@ -139,87 +136,87 @@ namespace UnrealBuildTool
 		/// </summary>
 		[XmlConfigFile]
 		[CommandLine("-MaxParallelActions")]
-		public int MaxParallelActions = 0;
+		public int MaxParallelActions { get; set; } = 0;
 
 		/// <summary>
 		/// Consider logical cores when determining how many total cpu cores are available.
 		/// </summary>
 		[XmlConfigFile(Name = "bAllCores")]
 		[CommandLine("-AllCores")]
-		public bool bAllCores = false;
+		public bool bAllCores { get; set; } = false;
 
 		/// <summary>
 		/// If true, force header regeneration. Intended for the build machine.
 		/// </summary>
 		[CommandLine("-ForceHeaderGeneration")]
 		[XmlConfigFile(Category = "UEBuildConfiguration")]
-		public bool bForceHeaderGeneration = false;
+		public bool bForceHeaderGeneration { get; set; } = false;
 
 		/// <summary>
 		/// If true, do not build UHT, assume it is already built.
 		/// </summary>
 		[Obsolete]
 		[XmlConfigFile(Category = "UEBuildConfiguration", Deprecated = true)]
-		public bool bDoNotBuildUHT = false;
+		public bool bDoNotBuildUHT { get; set; } = false;
 
 		/// <summary>
 		/// If true, fail if any of the generated header files is out of date.
 		/// </summary>
 		[CommandLine("-FailIfGeneratedCodeChanges")]
 		[XmlConfigFile(Category = "UEBuildConfiguration")]
-		public bool bFailIfGeneratedCodeChanges = false;
+		public bool bFailIfGeneratedCodeChanges { get; set; } = false;
 
 		/// <summary>
 		/// True if hot-reload from IDE is allowed.
 		/// </summary>
 		[CommandLine("-NoHotReloadFromIDE", Value = "false")]
 		[XmlConfigFile(Category = "UEBuildConfiguration")]
-		public bool bAllowHotReloadFromIDE = true;
+		public bool bAllowHotReloadFromIDE { get; set; } = true;
 
 		/// <summary>
 		/// If true, the Debug version of UnrealHeaderTool will be built and run instead of the Development version.
 		/// </summary>
 		[Obsolete]
 		[XmlConfigFile(Category = "UEBuildConfiguration", Deprecated = true)]
-		public bool bForceDebugUnrealHeaderTool = false;
+		public bool bForceDebugUnrealHeaderTool { get; set; } = false;
 
 		/// <summary>
 		/// If true, use C# UHT internal to UBT
 		/// </summary>
 		[Obsolete]
 		[XmlConfigFile(Category = "UEBuildConfiguration", Deprecated = true)]
-		public bool bUseBuiltInUnrealHeaderTool = true;
+		public bool bUseBuiltInUnrealHeaderTool { get; set; } = true;
 
 		/// <summary>
 		/// If true, generate warnings when C++ UHT is used
 		/// </summary>
 		[Obsolete]
 		[XmlConfigFile(Category = "UEBuildConfiguration", Deprecated = true)]
-		public bool bWarnOnCppUnrealHeaderTool = true;
+		public bool bWarnOnCppUnrealHeaderTool { get; set; } = true;
 
 		/// <summary>
 		/// Whether to skip compiling rules assemblies and just assume they are valid
 		/// </summary>
 		[CommandLine("-SkipRulesCompile")]
-		public bool bSkipRulesCompile = false;
+		public bool bSkipRulesCompile { get; set; } = false;
 
 		/// <summary>
 		/// Whether to force compiling rules assemblies, regardless of whether they are valid
 		/// </summary>
 		[CommandLine("-ForceRulesCompile")]
-		public bool bForceRulesCompile = false;
+		public bool bForceRulesCompile { get; set; } = false;
 
 		/// <summary>
 		/// Maximum recommended root path length.
 		/// </summary>
 		[XmlConfigFile(Category = "WindowsPlatform")]
-		public int MaxRootPathLength = 50;
+		public int MaxRootPathLength { get; set; } = 50;
 
 		/// <summary>
 		/// Maximum length of a path relative to the root directory. Used on Windows to ensure paths are portable between machines. Defaults to off.
 		/// </summary>
 		[XmlConfigFile(Category = "WindowsPlatform")]
-		public int MaxNestedPathLength = 200;
+		public int MaxNestedPathLength { get; set; } = 200;
 
 		/// <summary>
 		/// When single file targets are specified, via -File=, -SingleFile=, or -FileList=
@@ -228,7 +225,7 @@ namespace UnrealBuildTool
 		/// including the case where a file specified via -FileList= is empty.
 		/// </summary>
 		[CommandLine("-IgnoreInvalidFiles")]
-		public bool bIgnoreInvalidFiles;
+		public bool bIgnoreInvalidFiles { get; set; } = false;
 
 		/// <summary>
 		/// Instruct the executor to write compact output e.g. only errors, if supported by the executor.
@@ -236,7 +233,7 @@ namespace UnrealBuildTool
 		/// </summary>
 		[XmlConfigFile(Name = "bCompactOutput")]
 		[CommandLine("-CompactOutput")]
-		private bool bCompactOutputCommandLine = false;
+		private bool bCompactOutputCommandLine { get; set; } = false;
 
 		/// <summary>
 		/// If set, artifacts will be read
@@ -244,7 +241,7 @@ namespace UnrealBuildTool
 		[XmlConfigFile(Category = "BuildConfiguration")]
 		[CommandLine("-ArtifactReads", Value = "True")]
 		[CommandLine("-NoArtifactReads", Value = "False")]
-		public bool bArtifactRead = true;
+		public bool bArtifactRead { get; set; } = true;
 
 		/// <summary>
 		/// If set, artifacts will be written
@@ -252,21 +249,21 @@ namespace UnrealBuildTool
 		[XmlConfigFile(Category = "BuildConfiguration")]
 		[CommandLine("-ArtifactWrites", Value = "True")]
 		[CommandLine("-NoArtifactWrites", Value = "False")]
-		public bool bArtifactWrites = true;
+		public bool bArtifactWrites { get; set; } = true;
 
 		/// <summary>
 		/// If true, log all artifact cache misses as informational messages
 		/// </summary>
 		[XmlConfigFile(Category = "BuildConfiguration")]
 		[CommandLine("-LogArtifactCacheMisses", Value = "True")]
-		public bool bLogArtifactCacheMisses = false;
+		public bool bLogArtifactCacheMisses { get; set; } = false;
 
 		/// <summary>
 		/// Location to store the artifacts.
 		/// </summary>
 		[XmlConfigFile(Category = "BuildConfiguration")]
 		[CommandLine("-ArtifactDirectory=")]
-		public string ArtifactDirectory = String.Empty;
+		public string ArtifactDirectory { get; set; } = String.Empty;
 
 		/// <summary>
 		/// Instruct the executor to write compact output e.g. only errors, if supported by the executor,
@@ -278,13 +275,13 @@ namespace UnrealBuildTool
 		/// Whether to unify C++ code into larger files for faster compilation.
 		/// </summary>
 		[XmlConfigFile(Category = "BuildConfiguration")]
-		public bool? bUseUnityBuild;
+		public bool? bUseUnityBuild { get; set; }
 
 		/// <summary>
 		/// Whether to force C++ source files to be combined into larger files for faster compilation.
 		/// </summary>
 		[XmlConfigFile(Category = "BuildConfiguration")]
-		public bool? bForceUnityBuild;
+		public bool? bForceUnityBuild { get; set; }
 
 		/// <summary>
 		/// Intermediate environment. Determines if the intermediates end up in a different folder than normal.
@@ -297,7 +294,7 @@ namespace UnrealBuildTool
 				{
 					return UnrealIntermediateEnvironment.Default;
 				}
-				if (bUseUnityBuild == false)
+				else if (bUseUnityBuild == false)
 				{
 					return UnrealIntermediateEnvironment.NonUnity;
 				}
