UnrealBuildTool: ModuleRules cleanup pt.1

#rnx

diff --git a/Engine/Source/Programs/UnrealBuildTool/Configuration/ModuleRules.cs b/Engine/Source/Programs/UnrealBuildTool/Configuration/ModuleRules.cs
--- a/Engine/Source/Programs/UnrealBuildTool/Configuration/ModuleRules.cs
+++ b/Engine/Source/Programs/UnrealBuildTool/Configuration/ModuleRules.cs
@@ -590,7 +590,7 @@ namespace UnrealBuildTool
 		/// <summary>
 		/// Additional directories that contribute to this module but are not based on a subclass (NotForLicensees, etc)
 		/// </summary>
-		private List<DirectoryReference> AdditionalModuleDirectories = new List<DirectoryReference>();
+		private List<DirectoryReference> AdditionalModuleDirectories = new();
 
 		/// <summary>
 		/// The rules assembly to use when searching for modules
@@ -615,12 +615,12 @@ namespace UnrealBuildTool
 		/// <summary>
 		/// Rules for the target that this module belongs to
 		/// </summary>
-		public readonly ReadOnlyTargetRules Target;
+		public ReadOnlyTargetRules Target { get; init; }
 
 		/// <summary>
 		/// Type of module
 		/// </summary>
-		public ModuleType Type = ModuleType.CPlusPlus;
+		public ModuleType Type { get; set; } = ModuleType.CPlusPlus;
 
 		/// <summary>
 		/// Accessor for the target logger
@@ -659,7 +659,7 @@ namespace UnrealBuildTool
 		/// Subfolder of Binaries/PLATFORM folder to put this module in when building DLLs. This should only be used by modules that are found via searching like the
 		/// TargetPlatform or ShaderFormat modules. If FindModules is not used to track them down, the modules will not be found.
 		/// </summary>
-		public string BinariesSubFolder = "";
+		public string BinariesSubFolder { get; set; } = String.Empty;
 
 		private CodeOptimization? OptimizeCodeOverride;
 
@@ -670,7 +670,8 @@ namespace UnrealBuildTool
 		{
 			get
 			{
-				if (bCodeCoverage) {
+				if (bCodeCoverage)
+				{
 					return CodeOptimization.Never;
 				}
 
@@ -743,18 +744,18 @@ namespace UnrealBuildTool
 		/// <summary>
 		/// Explicit private PCH for this module. Implies that this module will not use a shared PCH.
 		/// </summary>
-		public string? PrivatePCHHeaderFile;
+		public string? PrivatePCHHeaderFile { get; set; }
 
 		/// <summary>
 		/// Header file name for a shared PCH provided by this module.  Must be a valid relative path to a public C++ header file.
 		/// This should only be set for header files that are included by a significant number of other C++ modules.
 		/// </summary>
-		public string? SharedPCHHeaderFile;
+		public string? SharedPCHHeaderFile { get; set; }
 
 		/// <summary>
 		/// Specifies an alternate name for intermediate directories and files for intermediates of this module. Useful when hitting path length limitations.
 		/// </summary>
-		public string? ShortName = null;
+		public string? ShortName { get; set; } = null;
 
 		/// <summary>
 		/// Precompiled header usage for this module
@@ -796,7 +797,7 @@ namespace UnrealBuildTool
 		/// Whether this module should be treated as an engine module (eg. using engine definitions, PCHs, compiled with optimizations enabled in DebugGame configurations, etc...).
 		/// Initialized to a default based on the rules assembly it was created from.
 		/// </summary>
-		public bool bTreatAsEngineModule;
+		public bool bTreatAsEngineModule { get; set; }
 
 		/// <summary>
 		/// Emits compilation errors for incorrect UE_LOG format strings.
@@ -852,18 +853,18 @@ namespace UnrealBuildTool
 		/// <summary>
 		/// Use run time type information
 		/// </summary>
-		public bool bUseRTTI = false;
+		public bool bUseRTTI { get; set; }
 
 		/// <summary>
 		/// Whether to direct MSVC to remove unreferenced COMDAT functions and data.
 		/// </summary>
 		/// <seealso href="https://learn.microsoft.com/en-us/cpp/build/reference/zc-inline-remove-unreferenced-comdat">zc-inline-remove-unreferenced-comdat</seealso>
-		public bool bVcRemoveUnreferencedComdat = true;
+		public bool bVcRemoveUnreferencedComdat { get; set; } = true;
 
 		/// <summary>
 		/// Enable code coverage compilation/linking support.
 		/// </summary>
-		public bool bCodeCoverage = false;
+		public bool bCodeCoverage { get; set; }
 
 		/// <summary>
 		/// Obsolete: Direct the compiler to generate AVX instructions wherever SSE or AVX intrinsics are used, on the platforms that support it.
@@ -880,29 +881,29 @@ namespace UnrealBuildTool
 		/// Direct the compiler to generate AVX instructions wherever SSE or AVX intrinsics are used, on the x64 platforms that support it.
 		/// Note that by enabling this you are changing the minspec for the PC platform, and the resultant executable will crash on machines without AVX support.
 		/// </summary>
-		public MinimumCpuArchitectureX64? MinCpuArchX64 = null;
+		public MinimumCpuArchitectureX64? MinCpuArchX64 { get; set; } = null;
 
 		/// <summary>
 		/// Enable buffer security checks.  This should usually be enabled as it prevents severe security risks.
 		/// </summary>
-		public bool bEnableBufferSecurityChecks = true;
+		public bool bEnableBufferSecurityChecks { get; set; } = true;
 
 		/// <summary>
 		/// Enable exception handling
 		/// </summary>
-		public bool bEnableExceptions = false;
+		public bool bEnableExceptions { get; set; }
 
 		/// <summary>
 		/// Enable objective C exception handling
 		/// </summary>
-		public bool bEnableObjCExceptions = false;
+		public bool bEnableObjCExceptions { get; set; }
 
 		/// <summary>
 		/// Enable objective C automatic reference counting (ARC)
 		/// If you set this to true you should not use shared PCHs for this module. The engine won't be extensively using ARC in the short term  
 		/// Not doing this will result in a compile errors because shared PCHs were compiled with different flags than consumer
 		/// </summary>
-		public bool bEnableObjCAutomaticReferenceCounting = false;
+		public bool bEnableObjCAutomaticReferenceCounting { get; set; }
 
 		/// <summary>
 		/// How to treat deterministic warnings (experimental).
@@ -931,7 +932,7 @@ namespace UnrealBuildTool
 		/// <summary>
 		/// Whether to enable all warnings as errors. UE enables most warnings as errors already, but disables a few (such as deprecation warnings).
 		/// </summary>`
-		public bool bWarningsAsErrors = false;
+		public bool bWarningsAsErrors { get; set; }
 
 		/// <summary>
 		/// How to treat unsafe implicit type cast warnings (e.g., double->float or int64->int32)
@@ -948,7 +949,7 @@ namespace UnrealBuildTool
 		/// <summary>
 		/// Enable warnings for using undefined identifiers in #if expressions
 		/// </summary>
-		public bool bEnableUndefinedIdentifierWarnings = true;
+		public bool bEnableUndefinedIdentifierWarnings { get; set; } = true;
 
 		/// <summary>
 		/// How to treat general module include path validation messages
@@ -989,20 +990,20 @@ namespace UnrealBuildTool
 		/// <summary>
 		/// Disable all static analysis - clang, msvc, pvs-studio.
 		/// </summary>
-		public bool bDisableStaticAnalysis = false;
+		public bool bDisableStaticAnalysis { get; set; }
 
 		/// <summary>
 		/// Enable additional analyzer extension warnings using the EspXEngine plugin. This is only supported for MSVC.
 		/// See https://learn.microsoft.com/en-us/cpp/code-quality/using-the-cpp-core-guidelines-checkers
 		/// This will add a large number of warnings by default. It's recommended to use StaticAnalyzerRulesets if this is enabled.
 		/// </summary>
-		public bool bStaticAnalyzerExtensions = false;
+		public bool bStaticAnalyzerExtensions { get; set; }
 
 		/// <summary>
 		/// The static analyzer rulesets that should be used to filter warnings. This is only supported for MSVC.
 		/// See https://learn.microsoft.com/en-us/cpp/code-quality/using-rule-sets-to-specify-the-cpp-rules-to-run
 		/// </summary>
-		public HashSet<FileReference> StaticAnalyzerRulesets = new HashSet<FileReference>();
+		public HashSet<FileReference> StaticAnalyzerRulesets = new();
 
 		/// <summary>
 		/// The static analyzer checkers that should be enabled rather than the defaults. This is only supported for Clang.
@@ -1012,7 +1013,7 @@ namespace UnrealBuildTool
 		///    'clang -Xclang -analyzer-checker-help-alpha' 
 		/// for the list of experimental checkers.
 		/// </summary>
-		public HashSet<string> StaticAnalyzerCheckers = new HashSet<string>();
+		public HashSet<string> StaticAnalyzerCheckers = new();
 
 		/// <summary>
 		/// The static analyzer default checkers that should be disabled. Unused if StaticAnalyzerCheckers is populated. This is only supported for Clang.
@@ -1023,7 +1024,7 @@ namespace UnrealBuildTool
 		///    'clang -Xclang -analyzer-checker-help-alpha' 
 		/// for the list of experimental checkers.
 		/// </summary>
-		public HashSet<string> StaticAnalyzerDisabledCheckers = new HashSet<string>() { "deadcode.DeadStores", "security.FloatLoopCounter" };
+		public HashSet<string> StaticAnalyzerDisabledCheckers = new() { "deadcode.DeadStores", "security.FloatLoopCounter" };
 
 		/// <summary>
 		/// The static analyzer non-default checkers that should be enabled. Unused if StaticAnalyzerCheckers is populated. This is only supported for Clang.
@@ -1033,7 +1034,7 @@ namespace UnrealBuildTool
 		///    'clang -Xclang -analyzer-checker-help-alpha' 
 		/// for the list of experimental checkers.
 		/// </summary>
-		public HashSet<string> StaticAnalyzerAdditionalCheckers = new HashSet<string>();
+		public HashSet<string> StaticAnalyzerAdditionalCheckers = new();
 
 		private bool? bUseUnityOverride;
 		/// <summary>
@@ -1058,23 +1059,23 @@ namespace UnrealBuildTool
 		/// <summary>
 		/// Whether to merge module and generated unity files for faster compilation.
 		/// </summary>
-		public bool bMergeUnityFiles = true;
+		public bool bMergeUnityFiles { get; set; } = true;
 
 		/// <summary>
 		/// The number of source files in this module before unity build will be activated for that module.  If set to
 		/// anything besides -1, will override the default setting which is controlled by MinGameModuleSourceFilesForUnityBuild
 		/// </summary>
-		public int MinSourceFilesForUnityBuildOverride = 0;
+		public int MinSourceFilesForUnityBuildOverride { get; set; } = 0;
 
 		/// <summary>
 		/// Overrides BuildConfiguration.MinFilesUsingPrecompiledHeader if non-zero.
 		/// </summary>
-		public int MinFilesUsingPrecompiledHeaderOverride = 0;
+		public int MinFilesUsingPrecompiledHeaderOverride { get; set; } = 0;
 
 		/// <summary>
 		/// Overrides Target.NumIncludedBytesPerUnityCPP if non-zero.
 		/// </summary>
-		public int NumIncludedBytesPerUnityCPPOverride = 0;
+		public int NumIncludedBytesPerUnityCPPOverride { get; set; } = 0;
 
 		/// <summary>
 		/// Helper function to get the number of byes per unity cpp file
@@ -1087,34 +1088,34 @@ namespace UnrealBuildTool
 		/// <summary>
 		/// Module uses a #import so must be built locally when compiling with SN-DBS
 		/// </summary>
-		public bool bBuildLocallyWithSNDBS = false;
+		public bool bBuildLocallyWithSNDBS { get; set; }
 
 		/// <summary>
 		/// Enable warnings for when there are .gen.cpp files that could be inlined in a matching handwritten cpp file
 		/// </summary>
-		public bool bEnableNonInlinedGenCppWarnings = false;
+		public bool bEnableNonInlinedGenCppWarnings { get; set; }
 
 		/// <summary>
 		/// Redistribution override flag for this module.
 		/// </summary>
-		public bool? IsRedistributableOverride = null;
+		public bool? IsRedistributableOverride { get; set; } = null;
 
 		/// <summary>
 		/// Whether the output from this module can be publicly distributed, even if it has code/
 		/// dependencies on modules that are not (i.e. CarefullyRedist, NotForLicensees, NoRedist).
 		/// This should be used when you plan to release binaries but not source.
 		/// </summary>
-		public bool bLegalToDistributeObjectCode = false;
+		public bool bLegalToDistributeObjectCode { get; set; }
 
 		/// <summary>
 		/// List of folders which are allowed to be referenced when compiling this binary, without propagating restricted folder names
 		/// </summary>
-		public List<string> AllowedRestrictedFolders = new List<string>();
+		public List<string> AllowedRestrictedFolders = new();
 
 		/// <summary>
 		/// Set of aliased restricted folder references
 		/// </summary>
-		public Dictionary<string, string> AliasRestrictedFolders = new Dictionary<string, string>();
+		public Dictionary<string, string> AliasRestrictedFolders = new();
 
 		/// <summary>
 		/// Enforce "include what you use" rules when PCHUsage is set to ExplicitOrSharedPCH; warns when monolithic headers (Engine.h, UnrealEd.h, etc...) 
@@ -1127,105 +1128,105 @@ namespace UnrealBuildTool
 		/// Allows "include what you use" to modify the source code when run. bEnforceIWYU must be true for this variable to matter.
 		/// 
 		/// </summary>
-		public IWYUSupport IWYUSupport = IWYUSupport.Full;
+		public IWYUSupport IWYUSupport { get; set; } = IWYUSupport.Full;
 
 		/// <summary>
 		/// Whether to add all the default include paths to the module (eg. the Source/Classes folder, subfolders under Source/Public).
 		/// </summary>
-		public bool bAddDefaultIncludePaths = true;
+		public bool bAddDefaultIncludePaths { get; set; } = true;
 
 		/// <summary>
 		/// Whether to ignore dangling (i.e. unresolved external) symbols in modules
 		/// </summary>
-		public bool bIgnoreUnresolvedSymbols = false;
+		public bool bIgnoreUnresolvedSymbols { get; set; }
 
 		/// <summary>
 		/// Whether this module should be precompiled. Defaults to the bPrecompile flag from the target. Clear this flag to prevent a module being precompiled.
 		/// </summary>
-		public bool bPrecompile;
+		public bool bPrecompile { get; set; }
 
 		/// <summary>
 		/// Whether this module should use precompiled data. Always true for modules created from installed assemblies.
 		/// </summary>
-		public bool bUsePrecompiled;
+		public bool bUsePrecompiled { get; set; }
 
 		/// <summary>
 		/// Whether this module can use PLATFORM_XXXX style defines, where XXXX is a confidential platform name. This is used to ensure engine or other 
 		/// shared code does not reveal confidential information inside an #if PLATFORM_XXXX block. Licensee game code may want to allow for them, however.
 		/// </summary>
-		public bool bAllowConfidentialPlatformDefines = false;
+		public bool bAllowConfidentialPlatformDefines { get; set; }
 
 		/// <summary>
 		/// Enables AutoRTFM instrumentation to this module only when AutoRTFMCompiler is enabled
 		/// </summary>
-		public bool bAllowAutoRTFMInstrumentation = false;
+		public bool bAllowAutoRTFMInstrumentation { get; set; }
 
 		/// <summary>
 		/// List of modules names (no path needed) with header files that our module's public headers needs access to, but we don't need to "import" or link against.
 		/// </summary>
-		public List<string> PublicIncludePathModuleNames = new List<string>();
+		public List<string> PublicIncludePathModuleNames = new();
 
 		/// <summary>
 		/// List of public dependency module names (no path needed) (automatically does the private/public include). These are modules that are required by our public source files.
 		/// </summary>
-		public List<string> PublicDependencyModuleNames = new List<string>();
+		public List<string> PublicDependencyModuleNames = new();
 
 		/// <summary>
 		/// List of modules name (no path needed) with header files that our module's private code files needs access to, but we don't need to "import" or link against.
 		/// </summary>
-		public List<string> PrivateIncludePathModuleNames = new List<string>();
+		public List<string> PrivateIncludePathModuleNames = new();
 
 		/// <summary>
 		/// List of private dependency module names.  These are modules that our private code depends on but nothing in our public
 		/// include files depend on.
 		/// </summary>
-		public List<string> PrivateDependencyModuleNames = new List<string>();
+		public List<string> PrivateDependencyModuleNames = new();
 
 		/// <summary>
 		/// Only for legacy reason, should not be used in new code. List of module dependencies that should be treated as circular references.  This modules must have already been added to
 		/// either the public or private dependent module list.
 		/// </summary>
-		public List<string> CircularlyReferencedDependentModules = new List<string>();
+		public List<string> CircularlyReferencedDependentModules = new();
 
 		/// <summary>
 		/// List of system/library include paths - typically used for External (third party) modules.  These are public stable header file directories that are not checked when resolving header dependencies.
 		/// </summary>
-		public List<string> PublicSystemIncludePaths = new List<string>();
+		public List<string> PublicSystemIncludePaths = new();
 
 		/// <summary>
 		/// (This setting is currently not need as we discover all files from the 'Public' folder) List of all paths to include files that are exposed to other modules
 		/// </summary>
-		public List<string> PublicIncludePaths = new List<string>();
+		public List<string> PublicIncludePaths = new();
 
 		/// <summary>
 		/// (This setting is currently not need as we discover all files from the 'Internal' folder) List of all paths to include files that are exposed to other internal modules
 		/// </summary>
-		public List<string> InternalIncludePaths = new List<string>();
+		public List<string> InternalIncludePaths = new();
 
 		/// <summary>
 		/// List of all paths to this module's internal include files, not exposed to other modules (at least one include to the 'Private' path, more if we want to avoid relative paths)
 		/// </summary>
-		public List<string> PrivateIncludePaths = new List<string>();
+		public List<string> PrivateIncludePaths = new();
 
 		/// <summary>
 		/// List of system library paths (directory of .lib files) - for External (third party) modules please use the PublicAdditionalLibaries instead
 		/// </summary>
-		public List<string> PublicSystemLibraryPaths = new List<string>();
+		public List<string> PublicSystemLibraryPaths = new();
 
 		/// <summary>
 		/// List of search paths for libraries at runtime (eg. .so files)
 		/// </summary>
-		public List<string> PrivateRuntimeLibraryPaths = new List<string>();
+		public List<string> PrivateRuntimeLibraryPaths = new();
 
 		/// <summary>
 		/// List of search paths for libraries at runtime (eg. .so files)
 		/// </summary>
-		public List<string> PublicRuntimeLibraryPaths = new List<string>();
+		public List<string> PublicRuntimeLibraryPaths = new();
 
 		/// <summary>
 		/// List of additional libraries (names of the .lib files including extension) - typically used for External (third party) modules
 		/// </summary>
-		public List<string> PublicAdditionalLibraries = new List<string>();
+		public List<string> PublicAdditionalLibraries = new();
 
 		/// <summary>
 		/// Per-architecture lists of dependencies for linking to ignore (useful when building for multiple architectures, and a lib only is needed for one architecture), it's up to the Toolchain to use this
@@ -1239,63 +1240,60 @@ namespace UnrealBuildTool
 		/// <returns>Directory where the module lives</returns>
 		public string GetModuleDirectory(string ModuleName)
 		{
-			FileReference? ModuleFileReference = RulesAssembly.GetModuleFileName(ModuleName);
-			if (ModuleFileReference == null)
-			{
-				throw new CompilationResultException(CompilationResult.RulesError, "Could not find a module named '{ModuleName}'.", ModuleName);
-			}
+			FileReference? ModuleFileReference = RulesAssembly.GetModuleFileName(ModuleName)
+				?? throw new CompilationResultException(CompilationResult.RulesError, "Could not find a module named '{ModuleName}'.", ModuleName);
 			return ModuleFileReference.Directory.FullName;
 		}
 
 		/// <summary>
 		/// List of additional pre-build libraries (names of the .lib files including extension) - typically used for additional targets which are still built, but using either TargetRules.PreBuildSteps or TargetRules.PreBuildTargets.
 		/// </summary>
-		public List<string> PublicPreBuildLibraries = new List<string>();
+		public List<string> PublicPreBuildLibraries = new();
 
 		/// <summary>
 		/// List of system libraries to use - these are typically referenced via name and then found via the system paths. If you need to reference a .lib file use the PublicAdditionalLibraries instead
 		/// </summary>
-		public List<string> PublicSystemLibraries = new List<string>();
+		public List<string> PublicSystemLibraries = new();
 
 		/// <summary>
 		/// List of XCode frameworks (iOS and MacOS)
 		/// </summary>
-		public List<string> PublicFrameworks = new List<string>();
+		public List<string> PublicFrameworks = new();
 
 		/// <summary>
 		/// List of weak frameworks (for OS version transitions)
 		/// </summary>
-		public List<string> PublicWeakFrameworks = new List<string>();
+		public List<string> PublicWeakFrameworks = new();
 
 		/// <summary>
 		/// List of addition frameworks - typically used for External (third party) modules on Mac and iOS
 		/// </summary>
-		public List<Framework> PublicAdditionalFrameworks = new List<Framework>();
+		public List<Framework> PublicAdditionalFrameworks = new();
 
 		/// <summary>
 		/// List of addition resources that should be copied to the app bundle for Mac or iOS
 		/// </summary>
-		public List<BundleResource> AdditionalBundleResources = new List<BundleResource>();
+		public List<BundleResource> AdditionalBundleResources = new();
 
 		/// <summary>
 		/// List of type libraries that we need to generate headers for (Windows only)
 		/// </summary>
-		public List<TypeLibrary> TypeLibraries = new List<TypeLibrary>();
+		public List<TypeLibrary> TypeLibraries = new();
 
 		/// <summary>
 		/// List of delay load DLLs - typically used for External (third party) modules
 		/// </summary>
-		public List<string> PublicDelayLoadDLLs = new List<string>();
+		public List<string> PublicDelayLoadDLLs = new();
 
 		/// <summary>
 		/// Private compiler definitions for this module
 		/// </summary>
-		public List<string> PrivateDefinitions = new List<string>();
+		public List<string> PrivateDefinitions = new();
 
 		/// <summary>
 		/// Public compiler definitions for this module
 		/// </summary>
-		public List<string> PublicDefinitions = new List<string>();
+		public List<string> PublicDefinitions = new();
 
 		/// <summary>
 		/// Append (or create)
@@ -1307,7 +1305,7 @@ namespace UnrealBuildTool
 			string WithEquals = Definition + "=";
 			for (int Index = 0; Index < PublicDefinitions.Count; Index++)
 			{
-				if (PublicDefinitions[Index].StartsWith(WithEquals))
+				if (PublicDefinitions[Index].StartsWith(WithEquals, StringComparison.Ordinal))
 				{
 					PublicDefinitions[Index] = PublicDefinitions[Index] + Text;
 					return;
@@ -1321,27 +1319,27 @@ namespace UnrealBuildTool
 		/// <summary>
 		/// Addition modules this module may require at run-time 
 		/// </summary>
-		public List<string> DynamicallyLoadedModuleNames = new List<string>();
+		public List<string> DynamicallyLoadedModuleNames = new();
 
 		/// <summary>
 		/// List of files which this module depends on at runtime. These files will be staged along with the target.
 		/// </summary>
-		public RuntimeDependencyList RuntimeDependencies = new RuntimeDependencyList();
+		public RuntimeDependencyList RuntimeDependencies = new();
 
 		/// <summary>
 		/// List of additional properties to be added to the build receipt
 		/// </summary>
-		public ReceiptPropertyList AdditionalPropertiesForReceipt = new ReceiptPropertyList();
+		public ReceiptPropertyList AdditionalPropertiesForReceipt = new();
 
 		/// <summary>
 		/// Which targets this module should be precompiled for
 		/// </summary>
-		public PrecompileTargetsType PrecompileForTargets = PrecompileTargetsType.Default;
+		public PrecompileTargetsType PrecompileForTargets { get; set; } = PrecompileTargetsType.Default;
 
 		/// <summary>
 		/// External files which invalidate the makefile if modified. Relative paths are resolved relative to the .build.cs file.
 		/// </summary>
-		public List<string> ExternalDependencies = new List<string>();
+		public List<string> ExternalDependencies = new();
 
 		/// <summary>
 		/// Subclass rules files which invalidate the makefile if modified.
@@ -1352,17 +1350,17 @@ namespace UnrealBuildTool
 		/// Whether this module requires the IMPLEMENT_MODULE macro to be implemented. Most UE modules require this, since we use the IMPLEMENT_MODULE macro
 		/// to do other global overloads (eg. operator new/delete forwarding to GMalloc).
 		/// </summary>
-		public bool? bRequiresImplementModule;
+		public bool? bRequiresImplementModule { get; set; }
 
 		/// <summary>
 		/// If this module has associated Verse code, this is the Verse root path of it
 		/// </summary>
-		public string? VersePath;
+		public string? VersePath { get; set; }
 
 		/// <summary>
 		/// Visibility of Verse code in this module's Source/Verse folder
 		/// </summary>
-		public VerseScope VerseScope = VerseScope.PublicUser;
+		public VerseScope VerseScope { get; set; } = VerseScope.PublicUser;
 
 		/// <summary>
 		/// Whether this module qualifies included headers from other modules relative to the root of their 'Public' folder. This reduces the number
@@ -1371,7 +1369,7 @@ namespace UnrealBuildTool
 		public bool bLegacyPublicIncludePaths
 		{
 			set => bLegacyPublicIncludePathsPrivate = value;
-			get => bLegacyPublicIncludePathsPrivate ?? ((DefaultBuildSettings < BuildSettingsVersion.V2) ? Target.bLegacyPublicIncludePaths : false);
+			get => bLegacyPublicIncludePathsPrivate ?? ((DefaultBuildSettings < BuildSettingsVersion.V2) && Target.bLegacyPublicIncludePaths);
 		}
 		private bool? bLegacyPublicIncludePathsPrivate;
 
@@ -1382,7 +1380,7 @@ namespace UnrealBuildTool
 		public bool bLegacyParentIncludePaths
 		{
 			set => bLegacyParentIncludePathsPrivate = value;
-			get => bLegacyParentIncludePathsPrivate ?? ((DefaultBuildSettings < BuildSettingsVersion.V3) ? Target.bLegacyParentIncludePaths : false);
+			get => bLegacyParentIncludePathsPrivate ?? ((DefaultBuildSettings < BuildSettingsVersion.V3) && Target.bLegacyParentIncludePaths);
 		}
 		private bool? bLegacyParentIncludePathsPrivate;
 
@@ -1391,17 +1389,17 @@ namespace UnrealBuildTool
 		/// Circular module dependencies result in slower builds. Disabling this option is strongly discouraged.
 		/// This option is ignored for Engine modules which will always be validated against the allow list.
 		/// </summary>
-		public bool bValidateCircularDependencies = true;
+		public bool bValidateCircularDependencies { get; set; } = true;
 
 		/// <summary>
 		/// Which stanard to use for compiling this module
 		/// </summary>
-		public CppStandardVersion? CppStandard;
+		public CppStandardVersion? CppStandard { get; set; }
 
 		/// <summary>
 		/// Which standard to use for compiling this module
 		/// </summary>
-		public CStandardVersion? CStandard;
+		public CStandardVersion? CStandard { get; set; }
 
 		/// <summary>
 		/// A list of subdirectory names and functions that are invoked to generate header files.
@@ -1413,19 +1411,12 @@ namespace UnrealBuildTool
 		/// <summary>
 		///  Control visibility of symbols
 		/// </summary>
-		public SymbolVisibility ModuleSymbolVisibility = ModuleRules.SymbolVisibility.Default;
+		public SymbolVisibility ModuleSymbolVisibility { get; set; } = ModuleRules.SymbolVisibility.Default;
 
 		/// <summary>
 		/// The AutoSDK directory for the active host platform
 		/// </summary>
-		public string? AutoSdkDirectory
-		{
-			get
-			{
-				DirectoryReference? AutoSdkDir;
-				return UEBuildPlatformSDK.TryGetHostPlatformAutoSDKDir(out AutoSdkDir) ? AutoSdkDir.FullName : null;
-			}
-		}
+		public string? AutoSdkDirectory => UEBuildPlatformSDK.TryGetHostPlatformAutoSDKDir(out DirectoryReference? AutoSdkDir) ? AutoSdkDir.FullName : null;
 
 		/// <summary>
 		/// The current engine directory
@@ -1453,6 +1444,7 @@ namespace UnrealBuildTool
 		/// <summary>
 		/// Property for the directory containing this module. Useful for adding paths to third party dependencies.
 		/// </summary>
+		[SuppressMessage("Naming", "CA1721:Property names should not match get methods", Justification = "GetModuleDirectory() is to get the path to other modules")]
 		public string ModuleDirectory => Directory.FullName;
 
 		/// <summary>
@@ -1711,7 +1703,7 @@ namespace UnrealBuildTool
 
 			// Skip platform extension modules. We only care about the base modules, not the platform overrides.
 			// The platform overrides get applied at a later stage when we actually come to build the module.
-			if (!UEBuildPlatform.GetPlatformFolderNames().Any(Name => ModuleType.Name.EndsWith("_" + Name)))
+			if (!UEBuildPlatform.GetPlatformFolderNames().Any(Name => ModuleType.Name.EndsWith("_" + Name, StringComparison.OrdinalIgnoreCase)))
 			{
 				IEnumerable<SupportedPlatformsAttribute> PlatformAttributes = ModuleType.GetCustomAttributes<SupportedPlatformsAttribute>();
 				IEnumerable<UnrealTargetPlatform> SupportedPlatforms = PlatformAttributes.SelectMany(x => x.Platforms).Distinct();
@@ -1824,10 +1816,11 @@ namespace UnrealBuildTool
 		}
 
 		internal bool IsTestModule => bIsTestModuleOverride ?? false;
+
 		/// <summary>
 		/// Whether this is a low level tests module.
 		/// </summary>
-		protected bool? bIsTestModuleOverride;
+		protected bool? bIsTestModuleOverride { get; set; }
 
 		/// <summary>
 		/// Returns the module directory for a given subclass of the module (platform extensions add subclasses of ModuleRules to add in platform-specific settings)
@@ -1885,6 +1878,7 @@ namespace UnrealBuildTool
 		/// </summary>
 		/// <param name="PackageName">The name of the third-party package</param>
 		/// <returns></returns>
+		[SuppressMessage("Globalization", "CA1308:Normalize strings to uppercase", Justification = "Architecture path is lowercase")]
 		public string GetVcPackageRoot(string PackageName)
 		{
 			// TODO: MacOS support, other platform support
@@ -1961,7 +1955,7 @@ namespace UnrealBuildTool
 			foreach (string Library in Libraries)
 			{
 				string LibraryPath = Path.Combine(VcPackageRoot, "lib", $"{Library}{LibraryExtension}");
-				if ((Target.Platform == UnrealTargetPlatform.Linux || Target.Platform == UnrealTargetPlatform.LinuxArm64 || Target.Platform == UnrealTargetPlatform.Mac) && !Library.StartsWith("lib"))
+				if ((Target.Platform == UnrealTargetPlatform.Linux || Target.Platform == UnrealTargetPlatform.LinuxArm64 || Target.Platform == UnrealTargetPlatform.Mac) && !Library.StartsWith("lib", StringComparison.OrdinalIgnoreCase))
 				{
 					LibraryPath = Path.Combine(VcPackageRoot, "lib", $"lib{Library}{LibraryExtension}");
 				}
