UnrealBuildTool: ModuleRules cleanup pt.2

#rnx

diff --git a/Engine/Source/Programs/UnrealBuildTool/Configuration/ModuleRules.cs b/Engine/Source/Programs/UnrealBuildTool/Configuration/ModuleRules.cs
--- a/Engine/Source/Programs/UnrealBuildTool/Configuration/ModuleRules.cs
+++ b/Engine/Source/Programs/UnrealBuildTool/Configuration/ModuleRules.cs
@@ -680,23 +680,23 @@ namespace UnrealBuildTool
 					return OptimizeCodeOverride.Value;
 				}
 
-				bool? ShouldOptimizeCode = null;
+				bool? shouldOptimizeCode = null;
 				if (Target.EnableOptimizeCodeForModules?.Contains(Name) ?? false)
 				{
-					ShouldOptimizeCode = true;
+					shouldOptimizeCode = true;
 				}
 
 				if (Target.DisableOptimizeCodeForModules?.Contains(Name) ?? false)
 				{
-					ShouldOptimizeCode = false;
+					shouldOptimizeCode = false;
 				}
 
-				if (!ShouldOptimizeCode.HasValue)
+				if (!shouldOptimizeCode.HasValue)
 				{
 					return CodeOptimization.Default;
 				}
 
-				return ShouldOptimizeCode.Value ? CodeOptimization.Always : CodeOptimization.Never;
+				return shouldOptimizeCode.Value ? CodeOptimization.Always : CodeOptimization.Never;
 			}
 			set => OptimizeCodeOverride = value;
 		}
@@ -1046,13 +1046,13 @@ namespace UnrealBuildTool
 			set => bUseUnityOverride = value;
 			get
 			{
-				bool UseUnity = true;
+				bool useUnity = true;
 				if (Target.DisableUnityBuildForModules?.Contains(Name) ?? false)
 				{
-					UseUnity = false;
+					useUnity = false;
 				}
 
-				return bUseUnityOverride ?? UseUnity;
+				return bUseUnityOverride ?? useUnity;
 			}
 		}
 
@@ -1236,13 +1236,13 @@ namespace UnrealBuildTool
 		/// <summary>
 		/// Returns the directory of where the passed in module name lives.
 		/// </summary>
-		/// <param name="ModuleName">Name of the module</param>
+		/// <param name="moduleName">Name of the module</param>
 		/// <returns>Directory where the module lives</returns>
-		public string GetModuleDirectory(string ModuleName)
+		public string GetModuleDirectory(string moduleName)
 		{
-			FileReference? ModuleFileReference = RulesAssembly.GetModuleFileName(ModuleName)
-				?? throw new CompilationResultException(CompilationResult.RulesError, "Could not find a module named '{ModuleName}'.", ModuleName);
-			return ModuleFileReference.Directory.FullName;
+			FileReference? moduleFileReference = RulesAssembly.GetModuleFileName(moduleName)
+				?? throw new CompilationResultException(CompilationResult.RulesError, "Could not find a module named '{ModuleName}'.", moduleName);
+			return moduleFileReference.Directory.FullName;
 		}
 
 		/// <summary>
@@ -1298,22 +1298,22 @@ namespace UnrealBuildTool
 		/// <summary>
 		/// Append (or create)
 		/// </summary>
-		/// <param name="Definition"></param>
-		/// <param name="Text"></param>
-		public void AppendStringToPublicDefinition(string Definition, string Text)
+		/// <param name="definition"></param>
+		/// <param name="text"></param>
+		public void AppendStringToPublicDefinition(string definition, string text)
 		{
-			string WithEquals = Definition + "=";
-			for (int Index = 0; Index < PublicDefinitions.Count; Index++)
+			string withEquals = definition + "=";
+			for (int index = 0; index < PublicDefinitions.Count; index++)
 			{
-				if (PublicDefinitions[Index].StartsWith(WithEquals, StringComparison.Ordinal))
+				if (PublicDefinitions[index].StartsWith(withEquals, StringComparison.Ordinal))
 				{
-					PublicDefinitions[Index] = PublicDefinitions[Index] + Text;
+					PublicDefinitions[index] = PublicDefinitions[index] + text;
 					return;
 				}
 			}
 
 			// if we get here, we need to make a new entry
-			PublicDefinitions.Add(Definition + "=" + Text);
+			PublicDefinitions.Add(definition + "=" + text);
 		}
 
 		/// <summary>
@@ -1416,7 +1416,7 @@ namespace UnrealBuildTool
 		/// <summary>
 		/// The AutoSDK directory for the active host platform
 		/// </summary>
-		public string? AutoSdkDirectory => UEBuildPlatformSDK.TryGetHostPlatformAutoSDKDir(out DirectoryReference? AutoSdkDir) ? AutoSdkDir.FullName : null;
+		public string? AutoSdkDirectory => UEBuildPlatformSDK.TryGetHostPlatformAutoSDKDir(out DirectoryReference? autoSdkDir) ? autoSdkDir.FullName : null;
 
 		/// <summary>
 		/// The current engine directory
@@ -1477,10 +1477,10 @@ namespace UnrealBuildTool
 		/// Constructor. For backwards compatibility while the parameterless constructor is being phased out, initialization which would happen here is done by 
 		/// RulesAssembly.CreateModulRules instead.
 		/// </summary>
-		/// <param name="Target">Rules for building this target</param>
-		public ModuleRules(ReadOnlyTargetRules Target)
+		/// <param name="target">Rules for building this target</param>
+		public ModuleRules(ReadOnlyTargetRules target)
 		{
-			this.Target = Target;
+			this.Target = target;
 		}
 #nullable restore
 
@@ -1490,13 +1490,13 @@ namespace UnrealBuildTool
 		///	Private, meaning the include paths for the included modules will not be exposed when giving this modules include paths
 		///	NOTE: There is no AddThirdPartyPublicStaticDependencies function.
 		/// </summary>
-		/// <param name="Target">The target this module belongs to</param>
-		/// <param name="ModuleNames">The names of the modules to add</param>
-		public void AddEngineThirdPartyPrivateStaticDependencies(ReadOnlyTargetRules Target, params string[] ModuleNames)
+		/// <param name="target">The target this module belongs to</param>
+		/// <param name="moduleNames">The names of the modules to add</param>
+		public void AddEngineThirdPartyPrivateStaticDependencies(ReadOnlyTargetRules target, params string[] moduleNames)
 		{
-			if (!bUsePrecompiled || Target.LinkType == TargetLinkType.Monolithic)
+			if (!bUsePrecompiled || target.LinkType == TargetLinkType.Monolithic)
 			{
-				PrivateDependencyModuleNames.AddRange(ModuleNames);
+				PrivateDependencyModuleNames.AddRange(moduleNames);
 			}
 		}
 
@@ -1506,23 +1506,23 @@ namespace UnrealBuildTool
 		///	Private, meaning the include paths for the included modules will not be exposed when giving this modules include paths
 		///	NOTE: There is no AddThirdPartyPublicDynamicDependencies function.
 		/// </summary>
-		/// <param name="Target">Rules for the target being built</param>
-		/// <param name="ModuleNames">The names of the modules to add</param>
-		public void AddEngineThirdPartyPrivateDynamicDependencies(ReadOnlyTargetRules Target, params string[] ModuleNames)
+		/// <param name="target">Rules for the target being built</param>
+		/// <param name="moduleNames">The names of the modules to add</param>
+		public void AddEngineThirdPartyPrivateDynamicDependencies(ReadOnlyTargetRules target, params string[] moduleNames)
 		{
-			if (!bUsePrecompiled || Target.LinkType == TargetLinkType.Monolithic)
+			if (!bUsePrecompiled || target.LinkType == TargetLinkType.Monolithic)
 			{
-				PrivateIncludePathModuleNames.AddRange(ModuleNames);
-				DynamicallyLoadedModuleNames.AddRange(ModuleNames);
+				PrivateIncludePathModuleNames.AddRange(moduleNames);
+				DynamicallyLoadedModuleNames.AddRange(moduleNames);
 			}
 		}
 
 		/// <summary>
 		/// Setup this module for Mesh Editor support (based on the settings in UEBuildConfiguration)
 		/// </summary>
-		public void EnableMeshEditorSupport(ReadOnlyTargetRules Target)
+		public void EnableMeshEditorSupport(ReadOnlyTargetRules target)
 		{
-			if (Target.bEnableMeshEditor == true)
+			if (target.bEnableMeshEditor)
 			{
 				PublicDefinitions.Add("ENABLE_MESH_EDITOR=1");
 			}
@@ -1535,13 +1535,13 @@ namespace UnrealBuildTool
 		/// <summary>
 		/// Setup this module for GameplayDebugger support
 		/// </summary>
-		public void SetupGameplayDebuggerSupport(ReadOnlyTargetRules Target, bool bAddAsPublicDependency = false)
+		public void SetupGameplayDebuggerSupport(ReadOnlyTargetRules target, bool bAddAsPublicDependency = false)
 		{
-			if (Target.bUseGameplayDebugger || Target.bUseGameplayDebuggerCore)
+			if (target.bUseGameplayDebugger || target.bUseGameplayDebuggerCore)
 			{
 				PublicDefinitions.Add("WITH_GAMEPLAY_DEBUGGER_CORE=1");
-				PublicDefinitions.Add("WITH_GAMEPLAY_DEBUGGER=" + (Target.bUseGameplayDebugger ? 1 : 0));
-				if (Target.bUseGameplayDebugger || (Target.bUseGameplayDebuggerCore && Target.Configuration != UnrealTargetConfiguration.Shipping))
+				PublicDefinitions.Add("WITH_GAMEPLAY_DEBUGGER=" + (target.bUseGameplayDebugger ? 1 : 0));
+				if (target.bUseGameplayDebugger || (target.bUseGameplayDebuggerCore && target.Configuration != UnrealTargetConfiguration.Shipping))
 				{
 					PublicDefinitions.Add("WITH_GAMEPLAY_DEBUGGER_MENU=1");
 				}
@@ -1553,7 +1553,7 @@ namespace UnrealBuildTool
 				if (bAddAsPublicDependency)
 				{
 					PublicDependencyModuleNames.Add("GameplayDebugger");
-					if (Target.Type == TargetType.Editor)
+					if (target.Type == TargetType.Editor)
 					{
 						PublicDependencyModuleNames.Add("GameplayDebuggerEditor");
 					}
@@ -1561,7 +1561,7 @@ namespace UnrealBuildTool
 				else
 				{
 					PrivateDependencyModuleNames.Add("GameplayDebugger");
-					if (Target.Type == TargetType.Editor)
+					if (target.Type == TargetType.Editor)
 					{
 						PrivateDependencyModuleNames.Add("GameplayDebuggerEditor");
 					}
@@ -1578,9 +1578,9 @@ namespace UnrealBuildTool
 		/// <summary>
 		/// Setup this module for Iris support (based on the settings in UEBuildConfiguration)
 		/// </summary>
-		public void SetupIrisSupport(ReadOnlyTargetRules Target, bool bAddAsPublicDependency = false)
+		public void SetupIrisSupport(ReadOnlyTargetRules target, bool bAddAsPublicDependency = false)
 		{
-			if (Target.bUseIris == true)
+			if (target.bUseIris == true)
 			{
 				PublicDefinitions.Add("UE_WITH_IRIS=1");
 				PrivateDefinitions.Add("UE_NET_HAS_IRIS_FASTARRAY_BINDING=1");
@@ -1613,9 +1613,9 @@ namespace UnrealBuildTool
 		/// <summary>
 		/// Setup this module for Chaos Visual Debugger support (Required for recording debug data that will be visualized in the Chaos Visual Debugger tool)
 		/// </summary>
-		public void SetupModuleChaosVisualDebuggerSupport(ReadOnlyTargetRules Target)
+		public void SetupModuleChaosVisualDebuggerSupport(ReadOnlyTargetRules target)
 		{
-			bool bHasChaosVisualDebuggerSupport = Target.bCompileChaosVisualDebuggerSupport && Target.Configuration != UnrealTargetConfiguration.Shipping;
+			bool bHasChaosVisualDebuggerSupport = target.bCompileChaosVisualDebuggerSupport && target.Configuration != UnrealTargetConfiguration.Shipping;
 			if (bHasChaosVisualDebuggerSupport)
 			{
 				PublicDependencyModuleNames.Add("ChaosVDRuntime");
@@ -1631,7 +1631,7 @@ namespace UnrealBuildTool
 		/// <summary>
 		/// Setup this module for physics support (based on the settings in UEBuildConfiguration)
 		/// </summary>
-		public void SetupModulePhysicsSupport(ReadOnlyTargetRules Target)
+		public void SetupModulePhysicsSupport(ReadOnlyTargetRules target)
 		{
 			PublicIncludePathModuleNames.AddRange(
 					new string[] {
@@ -1649,17 +1649,17 @@ namespace UnrealBuildTool
 
 			PublicDefinitions.Add("WITH_CLOTH_COLLISION_DETECTION=1");
 
-			SetupModuleChaosVisualDebuggerSupport(Target);
+			SetupModuleChaosVisualDebuggerSupport(target);
 
 			// Modules may still be relying on appropriate definitions for physics.
 			// Nothing in engine should use these anymore as they were all deprecated and 
 			// assumed to be in the following configuration from 5.1, this will cause
 			// deprecation warning to fire in any module still relying on these macros
 
-			Func<string, string, string, string> GetDeprecatedPhysicsMacro = (string Macro, string Value, string Version) =>
+			static string GetDeprecatedPhysicsMacro(string macro, string value, string version)
 			{
-				return Macro + "=UE_DEPRECATED_MACRO(" + Version + ", \"" + Macro + " is deprecated and should always be considered " + Value + ".\") " + Value;
-			};
+				return macro + "=UE_DEPRECATED_MACRO(" + version + ", \"" + macro + " is deprecated and should always be considered " + value + ".\") " + value;
+			}
 
 			PublicDefinitions.AddRange(
 				new string[]{
@@ -1681,49 +1681,49 @@ namespace UnrealBuildTool
 		/// <summary>
 		/// Determines if a module type is valid for a target, based on custom attributes
 		/// </summary>
-		/// <param name="ModuleType">The type of the module to check</param>
-		/// <param name="TargetRules">The target to check against</param>
-		/// <param name="InvalidReason">Out, reason this module was invalid</param>
+		/// <param name="moduleType">The type of the module to check</param>
+		/// <param name="targetRules">The target to check against</param>
+		/// <param name="invalidReason">Out, reason this module was invalid</param>
 		/// <returns>True if the module is valid, false otherwise</returns>
-		internal static bool IsValidForTarget(Type ModuleType, ReadOnlyTargetRules TargetRules, [NotNullWhen(false)] out string? InvalidReason)
+		internal static bool IsValidForTarget(Type moduleType, ReadOnlyTargetRules targetRules, [NotNullWhen(false)] out string? invalidReason)
 		{
-			IEnumerable<TargetType> SupportedTargetTypes = ModuleType.GetCustomAttributes<SupportedTargetTypesAttribute>().SelectMany(x => x.TargetTypes).Distinct();
-			if (SupportedTargetTypes.Any() && !SupportedTargetTypes.Contains(TargetRules.Type))
+			IEnumerable<TargetType> supportedTargetTypes = moduleType.GetCustomAttributes<SupportedTargetTypesAttribute>().SelectMany(x => x.TargetTypes).Distinct();
+			if (supportedTargetTypes.Any() && !supportedTargetTypes.Contains(targetRules.Type))
 			{
-				InvalidReason = $"TargetType '{TargetRules.Type}'";
+				invalidReason = $"TargetType '{targetRules.Type}'";
 				return false;
 			}
 
-			IEnumerable<UnrealTargetConfiguration> SupportedConfigurations = ModuleType.GetCustomAttributes<SupportedConfigurationsAttribute>().SelectMany(x => x.Configurations).Distinct();
-			if (SupportedConfigurations.Any() && !SupportedConfigurations.Contains(TargetRules.Configuration))
+			IEnumerable<UnrealTargetConfiguration> supportedConfigurations = moduleType.GetCustomAttributes<SupportedConfigurationsAttribute>().SelectMany(x => x.Configurations).Distinct();
+			if (supportedConfigurations.Any() && !supportedConfigurations.Contains(targetRules.Configuration))
 			{
-				InvalidReason = $"Configuration '{TargetRules.Configuration}'";
+				invalidReason = $"Configuration '{targetRules.Configuration}'";
 				return false;
 			}
 
 			// Skip platform extension modules. We only care about the base modules, not the platform overrides.
 			// The platform overrides get applied at a later stage when we actually come to build the module.
-			if (!UEBuildPlatform.GetPlatformFolderNames().Any(Name => ModuleType.Name.EndsWith("_" + Name, StringComparison.OrdinalIgnoreCase)))
+			if (!UEBuildPlatform.GetPlatformFolderNames().Any(name => moduleType.Name.EndsWith("_" + name, StringComparison.OrdinalIgnoreCase)))
 			{
-				IEnumerable<SupportedPlatformsAttribute> PlatformAttributes = ModuleType.GetCustomAttributes<SupportedPlatformsAttribute>();
-				IEnumerable<UnrealTargetPlatform> SupportedPlatforms = PlatformAttributes.SelectMany(x => x.Platforms).Distinct();
-				if (PlatformAttributes.Any() && !SupportedPlatforms.Contains(TargetRules.Platform))
+				IEnumerable<SupportedPlatformsAttribute> platformAttributes = moduleType.GetCustomAttributes<SupportedPlatformsAttribute>();
+				IEnumerable<UnrealTargetPlatform> supportedPlatforms = platformAttributes.SelectMany(x => x.Platforms).Distinct();
+				if (platformAttributes.Any() && !supportedPlatforms.Contains(targetRules.Platform))
 				{
-					InvalidReason = $"Platform '{TargetRules.Platform}'";
+					invalidReason = $"Platform '{targetRules.Platform}'";
 					return false;
 				}
 			}
 
-			InvalidReason = null;
+			invalidReason = null;
 			return true;
 		}
 
 		/// <summary>
 		/// Determines if this module can be precompiled for the current target.
 		/// </summary>
-		/// <param name="RulesFile">Path to the module rules file</param>
+		/// <param name="rulesFile">Path to the module rules file</param>
 		/// <returns>True if the module can be precompiled, false otherwise</returns>
-		internal bool IsValidForTarget(FileReference RulesFile)
+		internal bool IsValidForTarget(FileReference rulesFile)
 		{
 			if (Type == ModuleRules.ModuleType.CPlusPlus)
 			{
@@ -1732,7 +1732,7 @@ namespace UnrealBuildTool
 					case ModuleRules.PrecompileTargetsType.None:
 						return false;
 					case ModuleRules.PrecompileTargetsType.Default:
-						return (Target.Type == TargetType.Editor || !Unreal.GetExtensionDirs(Unreal.EngineDirectory, "Source/Developer").Any(Dir => RulesFile.IsUnderDirectory(Dir)) || Plugin != null);
+						return (Target.Type == TargetType.Editor || !Unreal.GetExtensionDirs(Unreal.EngineDirectory, "Source/Developer").Any(dir => rulesFile.IsUnderDirectory(dir)) || Plugin != null);
 					case ModuleRules.PrecompileTargetsType.Game:
 						return (Target.Type == TargetType.Client || Target.Type == TargetType.Server || Target.Type == TargetType.Game);
 					case ModuleRules.PrecompileTargetsType.Editor:
@@ -1747,46 +1747,46 @@ namespace UnrealBuildTool
 		/// <summary>
 		/// Determines whether a given platform is available in the context of the current Target
 		/// </summary>
-		/// <param name="InPlatform">The platform to check for</param>
+		/// <param name="inPlatform">The platform to check for</param>
 		/// <returns>True if it's available, false otherwise</returns>
-		protected bool IsPlatformAvailable(UnrealTargetPlatform InPlatform)
+		protected bool IsPlatformAvailable(UnrealTargetPlatform inPlatform)
 		{
-			return UEBuildPlatform.IsPlatformAvailableForTarget(InPlatform, Target);
+			return UEBuildPlatform.IsPlatformAvailableForTarget(inPlatform, Target);
 		}
 
 		/// <summary>
 		/// Returns all the modules that are in the given module group
 		/// </summary>
-		/// <param name="ModuleGroup">The name of the module group, as defined by a [ModuleGroup("Name")] attribute</param>
+		/// <param name="moduleGroup">The name of the module group, as defined by a [ModuleGroup("Name")] attribute</param>
 		/// <param name="bOnlyValid">Only include modules that are valid for the current Target</param>
 		/// <returns></returns>
-		public IEnumerable<string> GetModulesInGroup( string ModuleGroup, bool bOnlyValid = true )
+		public IEnumerable<string> GetModulesInGroup(string moduleGroup, bool bOnlyValid = true)
 		{
 			// figure out what platforms/groups aren't allowed with this opted in list
-			List<string>? DisallowedPlatformsAndGroups = (bOnlyValid && Target.OptedInModulePlatforms != null) ? Utils.MakeListOfUnsupportedPlatforms(Target.OptedInModulePlatforms.ToList(), false, Logger) : null;
+			List<string>? disallowedPlatformsAndGroups = (bOnlyValid && Target.OptedInModulePlatforms != null) ? Utils.MakeListOfUnsupportedPlatforms(Target.OptedInModulePlatforms.ToList(), false, Logger) : null;
 
-			foreach (Type ModuleType in RulesAssembly.GetTypes()
-				.Where( T => T.IsSubclassOf(typeof(ModuleRules)) && T.IsDefined(typeof(ModuleGroupsAttribute), true ))
-				.Where( T => !bOnlyValid || IsValidForTarget(T, Target, out string? _))
+			foreach (Type moduleType in RulesAssembly.GetTypes()
+				.Where(t => t.IsSubclassOf(typeof(ModuleRules)) && t.IsDefined(typeof(ModuleGroupsAttribute), true))
+				.Where(t => !bOnlyValid || IsValidForTarget(t, Target, out string? _))
 				)
 			{
 				// check if the module file is disallowed
-				if (DisallowedPlatformsAndGroups != null)
+				if (disallowedPlatformsAndGroups != null)
 				{
-					FileReference? ModuleFileName = RulesAssembly.GetModuleFileName(ModuleType.Name);
-					if (ModuleFileName != null)
+					FileReference? moduleFileName = RulesAssembly.GetModuleFileName(moduleType.Name);
+					if (moduleFileName != null)
 					{
-						if (ModuleFileName.ContainsAnyNames(DisallowedPlatformsAndGroups, Unreal.EngineDirectory) ||
-							(Target.ProjectFile != null && ModuleFileName.ContainsAnyNames(DisallowedPlatformsAndGroups, Target.ProjectFile.Directory)))
+						if (moduleFileName.ContainsAnyNames(disallowedPlatformsAndGroups, Unreal.EngineDirectory) ||
+							(Target.ProjectFile != null && moduleFileName.ContainsAnyNames(disallowedPlatformsAndGroups, Target.ProjectFile.Directory)))
 						{
 							continue;
 						}
 					}
 				}
 
-				if (ModuleType.GetCustomAttributes<ModuleGroupsAttribute>().Any( X => X.ModuleGroups.Contains(ModuleGroup)))
+				if (moduleType.GetCustomAttributes<ModuleGroupsAttribute>().Any(x => x.ModuleGroups.Contains(moduleGroup)))
 				{
-					yield return ModuleType.Name;
+					yield return moduleType.Name;
 				}
 			}
 		}
@@ -1797,13 +1797,13 @@ namespace UnrealBuildTool
 		/// </summary>
 		internal void PrepareModuleForTests()
 		{
-			TestTargetRules? TestTargetRules = Target.InnerTestTargetRules;
-			if (TestTargetRules == null)
+			TestTargetRules? testTargetRules = Target.InnerTestTargetRules;
+			if (testTargetRules == null)
 			{
 				return;
 			}
 
-			lock (TestTargetRules)
+			lock (testTargetRules)
 			{
 				if (Name != "LowLevelTestsRunner")
 				{
@@ -1825,14 +1825,14 @@ namespace UnrealBuildTool
 		/// <summary>
 		/// Returns the module directory for a given subclass of the module (platform extensions add subclasses of ModuleRules to add in platform-specific settings)
 		/// </summary>
-		/// <param name="Type">typeof the subclass</param>
+		/// <param name="type">typeof the subclass</param>
 		/// <returns>Directory where the subclass's .Build.cs lives, or null if not found</returns>
-		public DirectoryReference? GetModuleDirectoryForSubClass(Type Type)
+		public DirectoryReference? GetModuleDirectoryForSubClass(Type type)
 		{
-			DirectoryReference? Directory;
-			if (DirectoriesForModuleSubClasses.TryGetValue(Type, out Directory))
+			DirectoryReference? directory;
+			if (DirectoriesForModuleSubClasses.TryGetValue(type, out directory))
 			{
-				return Directory;
+				return directory;
 			}
 			return null;
 		}
@@ -1843,21 +1843,21 @@ namespace UnrealBuildTool
 		/// <returns>List of directories, or null if none were added</returns>
 		public DirectoryReference[] GetAllModuleDirectories()
 		{
-			List<DirectoryReference> AllDirectories = new List<DirectoryReference>(DirectoriesForModuleSubClasses.Values);
-			AllDirectories.AddRange(AdditionalModuleDirectories);
-			return AllDirectories.ToArray();
+			List<DirectoryReference> allDirectories = new List<DirectoryReference>(DirectoriesForModuleSubClasses.Values);
+			allDirectories.AddRange(AdditionalModuleDirectories);
+			return allDirectories.ToArray();
 		}
 
 		/// <summary>
 		/// Adds an additional module directory, if it exists (useful for NotForLicensees/NoRedist)
 		/// </summary>
-		/// <param name="Directory"></param>
+		/// <param name="directory"></param>
 		/// <returns>true if the directory exists</returns>
-		protected bool ConditionalAddModuleDirectory(DirectoryReference Directory)
+		protected bool ConditionalAddModuleDirectory(DirectoryReference directory)
 		{
-			if (DirectoryReference.Exists(Directory))
+			if (DirectoryReference.Exists(directory))
 			{
-				AdditionalModuleDirectories.Add(Directory);
+				AdditionalModuleDirectories.Add(directory);
 				return true;
 			}
 
@@ -1876,139 +1876,139 @@ namespace UnrealBuildTool
 		/// <summary>
 		/// Returns the VcPkg root directory for the build configuration
 		/// </summary>
-		/// <param name="PackageName">The name of the third-party package</param>
+		/// <param name="packageName">The name of the third-party package</param>
 		/// <returns></returns>
 		[SuppressMessage("Globalization", "CA1308:Normalize strings to uppercase", Justification = "Architecture path is lowercase")]
-		public string GetVcPackageRoot(string PackageName)
+		public string GetVcPackageRoot(string packageName)
 		{
 			// TODO: MacOS support, other platform support
-			string TargetPlatform = Target.Platform.ToString();
-			string? Platform = null;
-			string? Architecture = null;
-			string Linkage = String.Empty;
-			string Toolset = String.Empty;
+			string targetPlatform = Target.Platform.ToString();
+			string? platform = null;
+			string? architecture = null;
+			string linkage = String.Empty;
+			string toolset = String.Empty;
 			if (Target.Platform == UnrealTargetPlatform.Win64)
 			{
-				Platform = "windows";
-				Architecture = Target.WindowsPlatform.Architecture.ToString().ToLowerInvariant();
+				platform = "windows";
+				architecture = Target.WindowsPlatform.Architecture.ToString().ToLowerInvariant();
 				if (Target.bUseStaticCRT)
 				{
-					Linkage = "-static";
+					linkage = "-static";
 				}
 				else
 				{
-					Linkage = "-static-md";
+					linkage = "-static-md";
 				}
-				Toolset = "-v142";
+				toolset = "-v142";
 			}
 			else if (Target.Platform == UnrealTargetPlatform.Linux)
 			{
-				Architecture = "x86_64";
-				Platform = "unknown-linux-gnu";
+				architecture = "x86_64";
+				platform = "unknown-linux-gnu";
 			}
 			else if (Target.Platform == UnrealTargetPlatform.LinuxArm64)
 			{
-				Architecture = "aarch64";
-				Platform = "unknown-linux-gnueabi";
+				architecture = "aarch64";
+				platform = "unknown-linux-gnueabi";
 			}
 			else if (Target.Platform == UnrealTargetPlatform.Mac)
 			{
-				Architecture = "x86_64";
-				Platform = "osx";
+				architecture = "x86_64";
+				platform = "osx";
 			}
 
-			if (String.IsNullOrEmpty(TargetPlatform) || String.IsNullOrEmpty(Platform) || String.IsNullOrEmpty(Architecture))
+			if (String.IsNullOrEmpty(targetPlatform) || String.IsNullOrEmpty(platform) || String.IsNullOrEmpty(architecture))
 			{
 				throw new System.NotSupportedException($"Platform {Target.Platform} not currently supported by vcpkg");
 			}
 
-			string Triplet = $"{Architecture}-{Platform}{Linkage}{Toolset}";
+			string triplet = $"{architecture}-{platform}{linkage}{toolset}";
 
-			return Path.Combine("ThirdParty", "vcpkg", TargetPlatform, Triplet, $"{PackageName}_{Triplet}");
+			return Path.Combine("ThirdParty", "vcpkg", targetPlatform, triplet, $"{packageName}_{triplet}");
 		}
 
 		/// <summary>
 		/// Adds libraries compiled with vcpkg to the current module
 		/// </summary>
-		/// <param name="PackageName">The name of the third-party package</param>
-		/// <param name="AddInclude">Should the include directory be added to PublicIncludePaths</param>
-		/// <param name="Libraries">The names of the libaries to add to PublicAdditionalLibraries/</param>
-		public void AddVcPackage(string PackageName, bool AddInclude, params string[] Libraries)
+		/// <param name="packageName">The name of the third-party package</param>
+		/// <param name="addInclude">Should the include directory be added to PublicIncludePaths</param>
+		/// <param name="libraries">The names of the libaries to add to PublicAdditionalLibraries/</param>
+		public void AddVcPackage(string packageName, bool addInclude, params string[] libraries)
 		{
-			string VcPackageRoot = GetVcPackageRoot(PackageName);
+			string vcPackageRoot = GetVcPackageRoot(packageName);
 
-			if (!System.IO.Directory.Exists(VcPackageRoot))
+			if (!System.IO.Directory.Exists(vcPackageRoot))
 			{
-				throw new DirectoryNotFoundException(VcPackageRoot);
+				throw new DirectoryNotFoundException(vcPackageRoot);
 			}
 
-			string LibraryExtension = String.Empty;
+			string libraryExtension = String.Empty;
 			if (Target.Platform == UnrealTargetPlatform.Win64)
 			{
-				LibraryExtension = ".lib";
+				libraryExtension = ".lib";
 			}
 			else if (Target.Platform == UnrealTargetPlatform.Linux || Target.Platform == UnrealTargetPlatform.LinuxArm64 || Target.Platform == UnrealTargetPlatform.Mac)
 			{
-				LibraryExtension = ".a";
+				libraryExtension = ".a";
 			}
 
-			foreach (string Library in Libraries)
+			foreach (string library in libraries)
 			{
-				string LibraryPath = Path.Combine(VcPackageRoot, "lib", $"{Library}{LibraryExtension}");
-				if ((Target.Platform == UnrealTargetPlatform.Linux || Target.Platform == UnrealTargetPlatform.LinuxArm64 || Target.Platform == UnrealTargetPlatform.Mac) && !Library.StartsWith("lib", StringComparison.OrdinalIgnoreCase))
+				string libraryPath = Path.Combine(vcPackageRoot, "lib", $"{library}{libraryExtension}");
+				if ((Target.Platform == UnrealTargetPlatform.Linux || Target.Platform == UnrealTargetPlatform.LinuxArm64 || Target.Platform == UnrealTargetPlatform.Mac) && !library.StartsWith("lib", StringComparison.OrdinalIgnoreCase))
 				{
-					LibraryPath = Path.Combine(VcPackageRoot, "lib", $"lib{Library}{LibraryExtension}");
+					libraryPath = Path.Combine(vcPackageRoot, "lib", $"lib{library}{libraryExtension}");
 				}
-				if (!System.IO.File.Exists(LibraryPath))
+				if (!System.IO.File.Exists(libraryPath))
 				{
-					throw new FileNotFoundException(LibraryPath);
+					throw new FileNotFoundException(libraryPath);
 				}
-				PublicAdditionalLibraries.Add(LibraryPath);
+				PublicAdditionalLibraries.Add(libraryPath);
 			}
 
-			if (AddInclude)
+			if (addInclude)
 			{
-				string IncludePath = Path.Combine(VcPackageRoot, "include");
-				if (!System.IO.Directory.Exists(IncludePath))
+				string includePath = Path.Combine(vcPackageRoot, "include");
+				if (!System.IO.Directory.Exists(includePath))
 				{
-					throw new DirectoryNotFoundException(IncludePath);
+					throw new DirectoryNotFoundException(includePath);
 				}
 
-				PublicSystemIncludePaths.Add(Path.Combine(VcPackageRoot, "include"));
+				PublicSystemIncludePaths.Add(Path.Combine(vcPackageRoot, "include"));
 			}
 		}
 
 		/// <summary>
 		/// Replace an expected value in a list of definitions with a new value
 		/// </summary>
-		/// <param name="Definitions">List of definitions e.g. PublicDefinitions</param>
-		/// <param name="Name">Name of the define to change</param>
-		/// <param name="PreviousValue">Expected value</param>
-		/// <param name="NewValue">New value</param>
+		/// <param name="definitions">List of definitions e.g. PublicDefinitions</param>
+		/// <param name="name">Name of the define to change</param>
+		/// <param name="previousValue">Expected value</param>
+		/// <param name="newValue">New value</param>
 		/// <exception cref="Exception"></exception>
-		protected static void ChangeDefinition(List<string> Definitions, string Name, string PreviousValue, string NewValue)
+		protected static void ChangeDefinition(List<string> definitions, string name, string previousValue, string newValue)
 		{
-			if (!Definitions.Remove($"{Name}={PreviousValue}"))
+			if (!definitions.Remove($"{name}={previousValue}"))
 			{
 				throw new Exception("Failed to removed expected definition");
 			}
-			Definitions.Add($"{Name}={NewValue}");
+			definitions.Add($"{name}={newValue}");
 		}
 
 		/// <summary>
 		/// Replace an expected value in a list of module names with a new value
 		/// </summary>
-		/// <param name="Definitions">List of module names e.g. PublicDependencyModuleNames</param>
-		/// <param name="PreviousModule">Expected value</param>
-		/// <param name="NewModule">New value</param>
+		/// <param name="definitions">List of module names e.g. PublicDependencyModuleNames</param>
+		/// <param name="previousModule">Expected value</param>
+		/// <param name="newModule">New value</param>
 		/// <exception cref="Exception"></exception>
-		protected static void ReplaceModule(List<string> Definitions, string PreviousModule, string NewModule)
+		protected static void ReplaceModule(List<string> definitions, string previousModule, string newModule)
 		{
-			if (!Definitions.Remove(PreviousModule))
+			if (!definitions.Remove(previousModule))
 			{
 				throw new Exception("Failed to removed expected module name");
 			}
-			Definitions.Add(NewModule);
+			definitions.Add(newModule);
 		}
 	}
 }
