UnrealBuidTool: ModuleRules cleanup pt.3

#rnx

diff --git a/Engine/Source/Programs/UnrealBuildTool/Configuration/ModuleRules.cs b/Engine/Source/Programs/UnrealBuildTool/Configuration/ModuleRules.cs
--- a/Engine/Source/Programs/UnrealBuildTool/Configuration/ModuleRules.cs
+++ b/Engine/Source/Programs/UnrealBuildTool/Configuration/ModuleRules.cs
@@ -635,17 +635,9 @@ namespace UnrealBuildTool
 		public PackageOverrideType OverridePackageType
 		{
 			get => overridePackageType ?? PackageOverrideType.None;
-			set
-			{
-				if (!IsPlugin)
-				{
-					overridePackageType = value;
-				}
-				else
-				{
-					throw new CompilationResultException(CompilationResult.RulesError, "Module '{ModuleName}' cannot override package type because it is part of a plugin!", Name);
-				}
-			}
+			set => overridePackageType = !IsPlugin
+					? value
+					: throw new CompilationResultException(CompilationResult.RulesError, "Module '{ModuleName}' cannot override package type because it is part of a plugin!", Name);
 		}
 
 		private PackageOverrideType? overridePackageType;
@@ -1044,16 +1036,7 @@ namespace UnrealBuildTool
 		public bool bUseUnity
 		{
 			set => bUseUnityOverride = value;
-			get
-			{
-				bool useUnity = true;
-				if (Target.DisableUnityBuildForModules?.Contains(Name) ?? false)
-				{
-					useUnity = false;
-				}
-
-				return bUseUnityOverride ?? useUnity;
-			}
+			get => bUseUnityOverride ?? Target.DisableUnityBuildForModules?.Contains(Name) != true;
 		}
 
 		/// <summary>
@@ -1080,10 +1063,7 @@ namespace UnrealBuildTool
 		/// <summary>
 		/// Helper function to get the number of byes per unity cpp file
 		/// </summary>
-		public int GetNumIncludedBytesPerUnityCPP()
-		{
-			return (NumIncludedBytesPerUnityCPPOverride != 0 && !Target.bDisableModuleNumIncludedBytesPerUnityCPPOverride) ? NumIncludedBytesPerUnityCPPOverride : Target.NumIncludedBytesPerUnityCPP;
-		}
+		public int GetNumIncludedBytesPerUnityCPP() => (NumIncludedBytesPerUnityCPPOverride != 0 && !Target.bDisableModuleNumIncludedBytesPerUnityCPPOverride) ? NumIncludedBytesPerUnityCPPOverride : Target.NumIncludedBytesPerUnityCPP;
 
 		/// <summary>
 		/// Module uses a #import so must be built locally when compiling with SN-DBS
@@ -1411,7 +1391,7 @@ namespace UnrealBuildTool
 		/// <summary>
 		///  Control visibility of symbols
 		/// </summary>
-		public SymbolVisibility ModuleSymbolVisibility { get; set; } = ModuleRules.SymbolVisibility.Default;
+		public SymbolVisibility ModuleSymbolVisibility { get; set; } = SymbolVisibility.Default;
 
 		/// <summary>
 		/// The AutoSDK directory for the active host platform
@@ -1426,20 +1406,9 @@ namespace UnrealBuildTool
 		/// <summary>
 		/// Property for the directory containing this plugin. Useful for adding paths to third party dependencies.
 		/// </summary>
-		public string PluginDirectory
-		{
-			get
-			{
-				if (Plugin == null)
-				{
-					throw new CompilationResultException(CompilationResult.RulesError, "Module '{ModuleName}' does not belong to a plugin; PluginDirectory property is invalid.", Name);
-				}
-				else
-				{
-					return Plugin.Directory.FullName;
-				}
-			}
-		}
+		public string PluginDirectory => Plugin == null
+					? throw new CompilationResultException(CompilationResult.RulesError, "Module '{ModuleName}' does not belong to a plugin; PluginDirectory property is invalid.", Name)
+					: Plugin.Directory.FullName;
 
 		/// <summary>
 		/// Property for the directory containing this module. Useful for adding paths to third party dependencies.
@@ -1460,17 +1429,7 @@ namespace UnrealBuildTool
 		/// <summary>
 		/// Returns module's low level tests directory "Tests".
 		/// </summary>
-		public string TestsDirectory
-		{
-			get
-			{
-				if (IsTestModule)
-				{
-					return Directory.FullName;
-				}
-				return Path.Combine(Directory.FullName, "Tests");
-			}
-		}
+		public string TestsDirectory => IsTestModule ? Directory.FullName : Path.Combine(Directory.FullName, "Tests");
 
 #nullable disable
 		/// <summary>
@@ -1480,7 +1439,7 @@ namespace UnrealBuildTool
 		/// <param name="target">Rules for building this target</param>
 		public ModuleRules(ReadOnlyTargetRules target)
 		{
-			this.Target = target;
+			Target = target;
 		}
 #nullable restore
 
@@ -1520,17 +1479,7 @@ namespace UnrealBuildTool
 		/// <summary>
 		/// Setup this module for Mesh Editor support (based on the settings in UEBuildConfiguration)
 		/// </summary>
-		public void EnableMeshEditorSupport(ReadOnlyTargetRules target)
-		{
-			if (target.bEnableMeshEditor)
-			{
-				PublicDefinitions.Add("ENABLE_MESH_EDITOR=1");
-			}
-			else
-			{
-				PublicDefinitions.Add("ENABLE_MESH_EDITOR=0");
-			}
-		}
+		public void EnableMeshEditorSupport(ReadOnlyTargetRules target) => PublicDefinitions.Add($"ENABLE_MESH_EDITOR={(target.bEnableMeshEditor ? 1 : 0)}");
 
 		/// <summary>
 		/// Setup this module for GameplayDebugger support
@@ -1655,11 +1604,7 @@ namespace UnrealBuildTool
 			// Nothing in engine should use these anymore as they were all deprecated and 
 			// assumed to be in the following configuration from 5.1, this will cause
 			// deprecation warning to fire in any module still relying on these macros
-
-			static string GetDeprecatedPhysicsMacro(string macro, string value, string version)
-			{
-				return macro + "=UE_DEPRECATED_MACRO(" + version + ", \"" + macro + " is deprecated and should always be considered " + value + ".\") " + value;
-			}
+			static string GetDeprecatedPhysicsMacro(string macro, string value, string version) => $"{macro}=UE_DEPRECATED_MACRO({version}, \"{macro} is deprecated and should always be considered {value}.\") {value}";
 
 			PublicDefinitions.AddRange(
 				new string[]{
@@ -1725,19 +1670,19 @@ namespace UnrealBuildTool
 		/// <returns>True if the module can be precompiled, false otherwise</returns>
 		internal bool IsValidForTarget(FileReference rulesFile)
 		{
-			if (Type == ModuleRules.ModuleType.CPlusPlus)
+			if (Type == ModuleType.CPlusPlus)
 			{
 				switch (PrecompileForTargets)
 				{
-					case ModuleRules.PrecompileTargetsType.None:
+					case PrecompileTargetsType.None:
 						return false;
-					case ModuleRules.PrecompileTargetsType.Default:
+					case PrecompileTargetsType.Default:
 						return (Target.Type == TargetType.Editor || !Unreal.GetExtensionDirs(Unreal.EngineDirectory, "Source/Developer").Any(dir => rulesFile.IsUnderDirectory(dir)) || Plugin != null);
-					case ModuleRules.PrecompileTargetsType.Game:
+					case PrecompileTargetsType.Game:
 						return (Target.Type == TargetType.Client || Target.Type == TargetType.Server || Target.Type == TargetType.Game);
-					case ModuleRules.PrecompileTargetsType.Editor:
+					case PrecompileTargetsType.Editor:
 						return (Target.Type == TargetType.Editor);
-					case ModuleRules.PrecompileTargetsType.Any:
+					case PrecompileTargetsType.Any:
 						return true;
 				}
 			}
@@ -1749,10 +1694,7 @@ namespace UnrealBuildTool
 		/// </summary>
 		/// <param name="inPlatform">The platform to check for</param>
 		/// <returns>True if it's available, false otherwise</returns>
-		protected bool IsPlatformAvailable(UnrealTargetPlatform inPlatform)
-		{
-			return UEBuildPlatform.IsPlatformAvailableForTarget(inPlatform, Target);
-		}
+		protected bool IsPlatformAvailable(UnrealTargetPlatform inPlatform) => UEBuildPlatform.IsPlatformAvailableForTarget(inPlatform, Target);
 
 		/// <summary>
 		/// Returns all the modules that are in the given module group
@@ -1827,15 +1769,7 @@ namespace UnrealBuildTool
 		/// </summary>
 		/// <param name="type">typeof the subclass</param>
 		/// <returns>Directory where the subclass's .Build.cs lives, or null if not found</returns>
-		public DirectoryReference? GetModuleDirectoryForSubClass(Type type)
-		{
-			DirectoryReference? directory;
-			if (DirectoriesForModuleSubClasses.TryGetValue(type, out directory))
-			{
-				return directory;
-			}
-			return null;
-		}
+		public DirectoryReference? GetModuleDirectoryForSubClass(Type type) => DirectoriesForModuleSubClasses.TryGetValue(type, out DirectoryReference? directory) ? directory : null;
 
 		/// <summary>
 		/// Returns the directories for all subclasses of this module, as well as any additional directories specified by the rules
