UnrealBuildTool: Fix some editorconfig errors

#rnx

diff --git a/Engine/Source/Programs/UnrealBuildTool/Configuration/TargetDescriptor.cs b/Engine/Source/Programs/UnrealBuildTool/Configuration/TargetDescriptor.cs
--- a/Engine/Source/Programs/UnrealBuildTool/Configuration/TargetDescriptor.cs
+++ b/Engine/Source/Programs/UnrealBuildTool/Configuration/TargetDescriptor.cs
@@ -172,7 +172,7 @@ namespace UnrealBuildTool
 				}
 				return UnrealIntermediateEnvironment.Default;
 			}
-			set { IntermediateEnvironmentOverride = value; }
+			set => IntermediateEnvironmentOverride = value;
 		}
 		private UnrealIntermediateEnvironment? IntermediateEnvironmentOverride;
 
diff --git a/Engine/Source/Programs/UnrealBuildTool/Configuration/TestModuleRules.cs b/Engine/Source/Programs/UnrealBuildTool/Configuration/TestModuleRules.cs
--- a/Engine/Source/Programs/UnrealBuildTool/Configuration/TestModuleRules.cs
+++ b/Engine/Source/Programs/UnrealBuildTool/Configuration/TestModuleRules.cs
@@ -550,8 +550,8 @@ namespace UnrealBuildTool
 			/// </summary>
 			public Dictionary<UnrealTargetPlatform, bool> PlatformRunContainerized
 			{
-				get { return PlatformRunContainerizedPrivate; }
-				set { PlatformRunContainerizedPrivate = value; }
+				get => PlatformRunContainerizedPrivate;
+				set => PlatformRunContainerizedPrivate = value;
 			}
 
 		}
diff --git a/Engine/Source/Programs/UnrealBuildTool/Configuration/UEBuildTarget.cs b/Engine/Source/Programs/UnrealBuildTool/Configuration/UEBuildTarget.cs
--- a/Engine/Source/Programs/UnrealBuildTool/Configuration/UEBuildTarget.cs
+++ b/Engine/Source/Programs/UnrealBuildTool/Configuration/UEBuildTarget.cs
@@ -1402,7 +1402,7 @@ namespace UnrealBuildTool
 			if (ThisRules.RequiresUniqueEnvironment(RulesAssembly, Arguments, PropNamesThatRequiredUnique, out BaseTargetName))
 			{
 				throw new BuildException("{0} modifies the values of properties: [ {1} ]. This is not allowed, as {0} has build products in common with {2}.\nRemove the modified setting, change {0} to use a unique build environment by setting 'BuildEnvironment = TargetBuildEnvironment.Unique;' in the {3} constructor, or set bOverrideBuildEnvironment = true to force this setting on.", 
-					ThisTargetName, string.Join(", ", PropNamesThatRequiredUnique), BaseTargetName, ThisRules.GetType().Name);
+					ThisTargetName, String.Join(", ", PropNamesThatRequiredUnique), BaseTargetName, ThisRules.GetType().Name);
 			}
 
 			// Make sure that we don't explicitly enable or disable any plugins through the target rules. We can't do this with the shared build environment because it requires recompiling the "Projects" engine module.
diff --git a/Engine/Source/Programs/UnrealBuildTool/Platform/Windows/PVSToolChain.cs b/Engine/Source/Programs/UnrealBuildTool/Platform/Windows/PVSToolChain.cs
--- a/Engine/Source/Programs/UnrealBuildTool/Platform/Windows/PVSToolChain.cs
+++ b/Engine/Source/Programs/UnrealBuildTool/Platform/Windows/PVSToolChain.cs
@@ -535,7 +535,7 @@ namespace UnrealBuildTool
 
 		static Version GetAnalyzerVersion(FileReference AnalyzerPath)
 		{
-			String Output = String.Empty;
+			string Output = String.Empty;
 			Version? AnalyzerVersion = new Version(0, 0);
 
 			try
@@ -553,7 +553,7 @@ namespace UnrealBuildTool
 					PvsProc.WaitForExit();
 				}
 
-				const String VersionPattern = @"\d+(?:\.\d+)+";
+				const string VersionPattern = @"\d+(?:\.\d+)+";
 				Match Match = Regex.Match(Output, VersionPattern);
 
 				if (Match.Success)
@@ -624,7 +624,7 @@ namespace UnrealBuildTool
 			return cppCfgStandard;
 		}
 
-		public static bool ShouldCompileAsC(String compilerCommandLine, String sourceFileName)
+		public static bool ShouldCompileAsC(string compilerCommandLine, string sourceFileName)
 		{
 			int CFlagLastPosition = Math.Max(Math.Max(compilerCommandLine.LastIndexOf("/TC "), compilerCommandLine.LastIndexOf("/Tc ")),
 											 Math.Max(compilerCommandLine.LastIndexOf("-TC "), compilerCommandLine.LastIndexOf("-Tc ")));
