UnrealBuildTool: Cleanup XmlConfig

#rnx

diff --git a/Engine/Source/Programs/UnrealBuildTool/System/XmlConfig.cs b/Engine/Source/Programs/UnrealBuildTool/System/XmlConfig.cs
--- a/Engine/Source/Programs/UnrealBuildTool/System/XmlConfig.cs
+++ b/Engine/Source/Programs/UnrealBuildTool/System/XmlConfig.cs
@@ -29,22 +29,22 @@ namespace UnrealBuildTool
 			/// <summary>
 			/// Location of the file
 			/// </summary>
-			public FileReference Location;
+			public FileReference Location { get; init; }
 
 			/// <summary>
 			/// Which folder to display the config file under in the generated project files
 			/// </summary>
-			public string FolderName;
+			public string FolderName { get; init; }
 
 			/// <summary>
 			/// Constructor
 			/// </summary>
-			/// <param name="Location"></param>
-			/// <param name="FolderName"></param>
-			public InputFile(FileReference Location, string FolderName)
+			/// <param name="location"></param>
+			/// <param name="folderName"></param>
+			public InputFile(FileReference location, string folderName)
 			{
-				this.Location = Location;
-				this.FolderName = FolderName;
+				Location = location;
+				FolderName = folderName;
 			}
 		}
 
@@ -56,108 +56,108 @@ namespace UnrealBuildTool
 		/// <summary>
 		/// Parsed config values
 		/// </summary>
-		static XmlConfigData? Values;
+		static XmlConfigData? s_values;
 
 		/// <summary>
 		/// Cached serializer for the XML schema
 		/// </summary>
-		static XmlSerializer? CachedSchemaSerializer;
+		static XmlSerializer? s_cachedSchemaSerializer;
 
 		/// <summary>
 		/// Initialize the config system with the given types
 		/// </summary>
-		/// <param name="OverrideCacheFile">Force use of the cached XML config without checking if it's valid (useful for remote builds)</param>
-		/// <param name="ProjectRootDirectory">Read XML configuration with a Project directory</param>
-		/// <param name="Logger">Logger for output</param>
-		public static void ReadConfigFiles(FileReference? OverrideCacheFile, DirectoryReference? ProjectRootDirectory, ILogger Logger)
+		/// <param name="overrideCacheFile">Force use of the cached XML config without checking if it's valid (useful for remote builds)</param>
+		/// <param name="projectRootDirectory">Read XML configuration with a Project directory</param>
+		/// <param name="logger">Logger for output</param>
+		public static void ReadConfigFiles(FileReference? overrideCacheFile, DirectoryReference? projectRootDirectory, ILogger logger)
 		{
 			// Find all the configurable types
-			List<Type> ConfigTypes = FindConfigurableTypes();
+			List<Type> configTypes = FindConfigurableTypes();
 
 			// Update the cache if necessary
-			if (OverrideCacheFile != null)
+			if (overrideCacheFile != null)
 			{
 				// Set the cache file to the overriden value
-				CacheFile = OverrideCacheFile;
+				CacheFile = overrideCacheFile;
 
 				// Never rebuild the cache; just try to load it.
-				if (!XmlConfigData.TryRead(CacheFile, ConfigTypes, out Values))
+				if (!XmlConfigData.TryRead(CacheFile, configTypes, out s_values))
 				{
 					throw new BuildException("Unable to load XML config cache ({0})", CacheFile);
 				}
 			}
 			else
 			{
-				if (ProjectRootDirectory == null)
+				if (projectRootDirectory == null)
 				{
 					// Get the default cache file
 					CacheFile = FileReference.Combine(Unreal.EngineDirectory, "Intermediate", "Build", "XmlConfigCache.bin");
 					if (Unreal.IsEngineInstalled())
 					{
-						DirectoryReference? UserSettingsDir = Unreal.UserSettingDirectory;
-						if (UserSettingsDir != null)
+						DirectoryReference? userSettingsDir = Unreal.UserSettingDirectory;
+						if (userSettingsDir != null)
 						{
-							CacheFile = FileReference.Combine(UserSettingsDir, "UnrealEngine", String.Format("XmlConfigCache-{0}.bin", Unreal.RootDirectory.FullName.Replace(":", "").Replace(Path.DirectorySeparatorChar, '+')));
+							CacheFile = FileReference.Combine(userSettingsDir, "UnrealEngine", $"XmlConfigCache-{Unreal.RootDirectory.FullName.Replace(":", "", StringComparison.OrdinalIgnoreCase).Replace(Path.DirectorySeparatorChar, '+')}.bin");
 						}
 					}
 				}
 				else
 				{
-					CacheFile = FileReference.Combine(ProjectRootDirectory, "Intermediate", "Build", "XmlConfigCache.bin");
-					Values = null;
+					CacheFile = FileReference.Combine(projectRootDirectory, "Intermediate", "Build", "XmlConfigCache.bin");
+					s_values = null;
 				}
 
 				// Find all the input files
-				List<FileReference> TemporaryInputFileLocations = InputFiles.Select(x => x.Location).ToList();
+				List<FileReference> temporaryInputFileLocations = InputFiles.Select(x => x.Location).ToList();
 
-				if (ProjectRootDirectory != null)
+				if (projectRootDirectory != null)
 				{
-					FileReference ProjectRootConfigLocation = FileReference.Combine(ProjectRootDirectory, "Saved", "UnrealBuildTool", "BuildConfiguration.xml");
-					if (!FileReference.Exists(ProjectRootConfigLocation))
+					FileReference projectRootConfigLocation = FileReference.Combine(projectRootDirectory, "Saved", "UnrealBuildTool", "BuildConfiguration.xml");
+					if (!FileReference.Exists(projectRootConfigLocation))
 					{
-						CreateDefaultConfigFile(ProjectRootConfigLocation);
+						CreateDefaultConfigFile(projectRootConfigLocation);
 					}
 
-					TemporaryInputFileLocations.Add(ProjectRootConfigLocation);
+					temporaryInputFileLocations.Add(projectRootConfigLocation);
 				}
 
-				FileReference[] InputFileLocations = TemporaryInputFileLocations.ToArray();
+				FileReference[] inputFileLocations = temporaryInputFileLocations.ToArray();
 
 				// Get the path to the schema
-				FileReference SchemaFile = GetSchemaLocation(ProjectRootDirectory);
+				FileReference schemaFile = GetSchemaLocation(projectRootDirectory);
 
 				// Try to read the existing cache from disk
-				XmlConfigData? CachedValues;
-				if (IsCacheUpToDate(CacheFile, InputFileLocations) && FileReference.Exists(SchemaFile))
+				XmlConfigData? cachedValues;
+				if (IsCacheUpToDate(CacheFile, inputFileLocations) && FileReference.Exists(schemaFile))
 				{
-					if (XmlConfigData.TryRead(CacheFile, ConfigTypes, out CachedValues) && Enumerable.SequenceEqual(InputFileLocations, CachedValues.InputFiles))
+					if (XmlConfigData.TryRead(CacheFile, configTypes, out cachedValues) && Enumerable.SequenceEqual(inputFileLocations, cachedValues.InputFiles))
 					{
-						Values = CachedValues;
+						s_values = cachedValues;
 					}
 				}
 
 				// If that failed, regenerate it
-				if (Values == null)
+				if (s_values == null)
 				{
 					// Find all the configurable fields from the given types
-					Dictionary<string, Dictionary<string, XmlConfigData.TargetMember>> CategoryToFields = new Dictionary<string, Dictionary<string, XmlConfigData.TargetMember>>();
-					FindConfigurableFields(ConfigTypes, CategoryToFields);
+					Dictionary<string, Dictionary<string, XmlConfigData.TargetMember>> categoryToFields = new Dictionary<string, Dictionary<string, XmlConfigData.TargetMember>>();
+					FindConfigurableFields(configTypes, categoryToFields);
 
 					// Create a schema for the config files
-					XmlSchema Schema = CreateSchema(CategoryToFields);
+					XmlSchema schema = CreateSchema(categoryToFields);
 					if (!Unreal.IsEngineInstalled())
 					{
-						WriteSchema(Schema, SchemaFile);
+						WriteSchema(schema, schemaFile);
 					}
 
 					// Read all the XML files and validate them against the schema
-					Dictionary<Type, Dictionary<XmlConfigData.TargetMember, XmlConfigData.ValueInfo>> TypeToValues =
+					Dictionary<Type, Dictionary<XmlConfigData.TargetMember, XmlConfigData.ValueInfo>> typeToValues =
 						new Dictionary<Type, Dictionary<XmlConfigData.TargetMember, XmlConfigData.ValueInfo>>();
-					foreach (FileReference InputFile in InputFileLocations)
+					foreach (FileReference inputFile in inputFileLocations)
 					{
-						if (!TryReadFile(InputFile, CategoryToFields, TypeToValues, Schema, Logger))
+						if (!TryReadFile(inputFile, categoryToFields, typeToValues, schema, logger))
 						{
-							throw new BuildException("Failed to properly read XML file : {0}", InputFile.FullName);
+							throw new BuildException("Failed to properly read XML file : {0}", inputFile.FullName);
 						}
 					}
 
@@ -165,22 +165,22 @@ namespace UnrealBuildTool
 					DirectoryReference.CreateDirectory(CacheFile.Directory);
 
 					// Create the new cache
-					Values = new XmlConfigData(InputFileLocations, TypeToValues.ToDictionary(
+					s_values = new XmlConfigData(inputFileLocations, typeToValues.ToDictionary(
 							x => x.Key,
 							x => x.Value.Select(x => x.Value).ToArray()));
-					Values.Write(CacheFile);
+					s_values.Write(CacheFile);
 				}
 			}
 
 			// Apply all the static field values
-			foreach (KeyValuePair<Type, XmlConfigData.ValueInfo[]> TypeValuesPair in Values.TypeToValues)
+			foreach (KeyValuePair<Type, XmlConfigData.ValueInfo[]> typeValuesPair in s_values.TypeToValues)
 			{
-				foreach (XmlConfigData.ValueInfo MemberValue in TypeValuesPair.Value)
+				foreach (XmlConfigData.ValueInfo memberValue in typeValuesPair.Value)
 				{
-					if (MemberValue.Target.IsStatic)
+					if (memberValue.Target.IsStatic)
 					{
-						object Value = InstanceValue(MemberValue.Value, MemberValue.Target.Type);
-						MemberValue.Target.SetValue(null, Value);
+						object value = InstanceValue(memberValue.Value, memberValue.Target.Type);
+						memberValue.Target.SetValue(null, value);
 					}
 				}
 			}
@@ -192,47 +192,47 @@ namespace UnrealBuildTool
 		/// <returns>List of configurable types</returns>
 		static List<Type> FindConfigurableTypes()
 		{
-			List<Type> ConfigTypes = new List<Type>();
+			List<Type> configTypes = new List<Type>();
 			try
 			{
-				foreach (Type ConfigType in Assembly.GetExecutingAssembly().GetTypes())
+				foreach (Type configType in Assembly.GetExecutingAssembly().GetTypes())
 				{
-					if (HasXmlConfigFileAttribute(ConfigType))
+					if (HasXmlConfigFileAttribute(configType))
 					{
-						ConfigTypes.Add(ConfigType);
+						configTypes.Add(configType);
 					}
 				}
 			}
-			catch (ReflectionTypeLoadException Ex)
+			catch (ReflectionTypeLoadException ex)
 			{
-				Console.WriteLine("TypeLoadException: {0}", String.Join("\n", Ex.LoaderExceptions.Select(x => x?.Message)));
+				Console.WriteLine("TypeLoadException: {0}", String.Join("\n", ex.LoaderExceptions.Select(x => x?.Message)));
 				throw;
 			}
-			return ConfigTypes;
+			return configTypes;
 		}
 
 		/// <summary>
 		/// Determines whether the given type has a field with an XmlConfigFile attribute
 		/// </summary>
-		/// <param name="Type">The type to check</param>
+		/// <param name="type">The type to check</param>
 		/// <returns>True if the type has a field with the XmlConfigFile attribute</returns>
-		static bool HasXmlConfigFileAttribute(Type Type)
+		static bool HasXmlConfigFileAttribute(Type type)
 		{
-			foreach (FieldInfo Field in Type.GetFields(BindingFlags.Instance | BindingFlags.Static | BindingFlags.Public | BindingFlags.NonPublic))
+			foreach (FieldInfo field in type.GetFields(BindingFlags.Instance | BindingFlags.Static | BindingFlags.Public | BindingFlags.NonPublic))
 			{
-				foreach (CustomAttributeData CustomAttribute in Field.CustomAttributes)
+				foreach (CustomAttributeData customAttribute in field.CustomAttributes)
 				{
-					if (CustomAttribute.AttributeType == typeof(XmlConfigFileAttribute))
+					if (customAttribute.AttributeType == typeof(XmlConfigFileAttribute))
 					{
 						return true;
 					}
 				}
 			}
-			foreach (PropertyInfo Property in Type.GetProperties(BindingFlags.Instance | BindingFlags.Static | BindingFlags.Public | BindingFlags.NonPublic))
+			foreach (PropertyInfo property in type.GetProperties(BindingFlags.Instance | BindingFlags.Static | BindingFlags.Public | BindingFlags.NonPublic))
 			{
-				foreach (CustomAttributeData CustomAttribute in Property.CustomAttributes)
+				foreach (CustomAttributeData customAttribute in property.CustomAttributes)
 				{
-					if (CustomAttribute.AttributeType == typeof(XmlConfigFileAttribute))
+					if (customAttribute.AttributeType == typeof(XmlConfigFileAttribute))
 					{
 						return true;
 					}
@@ -244,22 +244,22 @@ namespace UnrealBuildTool
 		/// <summary>
 		/// Find the location of the XML config schema
 		/// </summary>
-		/// <param name="ProjectRootDirecory">Optional project root directory</param>
+		/// <param name="projectRootDirecory">Optional project root directory</param>
 		/// <returns>The location of the schema file</returns>
-		public static FileReference GetSchemaLocation(DirectoryReference? ProjectRootDirecory = null)
+		public static FileReference GetSchemaLocation(DirectoryReference? projectRootDirecory = null)
 		{
-			if (ProjectRootDirecory != null)
+			if (projectRootDirecory != null)
 			{
-				FileReference ProjectSchema = FileReference.Combine(ProjectRootDirecory, "Saved", "UnrealBuildTool", "BuildConfiguration.Schema.xsd");
-				if (FileReference.Exists(ProjectSchema))
+				FileReference projectSchema = FileReference.Combine(projectRootDirecory, "Saved", "UnrealBuildTool", "BuildConfiguration.Schema.xsd");
+				if (FileReference.Exists(projectSchema))
 				{
-					return ProjectSchema;
+					return projectSchema;
 				}
 			}
 			return FileReference.Combine(Unreal.EngineDirectory, "Saved", "UnrealBuildTool", "BuildConfiguration.Schema.xsd");
 		}
 
-		static InputFile[]? CachedInputFiles;
+		static InputFile[]? s_cachedInputFiles;
 
 		/// <summary>
 		/// Initialize the list of input files
@@ -268,15 +268,15 @@ namespace UnrealBuildTool
 		{
 			get
 			{
-				if (CachedInputFiles != null)
+				if (s_cachedInputFiles != null)
 				{
-					return CachedInputFiles;
+					return s_cachedInputFiles;
 				}
 
-				ILogger Logger = Log.Logger;
+				ILogger logger = Log.Logger;
 
 				// Find all the input file locations
-				List<InputFile> InputFilesFound = new List<InputFile>(5);
+				List<InputFile> inputFilesFound = new List<InputFile>(5);
 
 				// InputFile info and if a default file should be created if missing
 				List<KeyValuePair<InputFile, bool>> configs = new();
@@ -292,158 +292,152 @@ namespace UnrealBuildTool
 				}
 
 				// Check for the global config file under ProgramData/Unreal Engine/UnrealBuildTool
-				DirectoryReference? CommonProgramsFolder = DirectoryReference.FromString(Environment.GetFolderPath(Environment.SpecialFolder.CommonApplicationData));
-				if (CommonProgramsFolder != null)
+				DirectoryReference? commonProgramsFolder = DirectoryReference.FromString(Environment.GetFolderPath(Environment.SpecialFolder.CommonApplicationData));
+				if (commonProgramsFolder != null)
 				{
-					configs.Add(new(new InputFile(FileReference.Combine(CommonProgramsFolder, "Unreal Engine", "UnrealBuildTool", "BuildConfiguration.xml"), "Global (ProgramData)"), false));
+					configs.Add(new(new InputFile(FileReference.Combine(commonProgramsFolder, "Unreal Engine", "UnrealBuildTool", "BuildConfiguration.xml"), "Global (ProgramData)"), false));
 				}
 
 				// Check for the global config file under AppData/Unreal Engine/UnrealBuildTool (Roaming)
-				DirectoryReference? AppDataFolder = DirectoryReference.FromString(Environment.GetFolderPath(Environment.SpecialFolder.ApplicationData));
-				if (AppDataFolder != null)
+				DirectoryReference? appDataFolder = DirectoryReference.FromString(Environment.GetFolderPath(Environment.SpecialFolder.ApplicationData));
+				if (appDataFolder != null)
 				{
-					configs.Add(new(new InputFile(FileReference.Combine(AppDataFolder, "Unreal Engine", "UnrealBuildTool", "BuildConfiguration.xml"), "Global (AppData)"), true));
+					configs.Add(new(new InputFile(FileReference.Combine(appDataFolder, "Unreal Engine", "UnrealBuildTool", "BuildConfiguration.xml"), "Global (AppData)"), true));
 				}
 
 				// Check for the global config file under LocalAppData/Unreal Engine/UnrealBuildTool
-				DirectoryReference? LocalAppDataFolder = DirectoryReference.FromString(Environment.GetFolderPath(Environment.SpecialFolder.LocalApplicationData));
-				if (LocalAppDataFolder != null)
+				DirectoryReference? localAppDataFolder = DirectoryReference.FromString(Environment.GetFolderPath(Environment.SpecialFolder.LocalApplicationData));
+				if (localAppDataFolder != null)
 				{
-					configs.Add(new(new InputFile(FileReference.Combine(LocalAppDataFolder, "Unreal Engine", "UnrealBuildTool", "BuildConfiguration.xml"), "Global (LocalAppData)"), false));
+					configs.Add(new(new InputFile(FileReference.Combine(localAppDataFolder, "Unreal Engine", "UnrealBuildTool", "BuildConfiguration.xml"), "Global (LocalAppData)"), false));
 				}
 
 				// Check for the global config file under My Documents/Unreal Engine/UnrealBuildTool
-				DirectoryReference? PersonalFolder = DirectoryReference.FromString(Environment.GetFolderPath(Environment.SpecialFolder.Personal));
-				if (PersonalFolder != null)
+				DirectoryReference? personalFolder = DirectoryReference.FromString(Environment.GetFolderPath(Environment.SpecialFolder.Personal));
+				if (personalFolder != null)
 				{
-					configs.Add(new(new InputFile(FileReference.Combine(PersonalFolder, "Unreal Engine", "UnrealBuildTool", "BuildConfiguration.xml"), "Global (Documents)"), false));
+					configs.Add(new(new InputFile(FileReference.Combine(personalFolder, "Unreal Engine", "UnrealBuildTool", "BuildConfiguration.xml"), "Global (Documents)"), false));
 				}
 
 				foreach (KeyValuePair<InputFile, bool> config in configs)
 				{
 					if (config.Value && !FileReference.Exists(config.Key.Location))
 					{
-						Logger.LogDebug("Creating default config file at {ConfigLocation}", config.Key.Location);
+						logger.LogDebug("Creating default config file at {ConfigLocation}", config.Key.Location);
 						CreateDefaultConfigFile(config.Key.Location);
 					}
 					if (FileReference.Exists(config.Key.Location))
 					{
-						InputFilesFound.Add(config.Key);
+						inputFilesFound.Add(config.Key);
 					}
 					else
 					{
-						Logger.LogDebug("No config file at {ConfigLocation}", config.Key.Location);
+						logger.LogDebug("No config file at {ConfigLocation}", config.Key.Location);
 					}
 				}
 
-				CachedInputFiles = InputFilesFound.ToArray();
+				s_cachedInputFiles = inputFilesFound.ToArray();
 
-				Logger.LogDebug("Configuration will be read from:");
-				foreach (InputFile InputFile in InputFiles)
+				logger.LogDebug("Configuration will be read from:");
+				foreach (InputFile inputFile in InputFiles)
 				{
-					Logger.LogDebug("  {File}", InputFile.Location.FullName);
+					logger.LogDebug("  {File}", inputFile.Location.FullName);
 				}
 
-				return CachedInputFiles;
+				return s_cachedInputFiles;
 			}
 		}
 
 		/// <summary>
 		/// Create a default config file at the given location
 		/// </summary>
-		/// <param name="Location">Location to read from</param>
-		static void CreateDefaultConfigFile(FileReference Location)
+		/// <param name="location">Location to read from</param>
+		static void CreateDefaultConfigFile(FileReference location)
 		{
-			DirectoryReference.CreateDirectory(Location.Directory);
-			using (StreamWriter Writer = new StreamWriter(Location.FullName))
+			DirectoryReference.CreateDirectory(location.Directory);
+			using (StreamWriter writer = new StreamWriter(location.FullName))
 			{
-				Writer.WriteLine("<?xml version=\"1.0\" encoding=\"utf-8\" ?>");
-				Writer.WriteLine("<Configuration xmlns=\"{0}\">", XmlConfigFile.SchemaNamespaceURI);
-				Writer.WriteLine("</Configuration>");
+				writer.WriteLine("<?xml version=\"1.0\" encoding=\"utf-8\" ?>");
+				writer.WriteLine("<Configuration xmlns=\"{0}\">", XmlConfigFile.SchemaNamespaceURI);
+				writer.WriteLine("</Configuration>");
 			}
 		}
 
 		/// <summary>
 		/// Applies config values to the given object
 		/// </summary>
-		/// <param name="TargetObject">The object instance to be configured</param>
-		public static void ApplyTo(object TargetObject)
+		/// <param name="targetObject">The object instance to be configured</param>
+		public static void ApplyTo(object targetObject)
 		{
-			ILogger Logger = Log.Logger;
-			for (Type? TargetType = TargetObject.GetType(); TargetType != null; TargetType = TargetType.BaseType)
+			ILogger logger = Log.Logger;
+			for (Type? targetType = targetObject.GetType(); targetType != null; targetType = targetType.BaseType)
 			{
-				XmlConfigData.ValueInfo[]? FieldValues;
-				if (Values!.TypeToValues.TryGetValue(TargetType, out FieldValues))
+				XmlConfigData.ValueInfo[]? fieldValues;
+				if (s_values!.TypeToValues.TryGetValue(targetType, out fieldValues))
 				{
-					foreach (XmlConfigData.ValueInfo FieldValue in FieldValues)
+					foreach (XmlConfigData.ValueInfo fieldValue in fieldValues)
 					{
-						if (!FieldValue.Target.IsStatic)
+						if (!fieldValue.Target.IsStatic)
 						{
-							XmlConfigData.TargetMember TargetToWrite = FieldValue.Target;
+							XmlConfigData.TargetMember targetToWrite = fieldValue.Target;
 
 							// Check if setting has been deprecated
-							if (FieldValue.XmlConfigAttribute.Deprecated)
+							if (fieldValue.XmlConfigAttribute.Deprecated)
 							{
-								string CurrentSettingName = FieldValue.XmlConfigAttribute.Name != null ? FieldValue.XmlConfigAttribute.Name : FieldValue.Target.MemberInfo.Name;
+								string currentSettingName = fieldValue.XmlConfigAttribute.Name ?? fieldValue.Target.MemberInfo.Name;
 
-								Logger.LogWarning("Deprecated setting found in \"{SourceFile}\":", FieldValue.SourceFile);
-								Logger.LogWarning("The setting \"{Setting}\" is deprecated. Support for this setting will be removed in a future version of Unreal Engine.", CurrentSettingName);
+								logger.LogWarning("Deprecated setting found in \"{SourceFile}\":", fieldValue.SourceFile);
+								logger.LogWarning("The setting \"{Setting}\" is deprecated. Support for this setting will be removed in a future version of Unreal Engine.", currentSettingName);
 
-								if (FieldValue.XmlConfigAttribute.NewAttributeName != null)
+								if (fieldValue.XmlConfigAttribute.NewAttributeName != null)
 								{
 									// NewAttributeName is the name of a member in code. However, the log messages below are written from the XML's perspective,
 									// so we need to check if the new target member is not exposed under a custom name in the config.
-									string NewSettingName = GetMemberConfigAttributeName(TargetType, FieldValue.XmlConfigAttribute.NewAttributeName);
+									string newSettingName = GetMemberConfigAttributeName(targetType, fieldValue.XmlConfigAttribute.NewAttributeName);
 
-									Logger.LogWarning("Use \"{NewAttributeName}\" in place of \"{OldAttributeName}\"", NewSettingName, CurrentSettingName);
-									Logger.LogInformation("The value provided for deprecated setting \"{OldName}\" will be applied to \"{NewName}\"", CurrentSettingName, NewSettingName);
+									logger.LogWarning("Use \"{NewAttributeName}\" in place of \"{OldAttributeName}\"", newSettingName, currentSettingName);
+									logger.LogInformation("The value provided for deprecated setting \"{OldName}\" will be applied to \"{NewName}\"", currentSettingName, newSettingName);
 
-									TargetToWrite = GetTargetMember(TargetType, FieldValue.XmlConfigAttribute.NewAttributeName) ?? TargetToWrite;
+									targetToWrite = GetTargetMember(targetType, fieldValue.XmlConfigAttribute.NewAttributeName) ?? targetToWrite;
 								}
 							}
 
-							object ValueInstance = InstanceValue(FieldValue.Value, FieldValue.Target.Type);
-							TargetToWrite.SetValue(TargetObject, ValueInstance);
+							object valueInstance = InstanceValue(fieldValue.Value, fieldValue.Target.Type);
+							targetToWrite.SetValue(targetObject, valueInstance);
 						}
 					}
 				}
 			}
 		}
 
-		private static string GetMemberConfigAttributeName(Type TargetType, string MemberName)
+		private static string GetMemberConfigAttributeName(Type targetType, string memberName)
 		{
-			MemberInfo? MemberInfo = TargetType.GetRuntimeFields()
-				.FirstOrDefault(x => x.Name == MemberName);
+			MemberInfo? memberInfo = targetType.GetRuntimeFields().FirstOrDefault(x => x.Name == memberName) as MemberInfo
+				?? targetType.GetRuntimeProperties().FirstOrDefault(x => x.Name == memberName) as MemberInfo;
 
-			if (MemberInfo == null)
-			{
-				MemberInfo = TargetType.GetRuntimeProperties()
-					.FirstOrDefault(x => x.Name == MemberName);
-			}
-
-			XmlConfigFileAttribute? Attribute = MemberInfo?.GetCustomAttributes<XmlConfigFileAttribute>().FirstOrDefault();
+			XmlConfigFileAttribute? attribute = memberInfo?.GetCustomAttributes<XmlConfigFileAttribute>().FirstOrDefault();
 
-			return Attribute?.Name ?? MemberName;
+			return attribute?.Name ?? memberName;
 		}
 
-		private static XmlConfigData.TargetMember? GetTargetMember(Type TargetType, string MemberName)
+		private static XmlConfigData.TargetMember? GetTargetMember(Type targetType, string memberName)
 		{
 			// First, try to find the new field to which the setting should be actually applied.
-			FieldInfo? FieldInfo = TargetType.GetRuntimeFields()
-				.FirstOrDefault(x => x.Name == MemberName);
+			FieldInfo? fieldInfo = targetType.GetRuntimeFields()
+				.FirstOrDefault(x => x.Name == memberName);
 
-			if (FieldInfo != null)
+			if (fieldInfo != null)
 			{
-				return new XmlConfigData.TargetField(FieldInfo);
+				return new XmlConfigData.TargetField(fieldInfo);
 			}
 
 			// If not found, try to find the new property to which the setting should be actually applied.
-			PropertyInfo? PropertyInfo = TargetType.GetRuntimeProperties()
-				.FirstOrDefault(x => x.Name == MemberName);
+			PropertyInfo? propertyInfo = targetType.GetRuntimeProperties()
+				.FirstOrDefault(x => x.Name == memberName);
 
-			if (PropertyInfo != null)
+			if (propertyInfo != null)
 			{
-				return new XmlConfigData.TargetProperty(PropertyInfo);
+				return new XmlConfigData.TargetProperty(propertyInfo);
 			}
 
 			return null;
@@ -452,272 +446,269 @@ namespace UnrealBuildTool
 		/// <summary>
 		/// Instances a value for assignment to a target object
 		/// </summary>
-		/// <param name="Value">The value to instance</param>
-		/// <param name="ValueType">The type of value</param>
+		/// <param name="value">The value to instance</param>
+		/// <param name="valueType">The type of value</param>
 		/// <returns>New instance of the given value, if necessary</returns>
-		static object InstanceValue(object Value, Type ValueType)
+		static object InstanceValue(object value, Type valueType)
 		{
-			if (ValueType == typeof(string[]))
+			if (valueType == typeof(string[]))
 			{
-				return ((string[])Value).Clone();
+				return ((string[])value).Clone();
 			}
 			else
 			{
-				return Value;
+				return value;
 			}
 		}
 
 		/// <summary>
 		/// Gets a config value for a single value, without writing it to an instance of that class
 		/// </summary>
-		/// <param name="TargetType">Type to find config values for</param>
-		/// <param name="Name">Name of the field to receive</param>
-		/// <param name="Value">On success, receives the value of the field</param>
+		/// <param name="targetType">Type to find config values for</param>
+		/// <param name="name">Name of the field to receive</param>
+		/// <param name="value">On success, receives the value of the field</param>
 		/// <returns>True if the value was read, false otherwise</returns>
-		public static bool TryGetValue(Type TargetType, string Name, [NotNullWhen(true)] out object? Value)
+		public static bool TryGetValue(Type targetType, string name, [NotNullWhen(true)] out object? value)
 		{
 			// Find all the config values for this type
-			XmlConfigData.ValueInfo[]? FieldValues;
-			if (!Values!.TypeToValues.TryGetValue(TargetType, out FieldValues))
+			XmlConfigData.ValueInfo[]? fieldValues;
+			if (!s_values!.TypeToValues.TryGetValue(targetType, out fieldValues))
 			{
-				Value = null;
+				value = null;
 				return false;
 			}
 
 			// Find the value with the matching name
-			foreach (XmlConfigData.ValueInfo FieldValue in FieldValues)
+			foreach (XmlConfigData.ValueInfo fieldValue in fieldValues)
 			{
-				if (FieldValue.Target.MemberInfo.Name == Name)
+				if (fieldValue.Target.MemberInfo.Name == name)
 				{
-					Value = FieldValue.Value;
+					value = fieldValue.Value;
 					return true;
 				}
 			}
 
 			// Not found
-			Value = null;
+			value = null;
 			return false;
 		}
 
 		/// <summary>
 		/// Find all the configurable fields in the given types by searching for XmlConfigFile attributes.
 		/// </summary>
-		/// <param name="ConfigTypes">Array of types to search</param>
-		/// <param name="CategoryToFields">Dictionaries populated with category -> name -> field mappings on return</param>
-		static void FindConfigurableFields(IEnumerable<Type> ConfigTypes, Dictionary<string, Dictionary<string, XmlConfigData.TargetMember>> CategoryToFields)
+		/// <param name="configTypes">Array of types to search</param>
+		/// <param name="categoryToFields">Dictionaries populated with category -> name -> field mappings on return</param>
+		static void FindConfigurableFields(IEnumerable<Type> configTypes, Dictionary<string, Dictionary<string, XmlConfigData.TargetMember>> categoryToFields)
 		{
-			foreach (Type ConfigType in ConfigTypes)
+			foreach (Type configType in configTypes)
 			{
-				foreach (FieldInfo FieldInfo in ConfigType.GetFields(BindingFlags.Instance | BindingFlags.Static | BindingFlags.GetField | BindingFlags.Public | BindingFlags.NonPublic))
+				foreach (FieldInfo fieldInfo in configType.GetFields(BindingFlags.Instance | BindingFlags.Static | BindingFlags.GetField | BindingFlags.Public | BindingFlags.NonPublic))
 				{
-					ProcessConfigurableMember<FieldInfo>(ConfigType, FieldInfo, CategoryToFields, FieldInfo => new XmlConfigData.TargetField(FieldInfo));
+					ProcessConfigurableMember<FieldInfo>(configType, fieldInfo, categoryToFields, fieldInfo => new XmlConfigData.TargetField(fieldInfo));
 				}
-				foreach (PropertyInfo PropertyInfo in ConfigType.GetProperties(BindingFlags.Instance | BindingFlags.Static | BindingFlags.GetField | BindingFlags.Public | BindingFlags.NonPublic))
+				foreach (PropertyInfo propertyInfo in configType.GetProperties(BindingFlags.Instance | BindingFlags.Static | BindingFlags.GetField | BindingFlags.Public | BindingFlags.NonPublic))
 				{
-					ProcessConfigurableMember<PropertyInfo>(ConfigType, PropertyInfo, CategoryToFields, PropertyInfo => new XmlConfigData.TargetProperty(PropertyInfo));
+					ProcessConfigurableMember<PropertyInfo>(configType, propertyInfo, categoryToFields, propertyInfo => new XmlConfigData.TargetProperty(propertyInfo));
 				}
 			}
 		}
 
-		private static void ProcessConfigurableMember<MEMBER>(Type Type, MEMBER MemberInfo, Dictionary<string, Dictionary<string, XmlConfigData.TargetMember>> CategoryToFields, Func<MEMBER, XmlConfigData.TargetMember> CreateTarget)
+		private static void ProcessConfigurableMember<MEMBER>(Type type, MEMBER memberInfo, Dictionary<string, Dictionary<string, XmlConfigData.TargetMember>> categoryToFields, Func<MEMBER, XmlConfigData.TargetMember> createTarget)
 			where MEMBER : System.Reflection.MemberInfo
 		{
-			IEnumerable<XmlConfigFileAttribute> Attributes = MemberInfo.GetCustomAttributes<XmlConfigFileAttribute>();
-			foreach (XmlConfigFileAttribute Attribute in Attributes)
+			IEnumerable<XmlConfigFileAttribute> attributes = memberInfo.GetCustomAttributes<XmlConfigFileAttribute>();
+			foreach (XmlConfigFileAttribute attribute in attributes)
 			{
-				string CategoryName = Attribute.Category ?? Type.Name;
+				string categoryName = attribute.Category ?? type.Name;
 
-				Dictionary<string, XmlConfigData.TargetMember>? NameToTarget;
-				if (!CategoryToFields.TryGetValue(CategoryName, out NameToTarget))
+				Dictionary<string, XmlConfigData.TargetMember>? nameToTarget;
+				if (!categoryToFields.TryGetValue(categoryName, out nameToTarget))
 				{
-					NameToTarget = new Dictionary<string, XmlConfigData.TargetMember>();
-					CategoryToFields.Add(CategoryName, NameToTarget);
+					nameToTarget = new Dictionary<string, XmlConfigData.TargetMember>();
+					categoryToFields.Add(categoryName, nameToTarget);
 				}
 
-				NameToTarget[Attribute.Name ?? MemberInfo.Name] = CreateTarget(MemberInfo);
+				nameToTarget[attribute.Name ?? memberInfo.Name] = createTarget(memberInfo);
 			}
 		}
 
 		/// <summary>
 		/// Creates a schema from attributes in the given types
 		/// </summary>
-		/// <param name="CategoryToFields">Lookup for all field settings</param>
+		/// <param name="categoryToFields">Lookup for all field settings</param>
 		/// <returns>New schema instance</returns>
-		static XmlSchema CreateSchema(Dictionary<string, Dictionary<string, XmlConfigData.TargetMember>> CategoryToFields)
+		static XmlSchema CreateSchema(Dictionary<string, Dictionary<string, XmlConfigData.TargetMember>> categoryToFields)
 		{
 			// Create elements for all the categories
-			XmlSchemaAll RootAll = new XmlSchemaAll();
-			foreach (KeyValuePair<string, Dictionary<string, XmlConfigData.TargetMember>> CategoryPair in CategoryToFields)
+			XmlSchemaAll rootAll = new XmlSchemaAll();
+			foreach (KeyValuePair<string, Dictionary<string, XmlConfigData.TargetMember>> categoryPair in categoryToFields)
 			{
-				string CategoryName = CategoryPair.Key;
+				string categoryName = categoryPair.Key;
 
-				XmlSchemaAll CategoryAll = new XmlSchemaAll();
-				foreach (KeyValuePair<string, XmlConfigData.TargetMember> FieldPair in CategoryPair.Value)
+				XmlSchemaAll categoryAll = new XmlSchemaAll();
+				foreach (KeyValuePair<string, XmlConfigData.TargetMember> fieldPair in categoryPair.Value)
 				{
-					XmlSchemaElement Element = CreateSchemaFieldElement(FieldPair.Key, FieldPair.Value.Type);
-					CategoryAll.Items.Add(Element);
+					XmlSchemaElement element = CreateSchemaFieldElement(fieldPair.Key, fieldPair.Value.Type);
+					categoryAll.Items.Add(element);
 				}
 
-				XmlSchemaComplexType CategoryType = new XmlSchemaComplexType();
-				CategoryType.Particle = CategoryAll;
+				XmlSchemaComplexType categoryType = new XmlSchemaComplexType();
+				categoryType.Particle = categoryAll;
 
-				XmlSchemaElement CategoryElement = new XmlSchemaElement();
-				CategoryElement.Name = CategoryName;
-				CategoryElement.SchemaType = CategoryType;
-				CategoryElement.MinOccurs = 0;
-				CategoryElement.MaxOccurs = 1;
+				XmlSchemaElement categoryElement = new XmlSchemaElement();
+				categoryElement.Name = categoryName;
+				categoryElement.SchemaType = categoryType;
+				categoryElement.MinOccurs = 0;
+				categoryElement.MaxOccurs = 1;
 
-				RootAll.Items.Add(CategoryElement);
+				rootAll.Items.Add(categoryElement);
 			}
 
 			// Create the root element and schema object
-			XmlSchemaComplexType RootType = new XmlSchemaComplexType();
-			RootType.Particle = RootAll;
+			XmlSchemaComplexType rootType = new XmlSchemaComplexType();
+			rootType.Particle = rootAll;
 
-			XmlSchemaElement RootElement = new XmlSchemaElement();
-			RootElement.Name = XmlConfigFile.RootElementName;
-			RootElement.SchemaType = RootType;
+			XmlSchemaElement rootElement = new XmlSchemaElement();
+			rootElement.Name = XmlConfigFile.RootElementName;
+			rootElement.SchemaType = rootType;
 
-			XmlSchema Schema = new XmlSchema();
-			Schema.TargetNamespace = XmlConfigFile.SchemaNamespaceURI;
-			Schema.ElementFormDefault = XmlSchemaForm.Qualified;
-			Schema.Items.Add(RootElement);
+			XmlSchema schema = new XmlSchema();
+			schema.TargetNamespace = XmlConfigFile.SchemaNamespaceURI;
+			schema.ElementFormDefault = XmlSchemaForm.Qualified;
+			schema.Items.Add(rootElement);
 
 			// Finally compile it
-			XmlSchemaSet SchemaSet = new XmlSchemaSet();
-			SchemaSet.Add(Schema);
-			SchemaSet.Compile();
-			return SchemaSet.Schemas().OfType<XmlSchema>().First();
+			XmlSchemaSet schemaSet = new XmlSchemaSet();
+			schemaSet.Add(schema);
+			schemaSet.Compile();
+			return schemaSet.Schemas().OfType<XmlSchema>().First();
 		}
 
 		/// <summary>
 		/// Creates an XML schema element for reading a value of the given type
 		/// </summary>
-		/// <param name="Name">Name of the field</param>
-		/// <param name="Type">Type of the field</param>
+		/// <param name="name">Name of the field</param>
+		/// <param name="type">Type of the field</param>
 		/// <returns>New schema element representing the field</returns>
-		static XmlSchemaElement CreateSchemaFieldElement(string Name, Type Type)
+		static XmlSchemaElement CreateSchemaFieldElement(string name, Type type)
 		{
-			XmlSchemaElement Element = new XmlSchemaElement();
-			Element.Name = Name;
-			Element.MinOccurs = 0;
-			Element.MaxOccurs = 1;
+			XmlSchemaElement element = new XmlSchemaElement();
+			element.Name = name;
+			element.MinOccurs = 0;
+			element.MaxOccurs = 1;
 
-			if (TryGetNullableStructType(Type, out Type? InnerType))
+			if (TryGetNullableStructType(type, out Type? innerType))
 			{
-				Type = InnerType;
+				type = innerType;
 			}
 
-			if (Type == typeof(string))
+			if (type == typeof(string))
 			{
-				Element.SchemaTypeName = XmlSchemaType.GetBuiltInSimpleType(XmlTypeCode.String).QualifiedName;
+				element.SchemaTypeName = XmlSchemaType.GetBuiltInSimpleType(XmlTypeCode.String).QualifiedName;
 			}
-			else if (Type == typeof(bool))
+			else if (type == typeof(bool))
 			{
-				Element.SchemaTypeName = XmlSchemaType.GetBuiltInSimpleType(XmlTypeCode.Boolean).QualifiedName;
+				element.SchemaTypeName = XmlSchemaType.GetBuiltInSimpleType(XmlTypeCode.Boolean).QualifiedName;
 			}
-			else if (Type == typeof(int))
+			else if (type == typeof(int))
 			{
-				Element.SchemaTypeName = XmlSchemaType.GetBuiltInSimpleType(XmlTypeCode.Int).QualifiedName;
+				element.SchemaTypeName = XmlSchemaType.GetBuiltInSimpleType(XmlTypeCode.Int).QualifiedName;
 			}
-			else if (Type == typeof(float))
+			else if (type == typeof(float))
 			{
-				Element.SchemaTypeName = XmlSchemaType.GetBuiltInSimpleType(XmlTypeCode.Float).QualifiedName;
+				element.SchemaTypeName = XmlSchemaType.GetBuiltInSimpleType(XmlTypeCode.Float).QualifiedName;
 			}
-			else if (Type == typeof(double))
+			else if (type == typeof(double))
 			{
-				Element.SchemaTypeName = XmlSchemaType.GetBuiltInSimpleType(XmlTypeCode.Double).QualifiedName;
+				element.SchemaTypeName = XmlSchemaType.GetBuiltInSimpleType(XmlTypeCode.Double).QualifiedName;
 			}
-			else if (Type == typeof(FileReference))
+			else if (type == typeof(FileReference))
 			{
-				Element.SchemaTypeName = XmlSchemaType.GetBuiltInSimpleType(XmlTypeCode.String).QualifiedName;
+				element.SchemaTypeName = XmlSchemaType.GetBuiltInSimpleType(XmlTypeCode.String).QualifiedName;
 			}
-			else if (Type.IsEnum)
+			else if (type.IsEnum)
 			{
-				XmlSchemaSimpleTypeRestriction Restriction = new XmlSchemaSimpleTypeRestriction();
-				Restriction.BaseTypeName = XmlSchemaType.GetBuiltInSimpleType(XmlTypeCode.String).QualifiedName;
+				XmlSchemaSimpleTypeRestriction restriction = new XmlSchemaSimpleTypeRestriction();
+				restriction.BaseTypeName = XmlSchemaType.GetBuiltInSimpleType(XmlTypeCode.String).QualifiedName;
 
-				foreach (string EnumName in Enum.GetNames(Type))
+				foreach (string enumName in Enum.GetNames(type))
 				{
-					XmlSchemaEnumerationFacet Facet = new XmlSchemaEnumerationFacet();
-					Facet.Value = EnumName;
-					Restriction.Facets.Add(Facet);
+					XmlSchemaEnumerationFacet facet = new XmlSchemaEnumerationFacet();
+					facet.Value = enumName;
+					restriction.Facets.Add(facet);
 				}
 
-				XmlSchemaSimpleType EnumType = new XmlSchemaSimpleType();
-				EnumType.Content = Restriction;
-				Element.SchemaType = EnumType;
+				XmlSchemaSimpleType enumType = new XmlSchemaSimpleType();
+				enumType.Content = restriction;
+				element.SchemaType = enumType;
 			}
-			else if (Type == typeof(string[]))
+			else if (type == typeof(string[]))
 			{
-				XmlSchemaElement ItemElement = new XmlSchemaElement();
-				ItemElement.Name = "Item";
-				ItemElement.SchemaTypeName = XmlSchemaType.GetBuiltInSimpleType(XmlTypeCode.String).QualifiedName;
-				ItemElement.MinOccurs = 0;
-				ItemElement.MaxOccursString = "unbounded";
-
-				XmlSchemaSequence Sequence = new XmlSchemaSequence();
-				Sequence.Items.Add(ItemElement);
-
-				XmlSchemaComplexType ArrayType = new XmlSchemaComplexType();
-				ArrayType.Particle = Sequence;
-				Element.SchemaType = ArrayType;
+				XmlSchemaElement itemElement = new XmlSchemaElement();
+				itemElement.Name = "Item";
+				itemElement.SchemaTypeName = XmlSchemaType.GetBuiltInSimpleType(XmlTypeCode.String).QualifiedName;
+				itemElement.MinOccurs = 0;
+				itemElement.MaxOccursString = "unbounded";
+
+				XmlSchemaSequence sequence = new XmlSchemaSequence();
+				sequence.Items.Add(itemElement);
+
+				XmlSchemaComplexType arrayType = new XmlSchemaComplexType();
+				arrayType.Particle = sequence;
+				element.SchemaType = arrayType;
 			}
 			else
 			{
 				throw new Exception("Unsupported field type for XmlConfigFile attribute");
 			}
-			return Element;
+			return element;
 		}
 
 		/// <summary>
 		/// Writes a schema to the given location. Avoids writing it if the file is identical.
 		/// </summary>
-		/// <param name="Schema">The schema to be written</param>
-		/// <param name="Location">Location to write to</param>
-		static void WriteSchema(XmlSchema Schema, FileReference Location)
+		/// <param name="schema">The schema to be written</param>
+		/// <param name="location">Location to write to</param>
+		static void WriteSchema(XmlSchema schema, FileReference location)
 		{
-			XmlWriterSettings Settings = new XmlWriterSettings();
-			Settings.Indent = true;
-			Settings.IndentChars = "\t";
-			Settings.NewLineChars = Environment.NewLine;
-			Settings.OmitXmlDeclaration = true;
+			XmlWriterSettings settings = new XmlWriterSettings();
+			settings.Indent = true;
+			settings.IndentChars = "\t";
+			settings.NewLineChars = Environment.NewLine;
+			settings.OmitXmlDeclaration = true;
 
-			if (CachedSchemaSerializer == null)
-			{
-				CachedSchemaSerializer = XmlSerializer.FromTypes(new Type[] { typeof(XmlSchema) })[0]!;
-			}
+			s_cachedSchemaSerializer ??= XmlSerializer.FromTypes(new Type[] { typeof(XmlSchema) })[0]!;
 
-			StringBuilder Output = new StringBuilder();
-			Output.AppendLine("<?xml version=\"1.0\" encoding=\"utf-8\"?>");
-			using (XmlWriter Writer = XmlWriter.Create(Output, Settings))
+			StringBuilder output = new StringBuilder();
+			output.AppendLine("<?xml version=\"1.0\" encoding=\"utf-8\"?>");
+			using (XmlWriter writer = XmlWriter.Create(output, settings))
 			{
-				XmlSerializerNamespaces Namespaces = new XmlSerializerNamespaces();
-				Namespaces.Add("", "http://www.w3.org/2001/XMLSchema");
-				CachedSchemaSerializer.Serialize(Writer, Schema, Namespaces);
+				XmlSerializerNamespaces namespaces = new XmlSerializerNamespaces();
+				namespaces.Add("", "http://www.w3.org/2001/XMLSchema");
+				s_cachedSchemaSerializer.Serialize(writer, schema, namespaces);
 			}
 
-			string OutputText = Output.ToString();
-			if (!FileReference.Exists(Location) || File.ReadAllText(Location.FullName) != OutputText)
+			string outputText = output.ToString();
+			if (!FileReference.Exists(location) || File.ReadAllText(location.FullName) != outputText)
 			{
-				DirectoryReference.CreateDirectory(Location.Directory);
-				File.WriteAllText(Location.FullName, OutputText);
+				DirectoryReference.CreateDirectory(location.Directory);
+				File.WriteAllText(location.FullName, outputText);
 			}
 		}
 
 		/// <summary>
 		/// Tests whether a type is a nullable struct, and extracts the inner type if it is
 		/// </summary>
-		static bool TryGetNullableStructType(Type Type, [NotNullWhen(true)] out Type? InnerType)
+		static bool TryGetNullableStructType(Type type, [NotNullWhen(true)] out Type? innerType)
 		{
-			if (Type.IsGenericType && Type.GetGenericTypeDefinition() == typeof(Nullable<>))
+			if (type.IsGenericType && type.GetGenericTypeDefinition() == typeof(Nullable<>))
 			{
-				InnerType = Type.GetGenericArguments()[0];
+				innerType = type.GetGenericArguments()[0];
 				return true;
 			}
 			else
 			{
-				InnerType = null;
+				innerType = null;
 				return false;
 			}
 		}
@@ -725,60 +716,59 @@ namespace UnrealBuildTool
 		/// <summary>
 		/// Reads an XML config file and merges it to the given cache
 		/// </summary>
-		/// <param name="Location">Location to read from</param>
-		/// <param name="CategoryToFields">Lookup for configurable fields by category</param>
-		/// <param name="TypeToValues">Map of types to fields and their associated values</param>
-		/// <param name="Schema">Schema to validate against</param>
-		/// <param name="Logger">Logger for output</param>
+		/// <param name="location">Location to read from</param>
+		/// <param name="categoryToFields">Lookup for configurable fields by category</param>
+		/// <param name="typeToValues">Map of types to fields and their associated values</param>
+		/// <param name="schema">Schema to validate against</param>
+		/// <param name="logger">Logger for output</param>
 		/// <returns>True if the file was read successfully</returns>
-		static bool TryReadFile(FileReference Location, Dictionary<string, Dictionary<string, XmlConfigData.TargetMember>> CategoryToFields,
-			Dictionary<Type, Dictionary<XmlConfigData.TargetMember, XmlConfigData.ValueInfo>> TypeToValues, XmlSchema Schema, ILogger Logger)
+		static bool TryReadFile(FileReference location, Dictionary<string, Dictionary<string, XmlConfigData.TargetMember>> categoryToFields,
+			Dictionary<Type, Dictionary<XmlConfigData.TargetMember, XmlConfigData.ValueInfo>> typeToValues, XmlSchema schema, ILogger logger)
 		{
 			// Read the XML file, and validate it against the schema
-			XmlConfigFile? ConfigFile;
-			if (!XmlConfigFile.TryRead(Location, Schema, Logger, out ConfigFile))
+			XmlConfigFile? configFile;
+			if (!XmlConfigFile.TryRead(location, schema, logger, out configFile))
 			{
 				return false;
 			}
 
 			// Parse the document
-			foreach (XmlElement CategoryElement in ConfigFile.DocumentElement!.ChildNodes.OfType<XmlElement>())
+			foreach (XmlElement categoryElement in configFile.DocumentElement!.ChildNodes.OfType<XmlElement>())
 			{
-				Dictionary<string, XmlConfigData.TargetMember>? NameToField;
-				if (CategoryToFields.TryGetValue(CategoryElement.Name, out NameToField))
+				Dictionary<string, XmlConfigData.TargetMember>? nameToField;
+				if (categoryToFields.TryGetValue(categoryElement.Name, out nameToField))
 				{
-					foreach (XmlElement KeyElement in CategoryElement.ChildNodes.OfType<XmlElement>())
+					foreach (XmlElement keyElement in categoryElement.ChildNodes.OfType<XmlElement>())
 					{
-						XmlConfigData.TargetMember? Field;
-						object Value;
-						if (NameToField.TryGetValue(KeyElement.Name, out Field))
+						if (nameToField.TryGetValue(keyElement.Name, out XmlConfigData.TargetMember? field))
 						{
-							if (Field.Type == typeof(string[]))
+							object value;
+							if (field.Type == typeof(string[]))
 							{
-								Value = KeyElement.ChildNodes.OfType<XmlElement>().Where(x => x.Name == "Item").Select(x => x.InnerText).ToArray();
+								value = keyElement.ChildNodes.OfType<XmlElement>().Where(x => x.Name == "Item").Select(x => x.InnerText).ToArray();
 							}
-							else if (TryGetNullableStructType(Field.Type, out Type? StructType))
+							else if (TryGetNullableStructType(field.Type, out Type? structType))
 							{
-								Value = ParseValue(StructType, KeyElement.InnerText);
+								value = ParseValue(structType, keyElement.InnerText);
 							}
 							else
 							{
-								Value = ParseValue(Field.Type, KeyElement.InnerText);
+								value = ParseValue(field.Type, keyElement.InnerText);
 							}
 
 							// Add it to the set of values for the type containing this field
-							Dictionary<XmlConfigData.TargetMember, XmlConfigData.ValueInfo>? FieldToValue;
-							if (!TypeToValues.TryGetValue(Field.MemberInfo.DeclaringType!, out FieldToValue))
+							Dictionary<XmlConfigData.TargetMember, XmlConfigData.ValueInfo>? fieldToValue;
+							if (!typeToValues.TryGetValue(field.MemberInfo.DeclaringType!, out fieldToValue))
 							{
-								FieldToValue = new Dictionary<XmlConfigData.TargetMember, XmlConfigData.ValueInfo>();
-								TypeToValues.Add(Field.MemberInfo.DeclaringType!, FieldToValue);
+								fieldToValue = new Dictionary<XmlConfigData.TargetMember, XmlConfigData.ValueInfo>();
+								typeToValues.Add(field.MemberInfo.DeclaringType!, fieldToValue);
 							}
 
 							// Parse the corresponding value
-							XmlConfigData.ValueInfo FieldValue = new XmlConfigData.ValueInfo(Field, Value, Location,
-								Field.MemberInfo.GetCustomAttribute<XmlConfigFileAttribute>()!);
+							XmlConfigData.ValueInfo fieldValue = new XmlConfigData.ValueInfo(field, value, location,
+								field.MemberInfo.GetCustomAttribute<XmlConfigFileAttribute>()!);
 
-							FieldToValue[Field] = FieldValue;
+							fieldToValue[field] = fieldValue;
 						}
 					}
 				}
@@ -789,55 +779,55 @@ namespace UnrealBuildTool
 		/// <summary>
 		/// Parse the value for a field from its text based representation in an XML file
 		/// </summary>
-		/// <param name="FieldType">The type of field being read</param>
-		/// <param name="Text">Text to parse</param>
+		/// <param name="fieldType">The type of field being read</param>
+		/// <param name="text">Text to parse</param>
 		/// <returns>The object that was parsed</returns>
-		static object ParseValue(Type FieldType, string Text)
+		static object ParseValue(Type fieldType, string text)
 		{
 			// ignore whitespace in all fields except for Strings which we leave unprocessed
-			string TrimmedText = Text.Trim();
-			if (FieldType == typeof(string))
+			string trimmedText = text.Trim();
+			if (fieldType == typeof(string))
 			{
-				return Text;
+				return text;
 			}
-			else if (FieldType == typeof(bool) || FieldType == typeof(bool?))
+			else if (fieldType == typeof(bool) || fieldType == typeof(bool?))
 			{
-				if (TrimmedText == "1" || TrimmedText.Equals("true", StringComparison.InvariantCultureIgnoreCase))
+				if (trimmedText == "1" || trimmedText.Equals("true", StringComparison.OrdinalIgnoreCase))
 				{
 					return true;
 				}
-				else if (TrimmedText == "0" || TrimmedText.Equals("false", StringComparison.InvariantCultureIgnoreCase))
+				else if (trimmedText == "0" || trimmedText.Equals("false", StringComparison.OrdinalIgnoreCase))
 				{
 					return false;
 				}
 				else
 				{
-					throw new Exception(String.Format("Unable to convert '{0}' to boolean. 'true/false/0/1' are the supported formats.", Text));
+					throw new Exception(String.Format("Unable to convert '{0}' to boolean. 'true/false/0/1' are the supported formats.", text));
 				}
 			}
-			else if (FieldType == typeof(int))
+			else if (fieldType == typeof(int))
 			{
-				return Int32.Parse(TrimmedText);
+				return Int32.Parse(trimmedText);
 			}
-			else if (FieldType == typeof(float))
+			else if (fieldType == typeof(float))
 			{
-				return Single.Parse(TrimmedText, System.Globalization.CultureInfo.InvariantCulture);
+				return Single.Parse(trimmedText, System.Globalization.CultureInfo.InvariantCulture);
 			}
-			else if (FieldType == typeof(double))
+			else if (fieldType == typeof(double))
 			{
-				return Double.Parse(TrimmedText, System.Globalization.CultureInfo.InvariantCulture);
+				return Double.Parse(trimmedText, System.Globalization.CultureInfo.InvariantCulture);
 			}
-			else if (FieldType.IsEnum)
+			else if (fieldType.IsEnum)
 			{
-				return Enum.Parse(FieldType, TrimmedText);
+				return Enum.Parse(fieldType, trimmedText);
 			}
-			else if (FieldType == typeof(FileReference))
+			else if (fieldType == typeof(FileReference))
 			{
-				return FileReference.FromString(Text);
+				return FileReference.FromString(text);
 			}
 			else
 			{
-				throw new Exception(String.Format("Unsupported config type '{0}'", FieldType.Name));
+				throw new Exception(String.Format("Unsupported config type '{0}'", fieldType.Name));
 			}
 		}
 
@@ -845,30 +835,30 @@ namespace UnrealBuildTool
 		/// Checks that the given cache file exists and is newer than the given input files, and attempts to read it. Verifies that the resulting cache was created
 		/// from the same input files in the same order.
 		/// </summary>
-		/// <param name="CacheFile">Path to the config cache file</param>
-		/// <param name="InputFiles">The expected set of input files in the cache</param>
+		/// <param name="cacheFile">Path to the config cache file</param>
+		/// <param name="inputFiles">The expected set of input files in the cache</param>
 		/// <returns>True if the cache was valid and could be read, false otherwise.</returns>
-		static bool IsCacheUpToDate(FileReference CacheFile, FileReference[] InputFiles)
+		static bool IsCacheUpToDate(FileReference cacheFile, FileReference[] inputFiles)
 		{
 			// Always rebuild if the cache doesn't exist
-			if (!FileReference.Exists(CacheFile))
+			if (!FileReference.Exists(cacheFile))
 			{
 				return false;
 			}
 
 			// Get the timestamp for the cache
-			DateTime CacheWriteTime = File.GetLastWriteTimeUtc(CacheFile.FullName);
+			DateTime cacheWriteTime = File.GetLastWriteTimeUtc(cacheFile.FullName);
 
 			// Always rebuild if this executable is newer
-			if (File.GetLastWriteTimeUtc(Assembly.GetExecutingAssembly().Location) > CacheWriteTime)
+			if (File.GetLastWriteTimeUtc(Assembly.GetExecutingAssembly().Location) > cacheWriteTime)
 			{
 				return false;
 			}
 
 			// Check if any of the input files are newer than the cache
-			foreach (FileReference InputFile in InputFiles)
+			foreach (FileReference inputFile in inputFiles)
 			{
-				if (File.GetLastWriteTimeUtc(InputFile.FullName) > CacheWriteTime)
+				if (File.GetLastWriteTimeUtc(inputFile.FullName) > cacheWriteTime)
 				{
 					return false;
 				}
@@ -881,105 +871,105 @@ namespace UnrealBuildTool
 		/// <summary>
 		/// Generates documentation files for the available settings, by merging the XML documentation from the compiler.
 		/// </summary>
-		/// <param name="OutputFile">The documentation file to write</param>
-		/// <param name="Logger">Logger for output</param>
-		public static void WriteDocumentation(FileReference OutputFile, ILogger Logger)
+		/// <param name="outputFile">The documentation file to write</param>
+		/// <param name="logger">Logger for output</param>
+		public static void WriteDocumentation(FileReference outputFile, ILogger logger)
 		{
 			// Find all the configurable types
-			List<Type> ConfigTypes = FindConfigurableTypes();
+			List<Type> configTypes = FindConfigurableTypes();
 
 			// Find all the configurable fields from the given types
-			Dictionary<string, Dictionary<string, XmlConfigData.TargetMember>> CategoryToFields = new Dictionary<string, Dictionary<string, XmlConfigData.TargetMember>>();
-			FindConfigurableFields(ConfigTypes, CategoryToFields);
-			CategoryToFields = CategoryToFields.Where(x => x.Value.Count > 0).ToDictionary(x => x.Key, x => x.Value);
+			Dictionary<string, Dictionary<string, XmlConfigData.TargetMember>> categoryToFields = new Dictionary<string, Dictionary<string, XmlConfigData.TargetMember>>();
+			FindConfigurableFields(configTypes, categoryToFields);
+			categoryToFields = categoryToFields.Where(x => x.Value.Count > 0).ToDictionary(x => x.Key, x => x.Value);
 
 			// Get the path to the XML documentation
-			FileReference InputDocumentationFile = new FileReference(Assembly.GetExecutingAssembly().Location).ChangeExtension(".xml");
-			if (!FileReference.Exists(InputDocumentationFile))
+			FileReference inputDocumentationFile = new FileReference(Assembly.GetExecutingAssembly().Location).ChangeExtension(".xml");
+			if (!FileReference.Exists(inputDocumentationFile))
 			{
-				throw new BuildException("Generated assembly documentation not found at {0}.", InputDocumentationFile);
+				throw new BuildException("Generated assembly documentation not found at {0}.", inputDocumentationFile);
 			}
 
 			// Read the documentation
-			XmlDocument InputDocumentation = new XmlDocument();
-			InputDocumentation.Load(InputDocumentationFile.FullName);
+			XmlDocument inputDocumentation = new XmlDocument();
+			inputDocumentation.Load(inputDocumentationFile.FullName);
 
 			// Make sure we can write to the output file
-			if (FileReference.Exists(OutputFile))
+			if (FileReference.Exists(outputFile))
 			{
-				FileReference.MakeWriteable(OutputFile);
+				FileReference.MakeWriteable(outputFile);
 			}
 			else
 			{
-				DirectoryReference.CreateDirectory(OutputFile.Directory);
+				DirectoryReference.CreateDirectory(outputFile.Directory);
 			}
 
 			// Generate the documentation file
-			if (OutputFile.HasExtension(".udn"))
+			if (outputFile.HasExtension(".udn"))
 			{
-				WriteDocumentationUDN(OutputFile, InputDocumentation, CategoryToFields, Logger);
+				WriteDocumentationUDN(outputFile, inputDocumentation, categoryToFields, logger);
 			}
-			else if (OutputFile.HasExtension(".html"))
+			else if (outputFile.HasExtension(".html"))
 			{
-				WriteDocumentationHTML(OutputFile, InputDocumentation, CategoryToFields, Logger);
+				WriteDocumentationHTML(outputFile, inputDocumentation, categoryToFields, logger);
 			}
 			else
 			{
-				throw new BuildException("Unable to detect format from extension of output file ({0})", OutputFile);
+				throw new BuildException("Unable to detect format from extension of output file ({0})", outputFile);
 			}
 
 			// Success!
-			Logger.LogInformation("Written documentation to {OutputFile}.", OutputFile);
+			logger.LogInformation("Written documentation to {OutputFile}.", outputFile);
 		}
 
 		/// <summary>
 		/// Writes out documentation in UDN format
 		/// </summary>
-		/// <param name="OutputFile">The output file</param>
-		/// <param name="InputDocumentation">The XML documentation for this assembly</param>
-		/// <param name="CategoryToFields">Map of string to types to fields</param>
-		/// <param name="Logger">Logger for output</param>
-		private static void WriteDocumentationUDN(FileReference OutputFile, XmlDocument InputDocumentation, Dictionary<string, Dictionary<string, XmlConfigData.TargetMember>> CategoryToFields, ILogger Logger)
+		/// <param name="outputFile">The output file</param>
+		/// <param name="inputDocumentation">The XML documentation for this assembly</param>
+		/// <param name="categoryToFields">Map of string to types to fields</param>
+		/// <param name="logger">Logger for output</param>
+		private static void WriteDocumentationUDN(FileReference outputFile, XmlDocument inputDocumentation, Dictionary<string, Dictionary<string, XmlConfigData.TargetMember>> categoryToFields, ILogger logger)
 		{
-			using (StreamWriter Writer = new StreamWriter(OutputFile.FullName))
+			using (StreamWriter writer = new StreamWriter(outputFile.FullName))
 			{
-				Writer.WriteLine("Availability: NoPublish");
-				Writer.WriteLine("Title: Build Configuration Properties Page");
-				Writer.WriteLine("Crumbs:");
-				Writer.WriteLine("Description: This is a procedurally generated markdown page.");
-				Writer.WriteLine("Version: {0}.{1}", ReadOnlyBuildVersion.Current.MajorVersion, ReadOnlyBuildVersion.Current.MinorVersion);
-				Writer.WriteLine("");
-
-				foreach (KeyValuePair<string, Dictionary<string, XmlConfigData.TargetMember>> CategoryPair in CategoryToFields)
+				writer.WriteLine("Availability: NoPublish");
+				writer.WriteLine("Title: Build Configuration Properties Page");
+				writer.WriteLine("Crumbs:");
+				writer.WriteLine("Description: This is a procedurally generated markdown page.");
+				writer.WriteLine("Version: {0}.{1}", ReadOnlyBuildVersion.Current.MajorVersion, ReadOnlyBuildVersion.Current.MinorVersion);
+				writer.WriteLine("");
+
+				foreach (KeyValuePair<string, Dictionary<string, XmlConfigData.TargetMember>> categoryPair in categoryToFields)
 				{
-					string CategoryName = CategoryPair.Key;
-					Writer.WriteLine("### {0}", CategoryName);
-					Writer.WriteLine();
+					string categoryName = categoryPair.Key;
+					writer.WriteLine("### {0}", categoryName);
+					writer.WriteLine();
 
-					Dictionary<string, XmlConfigData.TargetMember> Fields = CategoryPair.Value;
-					foreach (KeyValuePair<string, XmlConfigData.TargetMember> FieldPair in Fields)
+					Dictionary<string, XmlConfigData.TargetMember> fields = categoryPair.Value;
+					foreach (KeyValuePair<string, XmlConfigData.TargetMember> fieldPair in fields)
 					{
 						// Get the XML comment for this field
-						List<string>? Lines;
-						if (!RulesDocumentation.TryGetXmlComment(InputDocumentation, FieldPair.Value.MemberInfo, Logger, out Lines) || Lines.Count == 0)
+						List<string>? lines;
+						if (!RulesDocumentation.TryGetXmlComment(inputDocumentation, fieldPair.Value.MemberInfo, logger, out lines) || lines.Count == 0)
 						{
 							continue;
 						}
 
 						// Write the result to the .udn file
-						Writer.WriteLine("$ {0} : {1}", FieldPair.Key, Lines[0]);
-						for (int Idx = 1; Idx < Lines.Count; Idx++)
+						writer.WriteLine("$ {0} : {1}", fieldPair.Key, lines[0]);
+						for (int idx = 1; idx < lines.Count; idx++)
 						{
-							if (Lines[Idx].StartsWith("*") || Lines[Idx].StartsWith("-"))
+							if (lines[idx].StartsWith("*", StringComparison.OrdinalIgnoreCase) || lines[idx].StartsWith("-", StringComparison.OrdinalIgnoreCase))
 							{
-								Writer.WriteLine("        * {0}", Lines[Idx].Substring(1).TrimStart());
+								writer.WriteLine("        * {0}", lines[idx].Substring(1).TrimStart());
 							}
 							else
 							{
-								Writer.WriteLine("    * {0}", Lines[Idx]);
+								writer.WriteLine("    * {0}", lines[idx]);
 							}
 						}
-						Writer.WriteLine();
+						writer.WriteLine();
 					}
 				}
 			}
@@ -988,67 +978,67 @@ namespace UnrealBuildTool
 		/// <summary>
 		/// Writes out documentation in HTML format
 		/// </summary>
-		/// <param name="OutputFile">The output file</param>
-		/// <param name="InputDocumentation">The XML documentation for this assembly</param>
-		/// <param name="CategoryToFields">Map of string to types to fields</param>
-		/// <param name="Logger">Logger for output</param>
-		private static void WriteDocumentationHTML(FileReference OutputFile, XmlDocument InputDocumentation, Dictionary<string, Dictionary<string, XmlConfigData.TargetMember>> CategoryToFields, ILogger Logger)
+		/// <param name="outputFile">The output file</param>
+		/// <param name="inputDocumentation">The XML documentation for this assembly</param>
+		/// <param name="categoryToFields">Map of string to types to fields</param>
+		/// <param name="logger">Logger for output</param>
+		private static void WriteDocumentationHTML(FileReference outputFile, XmlDocument inputDocumentation, Dictionary<string, Dictionary<string, XmlConfigData.TargetMember>> categoryToFields, ILogger logger)
 		{
-			using (StreamWriter Writer = new StreamWriter(OutputFile.FullName))
+			using (StreamWriter writer = new StreamWriter(outputFile.FullName))
 			{
-				Writer.WriteLine("<html>");
-				Writer.WriteLine("  <body>");
-				Writer.WriteLine("  <h2>BuildConfiguration Properties</h2>");
-				foreach (KeyValuePair<string, Dictionary<string, XmlConfigData.TargetMember>> CategoryPair in CategoryToFields)
+				writer.WriteLine("<html>");
+				writer.WriteLine("  <body>");
+				writer.WriteLine("  <h2>BuildConfiguration Properties</h2>");
+				foreach (KeyValuePair<string, Dictionary<string, XmlConfigData.TargetMember>> categoryPair in categoryToFields)
 				{
-					string CategoryName = CategoryPair.Key;
-					Writer.WriteLine("    <h3>{0}</h3>", CategoryName);
-					Writer.WriteLine("    <dl>");
+					string categoryName = categoryPair.Key;
+					writer.WriteLine("    <h3>{0}</h3>", categoryName);
+					writer.WriteLine("    <dl>");
 
-					Dictionary<string, XmlConfigData.TargetMember> Fields = CategoryPair.Value;
-					foreach (KeyValuePair<string, XmlConfigData.TargetMember> FieldPair in Fields)
+					Dictionary<string, XmlConfigData.TargetMember> fields = categoryPair.Value;
+					foreach (KeyValuePair<string, XmlConfigData.TargetMember> fieldPair in fields)
 					{
 						// Get the XML comment for this field
-						List<string>? Lines;
-						if (!RulesDocumentation.TryGetXmlComment(InputDocumentation, FieldPair.Value.MemberInfo, Logger, out Lines) || Lines.Count == 0)
+						List<string>? lines;
+						if (!RulesDocumentation.TryGetXmlComment(inputDocumentation, fieldPair.Value.MemberInfo, logger, out lines) || lines.Count == 0)
 						{
 							continue;
 						}
 
 						// Write the result to the .udn file
-						Writer.WriteLine("      <dt>{0}</dt>", FieldPair.Key);
+						writer.WriteLine("      <dt>{0}</dt>", fieldPair.Key);
 
-						if (Lines.Count == 1)
+						if (lines.Count == 1)
 						{
-							Writer.WriteLine("      <dd>{0}</dd>", Lines[0]);
+							writer.WriteLine("      <dd>{0}</dd>", lines[0]);
 						}
 						else
 						{
-							Writer.WriteLine("      <dd>");
-							for (int Idx = 0; Idx < Lines.Count; Idx++)
+							writer.WriteLine("      <dd>");
+							for (int idx = 0; idx < lines.Count; idx++)
 							{
-								if (Lines[Idx].StartsWith("*") || Lines[Idx].StartsWith("-"))
+								if (lines[idx].StartsWith("*", StringComparison.OrdinalIgnoreCase) || lines[idx].StartsWith("-", StringComparison.OrdinalIgnoreCase))
 								{
-									Writer.WriteLine("        <ul>");
-									for (; Idx < Lines.Count && (Lines[Idx].StartsWith("*") || Lines[Idx].StartsWith("-")); Idx++)
+									writer.WriteLine("        <ul>");
+									for (; idx < lines.Count && (lines[idx].StartsWith("*", StringComparison.OrdinalIgnoreCase) || lines[idx].StartsWith("-", StringComparison.OrdinalIgnoreCase)); idx++)
 									{
-										Writer.WriteLine("          <li>{0}</li>", Lines[Idx].Substring(1).TrimStart());
+										writer.WriteLine("          <li>{0}</li>", lines[idx].Substring(1).TrimStart());
 									}
-									Writer.WriteLine("        </ul>");
+									writer.WriteLine("        </ul>");
 								}
 								else
 								{
-									Writer.WriteLine("        {0}", Lines[Idx]);
+									writer.WriteLine("        {0}", lines[idx]);
 								}
 							}
-							Writer.WriteLine("      </dd>");
+							writer.WriteLine("      </dd>");
 						}
 					}
 
-					Writer.WriteLine("    </dl>");
+					writer.WriteLine("    </dl>");
 				}
-				Writer.WriteLine("  </body>");
-				Writer.WriteLine("</html>");
+				writer.WriteLine("  </body>");
+				writer.WriteLine("</html>");
 			}
 		}
 	}
diff --git a/Engine/Source/Programs/UnrealBuildTool/System/XmlConfigData.cs b/Engine/Source/Programs/UnrealBuildTool/System/XmlConfigData.cs
--- a/Engine/Source/Programs/UnrealBuildTool/System/XmlConfigData.cs
+++ b/Engine/Source/Programs/UnrealBuildTool/System/XmlConfigData.cs
@@ -23,7 +23,7 @@ namespace UnrealBuildTool
 		/// <summary>
 		/// List of input files. Stored to allow checking cache validity.
 		/// </summary>
-		public FileReference[] InputFiles;
+		public FileReference[] InputFiles { get; set; }
 
 		/// <summary>
 		/// Abstract description of a target data member.
@@ -61,17 +61,17 @@ namespace UnrealBuildTool
 		/// </summary>
 		public class TargetField : TargetMember
 		{
-			public override MemberInfo MemberInfo => FieldInfo;
-			public override Type Type => FieldInfo.FieldType;
-			public override bool IsStatic => FieldInfo.IsStatic;
-			public override Action<object?, object?> SetValue => FieldInfo.SetValue;
-			public override Func<object?, object?> GetValue => FieldInfo.GetValue;
+			public override MemberInfo MemberInfo => _fieldInfo;
+			public override Type Type => _fieldInfo.FieldType;
+			public override bool IsStatic => _fieldInfo.IsStatic;
+			public override Action<object?, object?> SetValue => _fieldInfo.SetValue;
+			public override Func<object?, object?> GetValue => _fieldInfo.GetValue;
 
-			private FieldInfo FieldInfo;
+			private readonly FieldInfo _fieldInfo;
 
-			public TargetField(FieldInfo FieldInfo)
+			public TargetField(FieldInfo fieldInfo)
 			{
-				this.FieldInfo = FieldInfo;
+				_fieldInfo = fieldInfo;
 			}
 		}
 
@@ -80,147 +80,141 @@ namespace UnrealBuildTool
 		/// </summary>
 		public class TargetProperty : TargetMember
 		{
-			public override MemberInfo MemberInfo => PropertyInfo;
-			public override Type Type => PropertyInfo.PropertyType;
-			public override bool IsStatic => PropertyInfo.GetGetMethod()!.IsStatic;
-			public override Action<object?, object?> SetValue => PropertyInfo.SetValue;
-			public override Func<object?, object?> GetValue => PropertyInfo.GetValue;
+			public override MemberInfo MemberInfo => _propertyInfo;
+			public override Type Type => _propertyInfo.PropertyType;
+			public override bool IsStatic => _propertyInfo.GetGetMethod()!.IsStatic;
+			public override Action<object?, object?> SetValue => _propertyInfo.SetValue;
+			public override Func<object?, object?> GetValue => _propertyInfo.GetValue;
 
-			private PropertyInfo PropertyInfo;
+			private readonly PropertyInfo _propertyInfo;
 
-			public TargetProperty(PropertyInfo PropertyInfo)
+			public TargetProperty(PropertyInfo propertyInfo)
 			{
-				this.PropertyInfo = PropertyInfo;
+				_propertyInfo = propertyInfo;
 			}
 		}
 
 		public class ValueInfo
 		{
-			public TargetMember Target;
-			public object Value;
-			public FileReference SourceFile;
-			public XmlConfigFileAttribute XmlConfigAttribute;
+			public TargetMember Target {  get; init; }
+			public object Value { get; init; }
+			public FileReference SourceFile { get; init; }
+			public XmlConfigFileAttribute XmlConfigAttribute { get; init; }
 
-			public ValueInfo(FieldInfo FieldInfo, object Value, FileReference SourceFile, XmlConfigFileAttribute XmlConfigAttribute)
+			public ValueInfo(FieldInfo fieldInfo, object value, FileReference sourceFile, XmlConfigFileAttribute xmlConfigAttribute)
+				: this(new TargetField(fieldInfo), value, sourceFile, xmlConfigAttribute)
 			{
-				Target = new TargetField(FieldInfo);
-				this.Value = Value;
-				this.SourceFile = SourceFile;
-				this.XmlConfigAttribute = XmlConfigAttribute;
 			}
 
-			public ValueInfo(PropertyInfo PropertyInfo, object Value, FileReference SourceFile, XmlConfigFileAttribute XmlConfigAttribute)
+			public ValueInfo(PropertyInfo propertyInfo, object value, FileReference sourceFile, XmlConfigFileAttribute xmlConfigAttribute)
+				: this(new TargetProperty(propertyInfo), value, sourceFile, xmlConfigAttribute)
 			{
-				Target = new TargetProperty(PropertyInfo);
-				this.Value = Value;
-				this.SourceFile = SourceFile;
-				this.XmlConfigAttribute = XmlConfigAttribute;
 			}
 
-			public ValueInfo(TargetMember Target, object Value, FileReference SourceFile, XmlConfigFileAttribute XmlConfigAttribute)
+			public ValueInfo(TargetMember target, object value, FileReference sourceFile, XmlConfigFileAttribute xmlConfigAttribute)
 			{
-				this.Target = Target;
-				this.Value = Value;
-				this.SourceFile = SourceFile;
-				this.XmlConfigAttribute = XmlConfigAttribute;
+				Target = target;
+				Value = value;
+				SourceFile = sourceFile;
+				XmlConfigAttribute = xmlConfigAttribute;
 			}
 		}
 
 		/// <summary>
 		/// Stores a mapping from type -> member -> value, with all the config values for configurable fields.
 		/// </summary>
-		public Dictionary<Type, ValueInfo[]> TypeToValues;
+		public Dictionary<Type, ValueInfo[]> TypeToValues { get; init; }
 
 		/// <summary>
 		/// Constructor
 		/// </summary>
-		/// <param name="InputFiles"></param>
-		/// <param name="TypeToValues"></param>
-		public XmlConfigData(FileReference[] InputFiles, Dictionary<Type, ValueInfo[]> TypeToValues)
+		/// <param name="inputFiles"></param>
+		/// <param name="typeToValues"></param>
+		public XmlConfigData(FileReference[] inputFiles, Dictionary<Type, ValueInfo[]> typeToValues)
 		{
-			this.InputFiles = InputFiles;
-			this.TypeToValues = TypeToValues;
+			InputFiles = inputFiles;
+			TypeToValues = typeToValues;
 		}
 
 		/// <summary>
 		/// Attempts to read a previous block of config values from disk
 		/// </summary>
-		/// <param name="Location">The file to read from</param>
-		/// <param name="Types">Array of valid types. Used to resolve serialized type names to concrete types.</param>
-		/// <param name="Data">On success, receives the parsed data</param>
+		/// <param name="location">The file to read from</param>
+		/// <param name="types">Array of valid types. Used to resolve serialized type names to concrete types.</param>
+		/// <param name="data">On success, receives the parsed data</param>
 		/// <returns>True if the data was read and is valid</returns>
-		public static bool TryRead(FileReference Location, IEnumerable<Type> Types, [NotNullWhen(true)] out XmlConfigData? Data)
+		public static bool TryRead(FileReference location, IEnumerable<Type> types, [NotNullWhen(true)] out XmlConfigData? data)
 		{
 			// Check the file exists first
-			if (!FileReference.Exists(Location))
+			if (!FileReference.Exists(location))
 			{
-				Data = null;
+				data = null;
 				return false;
 			}
 
 			// Read the cache from disk
-			using (BinaryReader Reader = new BinaryReader(File.Open(Location.FullName, FileMode.Open, FileAccess.Read, FileShare.Read)))
+			using (BinaryReader reader = new BinaryReader(File.Open(location.FullName, FileMode.Open, FileAccess.Read, FileShare.Read)))
 			{
 				// Check the serialization version matches
-				if (Reader.ReadInt32() != SerializationVersion)
+				if (reader.ReadInt32() != SerializationVersion)
 				{
-					Data = null;
+					data = null;
 					return false;
 				}
 
 				// Read the input files
-				FileReference[] InputFiles = Reader.ReadArray(() => Reader.ReadFileReference())!;
+				FileReference[] inputFiles = reader.ReadArray(() => reader.ReadFileReference())!;
 
 				// Read the types
-				int NumTypes = Reader.ReadInt32();
-				Dictionary<Type, ValueInfo[]> TypeToValues = new Dictionary<Type, ValueInfo[]>(NumTypes);
-				for (int TypeIdx = 0; TypeIdx < NumTypes; TypeIdx++)
+				int numTypes = reader.ReadInt32();
+				Dictionary<Type, ValueInfo[]> typeToValues = new Dictionary<Type, ValueInfo[]>(numTypes);
+				for (int typeIdx = 0; typeIdx < numTypes; typeIdx++)
 				{
 					// Read the type name
-					string TypeName = Reader.ReadString();
+					string typeName = reader.ReadString();
 
 					// Try to find it in the list of configurable types
-					Type? Type = Types.FirstOrDefault(x => x.Name == TypeName);
-					if (Type == null)
+					Type? type = types.FirstOrDefault(x => x.Name == typeName);
+					if (type == null)
 					{
-						Data = null;
+						data = null;
 						return false;
 					}
 
 					// Read all the values
-					ValueInfo[] Values = new ValueInfo[Reader.ReadInt32()];
-					for (int ValueIdx = 0; ValueIdx < Values.Length; ValueIdx++)
+					ValueInfo[] values = new ValueInfo[reader.ReadInt32()];
+					for (int valueIdx = 0; valueIdx < values.Length; valueIdx++)
 					{
-						string MemberName = Reader.ReadString();
+						string memberName = reader.ReadString();
 
-						XmlConfigData.TargetMember? TargetMember = GetTargetMemberWithAttribute<XmlConfigFileAttribute>(Type, MemberName);
+						TargetMember? targetMember = GetTargetMemberWithAttribute<XmlConfigFileAttribute>(type, memberName);
 
-						if (TargetMember != null)
+						if (targetMember != null)
 						{
 							// If TargetMember is not null, we know it has our attribute.
-							XmlConfigFileAttribute XmlConfigAttribute = TargetMember!.MemberInfo.GetCustomAttribute<XmlConfigFileAttribute>()!;
+							XmlConfigFileAttribute xmlConfigAttribute = targetMember!.MemberInfo.GetCustomAttribute<XmlConfigFileAttribute>()!;
 
 							// Try to parse the value and add it to the output array
-							object Value = Reader.ReadObject(TargetMember.Type)!;
+							object value = reader.ReadObject(targetMember.Type)!;
 
 							// Read the path of the config file that provided this setting
-							FileReference SourceFile = Reader.ReadFileReference();
+							FileReference sourceFile = reader.ReadFileReference();
 
-							Values[ValueIdx] = new ValueInfo(TargetMember, Value, SourceFile, XmlConfigAttribute);
+							values[valueIdx] = new ValueInfo(targetMember, value, sourceFile, xmlConfigAttribute);
 						}
 						else
 						{
-							Data = null;
+							data = null;
 							return false;
 						}
 					}
 
 					// Add it to the type map
-					TypeToValues.Add(Type, Values);
+					typeToValues.Add(type, values);
 				}
 
 				// Return the parsed data
-				Data = new XmlConfigData(InputFiles.ToArray(), TypeToValues);
+				data = new XmlConfigData(inputFiles.ToArray(), typeToValues);
 				return true;
 			}
 		}
@@ -229,28 +223,25 @@ namespace UnrealBuildTool
 		/// Find a data member (field or property) with the given name and attribute and returns TargetMember wrapper created for it.
 		/// </summary>
 		/// <typeparam name="T">Attribute a member has to have to be considered.</typeparam>
-		/// <param name="Type">Type which members are to be searched</param>
-		/// <param name="MemberName">Name of a member (field or property) to find.</param>
+		/// <param name="type">Type which members are to be searched</param>
+		/// <param name="memberName">Name of a member (field or property) to find.</param>
 		/// <returns>TargetMember wrapper or null if no member has been found.</returns>
-		private static XmlConfigData.TargetMember? GetTargetMemberWithAttribute<T>(Type Type, string MemberName)
-			where T : System.Attribute
+		private static TargetMember? GetTargetMemberWithAttribute<T>(Type type, string memberName)
+			where T : Attribute
 		{
-			T? XmlConfigAttribute = null;
-
-			FieldInfo? Field = Type.GetField(MemberName, BindingFlags.Public | BindingFlags.NonPublic | BindingFlags.Instance | BindingFlags.Static);
-			XmlConfigAttribute = Field?.GetCustomAttribute<T>();
-
-			if (Field != null && XmlConfigAttribute != null)
+			FieldInfo? field = type.GetField(memberName, BindingFlags.Public | BindingFlags.NonPublic | BindingFlags.Instance | BindingFlags.Static);
+			T? xmlConfigAttribute = field?.GetCustomAttribute<T>();
+			if (field != null && xmlConfigAttribute != null)
 			{
-				return new XmlConfigData.TargetField(Field);
+				return new TargetField(field);
 			}
 
-			PropertyInfo? Property = Type.GetProperty(MemberName, BindingFlags.Public | BindingFlags.NonPublic | BindingFlags.Instance | BindingFlags.Static);
-			XmlConfigAttribute = Property?.GetCustomAttribute<T>();
+			PropertyInfo? property = type.GetProperty(memberName, BindingFlags.Public | BindingFlags.NonPublic | BindingFlags.Instance | BindingFlags.Static);
+			xmlConfigAttribute = property?.GetCustomAttribute<T>();
 
-			if (Property != null && XmlConfigAttribute != null)
+			if (property != null && xmlConfigAttribute != null)
 			{
-				return new XmlConfigData.TargetProperty(Property);
+				return new TargetProperty(property);
 			}
 
 			return null;
@@ -259,28 +250,28 @@ namespace UnrealBuildTool
 		/// <summary>
 		/// Writes the coalesced config hierarchy to disk
 		/// </summary>
-		/// <param name="Location">File to write to</param>
-		public void Write(FileReference Location)
+		/// <param name="location">File to write to</param>
+		public void Write(FileReference location)
 		{
-			DirectoryReference.CreateDirectory(Location.Directory);
-			using (BinaryWriter Writer = new BinaryWriter(File.Open(Location.FullName, FileMode.Create, FileAccess.Write, FileShare.Read)))
+			DirectoryReference.CreateDirectory(location.Directory);
+			using (BinaryWriter writer = new BinaryWriter(File.Open(location.FullName, FileMode.Create, FileAccess.Write, FileShare.Read)))
 			{
-				Writer.Write(SerializationVersion);
+				writer.Write(SerializationVersion);
 
 				// Save all the input files. The cache will not be valid if these change.
-				Writer.Write(InputFiles, Item => Writer.Write(Item));
+				writer.Write(InputFiles, item => writer.Write(item));
 
 				// Write all the categories
-				Writer.Write(TypeToValues.Count);
-				foreach (KeyValuePair<Type, ValueInfo[]> TypePair in TypeToValues)
+				writer.Write(TypeToValues.Count);
+				foreach (KeyValuePair<Type, ValueInfo[]> typePair in TypeToValues)
 				{
-					Writer.Write(TypePair.Key.Name);
-					Writer.Write(TypePair.Value.Length);
-					foreach (ValueInfo MemberPair in TypePair.Value)
+					writer.Write(typePair.Key.Name);
+					writer.Write(typePair.Value.Length);
+					foreach (ValueInfo memberPair in typePair.Value)
 					{
-						Writer.Write(MemberPair.Target.MemberInfo.Name);
-						Writer.Write(MemberPair.Target.Type, MemberPair.Value);
-						Writer.Write(MemberPair.SourceFile);
+						writer.Write(memberPair.Target.MemberInfo.Name);
+						writer.Write(memberPair.Target.Type, memberPair.Value);
+						writer.Write(memberPair.SourceFile);
 					}
 				}
 			}
diff --git a/Engine/Source/Programs/UnrealBuildTool/System/XmlConfigFile.cs b/Engine/Source/Programs/UnrealBuildTool/System/XmlConfigFile.cs
--- a/Engine/Source/Programs/UnrealBuildTool/System/XmlConfigFile.cs
+++ b/Engine/Source/Programs/UnrealBuildTool/System/XmlConfigFile.cs
@@ -26,103 +26,97 @@ namespace UnrealBuildTool
 		/// <summary>
 		/// The file being read
 		/// </summary>
-		FileReference File;
+		readonly FileReference _file;
 
 		/// <summary>
 		/// Interface to the LineInfo on the active XmlReader
 		/// </summary>
-		IXmlLineInfo LineInfo = null!;
+		IXmlLineInfo _lineInfo = null!;
 
 		/// <summary>
 		/// Set to true if the reader encounters an error
 		/// </summary>
-		bool bHasErrors;
+		bool _bHasErrors;
 
 		/// <summary>
 		/// Private constructor. Use XmlConfigFile.TryRead to read an XML config file.
 		/// </summary>
-		private XmlConfigFile(FileReference InFile)
+		private XmlConfigFile(FileReference inFile)
 		{
-			File = InFile;
+			_file = inFile;
 		}
 
 		/// <summary>
 		/// Overrides XmlDocument.CreateElement() to construct ScriptElements rather than XmlElements
 		/// </summary>
-		public override XmlElement CreateElement(string? Prefix, string LocalName, string? NamespaceUri)
-		{
-			return new XmlConfigFileElement(File, LineInfo.LineNumber, Prefix!, LocalName, NamespaceUri, this);
-		}
+		public override XmlElement CreateElement(string? prefix, string localName, string? namespaceUri) => new XmlConfigFileElement(_file, _lineInfo.LineNumber, prefix!, localName, namespaceUri, this);
 
 		/// <summary>
 		/// Loads a script document from the given file
 		/// </summary>
-		/// <param name="File">The file to load</param>
-		/// <param name="Schema">The schema to validate against</param>
-		/// <param name="Logger">Logger for output</param>
-		/// <param name="OutConfigFile">If successful, the document that was read</param>
+		/// <param name="file">The file to load</param>
+		/// <param name="schema">The schema to validate against</param>
+		/// <param name="logger">Logger for output</param>
+		/// <param name="outConfigFile">If successful, the document that was read</param>
 		/// <returns>True if the document could be read, false otherwise</returns>
-		public static bool TryRead(FileReference File, XmlSchema Schema, ILogger Logger, [NotNullWhen(true)] out XmlConfigFile? OutConfigFile)
+		public static bool TryRead(FileReference file, XmlSchema schema, ILogger logger, [NotNullWhen(true)] out XmlConfigFile? outConfigFile)
 		{
-			XmlConfigFile ConfigFile = new XmlConfigFile(File);
+			XmlConfigFile configFile = new XmlConfigFile(file);
 
-			XmlReaderSettings Settings = new XmlReaderSettings();
-			Settings.Schemas.Add(Schema);
-			Settings.ValidationType = ValidationType.Schema;
-			Settings.ValidationEventHandler += ConfigFile.ValidationEvent;
+			XmlReaderSettings settings = new XmlReaderSettings();
+			settings.Schemas.Add(schema);
+			settings.ValidationType = ValidationType.Schema;
+			settings.ValidationEventHandler += configFile.ValidationEvent;
 
-			using (XmlReader Reader = XmlReader.Create(File.FullName, Settings))
+			using (XmlReader reader = XmlReader.Create(file.FullName, settings))
 			{
 				// Read the document
-				ConfigFile.LineInfo = (IXmlLineInfo)Reader;
+				configFile._lineInfo = (IXmlLineInfo)reader;
 				try
 				{
-					ConfigFile.Load(Reader);
+					configFile.Load(reader);
 				}
-				catch (XmlException Ex)
+				catch (XmlException ex)
 				{
-					if (!ConfigFile.bHasErrors)
+					if (!configFile._bHasErrors)
 					{
-						Log.TraceErrorTask(File, Ex.LineNumber, "{0}", Ex.Message);
-						ConfigFile.bHasErrors = true;
+						Log.TraceErrorTask(file, ex.LineNumber, "{0}", ex.Message);
+						configFile._bHasErrors = true;
 					}
 				}
 
 				// If we hit any errors while parsing
-				if (ConfigFile.bHasErrors)
+				if (configFile._bHasErrors)
 				{
-					OutConfigFile = null;
+					outConfigFile = null;
 					return false;
 				}
 
 				// Check that the root element is valid. If not, we didn't actually validate against the schema.
-				if (ConfigFile.DocumentElement?.Name != RootElementName)
+				if (configFile.DocumentElement?.Name != RootElementName)
 				{
-					Logger.LogError("Script does not have a root element called '{RootElementName}'", RootElementName);
-					OutConfigFile = null;
+					logger.LogError("Script does not have a root element called '{RootElementName}'", RootElementName);
+					outConfigFile = null;
 					return false;
 				}
-				if (ConfigFile.DocumentElement.NamespaceURI != SchemaNamespaceURI)
+				if (configFile.DocumentElement.NamespaceURI != SchemaNamespaceURI)
 				{
-					Logger.LogError("Script root element is not in the '{NamespaceUri}' namespace (add the xmlns=\"{NamespaceUri2}\" attribute)", SchemaNamespaceURI, SchemaNamespaceURI);
-					OutConfigFile = null;
+					logger.LogError("Script root element is not in the '{NamespaceUri}' namespace (add the xmlns=\"{NamespaceUri2}\" attribute)", SchemaNamespaceURI, SchemaNamespaceURI);
+					outConfigFile = null;
 					return false;
 				}
 			}
 
-			OutConfigFile = ConfigFile;
+			outConfigFile = configFile;
 			return true;
 		}
 
 		/// <summary>
 		/// Callback for validation errors in the document
 		/// </summary>
-		/// <param name="Sender">Standard argument for ValidationEventHandler</param>
-		/// <param name="Args">Standard argument for ValidationEventHandler</param>
-		void ValidationEvent(object? Sender, ValidationEventArgs Args)
-		{
-			Log.TraceWarningTask(File, Args.Exception.LineNumber, "{0}", Args.Message);
-		}
+		/// <param name="sender">Standard argument for ValidationEventHandler</param>
+		/// <param name="args">Standard argument for ValidationEventHandler</param>
+		void ValidationEvent(object? sender, ValidationEventArgs args) => Log.TraceWarningTask(_file, args.Exception.LineNumber, "{0}", args.Message);
 	}
 
 	/// <summary>
@@ -133,21 +127,21 @@ namespace UnrealBuildTool
 		/// <summary>
 		/// The file containing this element
 		/// </summary>
-		public readonly FileReference File;
+		public FileReference File { get; init; }
 
 		/// <summary>
 		/// The line number containing this element
 		/// </summary>
-		public readonly int LineNumber;
+		public int LineNumber { get; init; }
 
 		/// <summary>
 		/// Constructor
 		/// </summary>
-		public XmlConfigFileElement(FileReference InFile, int InLineNumber, string Prefix, string LocalName, string? NamespaceUri, XmlConfigFile ConfigFile)
-			: base(Prefix, LocalName, NamespaceUri, ConfigFile)
+		public XmlConfigFileElement(FileReference inFile, int inLineNumber, string prefix, string localName, string? namespaceUri, XmlConfigFile configFile)
+			: base(prefix, localName, namespaceUri, configFile)
 		{
-			File = InFile;
-			LineNumber = InLineNumber;
+			File = inFile;
+			LineNumber = inLineNumber;
 		}
 	}
 }
diff --git a/Engine/Source/Programs/UnrealBuildTool/System/XmlConfigFileAttribute.cs b/Engine/Source/Programs/UnrealBuildTool/System/XmlConfigFileAttribute.cs
--- a/Engine/Source/Programs/UnrealBuildTool/System/XmlConfigFileAttribute.cs
+++ b/Engine/Source/Programs/UnrealBuildTool/System/XmlConfigFileAttribute.cs
@@ -8,29 +8,29 @@ namespace UnrealBuildTool
 	/// Marks a field as being serializable from a config file
 	/// </summary>
 	[AttributeUsage(AttributeTargets.Field | AttributeTargets.Property, AllowMultiple = true)]
-	class XmlConfigFileAttribute : Attribute
+	sealed class XmlConfigFileAttribute : Attribute
 	{
 		/// <summary>
 		/// The category for this config value. Optional; defaults to the declaring type name.
 		/// </summary>
-		public string? Category = null;
+		public string? Category { get; set; } = null;
 
 		/// <summary>
 		/// Name of the key to read. Optional; defaults to the field name.
 		/// </summary>
-		public string? Name = null;
+		public string? Name { get; set; } = null;
 
 		/// <summary>
 		/// Use this field to indicate that the XML attribute has been deprecated, and that a warning should
 		/// be shown to the user if it is used.
 		/// A deprecated field should also be marked with the [Obsolete] attribute.
 		/// </summary>
-		public bool Deprecated = false;
+		public bool Deprecated { get; set; } = false;
 
 		/// <summary>
 		/// If the attribute has been deprecated because it has been renamed, this field can be used to apply the
 		/// value used for this field to another.
 		/// </summary>
-		public string? NewAttributeName = null;
+		public string? NewAttributeName { get; set; } = null;
 	}
 }
