UnrealBuildTool: If a dependency file fails to be read, gracefully handle the error and rebuild the action

#rnx

diff --git a/Engine/Source/Programs/UnrealBuildTool/System/CppDependencyCache.cs b/Engine/Source/Programs/UnrealBuildTool/System/CppDependencyCache.cs
--- a/Engine/Source/Programs/UnrealBuildTool/System/CppDependencyCache.cs
+++ b/Engine/Source/Programs/UnrealBuildTool/System/CppDependencyCache.cs
@@ -273,10 +273,6 @@ namespace UnrealBuildTool
 			{
 				return TryGetDependencyInfoInternal(InputFile, out OutInfo);
 			}
-			catch (BuildException Ex)
-			{
-				throw Ex;
-			}
 			catch (Exception Ex)
 			{
 				Logger.LogDebug("Unable to read {File}:\n{Ex}", InputFile, ExceptionUtils.FormatExceptionDetails(Ex));
