UnrealBuildTool: Ignore system paths when analyzing with PVS-Studio rather than just ThirdParty paths

#rnx

diff --git a/Engine/Source/Programs/UnrealBuildTool/Modes/BuildMode.cs b/Engine/Source/Programs/UnrealBuildTool/Modes/BuildMode.cs
--- a/Engine/Source/Programs/UnrealBuildTool/Modes/BuildMode.cs
+++ b/Engine/Source/Programs/UnrealBuildTool/Modes/BuildMode.cs
@@ -806,20 +806,18 @@ namespace UnrealBuildTool
 						}
 
 						// these are parsed by external tools wishing to open this file directly
-						foreach (LinkedAction BuildAction in MergedActionsToExecute)
+						foreach (LinkedAction BuildAction in MergedActionsToExecute.Where(BuildAction => BuildAction.ActionType == ActionType.Compile && (BuildAction.Inner as VCCompileAction)?.bIsAnalyzing != true))
 						{
-							if (BuildAction.ActionType == ActionType.Compile)
+							FileItem? PreprocessedFile = BuildAction.ProducedItems.FirstOrDefault(x => x.HasExtension(".i"));
+							if (PreprocessedFile != null)
 							{
-								FileItem? PreprocessedFile = BuildAction.ProducedItems.FirstOrDefault(x => x.HasExtension(".i"));
-								if (PreprocessedFile != null)
-								{
-									Logger.LogInformation("PreProcessPath: {File}", PreprocessedFile);
-								}
-								FileItem? AssemblyPath = BuildAction.ProducedItems.FirstOrDefault(x => x.HasExtension(".asm"));
-								if (AssemblyPath != null)
-								{
-									Logger.LogInformation("AssemblyPath: {File}", AssemblyPath);
-								}
+								Logger.LogInformation("PreProcessPath: {File}", PreprocessedFile);
+							}
+
+							FileItem? AssemblyPath = BuildAction.ProducedItems.FirstOrDefault(x => x.HasExtension(".asm"));
+							if (AssemblyPath != null)
+							{
+								Logger.LogInformation("AssemblyPath: {File}", AssemblyPath);
 							}
 						}
 					}
diff --git a/Engine/Source/Programs/UnrealBuildTool/Platform/Windows/PVSToolChain.cs b/Engine/Source/Programs/UnrealBuildTool/Platform/Windows/PVSToolChain.cs
--- a/Engine/Source/Programs/UnrealBuildTool/Platform/Windows/PVSToolChain.cs
+++ b/Engine/Source/Programs/UnrealBuildTool/Platform/Windows/PVSToolChain.cs
@@ -2,6 +2,7 @@
 
 using System;
 using System.Collections.Generic;
+using System.Collections.Immutable;
 using System.Diagnostics;
 using System.IO;
 using System.Linq;
@@ -382,6 +383,17 @@ namespace UnrealBuildTool
 		[CommandLine("-Output", Required = true)]
 		FileReference? OutputFile = null;
 
+		/// <summary>
+		/// Path to file list of paths to ignore
+		/// </summary>
+		[CommandLine("-Ignored")]
+		FileReference? IgnoredFile = null;
+
+		/// <summary>
+		/// If all ThirdParty code should be ignored
+		/// </summary>
+		bool IgnoreThirdParty = true;
+
 		/// <summary>
 		/// Execute the command
 		/// </summary>
@@ -399,6 +411,10 @@ namespace UnrealBuildTool
 			string[] InputFileLines = FileReference.ReadAllLines(InputFileList!);
 			FileReference[] InputFiles = InputFileLines.Select(x => x.Trim()).Where(x => x.Length > 0).Select(x => new FileReference(x)).ToArray();
 
+			// Read the ignore file
+			string[] IgnoreFileLines = IgnoredFile != null ? FileReference.ReadAllLines(IgnoredFile) : new string[] { };
+			DirectoryReference[] IgnoredDirectories = IgnoreFileLines.Select(x => x.Trim()).Where(x => x.Length > 0).Select(x => new DirectoryReference(x)).ToArray();
+
 			// Create the combined output file, and print the diagnostics to the log
 			HashSet<string> UniqueItems = new HashSet<string>();
 			List<string> OutputLines = new List<string>();
@@ -433,17 +449,26 @@ namespace UnrealBuildTool
 								{
 									bCanParse = true;
 
-									// Ignore anything in ThirdParty folders
-									if (FileName.Replace('/', '\\').IndexOf("\\ThirdParty\\", StringComparison.InvariantCultureIgnoreCase) == -1)
+									FileReference file = new FileReference(FileName);
+
+									// Ignore anything in the IgnoredDirectories folders
+									if (IgnoredDirectories.Any() && IgnoredDirectories.Any(x => file.IsUnderDirectory(x)))
 									{
-										// Output the line to the raw output file
-										RawWriter.WriteLine(Line);
-
-										// Output the line to the log
-										if (!bFalseAlarm && Level == 1)
-										{
-											Logger.LogWarning(KnownLogEvents.Compiler, "{Path}({LineNumber}): warning {WarningCode}: {WarningMessage}", LogValue.SourceFile(new FileReference(FileName), FileName), LineNumber, WarningCode, WarningMessage);
-										}
+										continue;
+									}
+
+									if (IgnoreThirdParty && file.FullName.Contains("ThirdParty", StringComparison.OrdinalIgnoreCase))
+									{
+										continue;
+									}
+
+									// Output the line to the raw output file
+									RawWriter.WriteLine(Line);
+
+									// Output the line to the log
+									if (!bFalseAlarm && Level == 1)
+									{
+										Logger.LogWarning(KnownLogEvents.Compiler, "{Path}({LineNumber}): warning {WarningCode}: {WarningMessage}", LogValue.SourceFile(file, FileName), LineNumber, WarningCode, WarningMessage);
 									}
 								}
 							}
@@ -472,12 +497,12 @@ namespace UnrealBuildTool
 		UnrealTargetPlatform Platform;
 		Version AnalyzerVersion;
 
-		public PVSToolChain(ReadOnlyTargetRules Target, ILogger Logger)
+		public PVSToolChain(ReadOnlyTargetRules Target, VCToolChain InInnerToolchain, ILogger Logger)
 			: base(Logger)
 		{
 			this.Target = Target;
 			Platform = Target.Platform;
-			InnerToolChain = new VCToolChain(Target, Logger);
+			InnerToolChain = InInnerToolchain;
 
 			AnalyzerFile = FileReference.Combine(Unreal.RootDirectory, "Engine", "Restricted", "NoRedist", "Extras", "ThirdPartyNotUE", "PVS-Studio", "PVS-Studio.exe");
 			if (!FileReference.Exists(AnalyzerFile))
@@ -724,7 +749,7 @@ namespace UnrealBuildTool
 						}
 					}
 				}
-				if (Platform == UnrealTargetPlatform.Win64)
+				if (Platform.IsInGroup(UnrealPlatformGroup.Microsoft))
 				{
 					ConfigFileContents.Append("platform=x64\n");
 				}
@@ -838,25 +863,31 @@ namespace UnrealBuildTool
 			}
 
 			TargetMakefile Makefile = MakefileBuilder.Makefile;
-			List<FileReference> InputFiles = Makefile.OutputItems.Select(x => x.Location).Where(x => x.HasExtension(".pvslog")).ToList();
+			ImmutableSortedSet<FileReference> InputFiles = Makefile.OutputItems.Select(x => x.Location).Where(x => x.HasExtension(".pvslog")).ToImmutableSortedSet();
 
 			// Collect the sourcefile items off of the Compile action added in CompileCPPFiles so that in SingleFileCompile mode the PVSGather step is also not filtered out
-			List<FileItem> CompileSourceFiles = Makefile.Actions.OfType<VCCompileAction>().Select(x => x.SourceFile!).ToList();
+			ImmutableSortedSet<FileItem> CompileSourceFiles = Makefile.Actions.OfType<VCCompileAction>().Select(x => x.SourceFile!).ToImmutableSortedSet();
+
+			// Store list of system paths that should be excluded
+			ImmutableSortedSet<DirectoryReference> SystemIncludePaths = Makefile.Actions.OfType<VCCompileAction>().SelectMany(x => x.SystemIncludePaths).ToImmutableSortedSet();
 
 			FileItem InputFileListItem = MakefileBuilder.CreateIntermediateTextFile(OutputFile.ChangeExtension(".input"), InputFiles.Select(x => x.FullName));
+			FileItem IgnoredFileListeItem = MakefileBuilder.CreateIntermediateTextFile(OutputFile.ChangeExtension(".ignored"), SystemIncludePaths.Select(x => x.FullName));
 
 			Action AnalyzeAction = MakefileBuilder.CreateAction(ActionType.Compile);
 			AnalyzeAction.ActionType = ActionType.PostBuildStep;
 			AnalyzeAction.CommandDescription = "Process PVS-Studio Results";
 			AnalyzeAction.CommandPath = Unreal.DotnetPath;
-			AnalyzeAction.CommandArguments = $"\"{Unreal.UnrealBuildToolDllPath}\" -Mode=PVSGather -Input=\"{InputFileListItem.Location}\" -Output=\"{OutputFile}\" ";
+			AnalyzeAction.CommandArguments = $"\"{Unreal.UnrealBuildToolDllPath}\" -Mode=PVSGather -Input=\"{InputFileListItem.Location}\" -Output=\"{OutputFile}\" -Ignored=\"{IgnoredFileListeItem.Location}\" ";
 			AnalyzeAction.WorkingDirectory = Unreal.EngineSourceDirectory;
 			AnalyzeAction.PrerequisiteItems.Add(InputFileListItem);
+			AnalyzeAction.PrerequisiteItems.Add(IgnoredFileListeItem);
 			AnalyzeAction.PrerequisiteItems.UnionWith(Makefile.OutputItems);
 			AnalyzeAction.PrerequisiteItems.UnionWith(CompileSourceFiles);
 			AnalyzeAction.ProducedItems.Add(FileItem.GetItemByFileReference(OutputFile));
 			AnalyzeAction.ProducedItems.Add(FileItem.GetItemByPath(OutputFile.FullName + "_does_not_exist")); // Force the gather step to always execute
 			AnalyzeAction.DeleteItems.UnionWith(AnalyzeAction.ProducedItems);
+			AnalyzeAction.bCanExecuteInUBA = false;
 
 			Makefile.OutputItems.AddRange(AnalyzeAction.ProducedItems);
 		}
diff --git a/Engine/Source/Programs/UnrealBuildTool/Platform/Windows/UEBuildWindows.cs b/Engine/Source/Programs/UnrealBuildTool/Platform/Windows/UEBuildWindows.cs
--- a/Engine/Source/Programs/UnrealBuildTool/Platform/Windows/UEBuildWindows.cs
+++ b/Engine/Source/Programs/UnrealBuildTool/Platform/Windows/UEBuildWindows.cs
@@ -1757,14 +1757,12 @@ namespace UnrealBuildTool
 		/// <returns>New toolchain instance.</returns>
 		public override UEToolChain CreateToolChain(ReadOnlyTargetRules Target)
 		{
+			VCToolChain toolchain = new VCToolChain(Target, Logger);
 			if (Target.StaticAnalyzer == StaticAnalyzer.PVSStudio)
 			{
-				return new PVSToolChain(Target, Logger);
-			}
-			else
-			{
-				return new VCToolChain(Target, Logger);
+				return new PVSToolChain(Target, toolchain, Logger);
 			}
+			return toolchain;
 		}
 
 		/// <summary>
