Adding support for >=2 Android devices in VS+AGDE when pushing libUnreal.so outside of .apk.
Implements support for patching .vcxproj.user files.

diff --git a/Engine/Source/Programs/UnrealBuildTool/Platform/Android/AndroidProjectGenerator.cs b/Engine/Source/Programs/UnrealBuildTool/Platform/Android/AndroidProjectGenerator.cs
--- a/Engine/Source/Programs/UnrealBuildTool/Platform/Android/AndroidProjectGenerator.cs
+++ b/Engine/Source/Programs/UnrealBuildTool/Platform/Android/AndroidProjectGenerator.cs
@@ -120,12 +120,6 @@ namespace UnrealBuildTool
 					Path.Combine(intermediateRootPath, "LLDBSymbolsLibs", "arm64") // support bDontBundleLibrariesInAPK
 				};
 				ProjectFileBuilder.AppendLine($"    <AndroidSymbolDirectories>{string.Join(";", symbolLocations)}</AndroidSymbolDirectories>");
-
-				// At this stage we don't know if bDontBundleLibrariesInAPK is enabled or not, so make a fail-safe check.
-				string pushSOScript = Path.Combine(
-					Path.GetDirectoryName(NMakeOutputPath.FullName)!,
-					"Push_" + Path.GetFileNameWithoutExtension(NMakeOutputPath.FullName) + "-arm64_so.bat");
-				ProjectFileBuilder.AppendLine($"    <AndroidPostApkInstallCommands>IF EXIST {pushSOScript} {pushSOScript};$(AndroidPostApkInstallCommands)</AndroidPostApkInstallCommands>");
 			}
 			else
 			{
@@ -139,23 +133,48 @@ namespace UnrealBuildTool
 			return (AGDEInstalled ? " -Architectures=arm64 -ForceAPKGeneration" : "") + base.GetExtraBuildArguments(InVSSettings);
 		}
 
-		public override string GetVisualStudioUserFileStrings(VSSettings InVSSettings,
-			string InConditionString, TargetRules InTargetRules, FileReference TargetRulesPath, FileReference ProjectFilePath)
+		public override string GetVisualStudioUserFileStrings(VisualStudioUserFileSettings VCUserFileSettings, VSSettings InVSSettings, string InConditionString, TargetRules InTargetRules, FileReference TargetRulesPath, FileReference ProjectFilePath, FileReference? NMakeOutputPath, string ProjectName, string? ForeignUProjectPath)
 		{
 			if (AGDEInstalled
 				&& (InVSSettings.Platform == UnrealTargetPlatform.Android)
 				&& ((InTargetRules.Type == TargetRules.TargetType.Client) || (InTargetRules.Type == TargetRules.TargetType.Game)))
 			{
-				string UserFileEntry = "<PropertyGroup " + InConditionString + ">\n";
-				UserFileEntry += "	<AndroidLldbStartupCommands>" +
-												"command script import \"" + Path.Combine(Unreal.EngineDirectory.FullName, "Extras", "LLDBDataFormatters", "UEDataFormatters_2ByteChars.py") + "\";" +
-												"$(AndroidLldbStartupCommands)" +
-											"</AndroidLldbStartupCommands>\n";
-				UserFileEntry += "</PropertyGroup>\n";
-				return UserFileEntry;
+				StringBuilder Out = new StringBuilder();
+				Out.AppendLine("  <PropertyGroup " + InConditionString + ">");
+
+				string LldbFormatterImport = $"command script import \"" + Path.Combine(Unreal.EngineDirectory.FullName, "Extras", "LLDBDataFormatters", "UEDataFormatters_2ByteChars.py") + "\"";
+				Out.AppendLine($"    <AndroidLldbStartupCommands>{LldbFormatterImport};$(AndroidLldbStartupCommands)</AndroidLldbStartupCommands>");
+				VCUserFileSettings.PatchProperty("AndroidLldbStartupCommands");
+
+				if (NMakeOutputPath != null)
+				{
+					// It's critical to have AndroidDebugTarget here and for it to be before properties that use it (e.g. AndroidPostApkInstallCommands), otherwise MSBuild will evaluate it to empty string.
+					Out.AppendLine("    <AndroidDebugTarget></AndroidDebugTarget>");
+
+					// At this stage we don't know if bDontBundleLibrariesInAPK is enabled or not, so make a fail-safe check.
+					string PushSOScript = Path.Combine(
+						Path.GetDirectoryName(NMakeOutputPath.FullName)!,
+						"Push_" + Path.GetFileNameWithoutExtension(NMakeOutputPath.FullName) + "-arm64_so.bat");
+
+					// AGDE specifies current debug target in AndroidDebugTarget property in a form of "model:serial:arch".
+					// AndroidDebugTarget is a special property and needs to be evaluated in-line. And the push script needs the device serial as first argument to push to the correct device. 
+					// MSBuild Property Functions allow to invoke limited C# expression from with-in MSBuild, see https://learn.microsoft.com/en-us/visualstudio/msbuild/property-functions?view=vs-2022
+					// They lack conditional statements, so this expression makes a branch-less variant by always appending "::" to AndroidDebugTarget,
+					// so regardless of what value it has, Split(':')[1] will never throw an exception.
+					string GetTargetDeviceSerial = "$([System.String]::Concat($(AndroidDebugTarget), \"::\").Split(':')[1])";
+					Out.AppendLine(
+						$"    <AndroidPostApkInstallCommands>IF EXIST {PushSOScript} {PushSOScript} {GetTargetDeviceSerial};$(AndroidPostApkInstallCommands)</AndroidPostApkInstallCommands>");
+
+					// Ensure the following properties are up to date even if .vcxproj.user already exists and are in correct order between each other.
+					VCUserFileSettings.PatchProperty("AndroidDebugTarget", true);
+					VCUserFileSettings.PatchProperty("AndroidPostApkInstallCommands");
+				}
+
+				Out.AppendLine("  </PropertyGroup>");
+				return Out.ToString();
 			}
 
-			return base.GetVisualStudioUserFileStrings(InVSSettings, InConditionString, InTargetRules, TargetRulesPath, ProjectFilePath);
+			return string.Empty;
 		}
 	}
 }
diff --git a/Engine/Source/Programs/UnrealBuildTool/Platform/Windows/WindowsProjectGenerator.cs b/Engine/Source/Programs/UnrealBuildTool/Platform/Windows/WindowsProjectGenerator.cs
--- a/Engine/Source/Programs/UnrealBuildTool/Platform/Windows/WindowsProjectGenerator.cs
+++ b/Engine/Source/Programs/UnrealBuildTool/Platform/Windows/WindowsProjectGenerator.cs
@@ -47,7 +47,7 @@ namespace UnrealBuildTool
 		}
 
 		/// <inheritdoc/>
-		public override string GetVisualStudioUserFileStrings(VSSettings InVSSettings, string InConditionString, TargetRules InTargetRules, FileReference TargetRulesPath, FileReference ProjectFilePath, string ProjectName, string? ForeignUProjectPath)
+		public override string GetVisualStudioUserFileStrings(VisualStudioUserFileSettings VCUserFileSettings, VSSettings InVSSettings, string InConditionString, TargetRules InTargetRules, FileReference TargetRulesPath, FileReference ProjectFilePath, FileReference? NMakeOutputPath, string ProjectName, string? ForeignUProjectPath)
 		{
 			StringBuilder VCUserFileContent = new StringBuilder();
 
diff --git a/Engine/Source/Programs/UnrealBuildTool/ProjectFiles/PlatformProjectGenerator.cs b/Engine/Source/Programs/UnrealBuildTool/ProjectFiles/PlatformProjectGenerator.cs
--- a/Engine/Source/Programs/UnrealBuildTool/ProjectFiles/PlatformProjectGenerator.cs
+++ b/Engine/Source/Programs/UnrealBuildTool/ProjectFiles/PlatformProjectGenerator.cs
@@ -270,35 +270,58 @@ namespace UnrealBuildTool
 		}
 
 		/// <summary>
-		/// Get the text to insert into the user file for the given platform/configuration/target
+		/// Additional configuration related to Visual Studio user file (.vcxproj.user). 
 		/// </summary>
-		/// <param name="InVSSettings">The ProjectFileSettings that contains the project platform/configuration/etc info being built</param>
-		/// <param name="InConditionString">The condition string </param>
-		/// <param name="InTargetRules">The target rules </param>
-		/// <param name="TargetRulesPath">The target rules path</param>
-		/// <param name="ProjectFilePath">The project file path</param>
-		/// <returns>The string to append to the user file</returns>
-		public virtual string GetVisualStudioUserFileStrings(VSSettings InVSSettings,
-			string InConditionString, TargetRules InTargetRules, FileReference TargetRulesPath, FileReference ProjectFilePath)
+		public class VisualStudioUserFileSettings
 		{
-			return "";
+			public IReadOnlySet<string> PropertiesToPatch => PropertiesToPatchContainer;
+
+			public IReadOnlySet<string> PropertiesToPatchOrderButPreserveValue => PropertiesToPatchOrderButPreserveValueContainer;
+
+			/// <summary>
+			/// Patch existing .vcxproj.user file based on the new file content and corresponding patching settings, by doing the following:
+			/// - Go over every &lt;PropertyGroup&gt; in the new document.
+			/// - Finds &lt;PropertyGroup&gt; in the current document with matching "Condition" attribute.
+			/// - If no property group is found, appends property group from new document to the end of current document.
+			/// - Otherwise, saves values of properties configured with bPreserveExistingValue == true.
+			/// - Removes all properties flagged for patching from current document property group.
+			/// - Adds all properties required for patching from new document property group as declared in new document order.
+			///
+			/// This makes possible to override some properties crucial for UE features while keeping rest of user configuration intact.
+			/// 
+			/// Otherwise, if .vcxproj.user doesn't exist, patching do nothing and instead creates a new file based on new file content.
+			/// </summary>
+			/// <param name="PropertyName">Name of the property to patch.</param>
+			/// <param name="bPreserveExistingValue">If true, the value of the property will be taken from existing user file, but relative order will be taken from newly generated file.</param>
+			public void PatchProperty(string PropertyName, bool bPreserveExistingValue = false)
+			{
+				if (!PropertiesToPatchContainer.Contains(PropertyName))
+					PropertiesToPatchContainer.Add(PropertyName);
+
+				if (bPreserveExistingValue && !PropertiesToPatchOrderButPreserveValueContainer.Contains(PropertyName))
+					PropertiesToPatchOrderButPreserveValueContainer.Add(PropertyName);
+			}
+
+			private HashSet<string> PropertiesToPatchContainer = new();
+			private HashSet<string> PropertiesToPatchOrderButPreserveValueContainer = new();
 		}
 
 		/// <summary>
 		/// Get the text to insert into the user file for the given platform/configuration/target
 		/// </summary>
+		/// <param name="VCUserFileSettings">Configuration for user file creation/patching/etc.</param>
 		/// <param name="InVSSettings">The ProjectFileSettings that contains the project platform/configuration/etc info being built</param>
 		/// <param name="InConditionString">The condition string </param>
 		/// <param name="InTargetRules">The target rules </param>
 		/// <param name="TargetRulesPath">The target rules path</param>
 		/// <param name="ProjectFilePath">The project file path</param>
+		/// <param name="NMakeOutputPath">Output path for NMake</param>
 		/// <param name="ProjectName">The name of the project</param>
 		/// <param name="ForeignUProjectPath">Path to foreign .uproject file, if any</param>
 		/// <returns>The string to append to the user file</returns>
-		public virtual string GetVisualStudioUserFileStrings(VSSettings InVSSettings,
-			string InConditionString, TargetRules InTargetRules, FileReference TargetRulesPath, FileReference ProjectFilePath, string ProjectName, string? ForeignUProjectPath)
+		public virtual string GetVisualStudioUserFileStrings(VisualStudioUserFileSettings VCUserFileSettings, VSSettings InVSSettings, string InConditionString, TargetRules InTargetRules, FileReference TargetRulesPath, FileReference ProjectFilePath, FileReference? NMakeOutputPath, string ProjectName, string? ForeignUProjectPath)
 		{
-			return GetVisualStudioUserFileStrings(InVSSettings, InConditionString, InTargetRules, TargetRulesPath, ProjectFilePath);
+			return string.Empty;
 		}
 
 		/// <summary>
diff --git a/Engine/Source/Programs/UnrealBuildTool/ProjectFiles/VisualStudio/VCProject.cs b/Engine/Source/Programs/UnrealBuildTool/ProjectFiles/VisualStudio/VCProject.cs
--- a/Engine/Source/Programs/UnrealBuildTool/ProjectFiles/VisualStudio/VCProject.cs
+++ b/Engine/Source/Programs/UnrealBuildTool/ProjectFiles/VisualStudio/VCProject.cs
@@ -12,6 +12,7 @@ using System.Text;
 using System.Text.RegularExpressions;
 using System.Threading.Tasks;
 using System.Xml;
+using System.Xml.Linq;
 using EpicGames.Core;
 using Microsoft.Extensions.Logging;
 using UnrealBuildBase;
@@ -904,6 +905,7 @@ namespace UnrealBuildTool
 			StringBuilder VCProjectFileContent = new StringBuilder();
 			StringBuilder VCFiltersFileContent = new StringBuilder();
 			StringBuilder VCUserFileContent = new StringBuilder();
+			VisualStudioUserFileSettings VCUserFileSettings = new VisualStudioUserFileSettings();
 
 			// Visual Studio doesn't require a *.vcxproj.filters file to even exist alongside the project unless
 			// it actually has something of substance in it.  We'll avoid saving it out unless we need to.
@@ -1406,7 +1408,7 @@ namespace UnrealBuildTool
 			// Write each configuration
 			foreach (ProjectConfigAndTargetCombination Combination in ProjectConfigAndTargetCombinations)
 			{
-				WriteConfiguration(ProjectName, Combination, VCProjectFileContent, PlatformProjectGenerators, bGenerateUserFileContent ? VCUserFileContent : null);
+				WriteConfiguration(ProjectName, Combination, VCProjectFileContent, PlatformProjectGenerators, bGenerateUserFileContent ? VCUserFileContent : null, bGenerateUserFileContent ? VCUserFileSettings : null);
 			}
 
 			{ 
@@ -1783,11 +1785,156 @@ namespace UnrealBuildTool
 				{
 					bSuccess = ProjectFileGenerator.WriteFileIfChanged(VCUserFilePath, VCUserFileContent.ToString(), Logger);
 				}
+				else
+				{
+					bSuccess = PatchVCUserFile(VCUserFilePath, VCUserFileContent.ToString(), VCUserFileSettings, Logger);
+				}
 			}
 
 			return bSuccess;
 		}
 
+		private bool PatchVCUserFile(string FileName, string NewFileContents, VisualStudioUserFileSettings UserFileSettings, ILogger Logger)
+		{
+			// Before we start, see if any relevant property is present in the new contents.
+			if (UserFileSettings.PropertiesToPatch.Any(NewFileContents.Contains) == false)
+			{
+				return true;
+			}
+
+			XDocument CurrentContent;
+			try
+			{
+				CurrentContent = XDocument.Load(FileName);
+			}
+			catch (Exception ex)
+			{
+				string Message = String.Format("Error while trying to parse XML file {0}.", FileName);
+				Logger.LogError("{Message}", Message);
+				throw new BuildException(ex, Message);
+			}
+
+			XDocument NewContent;
+			try
+			{
+				NewContent = XDocument.Parse(NewFileContents);
+			}
+			catch (Exception ex)
+			{
+				string Message = String.Format("Error while trying to parse XML new data for file {0} ('{1}').", FileName, NewFileContents);
+				Logger.LogError("{Message}", Message);
+				throw new BuildException(ex, Message);
+			}
+
+			if (CurrentContent.Root == null || NewContent.Root == null)
+			{
+				return true;
+			}
+
+			XNamespace NS = NewContent.Root.Name.Namespace;
+			// Skip patching if namespaces don't match (should never be the case?)
+			if (NS != CurrentContent.Root.Name.Namespace)
+			{
+				return false;
+			}
+
+			// Create dictionaries with key == Condition of each <PropertyGroup> and value == <PropertyGroup> XElement itself for both current and new document.
+			Dictionary<string, XElement> CurrentPropertyGroups = CurrentContent
+				.Descendants(NS + "PropertyGroup")
+				.Select(Element => (Attribute: Element.Attribute("Condition"), Element: Element))
+				.Where(Pair => Pair.Attribute != null)
+				.ToDictionary(Pair => Pair.Attribute!.Value, Pair => Pair.Element);
+
+			Dictionary<string, XElement> NewPropertyGroups = NewContent
+				.Descendants(NS + "PropertyGroup")
+				.Select(Element => (Attribute: Element.Attribute("Condition"), Element: Element))
+				.Where(Pair => Pair.Attribute != null)
+				.ToDictionary(Pair => Pair.Attribute!.Value, Pair => Pair.Element);
+
+			bool bNeedsSaving = false;
+
+			// Go over every <PropertyGroup> in new document.
+			foreach ((string Attribute, XElement NewPropertyGroup) in NewPropertyGroups)
+			{
+				// Check if new document <PropertyGroup> contains any properties that we need to patch in the current document.
+				if (NewPropertyGroup.Elements().Any(Element => UserFileSettings.PropertiesToPatch.Contains(Element.Name.LocalName)) == false)
+				{
+					continue;
+				}
+
+				// Check if <PropertyGroup> with same "Condition" attribute already exist in the current document.
+				// If yes, update required properties in existing <PropertyGroup> but preserve any other property in current document order.
+				if (CurrentPropertyGroups.TryGetValue(Attribute, out var CurrentPropertyGroup))
+				{
+					// Preserve values from current document for relevant properties by patching corresponding properties in new document. 
+					var ElementsToPreserveValuesFrom = CurrentPropertyGroup
+						.Elements()
+						.Where(Element => UserFileSettings.PropertiesToPatchOrderButPreserveValue.Contains(Element.Name.LocalName));
+					foreach (XElement CurrentElement in ElementsToPreserveValuesFrom)
+					{
+						XElement? NewElement = NewPropertyGroup.Element(CurrentElement.Name);
+						if (NewElement != null)
+						{
+							NewElement.Value = CurrentElement.Value;
+						}
+					}
+
+					XElement[] CurrentPropertyGroupElementsForPatch = CurrentPropertyGroup
+						.Elements()
+						.Where(Element => UserFileSettings.PropertiesToPatch.Contains(Element.Name.LocalName))
+						.ToArray();
+
+					XElement[] NewPropertyGroupElementsForPatch = NewPropertyGroup
+						.Elements()
+						.Where(Element => UserFileSettings.PropertiesToPatch.Contains(Element.Name.LocalName))
+						.ToArray();
+					
+					// Check if all properties are already has the correct value and order, and skip patching if so.
+					if (CurrentPropertyGroupElementsForPatch.Length == NewPropertyGroupElementsForPatch.Length &&
+						!CurrentPropertyGroupElementsForPatch.Where((CurrentProperty, i) => CurrentProperty.Name != NewPropertyGroupElementsForPatch[i].Name || CurrentProperty.Value != NewPropertyGroupElementsForPatch[i].Value).Any())
+					{
+						continue;
+					}
+
+					// Remove all existing properties that we need to update from existing document, this simplifies logic of adding them,
+					// because we need to update the values and ensure the order of properties is as declared in the new document,
+					// because the order of properties is important for MSBuild when one property uses a value of another property.
+					CurrentPropertyGroupElementsForPatch.Remove();
+
+					// Add new properties to existing <PropertyGroup> in the order as they are defined in new document.
+					CurrentPropertyGroup.Add(NewPropertyGroupElementsForPatch);
+
+					bNeedsSaving = true;
+				}
+				else // Otherwise add new <PropertyGroup> as-is to the end of existing document.
+				{
+					CurrentContent.Root.Add(NewPropertyGroup);
+					bNeedsSaving = true;
+				}
+			}
+
+			if (bNeedsSaving)
+			{
+				try
+				{
+					CurrentContent.Save(FileName);
+					Logger.LogDebug("Patching {Path}.", Path.GetFileName(FileName));
+				}
+				catch (Exception ex)
+				{
+					string Message = String.Format("Error while trying to write file {0}. The file is probably read-only.", FileName);
+					Logger.LogError("{Message}", Message);
+					throw new BuildException(ex, Message);
+				}
+			}
+			else
+			{
+				Logger.LogDebug("{Path} doesn't require patching.", Path.GetFileName(FileName));
+			}
+
+			return true;
+		}
+
 		private class ProjectConfigurationForGenerator : ProjectBuildConfiguration
 		{
 			public override string ConfigurationName => Combination.ProjectConfigurationName;
@@ -2066,7 +2213,7 @@ namespace UnrealBuildTool
 		}
 
 		// Anonymous function that writes project configuration data
-		private void WriteConfiguration(string ProjectName, ProjectConfigAndTargetCombination Combination, StringBuilder VCProjectFileContent, PlatformProjectGeneratorCollection PlatformProjectGenerators, StringBuilder? VCUserFileContent)
+		private void WriteConfiguration(string ProjectName, ProjectConfigAndTargetCombination Combination, StringBuilder VCProjectFileContent, PlatformProjectGeneratorCollection PlatformProjectGenerators, StringBuilder? VCUserFileContent, VisualStudioUserFileSettings? VCUserFileSettings)
 		{
 			UnrealTargetConfiguration Configuration = Combination.Configuration;
 
@@ -2096,6 +2243,7 @@ namespace UnrealBuildTool
 				}
 
 				DirectoryReference ProjectDirectory = ProjectFilePath.Directory;
+				FileReference? NMakePath = null;
 
 				if (IsStubProject)
 				{
@@ -2188,7 +2336,7 @@ namespace UnrealBuildTool
 					}
 
 					// Make the output file path
-					FileReference NMakePath = FileReference.Combine(OutputDirectory, BaseExeName);
+					NMakePath = FileReference.Combine(OutputDirectory, BaseExeName);
 					if (Configuration != TargetRulesObject.UndecoratedConfiguration)
 					{
 						NMakePath += "-" + UBTPlatformName + "-" + UBTConfigurationName;
@@ -2285,14 +2433,14 @@ namespace UnrealBuildTool
 					VCProjectFileContent.AppendLine("  </ItemDefinitionGroup>");
 				}
 
-				if (VCUserFileContent != null && Combination.ProjectTarget != null)
+				if (VCUserFileContent != null && VCUserFileSettings != null && Combination.ProjectTarget != null)
 				{
 					TargetRules TargetRulesObject = Combination.ProjectTarget.TargetRules!;
 
 					if (ProjGenerator != null)
 					{
 						string? ForeignUProjectPath = (IsForeignProject && !String.IsNullOrEmpty(UProjectPath)) ? UProjectPath : null;
-						VCUserFileContent.Append(ProjGenerator.GetVisualStudioUserFileStrings(new(Combination.Platform!.Value, Configuration, ProjectFileFormat, Combination.Architecture), ConditionString, TargetRulesObject, Combination.ProjectTarget.TargetFilePath, ProjectFilePath, ProjectName, ForeignUProjectPath));
+						VCUserFileContent.Append(ProjGenerator.GetVisualStudioUserFileStrings(VCUserFileSettings, new(Combination.Platform!.Value, Configuration, ProjectFileFormat, Combination.Architecture), ConditionString, TargetRulesObject, Combination.ProjectTarget.TargetFilePath, ProjectFilePath, NMakePath, ProjectName, ForeignUProjectPath));
 					}
 				}
 			}
