- Made AutomationTool references be non-copy-local, so it loads the .dlls from AutomatonTools\AutomationScripts. This solves the problem of making changes in automation scripts, then RunUAT would fail to get the changes

diff --git a/Engine/Source/Programs/UnrealBuildTool/ProjectFiles/VisualStudio/VCProjectFileGenerator.cs b/Engine/Source/Programs/UnrealBuildTool/ProjectFiles/VisualStudio/VCProjectFileGenerator.cs
--- a/Engine/Source/Programs/UnrealBuildTool/ProjectFiles/VisualStudio/VCProjectFileGenerator.cs
+++ b/Engine/Source/Programs/UnrealBuildTool/ProjectFiles/VisualStudio/VCProjectFileGenerator.cs
@@ -613,7 +613,8 @@ namespace UnrealBuildTool
 						new XElement(NS + "ItemGroup",
 							from AutomationProject in AutomationProjectFiles
 							select new XElement(NS + "ProjectReference",
-								new XAttribute("Include", AutomationProject.ProjectFilePath.MakeRelativeTo(AutomationToolDir))
+								new XAttribute("Include", AutomationProject.ProjectFilePath.MakeRelativeTo(AutomationToolDir)),
+								new XElement(NS + "Private", "false")
 							)
 						)
 					)
