- Made AutomationTool references be non-copy-local, so it loads the .dlls from AutomatonTools\AutomationScripts. This solves the problem of making changes in automation scripts, then RunUAT would fail to get the changes, because building AT would copy the .dll's next to the .exe, but then iterating on automation scripts and using RunUAT (without building AT.exe proper), would not copy new dlls on top of the ones next to the .exe
- Added a task to AutomationTool.csproj to delete the stale (bad) dlls that were next to the .exe

diff --git a/Engine/Source/Programs/UnrealBuildTool/ProjectFiles/VisualStudio/VCProjectFileGenerator.cs b/Engine/Source/Programs/UnrealBuildTool/ProjectFiles/VisualStudio/VCProjectFileGenerator.cs
--- a/Engine/Source/Programs/UnrealBuildTool/ProjectFiles/VisualStudio/VCProjectFileGenerator.cs
+++ b/Engine/Source/Programs/UnrealBuildTool/ProjectFiles/VisualStudio/VCProjectFileGenerator.cs
@@ -606,6 +606,7 @@ namespace UnrealBuildTool
 				XNamespace NS = XNamespace.Get("http://schemas.microsoft.com/developer/msbuild/2003");
 
 				DirectoryReference AutomationToolDir = DirectoryReference.Combine(Unreal.EngineSourceDirectory, "Programs", "AutomationTool");
+				DirectoryReference AutomationToolBinariesDir = DirectoryReference.Combine(Unreal.EngineDirectory, "Binaries", "DotNET", "AutomationTool");
 				XDocument AutomationToolDocument = new XDocument(
 					new XElement(NS + "Project",
 						new XAttribute("ToolsVersion", VCProjectFileGenerator.GetProjectFileToolVersionString(Settings.ProjectFileFormat)),
@@ -613,7 +614,22 @@ namespace UnrealBuildTool
 						new XElement(NS + "ItemGroup",
 							from AutomationProject in AutomationProjectFiles
 							select new XElement(NS + "ProjectReference",
-								new XAttribute("Include", AutomationProject.ProjectFilePath.MakeRelativeTo(AutomationToolDir))
+								new XAttribute("Include", AutomationProject.ProjectFilePath.MakeRelativeTo(AutomationToolDir)),
+								new XElement(NS + "Private", "false")
+							)
+						),
+						// Delete the private copied dlls in case they were ever next to the .exe - that is a bad place for them
+						new XElement(NS + "Target",
+							new XAttribute("Name", "CleanUpStaleDlls"),
+							new XAttribute("AfterTargets", "Build"),
+							AutomationProjectFiles.SelectMany(AutomationProject => {
+									string BaseFilename = FileReference.Combine(AutomationToolBinariesDir, AutomationProject.ProjectFilePath.GetFileNameWithoutExtension()).FullName;
+									return new List<XElement>() {
+										new XElement(NS + "Delete",	new XAttribute("Files", BaseFilename + ".dll")),
+										new XElement(NS + "Delete",	new XAttribute("Files", BaseFilename + ".dll.config")),
+										new XElement(NS + "Delete",	new XAttribute("Files", BaseFilename + ".pdb"))
+									};
+								}
 							)
 						)
 					)
diff --git a/Engine/Source/Programs/UnrealBuildTool/ProjectFiles/VisualStudioMac/VisualStudioMacProjectFileGenerator.cs b/Engine/Source/Programs/UnrealBuildTool/ProjectFiles/VisualStudioMac/VisualStudioMacProjectFileGenerator.cs
--- a/Engine/Source/Programs/UnrealBuildTool/ProjectFiles/VisualStudioMac/VisualStudioMacProjectFileGenerator.cs
+++ b/Engine/Source/Programs/UnrealBuildTool/ProjectFiles/VisualStudioMac/VisualStudioMacProjectFileGenerator.cs
@@ -1,5 +1,6 @@
 ﻿// Copyright Epic Games, Inc. All Rights Reserved.
 
+using System.Collections.Generic;
 using System.Linq;
 using System.Xml;
 using System.Xml.Linq;
@@ -88,6 +89,7 @@ namespace UnrealBuildTool
 				XNamespace NS = XNamespace.Get("http://schemas.microsoft.com/developer/msbuild/2003");
 
 				DirectoryReference AutomationToolDir = DirectoryReference.Combine(Unreal.EngineSourceDirectory, "Programs", "AutomationTool");
+				DirectoryReference AutomationToolBinariesDir = DirectoryReference.Combine(Unreal.EngineDirectory, "Binaries", "DotNET", "AutomationTool");
 				new XDocument(
 					new XElement(NS + "Project",
 						new XAttribute("ToolsVersion", VCProjectFileGenerator.GetProjectFileToolVersionString(Settings.ProjectFileFormat)),
@@ -100,6 +102,20 @@ namespace UnrealBuildTool
 								new XElement(NS + "Name", AutomationProject.ProjectFilePath.GetFileNameWithoutExtension()),
 								new XElement(NS + "Private", "false")
 							)
+						),
+						// Delete the private copied dlls in case they were ever next to the .exe - that is a bad place for them
+						new XElement(NS + "Target",
+							new XAttribute("Name", "CleanUpStaleDlls"),
+							new XAttribute("AfterTargets", "Build"),
+							AutomationProjectFiles.SelectMany(AutomationProject => {
+									string BaseFilename = FileReference.Combine(AutomationToolBinariesDir, AutomationProject.ProjectFilePath.GetFileNameWithoutExtension()).FullName;
+									return new List<XElement>() {
+										new XElement(NS + "Delete",	new XAttribute("Files", BaseFilename + ".dll")),
+										new XElement(NS + "Delete",	new XAttribute("Files", BaseFilename + ".dll.config")),
+										new XElement(NS + "Delete",	new XAttribute("Files", BaseFilename + ".pdb"))
+									};
+								}
+							)
 						)
 					)
 				).Save(FileReference.Combine(IntermediateProjectFilesPath, "AutomationTool.csproj.References").FullName);
