UnrealBuildTool: Don't register a file that doesn't exist in UBA

#rnx

diff --git a/Engine/Source/Programs/UnrealBuildTool/Executors/Experimental/UBAExecutor.cs b/Engine/Source/Programs/UnrealBuildTool/Executors/Experimental/UBAExecutor.cs
--- a/Engine/Source/Programs/UnrealBuildTool/Executors/Experimental/UBAExecutor.cs
+++ b/Engine/Source/Programs/UnrealBuildTool/Executors/Experimental/UBAExecutor.cs
@@ -719,7 +719,7 @@ namespace UnrealBuildTool
 
 					if (!enableDetour && process.ExitCode == 0)
 					{
-						_session!.RegisterNewFiles(action.ProducedItems.Select(x => x.FullName).ToArray());
+						_session!.RegisterNewFiles(action.ProducedItems.Where(x => FileReference.Exists(x.Location)).Select(x => x.FullName).ToArray());
 					}
 
 					TimeSpan processorTime = process.TotalProcessorTime;
