[UBA]
* Added IUBAAgentCoordinator interface and removed references to horde in UBAExecutor
* Fixed so horde agents are disposed if ctrl-c is pressed (so they get returned rather than time out)

diff --git a/Engine/Source/Programs/UnrealBuildTool/Executors/Experimental/UBAAgentCoordinator.cs b/Engine/Source/Programs/UnrealBuildTool/Executors/Experimental/UBAAgentCoordinator.cs
new file mode 100644
--- /dev/null
+++ b/Engine/Source/Programs/UnrealBuildTool/Executors/Experimental/UBAAgentCoordinator.cs
@@ -0,0 +1,23 @@
+// Copyright Epic Games, Inc. All Rights Reserved.
+
+using EpicGames.Core;
+using System;
+using System.Threading.Tasks;
+
+namespace UnrealBuildTool
+{
+
+	interface IUBAAgentCoordinator
+	{
+		DirectoryReference? GetUBARootDir();
+
+		Task InitAsync(UBAExecutor executor);
+
+		void Start(ImmediateActionQueue queue, Func<LinkedAction, bool> canRunRemotely);
+
+		void Stop();
+	
+		Task CloseAsync();
+	}
+
+}
\ No newline at end of file
diff --git a/Engine/Source/Programs/UnrealBuildTool/Executors/Experimental/UBAExecutor.cs b/Engine/Source/Programs/UnrealBuildTool/Executors/Experimental/UBAExecutor.cs
--- a/Engine/Source/Programs/UnrealBuildTool/Executors/Experimental/UBAExecutor.cs
+++ b/Engine/Source/Programs/UnrealBuildTool/Executors/Experimental/UBAExecutor.cs
@@ -1,9 +1,6 @@
 // Copyright Epic Games, Inc. All Rights Reserved.
 
 using EpicGames.Core;
-using EpicGames.Horde.Common;
-using EpicGames.Horde.Compute;
-using EpicGames.Horde.Compute.Clients;
 using EpicGames.UBA;
 using Microsoft.Extensions.Logging;
 using System;
@@ -13,7 +10,6 @@ using System.Linq;
 using System.Management;
 using System.Runtime.InteropServices;
 using System.Runtime.Versioning;
-using System.Threading;
 using System.Threading.Tasks;
 using UnrealBuildBase;
 using UnrealBuildTool.Artifacts;
@@ -244,11 +240,11 @@ namespace UnrealBuildTool
 	class UBAExecutor : ParallelExecutor
 	{
 		public UnrealBuildAcceleratorConfig UBAConfig { get; init; } = new();
-		public UnrealBuildAcceleratorHordeConfig HordeConfig { get; init; } = new();
 
 		public string Crypto { get; private set; } = String.Empty;
 		public IServer? Server { get; private set; }
 		ISessionServer? _session;
+		List<IUBAAgentCoordinator> _agentCoordinators = new List<IUBAAgentCoordinator>();
 		DirectoryReference? _rootDirRef;
 		bool _bIsCancelled;
 		bool _bIsRemoteActionsAllowed = true;
@@ -280,27 +276,16 @@ namespace UnrealBuildTool
 		}
 
 		public UBAExecutor(int maxLocalActions, bool bAllCores, bool bCompactOutput, Microsoft.Extensions.Logging.ILogger logger, CommandLineArguments? additionalArguments = null)
-			: this(maxLocalActions, bAllCores, bCompactOutput, new ThreadedLogger(logger))
+			: base(maxLocalActions, bAllCores, bCompactOutput, logger)
 		{
 			XmlConfig.ApplyTo(this);
 			XmlConfig.ApplyTo(UBAConfig);
-			XmlConfig.ApplyTo(HordeConfig);
 			CommandLine.ParseArguments(Environment.GetCommandLineArgs(), this, logger);
 			additionalArguments?.ApplyTo(this);
 			additionalArguments?.ApplyTo(UBAConfig);
-			additionalArguments?.ApplyTo(HordeConfig);
 
-			// Sentry is currently unsupported for non-Windows and non-x64
-			if (!OperatingSystem.IsWindows() || RuntimeInformation.ProcessArchitecture != Architecture.X64)
-			{
-				HordeConfig.UBASentryUrl = null;
-			}
-		}
-
-		private UBAExecutor(int maxLocalActions, bool bAllCores, bool bCompactOutput, ThreadedLogger logger)
-			: base(maxLocalActions, bAllCores, bCompactOutput, logger)
-		{
-			_threadedLogger = logger;
+			_threadedLogger = new ThreadedLogger(logger);
+			_agentCoordinators.Add(new UBAAgentCoordinatorHorde(logger, UBAConfig, additionalArguments));
 		}
 
 		private void PrintConfiguration()
@@ -347,15 +332,16 @@ namespace UnrealBuildTool
 				_rootDirRef = DirectoryReference.FromString(Environment.GetEnvironmentVariable("UBA_ROOT") ?? Environment.GetEnvironmentVariable("BOX_ROOT"));
 			}
 
-			using CancellationTokenSource cancellationTokenSource = new();
-			using Task<UBAHordeSession?> hordeSessionTask = !UBAConfig.bDisableRemote ? UBAHordeSession.TryCreateHordeSession(HordeConfig, this, UBAConfig.bStrict, logger, cancellationTokenSource.Token) : Task.FromResult<UBAHordeSession?>(null);
-
-			if (!UBAConfig.bDisableRemote && _rootDirRef == null)
+			if (!UBAConfig.bDisableRemote)
 			{
-				DirectoryReference? hordeSharedDir = DirectoryReference.FromString(Environment.GetEnvironmentVariable("UE_HORDE_SHARED_DIR"));
-				if (hordeSharedDir != null)
+				foreach (var coordinator in _agentCoordinators)
 				{
-					_rootDirRef = DirectoryReference.Combine(hordeSharedDir, "UbaHost");
+					await coordinator.InitAsync(this);
+
+					if (_rootDirRef == null)
+					{
+						_rootDirRef = coordinator.GetUBARootDir();
+					}
 				}
 			}
 
@@ -400,6 +386,10 @@ namespace UnrealBuildTool
 				_bIsCancelled = true;
 				_session?.CancelAll();
 				ubaStorage?.SaveCasTable();
+				foreach (var coordinator in _agentCoordinators)
+				{
+					coordinator.CloseAsync().Wait(2000); // Give coordinators some time to close (this makes coordinators like horde return resources faster)
+				}
 			}
 			Console.CancelKeyPress += CancelKeyPress;
 
@@ -407,7 +397,7 @@ namespace UnrealBuildTool
 			{
 				if (UBAConfig.bLaunchVisualizer && OperatingSystem.IsWindows())
 				{
-					_ = Task.Run(LaunchVisualizer, cancellationTokenSource.Token);
+					_ = Task.Run(LaunchVisualizer);
 				}
 
 				using EpicGames.UBA.ILogger ubaLogger = EpicGames.UBA.ILogger.CreateLogger(logger);
@@ -425,7 +415,7 @@ namespace UnrealBuildTool
 							Server.StartServer(UBAConfig.Host, UBAConfig.Port, Crypto);
 						}
 
-						bool success = ExecuteActionsInternal(inputActions, _session, hordeSessionTask, cancellationTokenSource, logger, actionArtifactCache);
+						bool success = ExecuteActionsInternal(inputActions, _session, logger, actionArtifactCache);
 
 						if (!UBAConfig.bDisableRemote)
 						{
@@ -445,11 +435,9 @@ namespace UnrealBuildTool
 			{
 				Console.CancelKeyPress -= CancelKeyPress;
 
-				cancellationTokenSource.Cancel();
-				UBAHordeSession? hordeSession = await hordeSessionTask;
-				if (hordeSession != null)
+				foreach (var coordinator in _agentCoordinators)
 				{
-					await hordeSession.DisposeAsync();
+					await coordinator.CloseAsync();
 				}
 				await _threadedLogger.FinishAsync();
 				_session = null;
@@ -505,7 +493,7 @@ namespace UnrealBuildTool
 		/// Executes the provided actions
 		/// </summary>
 		/// <returns>True if all the tasks successfully executed, or false if any of them failed.</returns>
-		bool ExecuteActionsInternal(IEnumerable<LinkedAction> inputActions, ISessionServer session, Task<UBAHordeSession?> hordeSessionTask, CancellationTokenSource hordeRequestCancellationSource, Microsoft.Extensions.Logging.ILogger logger, IActionArtifactCache? actionArtifactCache)
+		bool ExecuteActionsInternal(IEnumerable<LinkedAction> inputActions, ISessionServer session, Microsoft.Extensions.Logging.ILogger logger, IActionArtifactCache? actionArtifactCache)
 		{
 			using ImmediateActionQueue queue = CreateActionQueue(inputActions, actionArtifactCache, logger);
 			int actionLimit = Math.Min(NumParallelProcesses, queue.TotalActions);
@@ -539,12 +527,15 @@ namespace UnrealBuildTool
 					if (count <= NumParallelProcesses)
 					{
 						_bIsRemoteActionsAllowed = false;
-						hordeRequestCancellationSource.Cancel();
+						foreach (var coordinator in _agentCoordinators)
+						{
+							coordinator.Stop();
+						}
 						session!.DisableRemoteExecution();
 					}
 					else
 					{
-						// Tell UBA hordeSession max number of remote processes left. Providing this information makes it possible for UBA hordeSession to start disconnecting clients
+						// Tell UBA max number of remote processes left. Providing this information makes it possible for UBA to start disconnecting clients
 						session!.SetMaxRemoteProcessCount(count);
 					}
 				}
@@ -561,92 +552,22 @@ namespace UnrealBuildTool
 			// Add all actions we can add
 			queue.StartManyActions();
 
-			int timerPeriod = 5000;
-
-			bool shownNoAgentsFoundMessage = false;
-			Timer? hordeTimer = null;
 			try
 			{
-				hordeTimer = new(async (_) =>
+				foreach (var coordinator in _agentCoordinators)
 				{
-					hordeTimer?.Change(Timeout.Infinite, Timeout.Infinite);
-
-					UBAHordeSession? hordeSession = await hordeSessionTask;
-
-					if (queue.IsDone || hordeSession == null || !_bIsRemoteActionsAllowed)
-					{
-						return;
-					}
-
-					hordeSession.RemoveCompleteWorkers();
-
-					// We are assuming all active logical cores are already being used.. so queueWeight is essentially work that could be executed but can't because of bandwidth
-					double queueThreshold = UBAConfig.bForceBuildAllRemote ? 0 : 5;
-
-					try
-					{
-						double queueWeight = queue.EnumerateReadyToCompileActions().Where(x => CanRunRemotely(x)).Sum(x => x.Weight);
-
-						queueWeight -= hordeSession.QueuedUpCores();
-						while (true)
-						{
-							int currentLogicalCores = hordeSession.NumLogicalCores;
-
-							if (queueWeight <= queueThreshold || currentLogicalCores >= HordeConfig.HordeMaxCores)
-							{
-								break;
-							}
-
-							Requirements requirements = new()
-							{
-								Exclusive = true
-							};
-
-							if (!String.IsNullOrEmpty(HordeConfig.HordePool))
-							{
-								requirements.Pool = HordeConfig.HordePool;
-							}
-
-							if (HordeConfig.HordeCondition != null)
-							{
-								requirements.Condition = Condition.Parse(HordeConfig.HordeCondition);
-							}
-
-							if (!await hordeSession.AddWorkerAsync(requirements, hordeRequestCancellationSource.Token))
-							{
-								logger.LogDebug("No additional workers available");
-								break;
-							}
-							int coresAdded = hordeSession.NumLogicalCores - currentLogicalCores;
-							queueWeight -= coresAdded;
-						}
-					}
-					catch (NoComputeAgentsFoundException ex)
-					{
-						if (!shownNoAgentsFoundMessage)
-						{
-							logger.Log(UBAConfig.bStrict ? LogLevel.Warning : LogLevel.Information, KnownLogEvents.Systemic_Horde_Compute, ex, "No agents found matching requirements (cluster: {ClusterId}, requirements: {Requirements})", ex.ClusterId, ex.Requirements);
-							shownNoAgentsFoundMessage = true;
-						}
-					}
-					catch (Exception ex)
-					{
-						if (!hordeRequestCancellationSource.IsCancellationRequested)
-						{
-							logger.Log(UBAConfig.bStrict ? LogLevel.Error : LogLevel.Information, KnownLogEvents.Systemic_Horde_Compute, ex, "Unable to get worker: {Ex}", ex.ToString());
-						}
-					}
-
-					hordeTimer?.Change(timerPeriod, Timeout.Infinite);
-				}, null, HordeConfig.HordeDelay * 1000, timerPeriod);
+					coordinator.Start(queue, CanRunRemotely);
+				}
 
 				bool res = queue.RunTillDone().Result; // Using inline wait to avoid possible thread switch
-				hordeRequestCancellationSource.Cancel();
 				return res;
 			}
 			finally
 			{
-				hordeTimer?.Dispose();
+				foreach (var coordinator in _agentCoordinators)
+				{
+					coordinator.Stop();
+				}
 			}
 		}
 
@@ -869,4 +790,3 @@ namespace UnrealBuildTool
 		}
 	}
 }
-
diff --git a/Engine/Source/Programs/UnrealBuildTool/Executors/Experimental/UBAHordeSession.cs b/Engine/Source/Programs/UnrealBuildTool/Executors/Experimental/UBAHordeSession.cs
--- a/Engine/Source/Programs/UnrealBuildTool/Executors/Experimental/UBAHordeSession.cs
+++ b/Engine/Source/Programs/UnrealBuildTool/Executors/Experimental/UBAHordeSession.cs
@@ -11,11 +11,11 @@ using System.Threading;
 using System.Threading.Tasks;
 using EpicGames.Core;
 using EpicGames.Horde;
+using EpicGames.Horde.Common;
 using EpicGames.Horde.Compute;
 using EpicGames.Horde.Compute.Clients;
 using EpicGames.Horde.Storage;
 using EpicGames.Horde.Storage.Bundles;
-using EpicGames.Horde.Storage.Clients;
 using EpicGames.Horde.Storage.Nodes;
 using EpicGames.OIDC;
 using Microsoft.Extensions.Configuration;
@@ -354,7 +354,7 @@ namespace UnrealBuildTool
 				"LogicalCores"
 			};
 
-		public async Task<bool> AddWorkerAsync(Requirements requirements, CancellationToken cancellationToken)
+		public async Task<bool> AddWorkerAsync(Requirements requirements, UnrealBuildAcceleratorHordeConfig hordeConfig, CancellationToken cancellationToken)
 		{
 			if (_client == null)
 			{
@@ -439,7 +439,7 @@ namespace UnrealBuildTool
 					Port = ubaPort,
 					ProxyPort = ubaProxyPort,
 				};
-				worker.BackgroundTask = RunWorkerAsync(worker, lease, locator, exeName, workerLogger, _cancellationTokenSource.Token);
+				worker.BackgroundTask = RunWorkerAsync(worker, lease, locator, exeName, workerLogger, hordeConfig, _cancellationTokenSource.Token);
 				_workers.Add(worker);
 				lease = null; // Will be disposed by RunWorkerAsync
 
@@ -563,7 +563,7 @@ namespace UnrealBuildTool
 			return result.AccessToken;
 		}
 
-		async Task RunWorkerAsync(Worker self, IComputeLease lease, BlobLocator tool, string executable, ILogger logger, CancellationToken cancellationToken)
+		async Task RunWorkerAsync(Worker self, IComputeLease lease, BlobLocator tool, string executable, ILogger logger, UnrealBuildAcceleratorHordeConfig hordeConfig, CancellationToken cancellationToken)
 		{
 			logger.LogInformation("Running worker task..");
 			try
@@ -592,12 +592,12 @@ namespace UnrealBuildTool
 						await channel.UploadFilesAsync("", tool, _storage, cancellationToken);
 
 						string hordeHost = _owner.UBAConfig.Host;
-						if (!String.IsNullOrEmpty(_owner.HordeConfig.HordeHost))
+						if (!String.IsNullOrEmpty(hordeConfig.HordeHost))
 						{
-							hordeHost = _owner.HordeConfig.HordeHost;
+							hordeHost = hordeConfig.HordeHost;
 						}
 
-						bool useListen = !String.IsNullOrEmpty(_owner.HordeConfig.HordeHost);
+						bool useListen = !String.IsNullOrEmpty(hordeConfig.HordeHost);
 						List<string> arguments = new();
 
 						if (useListen)
@@ -616,9 +616,9 @@ namespace UnrealBuildTool
 
 						arguments.Add("-NoPoll");
 						arguments.Add("-Quiet");
-						if (!String.IsNullOrEmpty(_owner.HordeConfig.UBASentryUrl))
+						if (!String.IsNullOrEmpty(hordeConfig.UBASentryUrl))
 						{
-							arguments.Add($"-Sentry=\"{_owner.HordeConfig.UBASentryUrl}\"");
+							arguments.Add($"-Sentry=\"{hordeConfig.UBASentryUrl}\"");
 						}
 						arguments.Add("-ProxyPort=" + self.ProxyPort.AgentPort);
 						if (_owner.UBAConfig.bUseQuic)
@@ -686,4 +686,170 @@ namespace UnrealBuildTool
 			}
 		}
 	}
+
+	class UBAAgentCoordinatorHorde : IUBAAgentCoordinator
+	{
+		public UBAAgentCoordinatorHorde(ILogger logger, UnrealBuildAcceleratorConfig ubaConfig, CommandLineArguments? additionalArguments = null)
+		{
+			_logger = logger;
+			_ubaConfig = ubaConfig;
+
+			XmlConfig.ApplyTo(_hordeConfig);
+			additionalArguments?.ApplyTo(_hordeConfig);
+
+			// Sentry is currently unsupported for non-Windows and non-x64
+			if (!OperatingSystem.IsWindows() || RuntimeInformation.ProcessArchitecture != Architecture.X64)
+			{
+				_hordeConfig.UBASentryUrl = null;
+			}
+		}
+
+		public DirectoryReference? GetUBARootDir()
+		{
+			DirectoryReference? hordeSharedDir = DirectoryReference.FromString(Environment.GetEnvironmentVariable("UE_HORDE_SHARED_DIR"));
+			if (hordeSharedDir != null)
+			{
+				return DirectoryReference.Combine(hordeSharedDir, "UbaHost");
+			}
+			return null;
+		}
+
+		public async Task InitAsync(UBAExecutor executor)
+		{
+			if (_ubaConfig.bDisableRemote)
+			{
+				return;
+			}
+
+			_cancellationSource = new CancellationTokenSource();
+			_hordeSessionTask = UBAHordeSession.TryCreateHordeSession(_hordeConfig, executor, _ubaConfig.bStrict, _logger, _cancellationSource.Token);
+			await _hordeSessionTask;
+		}
+
+		public void Start(ImmediateActionQueue queue, Func<LinkedAction, bool> canRunRemotely)
+		{
+			int timerPeriod = 5000;
+			bool shownNoAgentsFoundMessage = false;
+
+			if (_hordeSessionTask == null)
+			{
+				return;
+			}
+
+			_timer = new(async (_) =>
+			{
+				_timer?.Change(Timeout.Infinite, Timeout.Infinite);
+
+				UBAHordeSession? hordeSession = await _hordeSessionTask!;
+
+				if (hordeSession == null)
+				{
+					return;
+				}
+
+				hordeSession.RemoveCompleteWorkers();
+
+				if (queue.IsDone)
+				{
+					return;
+				}
+
+				// We are assuming all active logical cores are already being used.. so queueWeight is essentially work that could be executed but can't because of bandwidth
+				double queueThreshold = _ubaConfig.bForceBuildAllRemote ? 0 : 5;
+
+				try
+				{
+					double queueWeight = queue.EnumerateReadyToCompileActions().Where(x => canRunRemotely(x)).Sum(x => x.Weight);
+
+					queueWeight -= hordeSession.QueuedUpCores();
+					while (true)
+					{
+						int currentLogicalCores = hordeSession.NumLogicalCores;
+
+						if (queueWeight <= queueThreshold || currentLogicalCores >= _hordeConfig.HordeMaxCores || _cancellationSource!.IsCancellationRequested)
+						{
+							break;
+						}
+
+						Requirements requirements = new()
+						{
+							Exclusive = true
+						};
+
+						if (!String.IsNullOrEmpty(_hordeConfig.HordePool))
+						{
+							requirements.Pool = _hordeConfig.HordePool;
+						}
+
+						if (_hordeConfig.HordeCondition != null)
+						{
+							requirements.Condition = Condition.Parse(_hordeConfig.HordeCondition);
+						}
+
+						if (!await hordeSession.AddWorkerAsync(requirements, _hordeConfig, _cancellationSource.Token))
+						{
+							_logger.LogDebug("No additional workers available");
+							break;
+						}
+						int coresAdded = hordeSession.NumLogicalCores - currentLogicalCores;
+						queueWeight -= coresAdded;
+					}
+				}
+				catch (NoComputeAgentsFoundException ex)
+				{
+					if (!shownNoAgentsFoundMessage)
+					{
+						_logger.Log(_ubaConfig.bStrict ? LogLevel.Warning : LogLevel.Information, KnownLogEvents.Systemic_Horde_Compute, ex, "No agents found matching requirements (cluster: {ClusterId}, requirements: {Requirements})", ex.ClusterId, ex.Requirements);
+						shownNoAgentsFoundMessage = true;
+					}
+				}
+				catch (Exception ex)
+				{
+					if (!_cancellationSource!.IsCancellationRequested)
+					{
+						_logger.Log(_ubaConfig.bStrict ? LogLevel.Error : LogLevel.Information, KnownLogEvents.Systemic_Horde_Compute, ex, "Unable to get worker: {Ex}", ex.ToString());
+					}
+				}
+
+				_timer?.Change(timerPeriod, Timeout.Infinite);
+			}, null, _hordeConfig.HordeDelay * 1000, timerPeriod);
+		}
+
+		public void Stop()
+		{
+			_cancellationSource?.Cancel();
+
+			_timer?.Dispose();
+		}
+
+		public async Task CloseAsync()
+		{
+			_cancellationSource?.Cancel();
+
+			if (_hordeSessionTask == null)
+			{
+				return;
+			}
+
+			UBAHordeSession? hordeSession = await _hordeSessionTask;
+			if (hordeSession != null)
+			{
+				await hordeSession.DisposeAsync();
+			}
+		}
+
+		public void Dispose()
+		{
+			Stop();
+			CloseAsync().Wait();
+		}
+
+		ILogger _logger;
+		UnrealBuildAcceleratorConfig _ubaConfig;
+		UnrealBuildAcceleratorHordeConfig _hordeConfig { get; init; } = new();
+
+		CancellationTokenSource? _cancellationSource;
+		Task<UBAHordeSession?>? _hordeSessionTask;
+		Timer? _timer;
+	}
 }
