[UBA]
* Added lock around disposing horde agent request timer to prevent it from being disposed while running

diff --git a/Engine/Source/Programs/UnrealBuildTool/Executors/Experimental/UBAHordeSession.cs b/Engine/Source/Programs/UnrealBuildTool/Executors/Experimental/UBAHordeSession.cs
--- a/Engine/Source/Programs/UnrealBuildTool/Executors/Experimental/UBAHordeSession.cs
+++ b/Engine/Source/Programs/UnrealBuildTool/Executors/Experimental/UBAHordeSession.cs
@@ -735,11 +735,15 @@ namespace UnrealBuildTool
 			{
 				return;
 			}
-
 			_timer = new(async (_) =>
 			{
 				_timer?.Change(Timeout.Infinite, Timeout.Infinite);
 
+				if (_cancellationSource!.IsCancellationRequested)
+				{
+					return;
+				}	
+
 				UBAHordeSession? hordeSession = await _hordeSessionTask!;
 
 				if (hordeSession == null)
@@ -818,8 +822,6 @@ namespace UnrealBuildTool
 		public void Stop()
 		{
 			_cancellationSource?.Cancel();
-
-			_timer?.Dispose();
 		}
 
 		public async Task CloseAsync()
