PipInstall: Handle updated or changed python versions in-engine by recreating pip virtual-env from scratch.

diff --git a/Engine/Plugins/Experimental/PythonScriptPlugin/Source/PythonScriptPlugin/Private/PipInstall.cpp b/Engine/Plugins/Experimental/PythonScriptPlugin/Source/PythonScriptPlugin/Private/PipInstall.cpp
--- a/Engine/Plugins/Experimental/PythonScriptPlugin/Source/PythonScriptPlugin/Private/PipInstall.cpp
+++ b/Engine/Plugins/Experimental/PythonScriptPlugin/Source/PythonScriptPlugin/Private/PipInstall.cpp
@@ -135,6 +135,28 @@ FString FPipInstall::WritePluginDependencies(const TArray<TSharedRef<IPlugin>>&
 	return MergedReqsFile;
 }
 
+void FPipInstall::CheckInvalidPipEnv()
+{
+	TRACE_CPUPROFILER_EVENT_SCOPE(FPipInstall::CheckInvalidPipEnv);
+
+	const FString PipInstallPath = GetPipInstallPath();
+	if (!FPaths::DirectoryExists(PipInstallPath))
+	{
+		return;
+	}
+
+	const FString VenvVersion = ParseVenvVersion(PipInstallPath);
+	if (VenvVersion == TEXT(PY_VERSION))
+	{
+		return;
+	}
+
+	UE_LOG(LogPython, Display, TEXT("Engine python version (%s) incompatible with venv (%s), recreating..."), TEXT(PY_VERSION), *VenvVersion);
+
+	IPlatformFile& PlatformFile = FPlatformFileManager::Get().GetPlatformFile();
+	PlatformFile.DeleteDirectoryRecursively(*PipInstallPath);
+}
+
 void FPipInstall::SetupPipEnv(FFeedbackContext* Context, bool bForceRebuild /* = false */)
 {
 	TRACE_CPUPROFILER_EVENT_SCOPE(FPipInstall::SetupPipEnv);
@@ -305,6 +327,35 @@ FString FPipInstall::GetPythonScriptPluginPath()
 	return PythonPlugin->GetBaseDir();
 }
 
+FString FPipInstall::ParseVenvVersion(const FString& InstallPath)
+{
+	FString VenvConfig = InstallPath / TEXT("pyvenv.cfg");
+	if (!FPaths::FileExists(VenvConfig))
+	{
+		return TEXT("");
+	}
+
+	TArray<FString> ConfigLines;
+	if (!FFileHelper::LoadFileToStringArray(ConfigLines, *VenvConfig))
+	{
+		return TEXT("");
+	}
+
+	for (FStringView Line : ConfigLines)
+	{
+		FStringView ChkLine = Line.TrimStartAndEnd();
+		if (!ChkLine.StartsWith(TEXT("version =")))
+		{
+			continue;
+		}
+
+		FStringView Version = ChkLine.RightChop(9).TrimStart();
+		return FString(Version);
+	}
+
+	return TEXT("");
+}
+
 FString FPipInstall::GetVenvInterpreter(const FString& InstallPath)
 {
 #if PLATFORM_WINDOWS
diff --git a/Engine/Plugins/Experimental/PythonScriptPlugin/Source/PythonScriptPlugin/Private/PipInstall.h b/Engine/Plugins/Experimental/PythonScriptPlugin/Source/PythonScriptPlugin/Private/PipInstall.h
--- a/Engine/Plugins/Experimental/PythonScriptPlugin/Source/PythonScriptPlugin/Private/PipInstall.h
+++ b/Engine/Plugins/Experimental/PythonScriptPlugin/Source/PythonScriptPlugin/Private/PipInstall.h
@@ -23,6 +23,8 @@ public:
 
 	static bool EnabledOnStartup();
 
+	static void CheckInvalidPipEnv();
+
 	static FString WritePluginsListing(TArray<TSharedRef<IPlugin>>& OutPythonPlugins);
 	static FString WritePluginDependencies(const TArray<TSharedRef<IPlugin>>& PythonPlugins, TArray<FString>& OutRequirements, TArray<FString>& OutExtraUrls);
 
@@ -37,9 +39,10 @@ private:
 	static void SetupPipInstallUtils(const FString& VenvInterp, FFeedbackContext* Context);
 	static bool CheckPipInstallUtils(const FString& VenvInterp, FFeedbackContext* Context);
 	static int32 RunPythonCmd(const FText& Description, const FString& VenvInterp, const FString& Cmd, FFeedbackContext* Context);
-	static bool RunLoggedSubprocess(const FText& Description, const FString& URL, const FString& Params, FFeedbackContext* Context, int32* OutExitCode);
+	static bool RunLoggedSubprocess(const FText& Description, const FString& URL, const FString& Params, FFeedbackContext* Context, int32* OutExitCode);\
 
 	static FString GetPythonScriptPluginPath();
+	static FString ParseVenvVersion(const FString& InstallPath);
 	static FString GetVenvInterpreter(const FString& InstallPath);
 	static bool CheckCompatiblePlatform(const TSharedPtr<FJsonObject>& JsonObject, const FString& PlatformName);
 };
diff --git a/Engine/Plugins/Experimental/PythonScriptPlugin/Source/PythonScriptPlugin/Private/PythonScriptPlugin.cpp b/Engine/Plugins/Experimental/PythonScriptPlugin/Source/PythonScriptPlugin/Private/PythonScriptPlugin.cpp
--- a/Engine/Plugins/Experimental/PythonScriptPlugin/Source/PythonScriptPlugin/Private/PythonScriptPlugin.cpp
+++ b/Engine/Plugins/Experimental/PythonScriptPlugin/Source/PythonScriptPlugin/Private/PythonScriptPlugin.cpp
@@ -1134,6 +1134,7 @@ void FPythonScriptPlugin::InitPipInstaller()
 	FFeedbackContext* Context = GWarn;
 	FScopedSlowTask PipInstallTask(0, LOCTEXT("PipInstall.RunInit", "Running Pip Init Tasks..."), true, *Context);
 
+	FPipInstall::CheckInvalidPipEnv();
 	const FString PipSitePackagePath = FPaths::ConvertRelativePathToFull(GetPipSitePackagesPath());
 
 	// Generate the input listing files of plugins with python dependencies and the listing of all requirements (installed or not)
diff --git a/Engine/Source/Programs/UnrealBuildTool/Modes/PipInstallMode.cs b/Engine/Source/Programs/UnrealBuildTool/Modes/PipInstallMode.cs
--- a/Engine/Source/Programs/UnrealBuildTool/Modes/PipInstallMode.cs
+++ b/Engine/Source/Programs/UnrealBuildTool/Modes/PipInstallMode.cs
@@ -10,6 +10,7 @@ using Microsoft.Extensions.Logging;
 using EpicGames.Core;
 using UnrealBuildBase;
 using System.IO;
+using System.Text.RegularExpressions;
 
 namespace UnrealBuildTool.Modes
 {
@@ -114,6 +115,9 @@ namespace UnrealBuildTool.Modes
 			PipEnv Pip = new(InstallDir, Platform, Logger, ProgressWriter.bWriteMarkup);
 			if ((Action & (PipAction)ActionBits.GenReqs) != 0)
 			{
+				// Make sure the virtual environment used for installs is compatible with python interpreter version
+				Pip.RemoveInvalidVenv(PythonInterpreter);
+
 				Pip.WritePluginsListing(Target, Logger);
 				if (!Pip.WritePluginDependencies())
 				{
@@ -180,7 +184,8 @@ namespace UnrealBuildTool.Modes
 
 		private UnrealTargetPlatform TargetPlatform;
 		private DirectoryReference InstallDir;
-		private FileReference PythonVenv;
+		private FileReference PythonVenvExe;
+		private string? PythonVenvVer;
 
 		private ILogger Logger;
 		private IBaseProgressLogFactory LoggerFactory;
@@ -193,12 +198,35 @@ namespace UnrealBuildTool.Modes
 
 			LoggerFactory = (UseProgressWriter) ? new PipProgressLogCreator(Logger) : new SimpleCmdLogCreator(Logger);
 
-			if (!DirectoryReference.Exists(InstallDir))
+			PythonVenvExe = GetVenvInterpreter(InstallDir, TargetPlatform);
+			PythonVenvVer = ParseVenvVersion(InstallDir);
+		}
+
+		public void RemoveInvalidVenv(FileReference? EnginePython)
+		{
+			// Always delete directory if can't find valid venv config
+			if ( PythonVenvVer == null )
 			{
-				DirectoryReference.CreateDirectory(InstallDir);
+				CleanVenvDir();
+				return;
+			}
+
+			FileReference EnginePythonInterp = GetEnginePythonInterpreter(EnginePython);
+			if (!FileReference.Exists(EnginePythonInterp))
+			{
+				Logger.LogError("PipInstall: Invalid path to UE python interpreter: {Interp}", EnginePythonInterp.ToString());
+				return;
 			}
 
-			PythonVenv = GetVenvInterpreter(InstallDir, TargetPlatform);
+			using (IBaseCmdProgressLogger SimpleLogger =new SimpleCmdLogger(Logger))
+			{
+				const string PyInterpVerCheckCmd = "import sys; exit(0) if f'{sys.version_info[0]}.{sys.version_info[1]}.{sys.version_info[2]}' == sys.argv[1] else exit(1)";
+				if (RunPythonCmd(EnginePythonInterp, $"-c \"{PyInterpVerCheckCmd}\" \"{PythonVenvVer}\"", SimpleLogger) != 0)
+				{
+					Logger.LogWarning("PipInstall: Found Incompatible virtual environment ({VenvVer}), removing...", PythonVenvVer);
+					CleanVenvDir();
+				}
+			}
 		}
 
 		public void WritePluginsListing(UEBuildTarget Target, ILogger Logger)
@@ -277,7 +305,7 @@ namespace UnrealBuildTool.Modes
 		{
 			using (IBaseCmdProgressLogger CmdLogger = LoggerFactory.Create("Creating pip installer virtual environment", 5))
 			{
-				if (!ForceRebuild && FileReference.Exists(PythonVenv))
+				if (!ForceRebuild && FileReference.Exists(PythonVenvExe))
 				{
 					return SetupPipInstallUtils(CmdLogger);
 				}
@@ -295,7 +323,7 @@ namespace UnrealBuildTool.Modes
 				}
 
 				int result = RunPythonCmd(EnginePythonInterp, $"-m venv \"{InstallDir}\"", CmdLogger);
-				if (result != 0 || !FileReference.Exists(PythonVenv))
+				if (result != 0 || !FileReference.Exists(PythonVenvExe))
 				{
 					return false;
 				}
@@ -374,6 +402,44 @@ namespace UnrealBuildTool.Modes
 				|| string.Equals(PlatformField, "All", StringComparison.InvariantCultureIgnoreCase);
 		}
 
+		private string? ParseVenvVersion(DirectoryReference VenvDir)
+		{
+			FileReference VenvConfig = FileReference.Combine(VenvDir, "pyvenv.cfg");
+			if ( !FileReference.Exists(VenvConfig) )
+			{
+				Logger.LogWarning("PipInstall: Unable to find venv config: {VenvFile}", VenvConfig);
+				return null;
+			}
+
+			string ConfigInfo = FileReference.ReadAllText(VenvConfig);
+			Match m = Regex.Match(ConfigInfo, @"version\s*=\s*(\d+\.\d+\.\d+)", RegexOptions.IgnoreCase);
+			if ( !m.Success )
+			{
+				Logger.LogWarning("PipInstall: Unable to match venv version config: {VenvFile}", ConfigInfo);
+				return null;
+			}
+
+			return m.Groups[1].Value;
+		}
+
+		private void CleanVenvDir()
+		{
+			if (!DirectoryReference.Exists(InstallDir))
+			{
+				DirectoryReference.CreateDirectory(InstallDir);
+				return;
+			}
+
+			// HACK: On windows these script files are set read-only and can't be deleted
+			foreach (FileReference File in DirectoryReference.EnumerateFiles(DirectoryReference.Combine(InstallDir,"Scripts")))
+			{
+				FileReference.SetAttributes(File, FileAttributes.Normal);
+			}
+
+			DirectoryReference.Delete(InstallDir, true);
+			DirectoryReference.CreateDirectory(InstallDir);
+		}
+
 		private bool PipInstall(FileReference RequirementsFile, string[]? ExtraUrls, bool OfflineOnly, string? ForceIndexUrl)
 		{
 			string[] Reqs = FileReference.ReadAllLines(RequirementsFile);
@@ -502,7 +568,7 @@ namespace UnrealBuildTool.Modes
 
 		private int RunPythonVenv(string Args, IBaseCmdProgressLogger InCmdLogger)
 		{
-			return RunPythonCmd(PythonVenv, Args, InCmdLogger);
+			return RunPythonCmd(PythonVenvExe, Args, InCmdLogger);
 		}
 
 		private int RunPythonCmd(FileReference PythonBin, string Args, IBaseCmdProgressLogger InCmdLogger)
