SNDBS: Fixed a missing .cpp dependencies when building precompiled headers with distribution on some platforms.

#rnx

diff --git a/Engine/Source/Programs/UnrealBuildTool/Executors/SNDBS.cs b/Engine/Source/Programs/UnrealBuildTool/Executors/SNDBS.cs
--- a/Engine/Source/Programs/UnrealBuildTool/Executors/SNDBS.cs
+++ b/Engine/Source/Programs/UnrealBuildTool/Executors/SNDBS.cs
@@ -268,6 +268,25 @@ namespace UnrealBuildTool
 						}).ToList();
 					}
 
+					// Add PCH source file dependencies for clang or clang-cl
+					if (a.bIsGCCCompiler)
+					{
+						// Look for any prerequisite actions that produce .pch files and add any .cpp files they depend on
+						var ExplicitInputFiles = Job["explicit_input_files"] as List<Dictionary<string, object>>;
+
+						if (ExplicitInputFiles != null)
+						{
+							ExplicitInputFiles.AddRange(a.PrerequisiteActions
+								.Where(Prereq => Prereq.ProducedItems.Any(Produced => Produced.AbsolutePath.EndsWith(".gch")))
+								.SelectMany(Prereq => Prereq.PrerequisiteItems, (_, PrereqFile) => PrereqFile.AbsolutePath)
+								.Where(Path => Path.EndsWith(".cpp"))
+								.Select(Path => new Dictionary<string, object>()
+								{
+									["filename"] = Path
+								}));
+						}
+					}
+
 					string CommandDescription = String.IsNullOrWhiteSpace(a.CommandDescription) ? a.ActionType.ToString() : a.CommandDescription;
 					Job["echo"] = $"{ProgressMarkupPrefix}{CommandDescription}:{a.StatusDescription}";
 
