UnrealBuildTool: Improve cycle route when circular dependencies are found

#rnx

diff --git a/Engine/Source/Programs/UnrealBuildTool/Configuration/UEBuildModule.cs b/Engine/Source/Programs/UnrealBuildTool/Configuration/UEBuildModule.cs
--- a/Engine/Source/Programs/UnrealBuildTool/Configuration/UEBuildModule.cs
+++ b/Engine/Source/Programs/UnrealBuildTool/Configuration/UEBuildModule.cs
@@ -1228,7 +1228,7 @@ namespace UnrealBuildTool
 				// No module has confessed its guilt, so this is an error.
 				if (String.IsNullOrEmpty(GuiltyModule))
 				{
-					string CycleChain = String.Join(" -> ", ReferenceStack);
+					string CycleChain = String.Join(" -> ", ReferenceStack.SkipWhile(x => x != this));
 					Logger.LogError("Circular dependency on {Name} detected:\n" +
 						"\tFull Route: {FullRoute}\n" +
 						"\tCycled Route: is {CycleRoute}\n" +
