Check directory exists before adding to dependency check

#android

diff --git a/Engine/Source/Programs/UnrealBuildTool/Platform/Android/UEDeployAndroid.cs b/Engine/Source/Programs/UnrealBuildTool/Platform/Android/UEDeployAndroid.cs
--- a/Engine/Source/Programs/UnrealBuildTool/Platform/Android/UEDeployAndroid.cs
+++ b/Engine/Source/Programs/UnrealBuildTool/Platform/Android/UEDeployAndroid.cs
@@ -2136,8 +2136,11 @@ namespace UnrealBuildTool
 				foreach (var SoDirName in new [] {"jni", "libs"})
 				{
 					string SoDirPath = Path.Combine(IntermediateAndroidPath, ArchRemapping[GetNDKArch(Arch)], SoDirName, GetNDKArch(Arch));
-					var files = Directory.EnumerateFiles(SoDirPath, "*.*", SearchOption.AllDirectories);
-					InputFiles.AddRange(files);
+					if (Directory.Exists(SoDirPath))
+					{
+						var files = Directory.EnumerateFiles(SoDirPath, "*.*", SearchOption.AllDirectories);
+						InputFiles.AddRange(files);
+					}
 				}
 
 				// add all Engine and Project build files to be safe
