- Allow for per-host-platform remote pools so Mac doesn't attempt to use Windows/Linux pool

#changelist validated
#virtualized

diff --git a/Engine/Source/Programs/UnrealBuildTool/Executors/Experimental/UBAHordeSession.cs b/Engine/Source/Programs/UnrealBuildTool/Executors/Experimental/UBAHordeSession.cs
--- a/Engine/Source/Programs/UnrealBuildTool/Executors/Experimental/UBAHordeSession.cs
+++ b/Engine/Source/Programs/UnrealBuildTool/Executors/Experimental/UBAHordeSession.cs
@@ -57,12 +57,65 @@ namespace UnrealBuildTool
 		public string? HordeOidcProvider { get; set; }
 
 		/// <summary>
-		/// Pool for the Horde agent to assign
+		/// Pool for the Horde agent to assign, calculated based on platform and overrides
+		/// </summary>
+		public string? HordePool
+		{
+			get
+			{
+				string? Pool = DefaultHordePool;
+				if (OverrideHordePool != null)
+				{
+					Pool = OverrideHordePool;
+				}
+				else if (OperatingSystem.IsWindows() && WindowsHordePool != null)
+				{
+					Pool = WindowsHordePool;
+				}
+				else if (OperatingSystem.IsMacOS() && MacHordePool != null)
+				{
+					Pool = MacHordePool;
+				}
+				else if (OperatingSystem.IsLinux() && LinuxHordePool != null)
+				{
+					Pool = LinuxHordePool;
+				}
+
+				Console.WriteLine("CHOSEN UBA POOL: {0}", Pool);
+				return Pool;
+			}
+		}
+
+		/// <summary>
+		/// Pool for the Horde agent to assign if no override current platform doesn't have it set
 		/// </summary>
 		[XmlConfigFile(Category = "Horde", Name = "Pool")]
+		public string? DefaultHordePool { get; set; }
+
+		/// <summary>
+		/// Pool for the Horde agent to assign, only used for commandline override
+		/// </summary>
 		[CommandLine("-BoxHordePool=")]
 		[CommandLine("-UBAHordePool=")]
-		public string? HordePool { get; set; }
+		public string? OverrideHordePool { get; set; }
+
+		/// <summary>
+		/// Pool for the Horde agent to assign when on Linux
+		/// </summary>
+		[XmlConfigFile(Category = "Horde", Name = "LinuxPool")]
+		public string? LinuxHordePool { get; set; }
+
+		/// <summary>
+		/// Pool for the Horde agent to assign when on Mac
+		/// </summary>
+		[XmlConfigFile(Category = "Horde", Name = "MacPool")]
+		public string? MacHordePool { get; set; }
+
+		/// <summary>
+		/// Pool for the Horde agent to assign when on Windows
+		/// </summary>
+		[XmlConfigFile(Category = "Horde", Name = "WindowsPool")]
+		public string? WindowsHordePool { get; set; }
 
 		/// <summary>
 		/// Requirements for the Horde agent to assign
