UnrealBuildTool: Print additional details about failed actions to the log file

#rnx

diff --git a/Engine/Source/Programs/UnrealBuildTool/Executors/ImmediateActionQueue.cs b/Engine/Source/Programs/UnrealBuildTool/Executors/ImmediateActionQueue.cs
--- a/Engine/Source/Programs/UnrealBuildTool/Executors/ImmediateActionQueue.cs
+++ b/Engine/Source/Programs/UnrealBuildTool/Executors/ImmediateActionQueue.cs
@@ -1125,20 +1125,31 @@ namespace UnrealBuildTool
 
 				if (exitCode != 0)
 				{
+					string exitCodeStr = String.Empty;
+					if ((uint)exitCode == 0xC0000005)
+					{
+						exitCodeStr = "(Access violation)";
+					}
+					else if ((uint)exitCode == 0xC0000409)
+					{
+						exitCodeStr = "(Stack buffer overflow)";
+					}
 
 					// If we have an error code but no output, chances are the tool crashed.  Generate more detailed information to let the
 					// user know something went wrong.
 					if (logLines == null || logLines.Count <= (action.bShouldOutputStatusDescription ? 0 : 1))
 					{
-						string exitCodeStr = string.Empty;
-						if ((uint)exitCode == 0xC0000005)
-							exitCodeStr = "(Access violation)";
-						else if ((uint)exitCode == 0xC0000409)
-							exitCodeStr = "(Stack buffer overflow)";
 						Logger.LogError("{TargetDetails} {Description}: Exited with error code {ExitCode} {ExitCodeStr}. The build will fail.", targetDetails, description, exitCode, exitCodeStr);
 						Logger.LogInformation("{TargetDetails} {Description}: WorkingDirectory {WorkingDirectory}", targetDetails, description, action.WorkingDirectory);
 						Logger.LogInformation("{TargetDetails} {Description}: {CommandPath} {CommandArguments}", targetDetails, description, action.CommandPath, action.CommandArguments);
 					}
+					// Always print error details to to the log file
+					else
+					{
+						Logger.LogDebug("{TargetDetails} {Description}: Exited with error code {ExitCode} {ExitCodeStr}. The build will fail.", targetDetails, description, exitCode, exitCodeStr);
+						Logger.LogDebug("{TargetDetails} {Description}: WorkingDirectory {WorkingDirectory}", targetDetails, description, action.WorkingDirectory);
+						Logger.LogDebug("{TargetDetails} {Description}: {CommandPath} {CommandArguments}", targetDetails, description, action.CommandPath, action.CommandArguments);
+					}
 
 					// prevent overwriting of error text
 					s_previousLineLength = -1;
