UnrealBuildTool: Fix incorrect comment for CppCompileEnvironment.bValidateInternalApi

#rnx

diff --git a/Engine/Source/Programs/UnrealBuildTool/System/CppCompileEnvironment.cs b/Engine/Source/Programs/UnrealBuildTool/System/CppCompileEnvironment.cs
--- a/Engine/Source/Programs/UnrealBuildTool/System/CppCompileEnvironment.cs
+++ b/Engine/Source/Programs/UnrealBuildTool/System/CppCompileEnvironment.cs
@@ -689,7 +689,7 @@ namespace UnrealBuildTool
 		public bool bValidateFormatStrings = true;
 
 		/// <summary>
-		/// Emits compilation errors for incorrect UE_LOG format strings.
+		/// Emits deprecated warnings\errors for internal API usage for non-engine modules.
 		/// </summary>
 		public bool bValidateInternalApi = false;
 
