PR #11345: Fix for rebuilding whole Unreal Engine sources in Rider

diff --git a/Engine/Source/Programs/UnrealBuildTool/Configuration/UEBuildTarget.cs b/Engine/Source/Programs/UnrealBuildTool/Configuration/UEBuildTarget.cs
--- a/Engine/Source/Programs/UnrealBuildTool/Configuration/UEBuildTarget.cs
+++ b/Engine/Source/Programs/UnrealBuildTool/Configuration/UEBuildTarget.cs
@@ -1131,6 +1131,10 @@ namespace UnrealBuildTool
 		/// Query
 		/// </summary>
 		Query,
+		/// <summary>
+		/// Generate project files
+		/// </summary>
+		GenerateProjectFiles,
 	}
 
 	/// <summary>
@@ -1741,6 +1745,9 @@ namespace UnrealBuildTool
 				case UnrealIntermediateEnvironment.Query:
 					TargetFolderName += "QRY";
 					break;
+				case UnrealIntermediateEnvironment.GenerateProjectFiles:
+					TargetFolderName += "GPF";
+					break;
 			}
 			return TargetFolderName;
 		}
diff --git a/Engine/Source/Programs/UnrealBuildTool/ProjectFiles/Rider/RiderProjectFile.cs b/Engine/Source/Programs/UnrealBuildTool/ProjectFiles/Rider/RiderProjectFile.cs
--- a/Engine/Source/Programs/UnrealBuildTool/ProjectFiles/Rider/RiderProjectFile.cs
+++ b/Engine/Source/Programs/UnrealBuildTool/ProjectFiles/Rider/RiderProjectFile.cs
@@ -111,7 +111,7 @@ namespace UnrealBuildTool
 							Platform, Configuration, ProjectArchitectures, Arguments);
 						try
 						{
-							UEBuildTarget BuildTarget = UEBuildTarget.Create(TargetDesc, false, false, false, Logger);
+							UEBuildTarget BuildTarget = UEBuildTarget.Create(TargetDesc, false, false, false, UnrealIntermediateEnvironment.GenerateProjectFiles, Logger);
 
 							FileReference OutputFile = FileReference.Combine(TargetFolder, $"{ProjectName}.json");
 							FileToTarget.Add(new TargetEntry(OutputFile, BuildTarget, bBuildByDefault));
