UnrealBuildTool: Move module and plugin validation out of BuildAsync and cleanup warnings

#rnx

diff --git a/Engine/Source/Programs/UnrealBuildTool/Configuration/UEBuildModule.cs b/Engine/Source/Programs/UnrealBuildTool/Configuration/UEBuildModule.cs
--- a/Engine/Source/Programs/UnrealBuildTool/Configuration/UEBuildModule.cs
+++ b/Engine/Source/Programs/UnrealBuildTool/Configuration/UEBuildModule.cs
@@ -341,16 +341,16 @@ namespace UnrealBuildTool
 		public HashSet<UEBuildModule> GetDependencies(bool bWithIncludePathModules, bool bWithDynamicallyLoadedModules)
 		{
 			HashSet<UEBuildModule> Modules = new HashSet<UEBuildModule>();
-			Modules.UnionWith(PublicDependencyModules!);
-			Modules.UnionWith(PrivateDependencyModules!);
+			Modules.UnionWith(PublicDependencyModules ?? new());
+			Modules.UnionWith(PrivateDependencyModules ?? new());
 			if (bWithIncludePathModules)
 			{
-				Modules.UnionWith(PublicIncludePathModules!);
-				Modules.UnionWith(PrivateIncludePathModules!);
+				Modules.UnionWith(PublicIncludePathModules ?? new());
+				Modules.UnionWith(PrivateIncludePathModules ?? new());
 			}
 			if (bWithDynamicallyLoadedModules)
 			{
-				Modules.UnionWith(DynamicallyLoadedModules!);
+				Modules.UnionWith(DynamicallyLoadedModules ?? new());
 			}
 			return Modules;
 		}
diff --git a/Engine/Source/Programs/UnrealBuildTool/Configuration/UEBuildPlatform.cs b/Engine/Source/Programs/UnrealBuildTool/Configuration/UEBuildPlatform.cs
--- a/Engine/Source/Programs/UnrealBuildTool/Configuration/UEBuildPlatform.cs
+++ b/Engine/Source/Programs/UnrealBuildTool/Configuration/UEBuildPlatform.cs
@@ -792,6 +792,13 @@ namespace UnrealBuildTool
 		{
 		}
 
+		/// <summary>
+		/// Validate a plugin's settings
+		/// </summary>
+		public virtual void ValidatePlugin(UEBuildPlugin Plugin, ReadOnlyTargetRules Target)
+		{
+		}
+
 		/// <summary>
 		/// Validate a UEBuildModule before it's processed
 		/// <param name="Module">The UEBuildModule that needs to be validated</param>
diff --git a/Engine/Source/Programs/UnrealBuildTool/Configuration/UEBuildPlugin.cs b/Engine/Source/Programs/UnrealBuildTool/Configuration/UEBuildPlugin.cs
--- a/Engine/Source/Programs/UnrealBuildTool/Configuration/UEBuildPlugin.cs
+++ b/Engine/Source/Programs/UnrealBuildTool/Configuration/UEBuildPlugin.cs
@@ -1,7 +1,10 @@
 // Copyright Epic Games, Inc. All Rights Reserved.
 
 using System.Collections.Generic;
+using System.Linq;
 using EpicGames.Core;
+using Microsoft.Extensions.Logging;
+using UnrealBuildBase;
 
 namespace UnrealBuildTool
 {
@@ -89,5 +92,40 @@ namespace UnrealBuildTool
 		{
 			return Info.Name;
 		}
+
+		/// <summary>
+		/// Public entry point to validate a module
+		/// </summary>
+		/// <param name="logger"></param>
+		/// <returns>true if there are any fatal errors</returns>
+		public bool ValidatePlugin(ILogger logger)
+		{
+			bool anyErrors = false;
+			if (Dependencies != null)
+			{
+				foreach (UEBuildPlugin dependencyPlugin in Dependencies)
+				{
+					if (dependencyPlugin.Descriptor.bIsSealed)
+					{
+						logger.LogError("Plugin '{PluginName}' cannot depend on plugin '{DependencyPluginName}' because it is sealed.", Name, dependencyPlugin.Name);
+						anyErrors = true;
+					}
+					else if (Descriptor.DisallowedPlugins != null && Descriptor.DisallowedPlugins.Contains(dependencyPlugin.Name))
+					{
+						logger.LogError("Plugin '{PluginName}' cannot depend on plugin '{DependencyPluginName}' because it is disallowed.", Name, dependencyPlugin.Name);
+						anyErrors = true;
+					}
+				}
+			}
+
+			// Check that any plugins with the NoCode specifier do not contain modules.
+			if (Descriptor.bNoCode && Modules.Any())
+			{
+				logger.LogError("Plugin '{PluginName}' cannot contain any code or modules. See the plugin descriptor property `NoCode`", Name);
+				anyErrors = true;
+			}
+
+			return anyErrors;
+		}
 	}
 }
diff --git a/Engine/Source/Programs/UnrealBuildTool/Configuration/UEBuildTarget.cs b/Engine/Source/Programs/UnrealBuildTool/Configuration/UEBuildTarget.cs
--- a/Engine/Source/Programs/UnrealBuildTool/Configuration/UEBuildTarget.cs
+++ b/Engine/Source/Programs/UnrealBuildTool/Configuration/UEBuildTarget.cs
@@ -1423,6 +1423,156 @@ namespace UnrealBuildTool
 			}
 		}
 
+		/// <summary>
+		/// Validate all plugins
+		/// </summary>
+		/// <param name="logger"></param>
+		/// <returns>true if there are any fatal errors</returns>
+		bool ValidatePlugins(ILogger logger)
+		{
+			if (BuildPlugins == null)
+			{
+				return false;
+			}
+
+			foreach (UEBuildPlugin plugin in BuildPlugins.OrderBy(x => x.Name))
+			{
+				UEBuildPlatform.GetBuildPlatform(Platform).ValidatePlugin(plugin, Rules);
+			}
+
+			// Intentionally not ordered by name to maintain the reference include order
+			return BuildPlugins.Select(plugin => ValidatePlugin(plugin, logger)).Any(x => x);
+		}
+
+		/// <summary>
+		/// Validates a plugin
+		/// </summary>
+		/// <param name="plugin"></param>
+		/// <param name="logger"></param>
+		/// <returns>true if there are any fatal errors</returns>
+		bool ValidatePlugin(UEBuildPlugin plugin, ILogger logger)
+		{
+			bool anyErrors = plugin.ValidatePlugin(logger);
+
+			// Check that each plugin declares its dependencies explicitly
+			foreach (UEBuildModule module in plugin.Modules)
+			{
+				HashSet<UEBuildModule> dependencyModules = module.GetDependencies(bWithIncludePathModules: true, bWithDynamicallyLoadedModules: true);
+				// Only warn on modules that can be compiled
+				foreach (UEBuildModuleCPP dependencyModule in dependencyModules.OfType<UEBuildModuleCPP>())
+				{
+					PluginInfo? dependencyPluginInfo = dependencyModule.Rules.Plugin;
+					bool isOptional = plugin.Descriptor.Plugins?.Where(x => dependencyPluginInfo?.Name == x.Name).Any(x => x.bOptional) ?? false;
+					if (dependencyPluginInfo != null
+						&& dependencyPluginInfo != plugin.Info
+						&& !Rules.InternalPluginDependencies.Contains(dependencyPluginInfo.Name)
+						&& !(plugin.Dependencies!.Any(x => dependencyPluginInfo == x.Info) || isOptional))
+					{
+						logger.LogWarning("Warning: Plugin '{PluginName}' does not list plugin '{DependencyPluginName}' as a dependency, but module '{ModuleName}' depends on module '{DependencyModuleName}'.", plugin.Name, dependencyPluginInfo.Name, module.Name, dependencyModule.Name);
+					}
+				}
+			}
+
+			return anyErrors;
+		}
+
+		/// <summary>
+		/// Validates all plugins
+		/// </summary>
+		/// <param name="logger"></param>
+		/// <returns>true if there are any fatal errors</returns>
+		bool ValidateModules(ILogger logger)
+		{
+			bool anyErrors = false;
+			foreach (UEBuildModule module in Modules.Values.OrderBy(x => x.Name))
+			{
+				UEBuildPlatform.GetBuildPlatform(Platform).ValidateModule(module, Rules);
+				anyErrors |= UEBuildPlatform.GetBuildPlatform(Platform).ValidateModuleIncludePaths(module, Rules, Modules.Values);
+			}
+
+			// Intentionally not ordered by name to maintain the reference include order
+			foreach (UEBuildModule module in Modules.Values)
+			{
+				anyErrors |= ValidateModule(module, logger);
+			}
+
+			return anyErrors;
+		}
+
+		/// <summary>
+		/// Validates a module
+		/// </summary>
+		/// <param name="module"></param>
+		/// <param name="logger"></param>
+		/// <returns>true if there are any fatal errors</returns>
+		bool ValidateModule(UEBuildModule module, ILogger logger)
+		{
+			bool anyErrors = module.ValidateModule("Target", logger);
+
+			Lazy<HashSet<UEBuildModule>> referencedModules = new(() => module.GetDependencies(bWithIncludePathModules: true, bWithDynamicallyLoadedModules: true));
+
+			// Check there aren't any engine binaries with dependencies on game modules. This can happen when game-specific plugins override engine plugins.
+			if (module.Binary != null)
+			{
+				foreach (UEBuildModule referencedModule in referencedModules.Value)
+				{
+					if (!module.Rules.Context.Scope.Contains(referencedModule.Rules.Context.Scope) && !IsEngineToPluginReferenceAllowed(module.Name, referencedModule.Name))
+					{
+						logger.LogError("Module '{ModuleName}' ({ScopeName}) should not reference module '{DependencyModuleName}' ({DependencyScopeName}). Hierarchy is {Hierarchy}.", module.Name, module.Rules.Context.Scope.Name, referencedModule.Name, referencedModule.Rules.Context.Scope.Name, referencedModule.Rules.Context.Scope.FormatHierarchy());
+						anyErrors = true;
+					}
+				}
+			}
+
+			// Check that each project module with a dependency on a plugin module, that the plugin is either enabled by default or it's enabled by the uproject file.
+			if (ProjectDescriptor?.Modules?.Any(x => x.Name == module.Name) == true)
+			{
+				string projectName = ProjectFile?.FullName ?? TargetRulesFile.GetFileName();
+				bool bAllowEnginePluginsEnabledByDefault = (!ProjectDescriptor?.DisableEnginePluginsByDefault) ?? false;
+
+				// Only warn on modules that can be compiled
+				foreach (UEBuildModuleCPP dependencyModule in referencedModules.Value.OfType<UEBuildModuleCPP>())
+				{
+					PluginInfo? dependencyPluginInfo = dependencyModule.Rules.Plugin;
+					// Is the modules plugin enabled by default?
+					if (dependencyPluginInfo != null && dependencyPluginInfo.IsEnabledByDefault(bAllowEnginePluginsEnabledByDefault) == false)
+					{
+						// Try and find the project plugin reference in the plugins list.
+						PluginReferenceDescriptor? projectPluginReference = ProjectDescriptor?.Plugins?.FirstOrDefault(x => x.Name == dependencyPluginInfo.Name);
+						if (projectPluginReference != null)
+						{
+							if (!projectPluginReference.bEnabled)
+							{
+								logger.LogWarning("Warning: {ProjectName} plugin dependency '{DependencyPluginName}' is not enabled, but module '{ModuleName}' depends on '{DependencyModuleName}'.", projectName, dependencyPluginInfo.Name, module.Name, dependencyModule.Name);
+							}
+							else if (!projectPluginReference.IsEnabledForTarget(TargetType))
+							{
+								logger.LogWarning("Warning: {ProjectName} plugin dependency '{DependencyPluginName}' is not enabled for target '{TargetType}', but module '{ModuleName}' depends on '{DependencyModuleName}'.", projectName, dependencyPluginInfo.Name, TargetType.ToString(), module.Name, dependencyModule.Name);
+							}
+							else if (!projectPluginReference.IsEnabledForPlatform(Platform))
+							{
+								logger.LogWarning("Warning: {ProjectName} plugin dependency '{DependencyPluginName}' is not enabled for platform '{Platform}', but module '{ModuleName}' depends on '{DependencyModuleName}'.", projectName, dependencyPluginInfo.Name, Platform.ToString(), module.Name, dependencyModule.Name);
+							}
+							else if (!projectPluginReference.IsEnabledForTargetConfiguration(Configuration))
+							{
+								logger.LogWarning("Warning: {ProjectName} plugin dependency '{DependencyPluginName}' is not enabled for target configuration '{Configuration}', but module '{ModuleName}' depends on '{DependencyModuleName}'.", projectName, dependencyPluginInfo.Name, Configuration.ToString(), module.Name, dependencyModule.Name);
+							}
+						}
+						else
+						{
+							// No project plugin reference exists. Check inclusion of a plugin was not made as part of a Target.cs file or Build.cs file
+							if (!Rules.EnablePlugins.Contains(dependencyPluginInfo.Name) && !Rules.InternalPluginDependencies.Contains(dependencyPluginInfo.Name))
+							{
+								logger.LogWarning("Warning: {ProjectName} does not list plugin '{DependencyPluginName}' as a dependency, but module '{ModuleName}' depends on '{DependencyModuleName}'.", projectName, dependencyPluginInfo.Name, module.Name, dependencyModule.Name);
+							}
+						}
+					}
+				}
+			}
+
+			return anyErrors;
+		}
+
 		/// <summary>
 		/// The target rules
 		/// </summary>
@@ -2330,136 +2480,6 @@ namespace UnrealBuildTool
 				}
 			}
 
-			// Build a mapping from module to its plugin
-			Dictionary<UEBuildModule, UEBuildPlugin> ModuleToPlugin = new Dictionary<UEBuildModule, UEBuildPlugin>();
-			foreach (UEBuildPlugin Plugin in BuildPlugins!)
-			{
-				foreach (UEBuildModule Module in Plugin.Modules)
-				{
-					if (!ModuleToPlugin.ContainsKey(Module))
-					{
-						ModuleToPlugin.Add(Module, Plugin);
-					}
-				}
-			}
-
-			// Check there aren't any engine binaries with dependencies on game modules. This can happen when game-specific plugins override engine plugins.
-			foreach (UEBuildModule Module in Modules.Values)
-			{
-				if (Module.Binary != null)
-				{
-					HashSet<UEBuildModule> ReferencedModules = Module.GetDependencies(bWithIncludePathModules: true, bWithDynamicallyLoadedModules: true);
-					foreach (UEBuildModule ReferencedModule in ReferencedModules)
-					{
-						if (!Module.Rules.Context.Scope.Contains(ReferencedModule.Rules.Context.Scope) && !IsEngineToPluginReferenceAllowed(Module.Name, ReferencedModule.Name))
-						{
-							throw new BuildException("Module '{0}' ({1}) should not reference module '{2}' ({3}). Hierarchy is {4}.", Module.Name, Module.Rules.Context.Scope.Name, ReferencedModule.Name, ReferencedModule.Rules.Context.Scope.Name, ReferencedModule.Rules.Context.Scope.FormatHierarchy());
-						}
-					}
-				}
-			}
-
-
-			// Check that each project module with a dependency on a plugin module, that the plugin is either enabled by default or it's enabled by the uproject file.
-			if (ProjectDescriptor?.Modules != null && ProjectDescriptor?.Plugins != null)
-			{
-				string ProjectName = ProjectFile?.GetFileName() ?? ".uproject";
-				bool bAllowEnginePluginsEnabledByDefault = !ProjectDescriptor.DisableEnginePluginsByDefault;
-
-				foreach (ModuleDescriptor Descriptor in ProjectDescriptor.Modules.Where(x => x.IsCompiledInConfiguration(Platform, Configuration, TargetName, TargetType, Rules.bBuildDeveloperTools, Rules.bBuildRequiresCookedData)))
-				{
-					if (Modules.TryGetValue(Descriptor.Name, out UEBuildModule? Module))
-					{
-						HashSet<UEBuildModule> DependencyModules = Module.GetDependencies(bWithIncludePathModules: true, bWithDynamicallyLoadedModules: true);
-						foreach (UEBuildModule DependencyModule in DependencyModules)
-						{
-							if (ModuleToPlugin.TryGetValue(DependencyModule, out UEBuildPlugin? DependencyPlugin))
-							{
-								PluginInfo? DependencyPluginInfo = Plugins.GetPlugin(DependencyPlugin.Name);
-
-								// Is the modules plugin enabled by default?
-								if (DependencyPluginInfo?.IsEnabledByDefault(bAllowEnginePluginsEnabledByDefault) == false)
-								{
-									// Try and find the project plugin reference in the plugins list.
-									PluginReferenceDescriptor? ProjectPluginReference = ProjectDescriptor.Plugins.FirstOrDefault(x => x.Name.Equals(DependencyPlugin.Name));
-									if (ProjectPluginReference != null)
-									{
-										if (!ProjectPluginReference.bEnabled)
-										{
-											Logger.LogWarning("Warning: {ProjectName} plugin dependency '{DependencyPluginName}' is not enabled, but module '{ModuleName}' depends on '{DependencyModuleName}'.", ProjectName, DependencyPlugin.Name, Module.Name, DependencyModule.Name);
-										}
-										else if (!ProjectPluginReference.IsEnabledForTarget(TargetType))
-										{
-											Logger.LogWarning("Warning: {ProjectName} plugin dependency '{DependencyPluginName}' is not enabled for target '{TargetType}', but module '{ModuleName}' depends on '{DependencyModuleName}'.", ProjectName, DependencyPlugin.Name, TargetType.ToString(), Module.Name, DependencyModule.Name);
-										}
-										else if (!ProjectPluginReference.IsEnabledForPlatform(Platform))
-										{
-											Logger.LogWarning("Warning: {ProjectName} plugin dependency '{DependencyPluginName}' is not enabled for platform '{Platform}', but module '{ModuleName}' depends on '{DependencyModuleName}'.", ProjectName, DependencyPlugin.Name, Platform.ToString(), Module.Name, DependencyModule.Name);
-										}
-										else if (!ProjectPluginReference.IsEnabledForTargetConfiguration(Configuration))
-										{
-											Logger.LogWarning("Warning: {ProjectName} plugin dependency '{DependencyPluginName}' is not enabled for target configuration '{Configuration}', but module '{ModuleName}' depends on '{DependencyModuleName}'.", ProjectName, DependencyPlugin.Name, Configuration.ToString(), Module.Name, DependencyModule.Name);
-										}
-									}
-									else
-									{
-										// No project plugin reference exists. Check inclusion of a plugin was not made as part of a Target.cs file or Build.cs file
-										if (!Rules.EnablePlugins.Contains(DependencyPlugin.Name) && !Rules.InternalPluginDependencies.Contains(DependencyPlugin.Name))
-										{
-											Logger.LogWarning("Warning: {ProjectName} does not list plugin '{DependencyPluginName}' as a dependency, but module '{ModuleName}' depends on '{DependencyModuleName}'.", ProjectName, DependencyPlugin.Name, Module.Name, DependencyModule.Name);
-										}
-									}
-								}
-							}
-						}
-					}
-				}
-			}
-
-			// Check that each plugin does not have a dependency on any sealed plugins
-			foreach (UEBuildPlugin Plugin in BuildPlugins)
-			{
-				if (Plugin.Dependencies != null)
-				{
-					foreach (UEBuildPlugin DependencyPlugin in Plugin.Dependencies)
-					{
-						if (DependencyPlugin.Descriptor.bIsSealed)
-						{
-							throw new BuildException("Plugin '{0}' cannot depend on plugin '{1}' because it is sealed.", Plugin.Name, DependencyPlugin.Name);
-						}
-						else if (Plugin.Descriptor.DisallowedPlugins != null && Plugin.Descriptor.DisallowedPlugins.Contains(DependencyPlugin.Name))
-						{
-							throw new BuildException("Plugin '{0}' cannot depend on plugin '{1}' because it is disallowed.", Plugin.Name, DependencyPlugin.Name);
-						}
-					}
-				}
-			}
-
-			// Check that each plugin declares its dependencies explicitly
-			foreach (UEBuildPlugin Plugin in BuildPlugins)
-			{
-				// Check that any plugins with the NoCode specifier do not contain modules.
-				if (Plugin.Descriptor.bNoCode && Plugin.Modules.Count > 0)
-				{
-					throw new BuildException("Plugin '{0}' cannot contain any code or modules. See the plugin descriptor property `NoCode`", Plugin.Name);
-				}
-
-				foreach (UEBuildModule Module in Plugin.Modules)
-				{
-					HashSet<UEBuildModule> DependencyModules = Module.GetDependencies(bWithIncludePathModules: true, bWithDynamicallyLoadedModules: true);
-					foreach (UEBuildModule DependencyModule in DependencyModules)
-					{
-						UEBuildPlugin? DependencyPlugin;
-						if (ModuleToPlugin.TryGetValue(DependencyModule, out DependencyPlugin) && DependencyPlugin != Plugin
-							&& !Rules.InternalPluginDependencies.Contains(DependencyPlugin.Name)
-							&& !Plugin.Dependencies!.Contains(DependencyPlugin))
-						{
-							Logger.LogWarning("Warning: Plugin '{PluginName}' does not list plugin '{DependencyPluginName}' as a dependency, but module '{ModuleName}' depends on '{DependencyModuleName}'.", Plugin.Name, DependencyPlugin.Name, Module.Name, DependencyModule.Name);
-						}
-					}
-				}
-			}
-
 			// Create the makefile
 			string ExternalMetadata = UEBuildPlatform.GetBuildPlatform(Platform).GetExternalBuildMetadata(ProjectFile);
 			TargetMakefile Makefile = new TargetMakefile(ExternalMetadata, Binaries[0].OutputFilePaths[0], ReceiptFileName,
@@ -2641,7 +2661,7 @@ namespace UnrealBuildTool
 			{
 				HashSet<FileItem> RetainOutputItems = new HashSet<FileItem>();
 
-				UEBuildPlugin? ForeignBuildPlugin = BuildPlugins.Find(x => x.File == ForeignPlugin);
+				UEBuildPlugin? ForeignBuildPlugin = BuildPlugins!.Find(x => x.File == ForeignPlugin);
 				
 				foreach (UEBuildPlugin Plugin in BuildPlugins)
 				{
@@ -3376,27 +3396,24 @@ namespace UnrealBuildTool
 				}
 			}
 
+			bool anyValidationErrors = false;
+
+			// Allow each platform to run any needed validation on the plugins.
+			using (GlobalTracer.Instance.BuildSpan("UEBuildTarget.PreBuildSetup().ValidatePlugins").StartActive())
+			{
+				anyValidationErrors |= ValidatePlugins(Logger);
+			}
+
 			// Allow each platform to run any needed validation on the module.
 			// ie: iOS/macOS do a static library version build check for unsupported clang version builds.
 			using (GlobalTracer.Instance.BuildSpan("UEBuildTarget.PreBuildSetup().ValidateModules").StartActive())
 			{
-				bool AnyErrors = false;
-				foreach (UEBuildModule Module in Modules.Values.OrderBy(x => x.Name))
-				{
-					UEBuildPlatform.GetBuildPlatform(Platform).ValidateModule(Module, Rules);
-					AnyErrors |= UEBuildPlatform.GetBuildPlatform(Platform).ValidateModuleIncludePaths(Module, Rules, Modules.Values);
-				}
-
-				// Intentionally not ordered by name to maintain the reference include order
-				foreach (UEBuildModule Module in Modules.Values)
-				{
-					Module.ValidateModule("Target", Logger);
-				}
+				anyValidationErrors |= ValidateModules(Logger);
+			}
 
-				if (AnyErrors)
-				{
-					throw new BuildException("Errors validating modules.");
-				}
+			if (anyValidationErrors)
+			{
+				throw new BuildException("Errors validating plugins or modules.");
 			}
 
 			if (!bCompileMonolithic)
