- Force disable Uba for all Link actions on Mac host
- Matched the CanRunRemotely match the force no-Uba for Mac shell commands

#changelist validated
#virtualized

diff --git a/Engine/Source/Programs/UnrealBuildTool/Executors/Experimental/UBAExecutor.cs b/Engine/Source/Programs/UnrealBuildTool/Executors/Experimental/UBAExecutor.cs
--- a/Engine/Source/Programs/UnrealBuildTool/Executors/Experimental/UBAExecutor.cs
+++ b/Engine/Source/Programs/UnrealBuildTool/Executors/Experimental/UBAExecutor.cs
@@ -571,12 +571,37 @@ namespace UnrealBuildTool
 			}
 		}
 
+
+		/// <summary>
+		/// Determine if an action must be run locally and with no detouring
+		/// </summary>
+		/// <param name="action">The action to check</param>
+		/// <returns>If this action must be local, non-detoured</returns>
+		bool ForceLocalNoDetour(LinkedAction action)
+		{
+			// Don't let Mac run shell commands through Uba as interposing dylibs into
+			// the shell results in dyld errors about no matching architecture.
+			// The shell is used to run various commands during a build like copy/ditto.
+			// So for these actions we need to make sure UBA is not used.
+			// Linking is similarly currently not working in Uba on Mac
+			bool bIsShellAction = action.CommandPath == BuildHostPlatform.Current.Shell;
+			bool bIsLinkAction = action.ActionType == ActionType.Link;
+			return System.OperatingSystem.IsMacOS() && (bIsShellAction || bIsLinkAction);
+		}
+
+
 		/// <summary>
 		/// Determine if an action is able to be run remotely
 		/// </summary>
 		/// <param name="action">The action to check</param>
 		/// <returns>If this action can be run remotely</returns>
-		bool CanRunRemotely(LinkedAction action) => action.bCanExecuteInUBA && action.bCanExecuteRemotely && !_localRetryActions.ContainsKey(action) && !_forcedRetryActions.ContainsKey(action) && (UBAConfig.bLinkRemote || action.ActionType != ActionType.Link);
+		bool CanRunRemotely(LinkedAction action) => 
+			action.bCanExecuteInUBA && 
+			action.bCanExecuteRemotely && 
+			!_localRetryActions.ContainsKey(action) && 
+			!_forcedRetryActions.ContainsKey(action) && 
+			(UBAConfig.bLinkRemote || action.ActionType != ActionType.Link) &&
+			!ForceLocalNoDetour(action);
 
 		ProcessStartInfo GetActionStartInfo(LinkedAction action, out FileItem? pchItem)
 		{
@@ -620,12 +645,7 @@ namespace UnrealBuildTool
 
 			return () =>
 			{
-				bool enableDetour = action.bCanExecuteInUBA && !_forcedRetryActions.ContainsKey(action) && 
-					// Don't let Mac run shell commands through Uba as interposing dylibs into
-					// the shell results in dyld errors about no matching architecture.
-					// The shell is used to run various commands during a build like copy/ditto.
-					//So for these actions we need to make sure UBA is not used.
-					(!System.OperatingSystem.IsMacOS() || action.CommandPath != BuildHostPlatform.Current.Shell);
+				bool enableDetour = !ForceLocalNoDetour(action) && action.bCanExecuteInUBA && !_forcedRetryActions.ContainsKey(action);
 
 				ProcessStartInfo startInfo = GetActionStartInfo(action, out FileItem? pchItem);
 				using (IProcess process = _session!.RunProcess(startInfo, false, null, enableDetour))
diff --git a/Engine/Source/Programs/UnrealBuildTool/Platform/Mac/MacToolChain.cs b/Engine/Source/Programs/UnrealBuildTool/Platform/Mac/MacToolChain.cs
--- a/Engine/Source/Programs/UnrealBuildTool/Platform/Mac/MacToolChain.cs
+++ b/Engine/Source/Programs/UnrealBuildTool/Platform/Mac/MacToolChain.cs
@@ -880,7 +880,6 @@ namespace UnrealBuildTool
 
 			// Only execute linking on the local Mac.
 			LinkAction.bCanExecuteRemotely = false;
-			LinkAction.bCanExecuteInUBA = false;
 
 			LinkAction.StatusDescription = Path.GetFileName(OutputFile.AbsolutePath);
 
