Renaming IgnorePluginDependency

diff --git a/Engine/Source/Programs/UnrealBuildTool/Configuration/TargetRules.cs b/Engine/Source/Programs/UnrealBuildTool/Configuration/TargetRules.cs
--- a/Engine/Source/Programs/UnrealBuildTool/Configuration/TargetRules.cs
+++ b/Engine/Source/Programs/UnrealBuildTool/Configuration/TargetRules.cs
@@ -2625,7 +2625,7 @@ namespace UnrealBuildTool
 		/// Checks if a plugin should be programmatically allowed in the build
 		/// </summary>
 		/// <returns>true if the plugin is allowed</returns>
-		public virtual bool IgnorePluginDependency(PluginInfo parentInfo, PluginReferenceDescriptor descriptor) => false;
+		public virtual bool ShouldIgnorePluginDependency(PluginInfo parentInfo, PluginReferenceDescriptor descriptor) => false;
 
 		/// <summary>
 		/// Checks if a property has been set with the RequiresUniqueBuildEnvironmentAttribute, and is different from it's base
diff --git a/Engine/Source/Programs/UnrealBuildTool/Configuration/UEBuildTarget.cs b/Engine/Source/Programs/UnrealBuildTool/Configuration/UEBuildTarget.cs
--- a/Engine/Source/Programs/UnrealBuildTool/Configuration/UEBuildTarget.cs
+++ b/Engine/Source/Programs/UnrealBuildTool/Configuration/UEBuildTarget.cs
@@ -4559,7 +4559,7 @@ namespace UnrealBuildTool
 					foreach (PluginReferenceDescriptor NextReference in Info.Descriptor.Plugins)
 					{
 						// Ignore any plugin dependency which is programmatically filtered out by the target
-						if (Rules.IgnorePluginDependency(Info, NextReference))
+						if (Rules.ShouldIgnorePluginDependency(Info, NextReference))
 						{
 							Logger.LogDebug("Ignoring plugin '{0}' (referenced via {1}) due to it being filtered out by the target rules.", NextReference.Name, Info.Name);
 							continue;
