Fixing compile error due to rename of IgnorePluginDependency

diff --git a/Engine/Source/Programs/UnrealBuildTool/Configuration/ReadOnlyTargetRules.cs b/Engine/Source/Programs/UnrealBuildTool/Configuration/ReadOnlyTargetRules.cs
--- a/Engine/Source/Programs/UnrealBuildTool/Configuration/ReadOnlyTargetRules.cs
+++ b/Engine/Source/Programs/UnrealBuildTool/Configuration/ReadOnlyTargetRules.cs
@@ -635,7 +635,7 @@ namespace UnrealBuildTool
 		/// Checks if a plugin should be programmatically allowed in the build
 		/// </summary>
 		/// <returns>true if the plugin is allowed</returns>
-		public bool IgnorePluginDependency(PluginInfo parentInfo, PluginReferenceDescriptor childDescriptor) => Inner.IgnorePluginDependency(parentInfo, childDescriptor);
+		public bool ShouldIgnorePluginDependency(PluginInfo parentInfo, PluginReferenceDescriptor childDescriptor) => Inner.ShouldIgnorePluginDependency(parentInfo, childDescriptor);
 
 		/// <summary>
 		/// Determines if the automation tests should be compiled based on the current configuration and optional forced settings.
