[UBA]
* Removed some log spam

diff --git a/Engine/Source/Programs/UnrealBuildTool/Executors/Experimental/UBAHordeSession.cs b/Engine/Source/Programs/UnrealBuildTool/Executors/Experimental/UBAHordeSession.cs
--- a/Engine/Source/Programs/UnrealBuildTool/Executors/Experimental/UBAHordeSession.cs
+++ b/Engine/Source/Programs/UnrealBuildTool/Executors/Experimental/UBAHordeSession.cs
@@ -81,7 +81,7 @@ namespace UnrealBuildTool
 					Pool = LinuxHordePool;
 				}
 
-				Console.WriteLine("CHOSEN UBA POOL: {0}", Pool);
+				//Console.WriteLine("CHOSEN UBA POOL: {0}", Pool);
 				return Pool;
 			}
 		}
