[ChaosVD] Fix CIS issue in Headless Chaos

The issue is happening because CVD support is enabled but UE trace is disabled therefore some methods are still compiled out. This change makes sure CVD code is compiled out if UE trace is also compiled out because CVD depends on it.

diff --git a/Engine/Source/Programs/UnrealBuildTool/Configuration/ModuleRules.cs b/Engine/Source/Programs/UnrealBuildTool/Configuration/ModuleRules.cs
--- a/Engine/Source/Programs/UnrealBuildTool/Configuration/ModuleRules.cs
+++ b/Engine/Source/Programs/UnrealBuildTool/Configuration/ModuleRules.cs
@@ -1564,7 +1564,8 @@ namespace UnrealBuildTool
 		/// </summary>
 		public void SetupModuleChaosVisualDebuggerSupport(ReadOnlyTargetRules target)
 		{
-			bool bHasChaosVisualDebuggerSupport = target.bCompileChaosVisualDebuggerSupport && target.Configuration != UnrealTargetConfiguration.Shipping;
+			//TODO: Emit some form of warning if CVD support is explicitly enabled, but UE Trace is disabled
+			bool bHasChaosVisualDebuggerSupport = target.bCompileChaosVisualDebuggerSupport && target.Configuration != UnrealTargetConfiguration.Shipping && target.bEnableTrace;
 			if (bHasChaosVisualDebuggerSupport)
 			{
 				PublicDependencyModuleNames.Add("ChaosVDRuntime");
