- Add BuildConfiguration entry for per target project generation in VC Project generator

diff --git a/Engine/Source/Programs/UnrealBuildTool/ProjectFiles/VisualStudio/VCProjectFileGenerator.cs b/Engine/Source/Programs/UnrealBuildTool/ProjectFiles/VisualStudio/VCProjectFileGenerator.cs
--- a/Engine/Source/Programs/UnrealBuildTool/ProjectFiles/VisualStudio/VCProjectFileGenerator.cs
+++ b/Engine/Source/Programs/UnrealBuildTool/ProjectFiles/VisualStudio/VCProjectFileGenerator.cs
@@ -93,6 +93,12 @@ namespace UnrealBuildTool
 		/// </summary>
 		[XmlConfigFile(Category = "VCProjectFileGenerator")]
 		public bool bBuildLiveCodingConsole = false;
+
+		/// <summary>
+		/// Whether to generate a project file for each individual target, and not include e.g. Editor/Client/Server in the Configuration.
+		/// </summary>
+		[XmlConfigFile(Category = "VCProjectFileGenerator")]
+		public bool bMakeProjectPerTarget = false;
 	}
 
 	/// <summary>
@@ -101,14 +107,14 @@ namespace UnrealBuildTool
 	class VCProjectFileGenerator : ProjectFileGenerator
 	{
 		/// <summary>
-		/// Set to true to enable a project for each target, and do not put the target type into the configuration
+		/// The settings object
 		/// </summary>
-		protected override bool bMakeProjectPerTarget => false;
+		protected VCProjectFileSettings Settings = new VCProjectFileSettings();
 
 		/// <summary>
-		/// The settings object
+		/// Set to true to enable a project for each target, and do not put the target type into the configuration
 		/// </summary>
-		protected VCProjectFileSettings Settings = new VCProjectFileSettings();
+		protected override bool bMakeProjectPerTarget => Settings.bMakeProjectPerTarget;
 
 		/// <summary>
 		/// Override for the build tool to use in generated projects. If the compiler version is specified on the command line, we use the same argument on the 
