UnrealBuildTool: Clean up some warnings in UBA

#rnx

diff --git a/Engine/Source/Programs/UnrealBuildTool/Executors/Experimental/UBAExecutor.cs b/Engine/Source/Programs/UnrealBuildTool/Executors/Experimental/UBAExecutor.cs
--- a/Engine/Source/Programs/UnrealBuildTool/Executors/Experimental/UBAExecutor.cs
+++ b/Engine/Source/Programs/UnrealBuildTool/Executors/Experimental/UBAExecutor.cs
@@ -244,7 +244,7 @@ namespace UnrealBuildTool
 		public string Crypto { get; private set; } = String.Empty;
 		public IServer? Server { get; private set; }
 		ISessionServer? _session;
-		List<IUBAAgentCoordinator> _agentCoordinators = new List<IUBAAgentCoordinator>();
+		readonly List<IUBAAgentCoordinator> _agentCoordinators = new List<IUBAAgentCoordinator>();
 		DirectoryReference? _rootDirRef;
 		bool _bIsCancelled;
 		bool _bIsRemoteActionsAllowed = true;
@@ -264,6 +264,7 @@ namespace UnrealBuildTool
 			{
 				_session?.Dispose();
 				_session = null;
+				_threadedLogger.Dispose();
 			}
 			base.Dispose(disposing);
 		}
@@ -334,7 +335,7 @@ namespace UnrealBuildTool
 
 			if (!UBAConfig.bDisableRemote)
 			{
-				foreach (var coordinator in _agentCoordinators)
+				foreach (IUBAAgentCoordinator coordinator in _agentCoordinators)
 				{
 					await coordinator.InitAsync(this);
 
@@ -386,7 +387,7 @@ namespace UnrealBuildTool
 				_bIsCancelled = true;
 				_session?.CancelAll();
 				ubaStorage?.SaveCasTable();
-				foreach (var coordinator in _agentCoordinators)
+				foreach (IUBAAgentCoordinator coordinator in _agentCoordinators)
 				{
 					coordinator.CloseAsync().Wait(2000); // Give coordinators some time to close (this makes coordinators like horde return resources faster)
 				}
@@ -435,7 +436,7 @@ namespace UnrealBuildTool
 			{
 				Console.CancelKeyPress -= CancelKeyPress;
 
-				foreach (var coordinator in _agentCoordinators)
+				foreach (IUBAAgentCoordinator coordinator in _agentCoordinators)
 				{
 					await coordinator.CloseAsync();
 				}
@@ -527,7 +528,7 @@ namespace UnrealBuildTool
 					if (count <= NumParallelProcesses)
 					{
 						_bIsRemoteActionsAllowed = false;
-						foreach (var coordinator in _agentCoordinators)
+						foreach (IUBAAgentCoordinator coordinator in _agentCoordinators)
 						{
 							coordinator.Stop();
 						}
@@ -554,7 +555,7 @@ namespace UnrealBuildTool
 
 			try
 			{
-				foreach (var coordinator in _agentCoordinators)
+				foreach (IUBAAgentCoordinator coordinator in _agentCoordinators)
 				{
 					coordinator.Start(queue, CanRunRemotely);
 				}
@@ -564,20 +565,19 @@ namespace UnrealBuildTool
 			}
 			finally
 			{
-				foreach (var coordinator in _agentCoordinators)
+				foreach (IUBAAgentCoordinator coordinator in _agentCoordinators)
 				{
 					coordinator.Stop();
 				}
 			}
 		}
 
-
 		/// <summary>
 		/// Determine if an action must be run locally and with no detouring
 		/// </summary>
 		/// <param name="action">The action to check</param>
 		/// <returns>If this action must be local, non-detoured</returns>
-		bool ForceLocalNoDetour(LinkedAction action)
+		static bool ForceLocalNoDetour(LinkedAction action)
 		{
 			// Don't let Mac run shell commands through Uba as interposing dylibs into
 			// the shell results in dyld errors about no matching architecture.
@@ -586,10 +586,9 @@ namespace UnrealBuildTool
 			// Linking is similarly currently not working in Uba on Mac
 			bool bIsShellAction = action.CommandPath == BuildHostPlatform.Current.Shell;
 			bool bIsLinkAction = action.ActionType == ActionType.Link;
-			return System.OperatingSystem.IsMacOS() && (bIsShellAction || bIsLinkAction);
+			return OperatingSystem.IsMacOS() && (bIsShellAction || bIsLinkAction);
 		}
 
-
 		/// <summary>
 		/// Determine if an action is able to be run remotely
 		/// </summary>
@@ -668,12 +667,7 @@ namespace UnrealBuildTool
 					List<string> logLines = process.LogLines;
 					logLines.RemoveAll((line) => line.StartsWith("   Creating library ", StringComparison.OrdinalIgnoreCase) && line.EndsWith(".exp", StringComparison.OrdinalIgnoreCase) || line.EndsWith("file(s) copied.", StringComparison.OrdinalIgnoreCase));
 
-					string? additionalDescription = null;
-					if (!enableDetour)
-					{
-						additionalDescription = "(UBA disabled)";
-					}
-
+					string? additionalDescription = !enableDetour ? "(UBA disabled)" : null;
 					ActionFinished(queue, new ExecuteResults(logLines, process.ExitCode, executionTime, processorTime, additionalDescription), action, pchItem, process);
 				}
 				return Task.CompletedTask;
@@ -696,7 +690,7 @@ namespace UnrealBuildTool
 					int sizeOfChar = System.OperatingSystem.IsWindows() ? 2 : 1;
 
 					int byteCount = 0;
-					foreach (var item in action.PrerequisiteItems)
+					foreach (FileItem item in action.PrerequisiteItems)
 					{
 						byteCount += (item.FullName.Length + 1) * sizeOfChar;
 						++knownInputsCount;
@@ -705,14 +699,19 @@ namespace UnrealBuildTool
 					knownInputs = new byte[byteCount + sizeOfChar];
 
 					int byteOffset = 0;
-					foreach (var item in action.PrerequisiteItems)
+					foreach (FileItem item in action.PrerequisiteItems)
 					{
-						var str = item.FullName;
+						string str = item.FullName;
 						int strBytes = str.Length * sizeOfChar;
 						if (sizeOfChar == 1) // Unmanaged size uses ascii
+						{
 							System.Buffer.BlockCopy(System.Text.Encoding.ASCII.GetBytes(str.ToCharArray()), 0, knownInputs, byteOffset, strBytes);
+						}
 						else
+						{
 							System.Buffer.BlockCopy(str.ToCharArray(), 0, knownInputs, byteOffset, strBytes);
+						}
+
 						byteOffset += strBytes + sizeOfChar;
 					}
 				}
diff --git a/Engine/Source/Programs/UnrealBuildTool/Executors/Experimental/UBAHordeSession.cs b/Engine/Source/Programs/UnrealBuildTool/Executors/Experimental/UBAHordeSession.cs
--- a/Engine/Source/Programs/UnrealBuildTool/Executors/Experimental/UBAHordeSession.cs
+++ b/Engine/Source/Programs/UnrealBuildTool/Executors/Experimental/UBAHordeSession.cs
@@ -361,7 +361,7 @@ namespace UnrealBuildTool
 			await using IBlobWriter writer = _storage.CreateBlobWriter();
 			DirectoryNode sandbox = new();
 			await sandbox.AddFilesAsync(baseDir, files, writer, serializerOptions: serializerOptions, cancellationToken: cancellationToken);
-			IBlobHandle<DirectoryNode> handle = await writer.WriteBlobAsync(sandbox);
+			IBlobHandle<DirectoryNode> handle = await writer.WriteBlobAsync(sandbox, cancellationToken);
 			await writer.FlushAsync(cancellationToken);
 			return handle.GetLocator();
 		}
@@ -684,8 +684,8 @@ namespace UnrealBuildTool
 						if (System.OperatingSystem.IsMacOS())
 						{
 							// we need to populate the cas with all known xcodes so we can serve all the ones we have installed
-							string XcodeVersion = Utils.RunLocalProcessAndReturnStdOut("/bin/sh", "-c '/usr/bin/defaults read $(xcode-select -p)/../version.plist ProductBuildVersion");
-							arguments.Add($"-populateCasFromXcodeVersion={XcodeVersion}");
+							string xcodeVersion = Utils.RunLocalProcessAndReturnStdOut("/bin/sh", "-c '/usr/bin/defaults read $(xcode-select -p)/../version.plist ProductBuildVersion");
+							arguments.Add($"-populateCasFromXcodeVersion={xcodeVersion}");
 						}
 
 						arguments.Add("-Dir=%UE_HORDE_SHARED_DIR%\\Uba");
@@ -707,7 +707,7 @@ namespace UnrealBuildTool
 						{
 							logger.LogInformation("{Line}", line);
 
-							if (shouldConnect && line.Contains("Listening on")) // This log entry means that the agent is ready for connections.
+							if (shouldConnect && line.Contains("Listening on", StringComparison.OrdinalIgnoreCase)) // This log entry means that the agent is ready for connections.
 							{
 								long totalMs = self.StartTime.ElapsedMilliseconds;
 								logger.LogInformation("Connecting to UbaAgent on {Ip}:{Port} (local agent port {AgentPort}) {Seconds}.{Milliseconds} seconds after assigned", 
@@ -740,20 +740,20 @@ namespace UnrealBuildTool
 		}
 	}
 
-	class UBAAgentCoordinatorHorde : IUBAAgentCoordinator
+	class UBAAgentCoordinatorHorde : IUBAAgentCoordinator, IDisposable
 	{
 		public UBAAgentCoordinatorHorde(ILogger logger, UnrealBuildAcceleratorConfig ubaConfig, CommandLineArguments? additionalArguments = null)
 		{
 			_logger = logger;
 			_ubaConfig = ubaConfig;
 
-			XmlConfig.ApplyTo(_hordeConfig);
-			additionalArguments?.ApplyTo(_hordeConfig);
+			XmlConfig.ApplyTo(HordeConfig);
+			additionalArguments?.ApplyTo(HordeConfig);
 
 			// Sentry is currently unsupported for non-Windows and non-x64
 			if (!OperatingSystem.IsWindows() || RuntimeInformation.ProcessArchitecture != Architecture.X64)
 			{
-				_hordeConfig.UBASentryUrl = null;
+				HordeConfig.UBASentryUrl = null;
 			}
 		}
 
@@ -775,7 +775,7 @@ namespace UnrealBuildTool
 			}
 
 			_cancellationSource = new CancellationTokenSource();
-			_hordeSessionTask = UBAHordeSession.TryCreateHordeSession(_hordeConfig, executor, _ubaConfig.bStrict, _logger, _cancellationSource.Token);
+			_hordeSessionTask = UBAHordeSession.TryCreateHordeSession(HordeConfig, executor, _ubaConfig.bStrict, _logger, _cancellationSource.Token);
 			await _hordeSessionTask;
 		}
 
@@ -823,7 +823,7 @@ namespace UnrealBuildTool
 					{
 						int currentLogicalCores = hordeSession.NumLogicalCores;
 
-						if (queueWeight <= queueThreshold || currentLogicalCores >= _hordeConfig.HordeMaxCores || _cancellationSource!.IsCancellationRequested)
+						if (queueWeight <= queueThreshold || currentLogicalCores >= HordeConfig.HordeMaxCores || _cancellationSource!.IsCancellationRequested)
 						{
 							break;
 						}
@@ -833,17 +833,17 @@ namespace UnrealBuildTool
 							Exclusive = true
 						};
 
-						if (!String.IsNullOrEmpty(_hordeConfig.HordePool))
+						if (!String.IsNullOrEmpty(HordeConfig.HordePool))
 						{
-							requirements.Pool = _hordeConfig.HordePool;
+							requirements.Pool = HordeConfig.HordePool;
 						}
 
-						if (_hordeConfig.HordeCondition != null)
+						if (HordeConfig.HordeCondition != null)
 						{
-							requirements.Condition = Condition.Parse(_hordeConfig.HordeCondition);
+							requirements.Condition = Condition.Parse(HordeConfig.HordeCondition);
 						}
 
-						if (!await hordeSession.AddWorkerAsync(requirements, _hordeConfig, _cancellationSource.Token))
+						if (!await hordeSession.AddWorkerAsync(requirements, HordeConfig, _cancellationSource.Token))
 						{
 							_logger.LogDebug("No additional workers available");
 							break;
@@ -869,7 +869,7 @@ namespace UnrealBuildTool
 				}
 
 				_timer?.Change(timerPeriod, Timeout.Infinite);
-			}, null, _hordeConfig.HordeDelay * 1000, timerPeriod);
+			}, null, HordeConfig.HordeDelay * 1000, timerPeriod);
 		}
 
 		public void Stop()
@@ -897,11 +897,24 @@ namespace UnrealBuildTool
 		{
 			Stop();
 			CloseAsync().Wait();
+			Dispose(true);
+			GC.SuppressFinalize(this);
 		}
 
-		ILogger _logger;
-		UnrealBuildAcceleratorConfig _ubaConfig;
-		UnrealBuildAcceleratorHordeConfig _hordeConfig { get; init; } = new();
+		protected virtual void Dispose(bool disposing)
+		{
+			if (disposing)
+			{
+				_cancellationSource?.Dispose();
+				_cancellationSource = null;
+				_timer?.Dispose();
+				_timer = null;
+			}
+		}
+
+		readonly ILogger _logger;
+		readonly UnrealBuildAcceleratorConfig _ubaConfig;
+		UnrealBuildAcceleratorHordeConfig HordeConfig { get; init; } = new();
 
 		CancellationTokenSource? _cancellationSource;
 		Task<UBAHordeSession?>? _hordeSessionTask;
