- Fix minor typos

diff --git a/Engine/Source/Programs/UnrealBuildTool/System/PluginDescriptor.cs b/Engine/Source/Programs/UnrealBuildTool/System/PluginDescriptor.cs
--- a/Engine/Source/Programs/UnrealBuildTool/System/PluginDescriptor.cs
+++ b/Engine/Source/Programs/UnrealBuildTool/System/PluginDescriptor.cs
@@ -205,7 +205,7 @@ namespace UnrealBuildTool
 		public bool bExplicitlyLoaded;
 
 		/// <summary>
-		/// When true, an empty SupportedTargetPlatforms is interpeted as 'no platforms' with the expectation that explict platforms will be added in plugin platform extensions
+		/// When true, an empty SupportedTargetPlatforms is interpreted as 'no platforms' with the expectation that explicit platforms will be added in plugin platform extensions
 		/// </summary>
 		public bool bHasExplicitPlatforms;
 
