Bump preferred MSVC compiler to14.38.33130 - VS2022 17.8

diff --git a/Engine/Source/Programs/UnrealBuildTool/Platform/Windows/MicrosoftPlatformSDK.Versions.cs b/Engine/Source/Programs/UnrealBuildTool/Platform/Windows/MicrosoftPlatformSDK.Versions.cs
--- a/Engine/Source/Programs/UnrealBuildTool/Platform/Windows/MicrosoftPlatformSDK.Versions.cs
+++ b/Engine/Source/Programs/UnrealBuildTool/Platform/Windows/MicrosoftPlatformSDK.Versions.cs
@@ -1,6 +1,6 @@
 // Copyright Epic Games, Inc. All Rights Reserved.
 
-using System.Collections.Generic;
+using System;
 using System.Diagnostics;
 using System.Linq;
 using EpicGames.Core;
@@ -40,10 +40,12 @@ namespace UnrealBuildTool
 
 		/// <summary>
 		/// Ranges of tested compiler toolchains to be used, in order of preference. If multiple toolchains in a range are present, the latest version will be preferred.
-		/// Note that the numbers here correspond to the installation *folders* rather than precise executable versions. 
+		/// Note that the numbers here correspond to the installation *folders* rather than precise executable versions.
 		/// </summary>
+		/// <seealso href="https://learn.microsoft.com/en-us/lifecycle/products/visual-studio-2022"/>
 		static readonly VersionNumberRange[] PreferredVisualCppVersions = new VersionNumberRange[]
 		{
+			VersionNumberRange.Parse("14.38.33130", "14.38.99999"), // VS2022 17.8.x
 			VersionNumberRange.Parse("14.37.32822", "14.37.99999"), // VS2022 17.7.x
 			VersionNumberRange.Parse("14.36.32532", "14.36.99999"), // VS2022 17.6.x
 			VersionNumberRange.Parse("14.35.32215", "14.35.99999"), // VS2022 17.5.x
@@ -53,11 +55,11 @@ namespace UnrealBuildTool
 		/// <summary>
 		/// Minimum Clang version required for MSVC toolchain versions
 		/// </summary>
-		static readonly IReadOnlyDictionary<VersionNumber, VersionNumber> MinimumRequiredClangVersion = new Dictionary<VersionNumber, VersionNumber>()
+		static readonly Tuple<VersionNumber, VersionNumber>[] MinimumRequiredClangVersion = new Tuple<VersionNumber, VersionNumber>[]
 		{
-			{ new VersionNumber(14, 37), new VersionNumber(16) }, // VS2022 17.7.x
-			{ new VersionNumber(14, 35), new VersionNumber(15) }, // VS2022 17.5.x - 17.6.x
-			{ new VersionNumber(14, 34), new VersionNumber(14) }, // VS2022 17.4.x
+			new(new VersionNumber(14, 37), new VersionNumber(16)), // VS2022 17.7.x - 17.8.x
+			new(new VersionNumber(14, 35), new VersionNumber(15)), // VS2022 17.5.x - 17.6.x
+			new(new VersionNumber(14, 34), new VersionNumber(14)), // VS2022 17.4.x
 		};
 
 		/// <summary>
@@ -98,14 +100,7 @@ namespace UnrealBuildTool
 		/// <returns></returns>
 		public static VersionNumber GetMinimumClangVersionForVcVersion(VersionNumber VcVersion)
 		{
-			foreach (KeyValuePair<VersionNumber, VersionNumber> Item in MinimumRequiredClangVersion)
-			{
-				if (VcVersion >= Item.Key)
-				{
-					return Item.Value;
-				}
-			}
-			return MinimumClangVersion;
+			return MinimumRequiredClangVersion.FirstOrDefault(x => VcVersion >= x.Item1)?.Item2 ?? MinimumClangVersion;
 		}
 
 		/// <summary>
diff --git a/Engine/Source/Programs/UnrealBuildTool/Platform/Windows/MicrosoftPlatformSDK.VisualStudio.cs b/Engine/Source/Programs/UnrealBuildTool/Platform/Windows/MicrosoftPlatformSDK.VisualStudio.cs
--- a/Engine/Source/Programs/UnrealBuildTool/Platform/Windows/MicrosoftPlatformSDK.VisualStudio.cs
+++ b/Engine/Source/Programs/UnrealBuildTool/Platform/Windows/MicrosoftPlatformSDK.VisualStudio.cs
@@ -32,9 +32,7 @@ namespace UnrealBuildTool
 		/// <summary>
 		/// Additional set of components that should be suggested to be installed for Visual Studio 2019.
 		/// </summary>
-		static readonly string[] VisualStudio2019SuggestedComponents = new string[]
-		{
-		};
+		static readonly string[] VisualStudio2019SuggestedComponents = System.Array.Empty<string>();
 
 		/// <summary>
 		/// Additional set of components that should be suggested to be installed for Visual Studio 2022.
@@ -42,7 +40,7 @@ namespace UnrealBuildTool
 		static readonly string[] VisualStudio2022SuggestedComponents = new string[]
 		{
 			"Microsoft.VisualStudio.Workload.ManagedDesktop",
-			"Microsoft.VisualStudio.Component.VC.14.37.17.7.x86.x64",
+			"Microsoft.VisualStudio.Component.VC.14.38.17.8.x86.x64",
 			"Microsoft.Net.Component.4.6.2.TargetingPack",
 		};
 
