Build fix.

diff --git a/Engine/Source/Programs/AutomationTool/BuildGraph/Tasks/DeployToolTask.cs b/Engine/Source/Programs/AutomationTool/BuildGraph/Tasks/DeployToolTask.cs
--- a/Engine/Source/Programs/AutomationTool/BuildGraph/Tasks/DeployToolTask.cs
+++ b/Engine/Source/Programs/AutomationTool/BuildGraph/Tasks/DeployToolTask.cs
@@ -160,7 +160,7 @@ namespace AutomationTool.Tasks
 
 			BlobSerializerOptions serializerOptions = BlobSerializerOptions.Create(infoResponse.ApiVersion);
 
-			IBlobHandle handle;
+			IBlobRef handle;
 
 			ToolId toolId = new ToolId(Parameters.Id);
 
@@ -198,7 +198,7 @@ namespace AutomationTool.Tasks
 				createPaused = true;
 			}
 
-			BlobLocator locator = handle.GetLocator();
+			BlobRefValue locator = handle.GetRefValue();
 			ToolDeploymentId deploymentId = await hordeHttpClient.CreateToolDeploymentAsync(toolId, Parameters.Version, duration, createPaused, locator);
 			Logger.LogInformation("Created {ToolId} deployment {DeploymentId}", toolId, deploymentId);
 		}
diff --git a/Engine/Source/Programs/UnrealBuildTool/Artifacts/HordeStorageArtifactCache.cs b/Engine/Source/Programs/UnrealBuildTool/Artifacts/HordeStorageArtifactCache.cs
--- a/Engine/Source/Programs/UnrealBuildTool/Artifacts/HordeStorageArtifactCache.cs
+++ b/Engine/Source/Programs/UnrealBuildTool/Artifacts/HordeStorageArtifactCache.cs
@@ -270,7 +270,7 @@ namespace UnrealBuildTool.Artifacts
 					{
 						artifactActions.AddRange(_pendingWrites.Where(x => x.Key == key));
 					}
-					ArtifactActionCollectionNode? node = await _store.TryReadRefAsync<ArtifactActionCollectionNode>(GetRefName(key), default, cancellationToken: cancellationToken);
+					ArtifactActionCollectionNode? node = await _store.TryReadRefTargetAsync<ArtifactActionCollectionNode>(GetRefName(key), default, cancellationToken: cancellationToken);
 					if (node != null)
 					{
 						foreach (HordeArtifactAction artifactAction in node.ArtifactActions.Values)
@@ -302,7 +302,7 @@ namespace UnrealBuildTool.Artifacts
 			{
 				output[index] = false;
 				ArtifactAction artifactAction = artifactActions[index];
-				ArtifactActionCollectionNode? node = await _store.TryReadRefAsync<ArtifactActionCollectionNode>(GetRefName(artifactAction.Key), default, cancellationToken: cancellationToken);
+				ArtifactActionCollectionNode? node = await _store.TryReadRefTargetAsync<ArtifactActionCollectionNode>(GetRefName(artifactAction.Key), default, cancellationToken: cancellationToken);
 				if (node != null)
 				{
 					if (node.ArtifactActions.TryGetValue(artifactAction.ActionKey, out HordeArtifactAction hordeArtifactAction))
@@ -469,7 +469,7 @@ namespace UnrealBuildTool.Artifacts
 					RefName refName = GetRefName(artifactAction.Key);
 
 					// Locate the destination collection for this key
-					ArtifactActionCollectionNode? node = _store!.TryReadRefAsync<ArtifactActionCollectionNode>(refName, default, cancellationToken: cancellationToken).Result;
+					ArtifactActionCollectionNode? node = _store!.TryReadRefTargetAsync<ArtifactActionCollectionNode>(refName, default, cancellationToken: cancellationToken).Result;
 					node ??= new ArtifactActionCollectionNode();
 
 					// Update the artifact action collection
@@ -483,7 +483,7 @@ namespace UnrealBuildTool.Artifacts
 					await writer.FlushAsync();
 
 					// Save the collection
-					await _store.WriteRefTargetAsync(refName, nodeRef, cancellationToken: cancellationToken);
+					await _store.WriteRefAsync(refName, nodeRef, cancellationToken: cancellationToken);
 				}, cancellationToken));
 			}
 			return tasks;
