Horde: Force use of SSL encryption if relay mode is active

diff --git a/Engine/Source/Programs/UnrealBuildTool/Executors/Experimental/UBAHordeSession.cs b/Engine/Source/Programs/UnrealBuildTool/Executors/Experimental/UBAHordeSession.cs
--- a/Engine/Source/Programs/UnrealBuildTool/Executors/Experimental/UBAHordeSession.cs
+++ b/Engine/Source/Programs/UnrealBuildTool/Executors/Experimental/UBAHordeSession.cs
@@ -260,6 +260,7 @@ namespace UnrealBuildTool
 			if (connectionMode == ConnectionMode.Relay && String.IsNullOrEmpty(_crypto))
 			{
 				_crypto = UBAExecutor.CreateCrypto();
+				_encryption = Encryption.Ssl;
 			}
 		}
 
