PipInstall: Keep virtual env dir if no config file found, allows offline installs into Lib/site-packages before pip installer has run.

diff --git a/Engine/Plugins/Experimental/PythonScriptPlugin/Source/PythonScriptPlugin/Private/PipInstall.cpp b/Engine/Plugins/Experimental/PythonScriptPlugin/Source/PythonScriptPlugin/Private/PipInstall.cpp
--- a/Engine/Plugins/Experimental/PythonScriptPlugin/Source/PythonScriptPlugin/Private/PipInstall.cpp
+++ b/Engine/Plugins/Experimental/PythonScriptPlugin/Source/PythonScriptPlugin/Private/PipInstall.cpp
@@ -315,6 +315,13 @@ void FPipInstall::CheckInvalidPipEnv()
 		return;
 	}
 
+	// If not a venv directory don't delete in case offline packages were added before editor run
+	FString VenvConfig = PipInstallPath / TEXT("pyvenv.cfg");
+	if (!FPaths::FileExists(VenvConfig))
+	{
+		return;
+	}
+
 	const FString VenvVersion = ParseVenvVersion(PipInstallPath);
 	if (VenvVersion == TEXT(PY_VERSION))
 	{
diff --git a/Engine/Source/Programs/UnrealBuildTool/Modes/PipInstallMode.cs b/Engine/Source/Programs/UnrealBuildTool/Modes/PipInstallMode.cs
--- a/Engine/Source/Programs/UnrealBuildTool/Modes/PipInstallMode.cs
+++ b/Engine/Source/Programs/UnrealBuildTool/Modes/PipInstallMode.cs
@@ -204,10 +204,16 @@ namespace UnrealBuildTool.Modes
 
 		public void RemoveInvalidVenv(FileReference? EnginePython)
 		{
-			// Always delete directory if can't find valid venv config
-			if ( PythonVenvVer == null )
+			// Invalid or non-existent virtual env
+			if (PythonVenvVer == null)
 			{
-				CleanVenvDir();
+				// Only delete virtual environment if version mismatch
+				FileReference VenvConfig = FileReference.Combine(InstallDir, "pyvenv.cfg");
+				if ( FileReference.Exists(VenvConfig) )
+				{
+					CleanVenvDir();
+				}
+
 				return;
 			}
 
@@ -405,9 +411,8 @@ namespace UnrealBuildTool.Modes
 		private string? ParseVenvVersion(DirectoryReference VenvDir)
 		{
 			FileReference VenvConfig = FileReference.Combine(VenvDir, "pyvenv.cfg");
-			if ( !FileReference.Exists(VenvConfig) )
+			if (!FileReference.Exists(VenvConfig))
 			{
-				Logger.LogWarning("PipInstall: Unable to find venv config: {VenvFile}", VenvConfig);
 				return null;
 			}
 
