UBA: Retry remote compiler crashes locally

#rnx

diff --git a/Engine/Source/Programs/UnrealBuildTool/Executors/Experimental/UBAExecutor.cs b/Engine/Source/Programs/UnrealBuildTool/Executors/Experimental/UBAExecutor.cs
--- a/Engine/Source/Programs/UnrealBuildTool/Executors/Experimental/UBAExecutor.cs
+++ b/Engine/Source/Programs/UnrealBuildTool/Executors/Experimental/UBAExecutor.cs
@@ -724,8 +724,18 @@ namespace UnrealBuildTool
 						RemoteActionFailedNoOutput(queue, action, e.ExitCode, e.ExecutingHost ?? "Unknown");
 						return;
 					}
+					else if ((uint)e.ExitCode == 0xC0000005)
+					{
+						RemoteActionFailedCrash(queue, action, e.ExitCode, e.ExecutingHost ?? "Unknown", "Access violation");
+						return;
+					}
+                    else if ((uint)e.ExitCode == 0xC0000409)
+                    {
+                        RemoteActionFailedCrash(queue, action, e.ExitCode, e.ExecutingHost ?? "Unknown", "Stack buffer overflow");
+                        return;
+                    }
 
-					string additionalDescription = $"[RemoteExecutor: {e.ExecutingHost}]";
+                    string additionalDescription = $"[RemoteExecutor: {e.ExecutingHost}]";
 					TimeSpan processorTime = e.TotalProcessorTime;
 					TimeSpan executionTime = e.TotalWallTime;
 					List<string> logLines = e.LogLines;
@@ -786,6 +796,23 @@ namespace UnrealBuildTool
 				_bActionsChanged = true;
 			}
 		}
+
+		void RemoteActionFailedCrash(ImmediateActionQueue queue, LinkedAction action, int exitCode, string executingHost, string error)
+		{
+			if (_bIsCancelled)
+			{
+				return;
+			}
+
+			_threadedLogger.LogWarning("{Description} {StatusDescription} [RemoteExecutor: {ExecutingHost}]: Exited with error code {ExitCode} ({Error}). This action will retry locally", action.CommandDescription, action.StatusDescription, executingHost, exitCode, error);
+			_localRetryActions.AddOrUpdate(action, false, (k, v) => false);
+			queue.RequeueAction(action);
+
+			lock (_actionsChangedLock)
+			{
+				_bActionsChanged = true;
+			}
+		}
 	}
 
 	/// <summary>
