[Uba]
* Cache prereq sort index to reduce overhead in time critical code

diff --git a/Engine/Source/Programs/UnrealBuildTool/Executors/ImmediateActionQueue.cs b/Engine/Source/Programs/UnrealBuildTool/Executors/ImmediateActionQueue.cs
--- a/Engine/Source/Programs/UnrealBuildTool/Executors/ImmediateActionQueue.cs
+++ b/Engine/Source/Programs/UnrealBuildTool/Executors/ImmediateActionQueue.cs
@@ -170,6 +170,11 @@ namespace UnrealBuildTool
 			/// Optional execution results
 			/// </summary>
 			public ExecuteResults? Results;
+
+			/// <summary>
+			/// Indices to prereq actions
+			/// </summary>
+			public int[] PrereqActionsSortIndex;
 		};
 
 		/// <summary>
@@ -385,6 +390,7 @@ namespace UnrealBuildTool
 					Status = ActionStatus.Queued,
 					Phase = action.ArtifactMode.HasFlag(ArtifactMode.Enabled) ? initialPhase : ActionPhase.Compile,
 					Results = null,
+					PrereqActionsSortIndex = action.PrerequisiteActions.Select(a => a.SortIndex).ToArray()
 				};
 			}
 
@@ -513,7 +519,7 @@ namespace UnrealBuildTool
 				var actionState = Actions[actionIndex];
 				if (actionState.Status == ActionStatus.Queued &&
 					actionState.Phase == ActionPhase.Compile &&
-					GetActionReadyState(actionState.Action) == ActionReadyState.Ready)
+					GetActionReadyState(actionState) == ActionReadyState.Ready)
 				{
 					yield return actionState.Action;
 				}
@@ -687,7 +693,7 @@ namespace UnrealBuildTool
 				}
 
 				// Based on the ready state, use this action or mark as an error
-				switch (GetActionReadyState(Actions[actionIndex].Action))
+				switch (GetActionReadyState(Actions[actionIndex]))
 				{
 					case ActionReadyState.NotReady:
 						break;
@@ -1242,20 +1248,20 @@ namespace UnrealBuildTool
 		/// </summary>
 		/// <param name="action">Action in question</param>
 		/// <returns>Action ready state</returns>
-		private ActionReadyState GetActionReadyState(LinkedAction action)
+		private ActionReadyState GetActionReadyState(ActionState action)
 		{
-			foreach (LinkedAction prereq in action.PrerequisiteActions)
+			foreach (int prereqIndex in action.PrereqActionsSortIndex)
 			{
 
 				// To avoid doing artifact checks on actions that might need compiling,
 				// we first make sure the action is in the compile phase
-				if (Actions[prereq.SortIndex].Phase != ActionPhase.Compile)
+				if (Actions[prereqIndex].Phase != ActionPhase.Compile)
 				{
 					return ActionReadyState.NotReady;
 				}
 
 				// Respect the compile status of the action
-				switch (Actions[prereq.SortIndex].Status)
+				switch (Actions[prereqIndex].Status)
 				{
 					case ActionStatus.Finished:
 						continue;
