[UBA]
* Optimization in UBAExecutor.cs where macos code path caused slowdowns on other platforms.

diff --git a/Engine/Source/Programs/UnrealBuildTool/Executors/Experimental/UBAExecutor.cs b/Engine/Source/Programs/UnrealBuildTool/Executors/Experimental/UBAExecutor.cs
--- a/Engine/Source/Programs/UnrealBuildTool/Executors/Experimental/UBAExecutor.cs
+++ b/Engine/Source/Programs/UnrealBuildTool/Executors/Experimental/UBAExecutor.cs
@@ -579,6 +579,8 @@ namespace UnrealBuildTool
 		/// <returns>If this action must be local, non-detoured</returns>
 		static bool ForceLocalNoDetour(LinkedAction action)
 		{
+			if (!OperatingSystem.IsMacOS()) // Below code is slow, so early out
+				return false;
 			// Don't let Mac run shell commands through Uba as interposing dylibs into
 			// the shell results in dyld errors about no matching architecture.
 			// The shell is used to run various commands during a build like copy/ditto.
@@ -586,7 +588,7 @@ namespace UnrealBuildTool
 			// Linking is similarly currently not working in Uba on Mac
 			bool bIsShellAction = action.CommandPath == BuildHostPlatform.Current.Shell;
 			bool bIsLinkAction = action.ActionType == ActionType.Link;
-			return OperatingSystem.IsMacOS() && (bIsShellAction || bIsLinkAction);
+			return (bIsShellAction || bIsLinkAction);
 		}
 
 		/// <summary>
@@ -596,10 +598,10 @@ namespace UnrealBuildTool
 		/// <returns>If this action can be run remotely</returns>
 		bool CanRunRemotely(LinkedAction action) => 
 			action.bCanExecuteInUBA && 
-			action.bCanExecuteRemotely && 
+			action.bCanExecuteRemotely &&
+			(UBAConfig.bLinkRemote || action.ActionType != ActionType.Link) &&
 			!_localRetryActions.ContainsKey(action) && 
 			!_forcedRetryActions.ContainsKey(action) && 
-			(UBAConfig.bLinkRemote || action.ActionType != ActionType.Link) &&
 			!ForceLocalNoDetour(action);
 
 		ProcessStartInfo GetActionStartInfo(LinkedAction action, out FileItem? pchItem)
