UnrealBuildTool: Add a warning if the windows compiler is not a preferred version.

Can be controled (Error, Warning, Off), default is Warning:
CommandLine: -ToolchainVersionWarningLevel=value
BuildConfiguration.xml: <WindowsPlatform><ToolchainVersionWarningLevel>value</ToolchainVersionWarningLevel></WindowsPlatform>
Ini: [/Script/WindowsTargetPlatform.WindowsTargetSettings] ToolchainVersionWarningLevel=value

diff --git a/Engine/Source/Programs/UnrealBuildTool/Platform/Windows/MicrosoftPlatformSDK.Versions.cs b/Engine/Source/Programs/UnrealBuildTool/Platform/Windows/MicrosoftPlatformSDK.Versions.cs
--- a/Engine/Source/Programs/UnrealBuildTool/Platform/Windows/MicrosoftPlatformSDK.Versions.cs
+++ b/Engine/Source/Programs/UnrealBuildTool/Platform/Windows/MicrosoftPlatformSDK.Versions.cs
@@ -4,6 +4,7 @@ using System;
 using System.Diagnostics;
 using System.Linq;
 using EpicGames.Core;
+using UnrealBuildTool.PS4;
 
 namespace UnrealBuildTool
 {
@@ -93,6 +94,51 @@ namespace UnrealBuildTool
 			maxVersion = MaximumWindowsSDKVersion?.ToString();
 		}
 
+		/// <summary>
+		/// If a toolchain version is a preferred version
+		/// </summary>
+		/// <param name="toolchain">The toolchain type</param>
+		/// <param name="version">The version number</param>
+		/// <returns>If the version is preferred</returns>
+		public static bool IsPreferredVersion(WindowsCompiler toolchain, VersionNumber version)
+		{
+			if (toolchain.IsMSVC())
+			{
+				return PreferredVisualCppVersions.Any(x => x.Contains(version));
+			}
+			else if (toolchain.IsClang())
+			{
+				return PreferredClangVersions.Any(x => x.Contains(version));
+			}
+			else if (toolchain.IsIntel())
+			{
+				return PreferredIntelOneApiVersions.Any(x => x.Contains(version));
+			}
+			return false;
+		}
+
+		/// <summary>
+		/// Get the latest preferred toolchain version
+		/// </summary>
+		/// <param name="toolchain">The toolchain type</param>
+		/// <returns>The version number</returns>
+		public static VersionNumber GetLatestPreferredVersion(WindowsCompiler toolchain)
+		{
+			if (toolchain.IsMSVC())
+			{
+				return PreferredVisualCppVersions.Select(x => x.Min).Max()!;
+			}
+			else if (toolchain.IsClang())
+			{
+				return PreferredClangVersions.Select(x => x.Min).Max()!;
+			}
+			else if (toolchain.IsIntel())
+			{
+				return PreferredIntelOneApiVersions.Select(x => x.Min).Max()!;
+			}
+			return new VersionNumber(0);
+		}
+
 		/// <summary>
 		/// The minimum supported Clang version for a given MSVC toolchain
 		/// </summary>
diff --git a/Engine/Source/Programs/UnrealBuildTool/Platform/Windows/MicrosoftPlatformSDK.cs b/Engine/Source/Programs/UnrealBuildTool/Platform/Windows/MicrosoftPlatformSDK.cs
--- a/Engine/Source/Programs/UnrealBuildTool/Platform/Windows/MicrosoftPlatformSDK.cs
+++ b/Engine/Source/Programs/UnrealBuildTool/Platform/Windows/MicrosoftPlatformSDK.cs
@@ -698,6 +698,21 @@ namespace UnrealBuildTool
 			}
 		}
 
+		/// <summary>
+		/// Dump all available toolchain info for a compiler
+		/// </summary>
+		/// <param name="Compiler">The compiler to filter by</param>
+		/// <param name="Architecture">Architecture that must be supported</param>
+		/// <param name="Logger">The ILogger interface to write to</param>
+		public static void DumpAllToolChainInstallations(WindowsCompiler Compiler, UnrealArch Architecture, ILogger Logger)
+		{
+			List<ToolChainInstallation>? ToolChains = CachedToolChainInstallations.GetValueOrDefault(Compiler);
+			if (ToolChains != null)
+			{
+				DumpToolChains(ToolChains, x => x.ThenBy(x => x.ReleaseChannel).ThenByDescending(x => x.Version), Architecture, Logger);
+			}
+		}
+
 		static List<ToolChainInstallation> FindToolChainInstallations(WindowsCompiler Compiler, ILogger Logger)
 		{
 			return CachedToolChainInstallations.GetOrAdd(Compiler, _ =>
diff --git a/Engine/Source/Programs/UnrealBuildTool/Platform/Windows/UEBuildWindows.cs b/Engine/Source/Programs/UnrealBuildTool/Platform/Windows/UEBuildWindows.cs
--- a/Engine/Source/Programs/UnrealBuildTool/Platform/Windows/UEBuildWindows.cs
+++ b/Engine/Source/Programs/UnrealBuildTool/Platform/Windows/UEBuildWindows.cs
@@ -287,6 +287,15 @@ namespace UnrealBuildTool
 		}
 		= UnrealArch.X64;
 
+		/// <summary>
+		/// Warning level when reporting toolchains that are not in the preferred version list
+		/// </summary>
+		/// <seealso cref="MicrosoftPlatformSDK.PreferredVisualCppVersions"/>
+		[ConfigFile(ConfigHierarchyType.Engine, "/Script/WindowsTargetPlatform.WindowsTargetSettings", "ToolchainVersionWarningLevel")]
+		[XmlConfigFile(Category = "WindowsPlatform")]
+		[CommandLine("-ToolchainVersionWarningLevel=")]
+		public WarningLevel ToolchainVersionWarningLevel { get; set; } = WarningLevel.Warning;
+
 		/// <summary>
 		/// The specific compiler version to use. This may be a specific version number (for example, "14.13.26128"), the string "Latest" to select the newest available version, or
 		/// the string "Preview" to select the newest available preview version. By default, and if it is available, we use the toolchain version indicated by
@@ -793,6 +802,8 @@ namespace UnrealBuildTool
 
 		public UnrealArch Architecture => Inner.Architecture;
 
+		public WarningLevel ToolchainVersionWarningLevel => Inner.ToolchainVersionWarningLevel;
+
 		public string? CompilerVersion => Inner.CompilerVersion;
 
 		public string? ToolchainVerison => Inner.ToolchainVersion;
@@ -1190,6 +1201,31 @@ namespace UnrealBuildTool
 				}
 			}
 
+			if (Target.WindowsPlatform.ToolchainVersionWarningLevel != WarningLevel.Off)
+			{
+				if (!MicrosoftPlatformSDK.IsPreferredVersion(Target.WindowsPlatform.Compiler, Target.WindowsPlatform.Environment.CompilerVersion))
+				{
+					VersionNumber preferred = MicrosoftPlatformSDK.GetLatestPreferredVersion(Target.WindowsPlatform.Compiler);
+					MicrosoftPlatformSDK.DumpAllToolChainInstallations(Target.WindowsPlatform.Compiler, Target.Architecture, Logger);
+					if (Target.WindowsPlatform.ToolchainVersionWarningLevel == WarningLevel.Error)
+					{
+						throw new BuildLogEventException("{Compiler} compiler version {Version} is not a preferred version. Please use the latest preferred version {PreferredVersion}", WindowsPlatform.GetCompilerName(Target.WindowsPlatform.Compiler), Target.WindowsPlatform.Environment.CompilerVersion, preferred);
+					}
+					Logger.LogInformation("{Compiler} compiler version {Version} is not a preferred version. Please use the latest preferred version {PreferredVersion}", WindowsPlatform.GetCompilerName(Target.WindowsPlatform.Compiler), Target.WindowsPlatform.Environment.CompilerVersion, preferred);
+				}
+
+				if (Target.WindowsPlatform.Compiler != Target.WindowsPlatform.ToolChain && !MicrosoftPlatformSDK.IsPreferredVersion(Target.WindowsPlatform.ToolChain, Target.WindowsPlatform.Environment.ToolChainVersion))
+				{
+					VersionNumber preferred = MicrosoftPlatformSDK.GetLatestPreferredVersion(Target.WindowsPlatform.ToolChain);
+					MicrosoftPlatformSDK.DumpAllToolChainInstallations(Target.WindowsPlatform.ToolChain, Target.Architecture, Logger);
+					if (Target.WindowsPlatform.ToolchainVersionWarningLevel == WarningLevel.Error)
+					{
+						throw new BuildLogEventException("{Toolchain} toolchain version {Version} is not a preferred version. Please use the latest preferred version {PreferredVersion}", WindowsPlatform.GetCompilerName(Target.WindowsPlatform.ToolChain), Target.WindowsPlatform.Environment.ToolChainVersion, preferred);
+					}
+					Logger.LogInformation("{Toolchain} toolchain version {Version} is not a preferred version. Please use a preferred toolchain such as {PreferredVersion}", WindowsPlatform.GetCompilerName(Target.WindowsPlatform.ToolChain), Target.WindowsPlatform.Environment.ToolChainVersion, preferred);
+				}
+			}
+
 			//			@Todo: Still getting reports of frequent OOM issues with this enabled as of 15.7.
 			//			// Enable fast PDB linking if we're on VS2017 15.7 or later. Previous versions have OOM issues with large projects.
 			//			if(!Target.bFormalBuild && !Target.bUseFastPDBLinking.HasValue && Target.WindowsPlatform.Compiler.IsMSVC())
diff --git a/Engine/Source/Programs/UnrealBuildTool/Platform/Windows/VCToolChain.cs b/Engine/Source/Programs/UnrealBuildTool/Platform/Windows/VCToolChain.cs
--- a/Engine/Source/Programs/UnrealBuildTool/Platform/Windows/VCToolChain.cs
+++ b/Engine/Source/Programs/UnrealBuildTool/Platform/Windows/VCToolChain.cs
@@ -109,6 +109,19 @@ namespace UnrealBuildTool
 			{
 				Lines.Add($"Using {WindowsPlatform.GetCompilerName(EnvVars.Compiler)} {EnvVars.CompilerVersion} compiler ({EnvVars.CompilerDir}) with {WindowsPlatform.GetCompilerName(EnvVars.ToolChain)} {EnvVars.ToolChainVersion} runtime ({EnvVars.ToolChainDir}) and Windows {EnvVars.WindowsSdkVersion} SDK ({EnvVars.WindowsSdkDir}).");
 			}
+
+			if (Target.WindowsPlatform.ToolchainVersionWarningLevel == WarningLevel.Warning)
+			{
+				if (!MicrosoftPlatformSDK.IsPreferredVersion(EnvVars.Compiler, EnvVars.CompilerVersion))
+				{
+					Lines.Add($"Warning: {WindowsPlatform.GetCompilerName(EnvVars.Compiler)} compiler is not a preferred version");
+				}
+
+				if (EnvVars.Compiler != EnvVars.ToolChain && !MicrosoftPlatformSDK.IsPreferredVersion(EnvVars.ToolChain, EnvVars.ToolChainVersion))
+				{
+					Lines.Add($"Warning: {WindowsPlatform.GetCompilerName(EnvVars.ToolChain)} runtime is not a preferred version");
+				}
+			}
 		}
 
 		public override void GetExternalDependencies(HashSet<FileItem> ExternalDependencies)
