UBT: Fix build error (reference to missing namespace).

diff --git a/Engine/Source/Programs/UnrealBuildTool/Platform/Windows/MicrosoftPlatformSDK.Versions.cs b/Engine/Source/Programs/UnrealBuildTool/Platform/Windows/MicrosoftPlatformSDK.Versions.cs
--- a/Engine/Source/Programs/UnrealBuildTool/Platform/Windows/MicrosoftPlatformSDK.Versions.cs
+++ b/Engine/Source/Programs/UnrealBuildTool/Platform/Windows/MicrosoftPlatformSDK.Versions.cs
@@ -4,7 +4,6 @@ using System;
 using System.Diagnostics;
 using System.Linq;
 using EpicGames.Core;
-using UnrealBuildTool.PS4;
 
 namespace UnrealBuildTool
 {
