Adding and switching to the 5.5 include order version.

diff --git a/Engine/Source/Programs/UnrealBuildTool/Configuration/TargetRules.cs b/Engine/Source/Programs/UnrealBuildTool/Configuration/TargetRules.cs
--- a/Engine/Source/Programs/UnrealBuildTool/Configuration/TargetRules.cs
+++ b/Engine/Source/Programs/UnrealBuildTool/Configuration/TargetRules.cs
@@ -176,12 +176,13 @@ namespace UnrealBuildTool
 		/// <summary>
 		/// Include order used in Unreal 5.2
 		/// </summary>
-		[Obsolete("The Unreal 5.2 include order is deprecated and will be unsupported in 5.5.")]
+		[Obsolete("The Unreal 5.2 include order is unsupported.")]
 		Unreal5_2,
 
 		/// <summary>
 		/// Include order used in Unreal 5.3
 		/// </summary>
+		[Obsolete("The Unreal 5.3 include order is deprecated and will be unsupported in 5.6.")]
 		Unreal5_3,
 
 		/// <summary>
@@ -189,19 +190,24 @@ namespace UnrealBuildTool
 		/// </summary>
 		Unreal5_4,
 
+		/// <summary>
+		/// Include order used in Unreal 5.5
+		/// </summary>
+		Unreal5_5,
+
 		// *** When adding new entries here, be sure to update UEBuildModuleCPP.CurrentIncludeOrderDefine to ensure that the correct guard is used. ***
 
 		/// <summary>
 		/// Always use the latest version of include order. This value is updated every Unreal release, use with caution if you intend to integrate newer Unreal releases.
 		/// </summary>
-		Latest = Unreal5_4,
+		Latest = Unreal5_5,
 
 		/// <summary>
 		/// Contains the oldest version of include order that the engine supports.
 		/// Do not delete old enum entries to prevent breaking project generation
 		/// </summary>
 #pragma warning disable CS0618 // Type or member is obsolete
-		Oldest = Unreal5_2,
+		Oldest = Unreal5_3,
 #pragma warning restore CS0618 // Type or member is obsolete
 	}
 
@@ -400,11 +406,13 @@ namespace UnrealBuildTool
 			switch (inVersion)
 			{
 				default:
+				case EngineIncludeOrderVersion.Unreal5_5:
+					return "UE_ENABLE_INCLUDE_ORDER_DEPRECATED_IN_5_5";
 				case EngineIncludeOrderVersion.Unreal5_4:
 					return "UE_ENABLE_INCLUDE_ORDER_DEPRECATED_IN_5_4";
+#pragma warning disable CS0618 // Type or member is obsolete
 				case EngineIncludeOrderVersion.Unreal5_3:
 					return "UE_ENABLE_INCLUDE_ORDER_DEPRECATED_IN_5_3";
-#pragma warning disable CS0618 // Type or member is obsolete
 				case EngineIncludeOrderVersion.Unreal5_2:
 					return "UE_ENABLE_INCLUDE_ORDER_DEPRECATED_IN_5_2";
 #pragma warning restore CS0618 // Type or member is obsolete
@@ -426,9 +434,10 @@ namespace UnrealBuildTool
 			{
 #pragma warning disable CS0618 // Type or member is obsolete
 				GetDeprecationDefine(EngineIncludeOrderVersion.Unreal5_2),
-#pragma warning restore CS0618 // Type or member is obsolete
 				GetDeprecationDefine(EngineIncludeOrderVersion.Unreal5_3),
+#pragma warning restore CS0618 // Type or member is obsolete
 				GetDeprecationDefine(EngineIncludeOrderVersion.Unreal5_4),
+				GetDeprecationDefine(EngineIncludeOrderVersion.Unreal5_5),
 			};
 		}
 
@@ -443,9 +452,10 @@ namespace UnrealBuildTool
 			{
 #pragma warning disable CS0618 // Type or member is obsolete
 				GetDeprecationDefine(EngineIncludeOrderVersion.Unreal5_2, inVersion),
-#pragma warning restore CS0618 // Type or member is obsolete
 				GetDeprecationDefine(EngineIncludeOrderVersion.Unreal5_3, inVersion),
+#pragma warning restore CS0618 // Type or member is obsolete
 				GetDeprecationDefine(EngineIncludeOrderVersion.Unreal5_4, inVersion),
+				GetDeprecationDefine(EngineIncludeOrderVersion.Unreal5_5, inVersion),
 			};
 		}
 
diff --git a/Templates/TP_Blank/Source/TP_Blank.Target.cs b/Templates/TP_Blank/Source/TP_Blank.Target.cs
--- a/Templates/TP_Blank/Source/TP_Blank.Target.cs
+++ b/Templates/TP_Blank/Source/TP_Blank.Target.cs
@@ -9,7 +9,7 @@ public class TP_BlankTarget : TargetRules
 	{
 		Type = TargetType.Game;
 		DefaultBuildSettings = BuildSettingsVersion.V5;
-		IncludeOrderVersion = EngineIncludeOrderVersion.Unreal5_4;
+		IncludeOrderVersion = EngineIncludeOrderVersion.Unreal5_5;
 		ExtraModuleNames.Add("TP_Blank");
 	}
 }
diff --git a/Templates/TP_Blank/Source/TP_BlankEditor.Target.cs b/Templates/TP_Blank/Source/TP_BlankEditor.Target.cs
--- a/Templates/TP_Blank/Source/TP_BlankEditor.Target.cs
+++ b/Templates/TP_Blank/Source/TP_BlankEditor.Target.cs
@@ -9,7 +9,7 @@ public class TP_BlankEditorTarget : TargetRules
 	{
 		Type = TargetType.Editor;
 		DefaultBuildSettings = BuildSettingsVersion.V5;
-		IncludeOrderVersion = EngineIncludeOrderVersion.Unreal5_4;
+		IncludeOrderVersion = EngineIncludeOrderVersion.Unreal5_5;
 		ExtraModuleNames.Add("TP_Blank");
 	}
 }
diff --git a/Templates/TP_FirstPerson/Source/TP_FirstPerson.Target.cs b/Templates/TP_FirstPerson/Source/TP_FirstPerson.Target.cs
--- a/Templates/TP_FirstPerson/Source/TP_FirstPerson.Target.cs
+++ b/Templates/TP_FirstPerson/Source/TP_FirstPerson.Target.cs
@@ -9,7 +9,7 @@ public class TP_FirstPersonTarget : TargetRules
 	{
 		Type = TargetType.Game;
 		DefaultBuildSettings = BuildSettingsVersion.V5;
-		IncludeOrderVersion = EngineIncludeOrderVersion.Unreal5_4;
+		IncludeOrderVersion = EngineIncludeOrderVersion.Unreal5_5;
 		ExtraModuleNames.Add("TP_FirstPerson");
 	}
 }
diff --git a/Templates/TP_FirstPerson/Source/TP_FirstPersonEditor.Target.cs b/Templates/TP_FirstPerson/Source/TP_FirstPersonEditor.Target.cs
--- a/Templates/TP_FirstPerson/Source/TP_FirstPersonEditor.Target.cs
+++ b/Templates/TP_FirstPerson/Source/TP_FirstPersonEditor.Target.cs
@@ -9,7 +9,7 @@ public class TP_FirstPersonEditorTarget : TargetRules
 	{
 		Type = TargetType.Editor;
 		DefaultBuildSettings = BuildSettingsVersion.V5;
-		IncludeOrderVersion = EngineIncludeOrderVersion.Unreal5_4;
+		IncludeOrderVersion = EngineIncludeOrderVersion.Unreal5_5;
 		ExtraModuleNames.Add("TP_FirstPerson");
 	}
 }
diff --git a/Templates/TP_Puzzle/Source/TP_Puzzle.Target.cs b/Templates/TP_Puzzle/Source/TP_Puzzle.Target.cs
--- a/Templates/TP_Puzzle/Source/TP_Puzzle.Target.cs
+++ b/Templates/TP_Puzzle/Source/TP_Puzzle.Target.cs
@@ -9,7 +9,7 @@ public class TP_PuzzleTarget : TargetRules
 	{
 		Type = TargetType.Game;
 		DefaultBuildSettings = BuildSettingsVersion.V5;
-		IncludeOrderVersion = EngineIncludeOrderVersion.Unreal5_4;
+		IncludeOrderVersion = EngineIncludeOrderVersion.Unreal5_5;
 		ExtraModuleNames.Add("TP_Puzzle");
 	}
 }
diff --git a/Templates/TP_Puzzle/Source/TP_PuzzleEditor.Target.cs b/Templates/TP_Puzzle/Source/TP_PuzzleEditor.Target.cs
--- a/Templates/TP_Puzzle/Source/TP_PuzzleEditor.Target.cs
+++ b/Templates/TP_Puzzle/Source/TP_PuzzleEditor.Target.cs
@@ -9,7 +9,7 @@ public class TP_PuzzleEditorTarget : TargetRules
 	{
 		Type = TargetType.Editor;
 		DefaultBuildSettings = BuildSettingsVersion.V5;
-		IncludeOrderVersion = EngineIncludeOrderVersion.Unreal5_4;
+		IncludeOrderVersion = EngineIncludeOrderVersion.Unreal5_5;
 		ExtraModuleNames.Add("TP_Puzzle");
 	}
 }
diff --git a/Templates/TP_SIM_Blank/Source/TP_SIM_Blank.Target.cs b/Templates/TP_SIM_Blank/Source/TP_SIM_Blank.Target.cs
--- a/Templates/TP_SIM_Blank/Source/TP_SIM_Blank.Target.cs
+++ b/Templates/TP_SIM_Blank/Source/TP_SIM_Blank.Target.cs
@@ -9,7 +9,7 @@ public class TP_SIM_BlankTarget : TargetRules
 	{
 		Type = TargetType.Game;
 		DefaultBuildSettings = BuildSettingsVersion.V5;
-		IncludeOrderVersion = EngineIncludeOrderVersion.Unreal5_4;
+		IncludeOrderVersion = EngineIncludeOrderVersion.Unreal5_5;
 		ExtraModuleNames.AddRange( new string[] { "TP_SIM_Blank" } );
 	}
 }
diff --git a/Templates/TP_SIM_Blank/Source/TP_SIM_BlankEditor.Target.cs b/Templates/TP_SIM_Blank/Source/TP_SIM_BlankEditor.Target.cs
--- a/Templates/TP_SIM_Blank/Source/TP_SIM_BlankEditor.Target.cs
+++ b/Templates/TP_SIM_Blank/Source/TP_SIM_BlankEditor.Target.cs
@@ -9,7 +9,7 @@ public class TP_SIM_BlankEditorTarget : TargetRules
 	{
 		Type = TargetType.Editor;
 		DefaultBuildSettings = BuildSettingsVersion.V5;
-		IncludeOrderVersion = EngineIncludeOrderVersion.Unreal5_4;
+		IncludeOrderVersion = EngineIncludeOrderVersion.Unreal5_5;
 		ExtraModuleNames.AddRange( new string[] { "TP_SIM_Blank" } );
 	}
 }
diff --git a/Templates/TP_SIM_Sandbox/Source/TP_SIM_Sandbox.Target.cs b/Templates/TP_SIM_Sandbox/Source/TP_SIM_Sandbox.Target.cs
--- a/Templates/TP_SIM_Sandbox/Source/TP_SIM_Sandbox.Target.cs
+++ b/Templates/TP_SIM_Sandbox/Source/TP_SIM_Sandbox.Target.cs
@@ -9,7 +9,7 @@ public class TP_SIM_SandboxTarget : TargetRules
 	{
 		Type = TargetType.Game;
 		DefaultBuildSettings = BuildSettingsVersion.V5;
-		IncludeOrderVersion = EngineIncludeOrderVersion.Unreal5_4;
+		IncludeOrderVersion = EngineIncludeOrderVersion.Unreal5_5;
 		ExtraModuleNames.AddRange( new string[] { "TP_SIM_Sandbox" } );
 	}
 }
diff --git a/Templates/TP_SIM_Sandbox/Source/TP_SIM_SandboxEditor.Target.cs b/Templates/TP_SIM_Sandbox/Source/TP_SIM_SandboxEditor.Target.cs
--- a/Templates/TP_SIM_Sandbox/Source/TP_SIM_SandboxEditor.Target.cs
+++ b/Templates/TP_SIM_Sandbox/Source/TP_SIM_SandboxEditor.Target.cs
@@ -9,7 +9,7 @@ public class TP_SIM_SandboxEditorTarget : TargetRules
 	{
 		Type = TargetType.Editor;
 		DefaultBuildSettings = BuildSettingsVersion.V5;
-		IncludeOrderVersion = EngineIncludeOrderVersion.Unreal5_4;
+		IncludeOrderVersion = EngineIncludeOrderVersion.Unreal5_5;
 		ExtraModuleNames.AddRange( new string[] { "TP_SIM_Sandbox" } );
 	}
 }
diff --git a/Templates/TP_ThirdPerson/Source/TP_ThirdPerson.Target.cs b/Templates/TP_ThirdPerson/Source/TP_ThirdPerson.Target.cs
--- a/Templates/TP_ThirdPerson/Source/TP_ThirdPerson.Target.cs
+++ b/Templates/TP_ThirdPerson/Source/TP_ThirdPerson.Target.cs
@@ -9,7 +9,7 @@ public class TP_ThirdPersonTarget : TargetRules
 	{
 		Type = TargetType.Game;
 		DefaultBuildSettings = BuildSettingsVersion.V5;
-		IncludeOrderVersion = EngineIncludeOrderVersion.Unreal5_4;
+		IncludeOrderVersion = EngineIncludeOrderVersion.Unreal5_5;
 		ExtraModuleNames.Add("TP_ThirdPerson");
 	}
 }
diff --git a/Templates/TP_ThirdPerson/Source/TP_ThirdPersonEditor.Target.cs b/Templates/TP_ThirdPerson/Source/TP_ThirdPersonEditor.Target.cs
--- a/Templates/TP_ThirdPerson/Source/TP_ThirdPersonEditor.Target.cs
+++ b/Templates/TP_ThirdPerson/Source/TP_ThirdPersonEditor.Target.cs
@@ -9,7 +9,7 @@ public class TP_ThirdPersonEditorTarget : TargetRules
 	{
 		Type = TargetType.Editor;
 		DefaultBuildSettings = BuildSettingsVersion.V5;
-		IncludeOrderVersion = EngineIncludeOrderVersion.Unreal5_4;
+		IncludeOrderVersion = EngineIncludeOrderVersion.Unreal5_5;
 		ExtraModuleNames.Add("TP_ThirdPerson");
 	}
 }
diff --git a/Templates/TP_TopDown/Source/TP_TopDown.Target.cs b/Templates/TP_TopDown/Source/TP_TopDown.Target.cs
--- a/Templates/TP_TopDown/Source/TP_TopDown.Target.cs
+++ b/Templates/TP_TopDown/Source/TP_TopDown.Target.cs
@@ -9,7 +9,7 @@ public class TP_TopDownTarget : TargetRules
 	{
 		Type = TargetType.Game;
 		DefaultBuildSettings = BuildSettingsVersion.V5;
-		IncludeOrderVersion = EngineIncludeOrderVersion.Unreal5_4;
+		IncludeOrderVersion = EngineIncludeOrderVersion.Unreal5_5;
 		ExtraModuleNames.Add("TP_TopDown");
 	}
 }
diff --git a/Templates/TP_TopDown/Source/TP_TopDownEditor.Target.cs b/Templates/TP_TopDown/Source/TP_TopDownEditor.Target.cs
--- a/Templates/TP_TopDown/Source/TP_TopDownEditor.Target.cs
+++ b/Templates/TP_TopDown/Source/TP_TopDownEditor.Target.cs
@@ -9,7 +9,7 @@ public class TP_TopDownEditorTarget : TargetRules
 	{
 		Type = TargetType.Editor;
 		DefaultBuildSettings = BuildSettingsVersion.V5;
-		IncludeOrderVersion = EngineIncludeOrderVersion.Unreal5_4;
+		IncludeOrderVersion = EngineIncludeOrderVersion.Unreal5_5;
 		ExtraModuleNames.Add("TP_TopDown");
 	}
 }
diff --git a/Templates/TP_VehicleAdv/Source/TP_VehicleAdv.Target.cs b/Templates/TP_VehicleAdv/Source/TP_VehicleAdv.Target.cs
--- a/Templates/TP_VehicleAdv/Source/TP_VehicleAdv.Target.cs
+++ b/Templates/TP_VehicleAdv/Source/TP_VehicleAdv.Target.cs
@@ -9,7 +9,7 @@ public class TP_VehicleAdvTarget : TargetRules
 	{
 		Type = TargetType.Game;
 		DefaultBuildSettings = BuildSettingsVersion.V5;
-		IncludeOrderVersion = EngineIncludeOrderVersion.Unreal5_4;
+		IncludeOrderVersion = EngineIncludeOrderVersion.Unreal5_5;
 		ExtraModuleNames.Add("TP_VehicleAdv");
 	}
 }
diff --git a/Templates/TP_VehicleAdv/Source/TP_VehicleAdvEditor.Target.cs b/Templates/TP_VehicleAdv/Source/TP_VehicleAdvEditor.Target.cs
--- a/Templates/TP_VehicleAdv/Source/TP_VehicleAdvEditor.Target.cs
+++ b/Templates/TP_VehicleAdv/Source/TP_VehicleAdvEditor.Target.cs
@@ -9,7 +9,7 @@ public class TP_VehicleAdvEditorTarget : TargetRules
 	{
 		Type = TargetType.Editor;
 		DefaultBuildSettings = BuildSettingsVersion.V5;
-		IncludeOrderVersion = EngineIncludeOrderVersion.Unreal5_4;
+		IncludeOrderVersion = EngineIncludeOrderVersion.Unreal5_5;
 		ExtraModuleNames.Add("TP_VehicleAdv");
 	}
 }
