PR #11383: Project Generator for VS Workspace support

diff --git a/Engine/Source/Programs/UnrealBuildTool/Modes/GenerateProjectFilesMode.cs b/Engine/Source/Programs/UnrealBuildTool/Modes/GenerateProjectFilesMode.cs
--- a/Engine/Source/Programs/UnrealBuildTool/Modes/GenerateProjectFilesMode.cs
+++ b/Engine/Source/Programs/UnrealBuildTool/Modes/GenerateProjectFilesMode.cs
@@ -33,6 +33,7 @@ namespace UnrealBuildTool
 		[CommandLine("-EddieProjectFiles", Value = nameof(ProjectFileFormat.Eddie))]
 		[CommandLine("-VSCode", Value = nameof(ProjectFileFormat.VisualStudioCode))]
 		[CommandLine("-VSMac", Value = nameof(ProjectFileFormat.VisualStudioMac))]
+		[CommandLine("-VSWorkspace", Value = nameof(ProjectFileFormat.VisualStudioWorkspace))]
 		[CommandLine("-CLion", Value = nameof(ProjectFileFormat.CLion))]
 		[CommandLine("-Rider", Value = nameof(ProjectFileFormat.Rider))]
 #if __VPROJECT_AVAILABLE__
@@ -216,6 +217,9 @@ namespace UnrealBuildTool
 					case ProjectFileFormat.VisualStudioCode:
 						Generator = new VSCodeProjectFileGenerator(ProjectFile);
 						break;
+					case ProjectFileFormat.VisualStudioWorkspace:
+						Generator = new VSWorkspaceProjectFileGenerator(ProjectFile, Arguments);
+						break;
 					case ProjectFileFormat.CLion:
 						Generator = new CLionGenerator(ProjectFile);
 						break;
diff --git a/Engine/Source/Programs/UnrealBuildTool/ProjectFiles/ProjectFileGenerator.cs b/Engine/Source/Programs/UnrealBuildTool/ProjectFiles/ProjectFileGenerator.cs
--- a/Engine/Source/Programs/UnrealBuildTool/ProjectFiles/ProjectFileGenerator.cs
+++ b/Engine/Source/Programs/UnrealBuildTool/ProjectFiles/ProjectFileGenerator.cs
@@ -143,6 +143,7 @@ namespace UnrealBuildTool
 		CodeLite,
 		VisualStudio,
 		VisualStudio2022,
+		VisualStudioWorkspace,
 		XCode,
 		Eddie,
 		VisualStudioCode,
@@ -2688,7 +2689,7 @@ namespace UnrealBuildTool
 		/// <param name="ProgramProjects">Map of program names to all of the program projects we created</param>
 		/// <param name="RulesAssemblies">Map of RuleAssemblies to their base folders</param>
 		/// <param name="Logger">Logger for output</param>
-		private void AddProjectsForAllTargets(
+		protected void AddProjectsForAllTargets(
 			PlatformProjectGeneratorCollection PlatformProjectGenerators,
 			List<FileReference> AllGames,
 			List<FileReference> AllTargetFiles,
@@ -3055,7 +3056,7 @@ namespace UnrealBuildTool
 		/// </summary>
 		/// <param name="BaseName">The base name for the project file</param>
 		/// <returns>Full path to the project file</returns>
-		protected FileReference GetProjectLocation(string BaseName)
+		protected virtual FileReference GetProjectLocation(string BaseName)
 		{
 			return FileReference.Combine(IntermediateProjectFilesPath, BaseName + ProjectFileExtension);
 		}
diff --git a/Engine/Source/Programs/UnrealBuildTool/ProjectFiles/VisualStudioWorkspace/VSWorkspaceProjectFile.cs b/Engine/Source/Programs/UnrealBuildTool/ProjectFiles/VisualStudioWorkspace/VSWorkspaceProjectFile.cs
new file mode 100644
--- /dev/null
+++ b/Engine/Source/Programs/UnrealBuildTool/ProjectFiles/VisualStudioWorkspace/VSWorkspaceProjectFile.cs
@@ -0,0 +1,336 @@
+// Copyright Epic Games, Inc. All Rights Reserved.
+
+using System;
+using System.Collections.Generic;
+using System.IO;
+using System.Linq;
+using System.Text;
+using System.Text.Json;
+using EpicGames.Core;
+using Microsoft.Extensions.Logging;
+using UnrealBuildBase;
+
+namespace UnrealBuildTool
+{
+	internal class VSWorkspaceProjectFile : ProjectFile
+	{
+		private readonly DirectoryReference RootPath;
+		private readonly HashSet<TargetType> TargetTypes;
+		private readonly CommandLineArguments Arguments;
+
+		/// <summary>
+		/// Collection of output files for this project
+		/// </summary>
+		public List<ExportedTargetInfo> ExportedTargetProjects { get; set; } = new();
+
+		public VSWorkspaceProjectFile(FileReference InProjectFilePath, DirectoryReference BaseDir,
+			DirectoryReference RootPath, HashSet<TargetType> TargetTypes, CommandLineArguments Arguments)
+			: base(InProjectFilePath, BaseDir)
+		{
+			this.RootPath = RootPath;
+			this.TargetTypes = TargetTypes;
+			this.Arguments = Arguments;
+		}
+
+		/// <summary>
+		/// Write project file info in JSON file.
+		/// For every combination of <c>UnrealTargetPlatform</c>, <c>UnrealTargetConfiguration</c> and <c>TargetType</c>
+		/// will be generated separate JSON file.
+		/// </summary>
+		/// <remarks>
+		/// * <c>TargetType.Editor</c> will be generated for current platform only and will ignore <c>UnrealTargetConfiguration.Test</c> and <c>UnrealTargetConfiguration.Shipping</c> configurations
+		/// * <c>TargetType.Program</c>  will be generated for current platform only and <c>UnrealTargetConfiguration.Development</c> configuration only 
+		/// </remarks>
+		/// <param name="InPlatforms"></param>
+		/// <param name="InConfigurations"></param>
+		/// <param name="PlatformProjectGenerators"></param>
+		/// <param name="Minimize"></param>
+		/// <param name="Logger"></param>
+		/// <returns></returns>
+		public bool WriteProjectFile(List<UnrealTargetPlatform> InPlatforms,
+			List<UnrealTargetConfiguration> InConfigurations,
+			PlatformProjectGeneratorCollection PlatformProjectGenerators, JsonWriterStyle Minimize, ILogger Logger)
+		{
+			DirectoryReference ProjectRootFolder = RootPath;
+
+			Dictionary<FileReference, (UEBuildTarget BuildTarget, bool bBuildByDefault)> FileToTarget = new();
+
+			foreach (UnrealTargetPlatform Platform in InPlatforms)
+			{
+				foreach (UnrealTargetConfiguration Configuration in InConfigurations)
+				{
+					foreach (ProjectTarget ProjectTarget in ProjectTargets.OfType<ProjectTarget>())
+					{
+						if (TargetTypes.Any() && !TargetTypes.Contains(ProjectTarget.TargetRules!.Type))
+						{
+							continue;
+						}
+
+						// Skip Programs for all configs except for current platform + Development & Debug configurations
+						if (ProjectTarget.TargetRules!.Type == TargetType.Program &&
+							(BuildHostPlatform.Current.Platform != Platform ||
+							 !(Configuration == UnrealTargetConfiguration.Development || Configuration == UnrealTargetConfiguration.Debug)))
+						{
+							continue;
+						}
+
+						// Skip Editor for all platforms except for current platform
+						if (ProjectTarget.TargetRules.Type == TargetType.Editor && 
+							(BuildHostPlatform.Current.Platform != Platform || 
+							(Configuration == UnrealTargetConfiguration.Test || Configuration == UnrealTargetConfiguration.Shipping)))
+						{
+							continue;
+						}
+
+						bool bBuildByDefault = ShouldBuildByDefaultForSolutionTargets && ProjectTarget.SupportedPlatforms.Contains(Platform);
+
+						UnrealArchitectures ProjectArchitectures = UEBuildPlatform
+							.GetBuildPlatform(Platform)
+							.ArchitectureConfig.ActiveArchitectures(ProjectTarget.UnrealProjectFilePath, ProjectTarget.Name);
+
+						TargetDescriptor TargetDesc = new(ProjectTarget.UnrealProjectFilePath, ProjectTarget.Name,
+							Platform, Configuration, ProjectArchitectures, Arguments);
+
+						try
+						{
+							FileReference OutputFile = FileReference.Combine(ProjectRootFolder, 
+								$"{ProjectTarget.TargetFilePath.GetFileNameWithoutAnyExtensions()}_{Configuration}_{Platform}.json");
+
+							UEBuildTarget BuildTarget = UEBuildTarget.Create(TargetDesc, false, false, false, Logger);
+							FileToTarget.Add(OutputFile, (BuildTarget, bBuildByDefault));
+						}
+						catch (Exception Ex)
+						{
+							Logger.LogWarning("Exception while generating include data for Target:{Target}, Platform: {Platform}, Configuration: {Configuration}", TargetDesc.Name, Platform.ToString(), Configuration.ToString());
+							Logger.LogWarning("{Ex}", Ex.ToString());
+						}
+					}
+				}
+			}
+
+			foreach (var Entry in FileToTarget)
+			{
+				var OutputFile = Entry.Key;
+				var BuildTarget = Entry.Value.BuildTarget;
+				var bBuildByDefault = Entry.Value.bBuildByDefault;
+
+				try
+				{
+					BuildTarget.PreBuildSetup(Logger);
+
+					ExportedTargetInfo TargetInfo = ExportTarget(BuildTarget, bBuildByDefault, PlatformProjectGenerators, Logger);
+
+					DirectoryReference.CreateDirectory(OutputFile.Directory);
+					using FileStream Stream = new(OutputFile.FullName, FileMode.Create, FileAccess.Write);
+					JsonSerializer.Serialize(Stream, TargetInfo, options: new JsonSerializerOptions()
+					{
+						PropertyNamingPolicy = JsonNamingPolicy.CamelCase,
+						WriteIndented = Minimize == JsonWriterStyle.Readable,
+					});
+
+					ExportedTargetProjects.Add(TargetInfo);
+				}
+				catch (Exception Ex)
+				{
+					Logger.LogWarning("Exception while generating include data for Target:{Target}, Platform: {Platform}, Configuration: {Configuration}",
+						BuildTarget.AppName, BuildTarget.Platform.ToString(), BuildTarget.Configuration.ToString());
+					Logger.LogWarning("{Ex}", Ex.ToString());
+				}
+			}
+
+			return true;
+		}
+
+		/// <summary>
+		/// Write a Target to a JSON writer. Is array is empty, don't write anything
+		/// </summary>
+		/// <param name="Target"></param>
+		/// <param name="bBuildByDefault"></param>
+		/// <param name="PlatformProjectGenerators"></param>
+		/// <param name="Logger">Logger for output</param>
+		private ExportedTargetInfo ExportTarget(
+			UEBuildTarget Target, bool bBuildByDefault, PlatformProjectGeneratorCollection PlatformProjectGenerators, ILogger Logger)
+		{
+			ExportedTargetInfo TargetInfo = new()
+			{
+				TargetName = Target.TargetName,
+				TargetPath = Target.TargetRulesFile.FullName,
+				ProjectPath = Target.ProjectFile?.FullName ?? String.Empty,
+				TargetType = Target.TargetType.ToString(),
+				Platform = Target.Platform.ToString(),
+				Configuration = Target.Configuration.ToString(),
+				BuildInfo = ExportBuildInfo(Target, PlatformProjectGenerators, bBuildByDefault, Logger)
+			};
+
+			UEToolChain TargetToolChain = Target.CreateToolchain(Target.Platform, Logger);
+			CppCompileEnvironment GlobalCompileEnvironment = Target.CreateCompileEnvironmentForProjectFiles(Logger);
+
+			HashSet<string> ModuleNames = new();
+			foreach (UEBuildBinary Binary in Target.Binaries)
+			{
+				CppCompileEnvironment BinaryCompileEnvironment = Binary.CreateBinaryCompileEnvironment(GlobalCompileEnvironment);
+				IEnumerable<UEBuildModuleCPP> CandidateModules = Binary.Modules.Where(x => x is UEBuildModuleCPP).Cast<UEBuildModuleCPP>();
+
+				foreach (var ModuleCpp in CandidateModules)
+				{
+					if (!ModuleNames.Add(ModuleCpp.Name))
+					{
+						continue;
+					}
+
+					CppCompileEnvironment ModuleCompileEnvironment = ModuleCpp.CreateCompileEnvironmentForIntellisense(Target.Rules, BinaryCompileEnvironment, Logger);
+					TargetInfo.ModuleToCompileSettings.Add(ModuleCpp.Name, ExportModule(ModuleCpp, TargetToolChain, ModuleCompileEnvironment));
+
+					foreach (DirectoryReference ModuleDirectory in ModuleCpp.ModuleDirectories)
+					{
+						TargetInfo.DirToModule.TryAdd(ModuleDirectory.FullName, ModuleCpp.Name);
+					}
+
+					if (ModuleCpp.GeneratedCodeDirectory != null)
+					{
+						TargetInfo.DirToModule.TryAdd(ModuleCpp.GeneratedCodeDirectory.FullName, ModuleCpp.Name);
+					}
+				}
+			}
+
+			return TargetInfo;
+		}
+
+		/// <summary>
+		/// Write a Module to a JSON writer. If array is empty, don't write anything
+		/// </summary>
+		/// <param name="TargetToolChain"></param>
+		/// <param name="ModuleCompileEnvironment"></param>
+		/// <param name="Module"></param>
+		private static ExportedModuleInfo ExportModule(UEBuildModuleCPP Module, UEToolChain TargetToolChain, CppCompileEnvironment ModuleCompileEnvironment)
+		{
+			ExportedModuleInfo Result = new()
+			{
+				Name = Module.Name,
+				Directory = Module.ModuleDirectory.FullName,
+				Rules = Module.RulesFile.FullName,
+				GeneratedCodeDirectory = Module.GeneratedCodeDirectory != null ? Module.GeneratedCodeDirectory.FullName : String.Empty,
+				Standard = ModuleCompileEnvironment.CppStandard.ToString(),
+			};
+
+			Result.IncludePaths.AddRange(Module.PublicIncludePaths.Select(x => x.FullName));
+			Result.IncludePaths.AddRange(Module.PublicSystemIncludePaths.Select(x => x.FullName));
+			Result.IncludePaths.AddRange(Module.InternalIncludePaths.Select(x => x.FullName));
+			Result.IncludePaths.AddRange(Module.LegacyPublicIncludePaths.Select(x => x.FullName));
+			Result.IncludePaths.AddRange(Module.LegacyParentIncludePaths.Select(x => x.FullName));
+			Result.IncludePaths.AddRange(Module.PrivateIncludePaths.Select(x => x.FullName));
+
+			Result.IncludePaths.AddRange(ModuleCompileEnvironment.UserIncludePaths.Select(x => x.FullName));
+
+			Result.IncludePaths.AddRange(Module.PublicSystemLibraryPaths.Select(x => x.FullName));
+			Result.IncludePaths.AddRange(Module.PublicSystemLibraries.Concat(Module.PublicLibraries.Select(x => x.FullName)));
+
+			if (TargetToolChain is VCToolChain TargetVCToolChain)
+			{
+				Result.IncludePaths.AddRange(TargetVCToolChain.GetVCIncludePaths().Select(x => x.ToString()));
+			}
+
+			Result.Defines.AddRange(Module.PublicDefinitions);
+			Result.Defines.AddRange(Module.Rules.PrivateDefinitions);
+			Result.Defines.AddRange(Module.Rules.bTreatAsEngineModule ? Array.Empty<string>() : Module.Rules.Target.ProjectDefinitions);
+			Result.Defines.AddRange(Module.GetEmptyApiMacros());
+
+			var ForcedIncludes = ModuleCompileEnvironment.ForceIncludeFiles.ToList();
+			if (ModuleCompileEnvironment.PrecompiledHeaderAction == PrecompiledHeaderAction.Include)
+			{
+				FileItem IncludeHeader = FileItem.GetItemByFileReference(ModuleCompileEnvironment.PrecompiledHeaderIncludeFilename!);
+				ForcedIncludes.Insert(0, IncludeHeader);
+			}
+
+			Result.ForcedIncludes.AddRange(ForcedIncludes.Select(x => x.FullName));
+
+			Result.CompilerPath = TargetToolChain.GetCppCompilerPath()?.ToString();
+			Result.CompilerArgs.AddRange(TargetToolChain.GetGlobalCommandLineArgs(ModuleCompileEnvironment));
+			Result.CompilerAdditionalArgs.Add("c", TargetToolChain.GetCCommandLineArgs(ModuleCompileEnvironment).ToList());
+			Result.CompilerAdditionalArgs.Add("cpp", TargetToolChain.GetCPPCommandLineArgs(ModuleCompileEnvironment).ToList());
+
+			return Result;
+		}
+
+		private TargetBuildInfo? ExportBuildInfo(UEBuildTarget Target, PlatformProjectGeneratorCollection PlatformProjectGenerators, bool bBuildByDefault, ILogger Logger)
+		{
+			if (!BuildHostPlatform.Current.Platform.IsInGroup(UnrealPlatformGroup.Windows))
+			{
+				Logger.LogWarning("Unsupported platform for Build Information: {Platform}", BuildHostPlatform.Current.Platform.ToString());
+				return null;
+			}
+
+			if (IsStubProject)
+			{
+				return null;
+			}
+
+			ProjectTarget ProjectTarget = ProjectTargets.OfType<ProjectTarget>().Single(It => Target.TargetRulesFile == It.TargetFilePath);
+			UnrealTargetPlatform Platform = Target.Platform;
+			UnrealTargetConfiguration Configuration = Target.Configuration;
+
+			string UProjectPath = "";
+			if (IsForeignProject)
+			{
+				UProjectPath = String.Format("\"{0}\"", ProjectTarget.UnrealProjectFilePath!.FullName);
+			}
+
+			PlatformProjectGenerator? ProjGenerator = PlatformProjectGenerators.GetPlatformProjectGenerator(Platform, true);
+			VCProjectFile.BuildCommandBuilder BuildCommandBuilder = new(
+				new PlatformProjectGenerator.VSSettings(Platform, Configuration, VCProjectFileFormat.Default, null), ProjectTarget, UProjectPath)
+			{
+				ProjectGenerator = ProjGenerator,
+				bIsForeignProject = IsForeignProject
+			};
+
+			string BuildArguments = BuildCommandBuilder.GetBuildArguments();
+
+			return new TargetBuildInfo()
+			{
+				BuildCmd = $"{BuildCommandBuilder.BuildScript.FullName} {BuildArguments}",
+				RebuildCmd = $"{BuildCommandBuilder.RebuildScript.FullName} {BuildArguments}",
+				CleanCmd = $"{BuildCommandBuilder.CleanScript.FullName} {BuildArguments}",
+				PrimaryOutput = Target.Binaries[0].OutputFilePath.FullName,
+				BuildByDefault = bBuildByDefault,
+			};
+		}
+
+		internal class ExportedTargetInfo
+		{
+			public string TargetName { get; set; } = String.Empty;
+			public string TargetPath { get; set; } = String.Empty;
+			public string ProjectPath { get; set; } = String.Empty;
+			public string TargetType { get; set; } = String.Empty;
+			public string Platform { get; set; } = String.Empty;
+			public string Configuration { get; set; } = String.Empty;
+			public TargetBuildInfo? BuildInfo { get; set; }
+			public Dictionary<string, ExportedModuleInfo> ModuleToCompileSettings { get; set; } = new();
+			public Dictionary<string, string> DirToModule { get; set; } = new();
+		}
+
+		internal class ExportedModuleInfo
+		{
+			public string Name { get; set; } = String.Empty;
+			public string Directory { get; set; } = String.Empty;
+			public string Rules { get; set; } = String.Empty;
+			public string GeneratedCodeDirectory { get; set; } = String.Empty;
+			public List<string> IncludePaths { get; set; } = new();
+			public List<string> Defines { get; set; } = new();
+			public string? Standard { get; set; }
+			public List<string> ForcedIncludes { get; set; } = new();
+			public string? CompilerPath { get; set; }
+			public List<string> CompilerArgs { get; set; } = new();
+			public Dictionary<string, List<string>> CompilerAdditionalArgs { get; set; } = new();
+			public string? WindowsSdkVersion { get; set; }
+		}
+
+		internal class TargetBuildInfo
+		{
+			public string BuildCmd { get; set; } = String.Empty;
+			public string RebuildCmd { get; set; } = String.Empty;
+			public string CleanCmd { get; set; } = String.Empty;
+			public string PrimaryOutput { get; set; } = String.Empty;
+			public bool BuildByDefault { get; internal set; }
+		}
+	}
+}
diff --git a/Engine/Source/Programs/UnrealBuildTool/ProjectFiles/VisualStudioWorkspace/VSWorkspaceProjectFileGenerator.cs b/Engine/Source/Programs/UnrealBuildTool/ProjectFiles/VisualStudioWorkspace/VSWorkspaceProjectFileGenerator.cs
new file mode 100644
--- /dev/null
+++ b/Engine/Source/Programs/UnrealBuildTool/ProjectFiles/VisualStudioWorkspace/VSWorkspaceProjectFileGenerator.cs
@@ -0,0 +1,258 @@
+// Copyright Epic Games, Inc. All Rights Reserved.
+
+using System;
+using System.Collections.Generic;
+using System.IO;
+using System.Linq;
+using System.Text.Json;
+using EpicGames.Core;
+using Microsoft.Extensions.Logging;
+using UnrealBuildBase;
+
+namespace UnrealBuildTool
+{
+	class VSWorkspaceProjectFileGenerator : ProjectFileGenerator
+	{
+		public override string ProjectFileExtension => ".json";
+
+		// These properties are used by Visual Studio to determine where to read the project files.
+		// So they must remain constant.
+		private const string VSUnrealWorkspaceFileName = ".vs-unreal-workspace";
+		private const string ProjectFilesFolder = "VisualStudio";
+
+		private readonly CommandLineArguments Arguments;
+
+		/// <summary>
+		/// List of deprecated platforms.
+		/// Don't generate project model for these platforms unless they are specified in "Platforms" console arguments. 
+		/// </summary>
+		/// <returns></returns>
+		private readonly HashSet<UnrealTargetPlatform> DeprecatedPlatforms = new();
+
+		/// <summary>
+		/// Platforms to generate project files for
+		/// </summary>
+		[CommandLine("-Platforms=", ListSeparator = '+')]
+		HashSet<UnrealTargetPlatform> Platforms = new();
+
+		/// <summary>
+		/// Target types to generate project files for
+		/// </summary>
+		[CommandLine("-TargetTypes=", ListSeparator = '+')]
+		HashSet<TargetType> TargetTypes = new();
+
+		/// <summary>
+		/// Target configurations to generate project files for
+		/// </summary>
+		[CommandLine("-TargetConfigurations=", ListSeparator = '+')]
+		HashSet<UnrealTargetConfiguration> TargetConfigurations = new();
+
+		/// <summary>
+		/// Projects to generate project files for
+		/// </summary>
+		[CommandLine("-ProjectNames=", ListSeparator = '+')]
+		HashSet<string> ProjectNames = new();
+
+		/// <summary>
+		/// Should format JSON files in human readable form, or use packed one without indents
+		/// </summary>
+		[CommandLine("-Minimize", Value = "Compact")]
+		private JsonWriterStyle Minimize = JsonWriterStyle.Readable;
+
+		public VSWorkspaceProjectFileGenerator(FileReference? InOnlyGameProject,
+			CommandLineArguments InArguments)
+			: base(InOnlyGameProject)
+		{
+			Arguments = InArguments;
+			Arguments.ApplyTo(this);
+		}
+
+		public override bool ShouldGenerateIntelliSenseData() => true;
+
+		public override void CleanProjectFiles(DirectoryReference InPrimaryProjectDirectory, string InPrimaryProjectName,
+			DirectoryReference InIntermediateProjectFilesDirectory, ILogger Logger)
+		{
+			DirectoryReference.Delete(InPrimaryProjectDirectory);
+		}
+
+		protected override void ConfigureProjectFileGeneration(string[] Arguments, ref bool IncludeAllPlatforms, ILogger Logger)
+		{
+			base.ConfigureProjectFileGeneration(Arguments, ref IncludeAllPlatforms, Logger);
+		}
+
+		protected override ProjectFile AllocateProjectFile(FileReference InitFilePath, DirectoryReference BaseDir)
+		{
+			VSWorkspaceProjectFile projectFile = new(InitFilePath, BaseDir, RootPath: InitFilePath.Directory,
+				Arguments: Arguments, TargetTypes: TargetTypes);
+			return projectFile;
+		}
+
+		protected override bool WriteProjectFiles(PlatformProjectGeneratorCollection PlatformProjectGenerators, ILogger Logger)
+		{
+			using ProgressWriter Progress = new("Writing project files...", true, Logger);
+			List<ProjectFile> ProjectsToGenerate = new(GeneratedProjectFiles);
+			if (ProjectNames.Any())
+			{
+				ProjectsToGenerate = ProjectsToGenerate.Where(Project =>
+					ProjectNames.Contains(Project.ProjectFilePath.GetFileNameWithoutAnyExtensions())).ToList();
+			}
+
+			int TotalProjectFileCount = ProjectsToGenerate.Count;
+
+			HashSet<UnrealTargetPlatform> PlatformsToGenerate = new(SupportedPlatforms);
+			if (Platforms.Any())
+			{
+				PlatformsToGenerate.IntersectWith(Platforms);
+			}
+
+			List<UnrealTargetPlatform> FilteredPlatforms = PlatformsToGenerate.Where(Platform =>
+			{
+				// Skip deprecated unless explicitly specified in the command line.
+				return (!DeprecatedPlatforms.Contains(Platform) || Platforms.Contains(Platform))
+					   && UEBuildPlatform.IsPlatformAvailable(Platform);
+			}).ToList();
+
+			HashSet<UnrealTargetConfiguration> ConfigurationsToGenerate = new(SupportedConfigurations);
+			if (TargetConfigurations.Any())
+			{
+				ConfigurationsToGenerate.IntersectWith(TargetConfigurations);
+			}
+
+			List<UnrealTargetConfiguration> Configurations = ConfigurationsToGenerate.ToList();
+
+			for (int ProjectFileIndex = 0; ProjectFileIndex < ProjectsToGenerate.Count; ++ProjectFileIndex)
+			{
+				if (ProjectsToGenerate[ProjectFileIndex] is not VSWorkspaceProjectFile CurrentProject)
+				{
+					return false;
+				}
+
+				if (!CurrentProject.WriteProjectFile(FilteredPlatforms, Configurations, PlatformProjectGenerators, Minimize, Logger))
+				{
+					return false;
+				}
+
+				Progress.Write(ProjectFileIndex + 1, TotalProjectFileCount);
+			}
+
+			Progress.Write(TotalProjectFileCount, TotalProjectFileCount);
+
+			return true;
+		}
+
+		public override bool GenerateProjectFiles(PlatformProjectGeneratorCollection PlatformProjectGenerators,
+			String[] Arguments, bool bCacheDataForEditor, ILogger Logger)
+		{
+			bool IncludeAllPlatforms = true;
+			ConfigureProjectFileGeneration(Arguments, ref IncludeAllPlatforms, Logger);
+
+			if (bGeneratingGameProjectFiles || Unreal.IsEngineInstalled())
+			{
+				PrimaryProjectPath = OnlyGameProject!.Directory;
+				PrimaryProjectName = OnlyGameProject.GetFileNameWithoutExtension();
+
+				IntermediateProjectFilesPath =
+					DirectoryReference.Combine(PrimaryProjectPath, "Intermediate", "ProjectFiles");
+			}
+
+			SetupSupportedPlatformsAndConfigurations(IncludeAllPlatforms: true, Logger, out string SupportedPlatformNames);
+			Logger.LogDebug("Supported platforms: {Platforms}", SupportedPlatformNames);
+
+			List<FileReference> AllGames = FindGameProjects(Logger);
+
+			{
+				// Find all of the target files.
+				List<FileReference> AllTargetFiles = DiscoverTargets(
+					AllGames,
+					Logger,
+					OnlyGameProject,
+					SupportedPlatforms,
+					bIncludeEngineSource: bIncludeEngineSource,
+					bIncludeTempTargets: bIncludeTempTargets);
+
+				// If there are multiple targets of a given type for a project, use the order to determine which one gets a suffix.
+				AllTargetFiles = AllTargetFiles.OrderBy(x => x.FullName, StringComparer.OrdinalIgnoreCase).ToList();
+
+				List<ProjectFile> EngineProjects = new();
+				List<ProjectFile> GameProjects = new();
+				List<ProjectFile> ModProjects = new();
+				Dictionary<FileReference, ProjectFile> ProgramProjects = new();
+				Dictionary<RulesAssembly, DirectoryReference> RulesAssemblies = new();
+				Dictionary<ProjectFile, FileReference> ProjectFileToUProjectFile = new();
+
+				AddProjectsForAllTargets(
+					PlatformProjectGenerators,
+					AllGames,
+					AllTargetFiles,
+					Arguments,
+					EngineProjects,
+					GameProjects,
+					ProjectFileToUProjectFile,
+					ProgramProjects,
+					RulesAssemblies,
+					Logger);
+
+				AddAllGameProjects(GameProjects);
+			}
+
+			WriteProjectFiles(PlatformProjectGenerators, Logger);
+			WritePrimaryProjectFile(UBTProject, PlatformProjectGenerators, Logger);
+
+			return true;
+		}
+
+		protected override bool WritePrimaryProjectFile(ProjectFile? UBTProject,
+			PlatformProjectGeneratorCollection PlatformProjectGenerators,
+			ILogger Logger)
+		{
+			try
+			{
+				FileReference PrimaryProjectFile = FileReference.Combine(
+					IntermediateProjectFilesPath, ProjectFilesFolder, VSUnrealWorkspaceFileName);
+
+				DirectoryReference.CreateDirectory(PrimaryProjectFile.Directory);
+
+				// Collect all the resulting project files and aggregate the target-level data
+				var AggregatedProjectInfo = GeneratedProjectFiles
+					.Where(Project => Project is VSWorkspaceProjectFile)
+					.OfType<VSWorkspaceProjectFile>()
+					.SelectMany(Project => Project.ExportedTargetProjects)
+					.GroupBy(TargetProject => TargetProject.TargetName)
+					.Select(g => (g.Key, Target: new
+					{
+						TargetType = g.Select(i => i.TargetType).Distinct().Single(),
+						TargetPath = g.Select(i => i.TargetPath).Distinct().Single(),
+						ProjectPath = g.Select(i => i.ProjectPath).Distinct().Single(),
+						Configurations = g.Select(i => i.Configuration).Distinct().ToList(),
+						Platforms = g.Select(i => i.Platform).Distinct().ToList(),
+					}));
+
+				// The inner Targets object is needed for schema compatibility with the Query Mode API.
+				var Result = new
+				{
+					Targets = AggregatedProjectInfo.ToDictionary(item => item.Key, item => item.Target)
+				};
+
+				using FileStream Stream = new(PrimaryProjectFile.FullName, FileMode.Create, FileAccess.Write);
+				JsonSerializer.Serialize(Stream, Result, new JsonSerializerOptions
+				{
+					PropertyNamingPolicy = JsonNamingPolicy.CamelCase,
+					WriteIndented = Minimize == JsonWriterStyle.Readable,
+				});
+			}
+			catch (Exception Ex)
+			{
+				Logger.LogWarning("Exception while writing root project file: {Ex}", Ex.ToString());
+				return false;
+			}
+
+			return true;
+		}
+
+		/// <inheritdoc />
+		protected override FileReference GetProjectLocation(string BaseName)
+		{
+			return FileReference.Combine(IntermediateProjectFilesPath, ProjectFilesFolder, BaseName + ProjectFileExtension);
+		}
+	}
+}
