PR #10515: Add architecture info to Rider project files

diff --git a/Engine/Source/Programs/UnrealBuildTool/ProjectFiles/Rider/RiderProjectFile.cs b/Engine/Source/Programs/UnrealBuildTool/ProjectFiles/Rider/RiderProjectFile.cs
--- a/Engine/Source/Programs/UnrealBuildTool/ProjectFiles/Rider/RiderProjectFile.cs
+++ b/Engine/Source/Programs/UnrealBuildTool/ProjectFiles/Rider/RiderProjectFile.cs
@@ -798,10 +798,13 @@ namespace UnrealBuildTool
 				bEnableCoroutines = CompileEnvironment.bEnableCoroutines
 			};
 
-			if (CurrentTarget!.Platform.IsInGroup(UnrealPlatformGroup.Windows))
+			if (CompileEnvironment.Architectures.Architectures.Count >= 1)
 			{
-				ToolchainInfo.Architecture = CurrentTarget.Rules.WindowsPlatform.Architecture.WindowsToolChain;
+				ToolchainInfo.Architecture = CompileEnvironment.Architectures.Architectures[0].ToString();
+			}
 
+			if (CurrentTarget!.Platform.IsInGroup(UnrealPlatformGroup.Windows))
+			{
 				WindowsCompiler WindowsPlatformCompiler = CurrentTarget.Rules.WindowsPlatform.Compiler;
 				ToolchainInfo.bStrictConformanceMode = WindowsPlatformCompiler.IsMSVC() && CurrentTarget.Rules.WindowsPlatform.bStrictConformanceMode;
 				ToolchainInfo.Compiler = WindowsPlatformCompiler.ToString();
