PR #10594: Add list of subrules files to generated Rider project files.

diff --git a/Engine/Source/Programs/UnrealBuildTool/ProjectFiles/Rider/RiderProjectFile.cs b/Engine/Source/Programs/UnrealBuildTool/ProjectFiles/Rider/RiderProjectFile.cs
--- a/Engine/Source/Programs/UnrealBuildTool/ProjectFiles/Rider/RiderProjectFile.cs
+++ b/Engine/Source/Programs/UnrealBuildTool/ProjectFiles/Rider/RiderProjectFile.cs
@@ -292,6 +292,7 @@ namespace UnrealBuildTool
 			Writer.WriteValue("Name", Module.Name);
 			Writer.WriteValue("Directory", Module.ModuleDirectory.FullName);
 			Writer.WriteValue("Rules", Module.RulesFile.FullName);
+			ExportJsonStringArray(Writer, "SubRules", Module.Rules.SubclassRules);
 			Writer.WriteValue("PCHUsage", Module.Rules.PCHUsage.ToString());
 
 			if (Module.Rules.PrivatePCHHeaderFile != null)
@@ -408,7 +409,7 @@ namespace UnrealBuildTool
 		/// <param name="Writer">Writer for the array data</param>
 		/// <param name="ArrayName">Name of the array property</param>
 		/// <param name="Strings">Sequence of strings to write. May be null.</param>
-		private static void ExportJsonStringArray(JsonWriter Writer, string ArrayName, IEnumerable<string> Strings)
+		private static void ExportJsonStringArray(JsonWriter Writer, string ArrayName, IEnumerable<string>? Strings)
 		{
 			if (Strings == null || !Strings.Any())
 			{
