PR #10946: [Rider] Fix IntelliSense by adding include order deprecation defines

diff --git a/Engine/Source/Programs/UnrealBuildTool/ProjectFiles/Rider/RiderProjectFile.cs b/Engine/Source/Programs/UnrealBuildTool/ProjectFiles/Rider/RiderProjectFile.cs
--- a/Engine/Source/Programs/UnrealBuildTool/ProjectFiles/Rider/RiderProjectFile.cs
+++ b/Engine/Source/Programs/UnrealBuildTool/ProjectFiles/Rider/RiderProjectFile.cs
@@ -325,7 +325,7 @@ namespace UnrealBuildTool
 			ExportJsonStringArray(Writer, "PublicWeakFrameworks", Module.PublicWeakFrameworks);
 			ExportJsonStringArray(Writer, "PublicDelayLoadDLLs", Module.PublicDelayLoadDLLs);
 			ExportJsonStringArray(Writer, "PublicDefinitions", Module.PublicDefinitions);
-			ExportJsonStringArray(Writer, "PrivateDefinitions", Module.Rules.PrivateDefinitions);
+			ExportJsonStringArray(Writer, "PrivateDefinitions", Module.Rules.PrivateDefinitions.Concat(EngineIncludeOrderHelper.GetDeprecationDefines(Module.Rules.IncludeOrderVersion)));
 			ExportJsonStringArray(Writer, "ProjectDefinitions", /* TODO: Add method ShouldAddProjectDefinitions */ !Module.Rules.bTreatAsEngineModule ? Module.Rules.Target.ProjectDefinitions : new string[0]);
 			ExportJsonStringArray(Writer, "ApiDefinitions", Module.GetEmptyApiMacros());
 			Writer.WriteValue("ShouldAddLegacyPublicIncludePaths", Module.Rules.bLegacyPublicIncludePaths);
