PR #11280: Fix missing dash symbol in ShortArchNames

#android

diff --git a/Engine/Source/Programs/UnrealBuildTool/Platform/Android/AndroidToolChain.cs b/Engine/Source/Programs/UnrealBuildTool/Platform/Android/AndroidToolChain.cs
--- a/Engine/Source/Programs/UnrealBuildTool/Platform/Android/AndroidToolChain.cs
+++ b/Engine/Source/Programs/UnrealBuildTool/Platform/Android/AndroidToolChain.cs
@@ -24,12 +24,12 @@ namespace UnrealBuildTool
 			{ UnrealArch.X64,   "-x64" },
 		};
 
-		// sh0rt names for the above suffixes
+		// short names for the above suffixes
 		public static readonly Dictionary<string, string> ShortArchNames = new Dictionary<string, string>()
 		{
 			{ "", "" },
-			{ "arm64", "a8" },
-			{ "x64", "x6" },
+			{ "-arm64", "a8" },
+			{ "-x64", "x6" },
 		};
 
 		public enum ClangSanitizer
