PR #11227: [Rider] Fixed missing system include paths on macOS

diff --git a/Engine/Source/Programs/UnrealBuildTool/ProjectFiles/Rider/RiderProjectFile.cs b/Engine/Source/Programs/UnrealBuildTool/ProjectFiles/Rider/RiderProjectFile.cs
--- a/Engine/Source/Programs/UnrealBuildTool/ProjectFiles/Rider/RiderProjectFile.cs
+++ b/Engine/Source/Programs/UnrealBuildTool/ProjectFiles/Rider/RiderProjectFile.cs
@@ -852,8 +852,7 @@ namespace UnrealBuildTool
 
 		private class XcrunRunner
 		{
-			private readonly Dictionary<string, IList<string>> CachedIncludePaths =
-				new Dictionary<string, IList<string>>();
+			private readonly Dictionary<string, IList<string>> CachedIncludePaths = new Dictionary<string, IList<string>>();
 
 			private string CurrentlyProcessedSDK = String.Empty;
 			private Process? XcrunProcess;
@@ -869,13 +868,13 @@ namespace UnrealBuildTool
 				string SDKPath = GetSDKPath(Architecture, Platform, Logger);
 				if (!CachedIncludePaths.ContainsKey(SDKPath))
 				{
-					CalculateSystemIncludePaths(SDKPath);
+					CalculateSystemIncludePaths(SDKPath, Logger);
 				}
 
 				return CachedIncludePaths[SDKPath];
 			}
 
-			private void CalculateSystemIncludePaths(string SDKPath)
+			private void CalculateSystemIncludePaths(string SDKPath, ILogger Logger)
 			{
 				if (!String.IsNullOrEmpty(CurrentlyProcessedSDK))
 				{
@@ -887,12 +886,13 @@ namespace UnrealBuildTool
 				using (XcrunProcess = new Process())
 				{
 					string AppName = "xcrun";
-					string Arguments = "clang++ -Wp,-v -x c++ - -fsyntax-only" +
-									   (String.IsNullOrEmpty(SDKPath) ? String.Empty : (" -isysroot " + SDKPath));
+					string SystemRootArgument = String.IsNullOrEmpty(SDKPath) ? String.Empty : (" -isysroot " + SDKPath);
+					string Arguments = "clang++ -Wp,-v -x c++ - -fsyntax-only" + SystemRootArgument;
 					XcrunProcess.StartInfo.FileName = AppName;
 					XcrunProcess.StartInfo.Arguments = Arguments;
 					XcrunProcess.StartInfo.UseShellExecute = false;
 					XcrunProcess.StartInfo.CreateNoWindow = true;
+
 					// For some weird reason output of this command is written to error channel so we're redirecting both channels
 					XcrunProcess.StartInfo.RedirectStandardOutput = true;
 					XcrunProcess.StartInfo.RedirectStandardError = true;
@@ -901,8 +901,15 @@ namespace UnrealBuildTool
 					XcrunProcess.Start();
 					XcrunProcess.BeginOutputReadLine();
 					XcrunProcess.BeginErrorReadLine();
-					// xcrun is not finished on it's own. It should be killed by OnOutputDataReceived when reading is finished. But we'll add timeout as a safeguard
-					XcrunProcess.WaitForExit(3000);
+
+					// xcrun is not finished on it's own. It should be killed by OnOutputDataReceived when reading is finished. But we'll add timeout as a safeguard.
+					// While usually it is fast, first launch on macOS might take ~10 seconds so timeout is quite big: [REDACTED_PROFILE_URL]/llvm-project/issues/75179
+					bool HasExited = XcrunProcess.WaitForExit(30_000);
+					if (!HasExited)
+					{
+						Logger.LogWarning("xcrun didn't finish in 30 second. List of system include paths will not be complete");
+						XcrunProcess.Kill();
+					}
 				}
 
 				XcrunProcess = null;
