PR #11261: [Rider] Fixed missing toolchain include paths on Linux

diff --git a/Engine/Source/Programs/UnrealBuildTool/ProjectFiles/Rider/RiderProjectFile.cs b/Engine/Source/Programs/UnrealBuildTool/ProjectFiles/Rider/RiderProjectFile.cs
--- a/Engine/Source/Programs/UnrealBuildTool/ProjectFiles/Rider/RiderProjectFile.cs
+++ b/Engine/Source/Programs/UnrealBuildTool/ProjectFiles/Rider/RiderProjectFile.cs
@@ -565,30 +565,23 @@ namespace UnrealBuildTool
 					}
 				}
 
-				UEBuildPlatform BuildPlatform;
-
-				if (TargetArchitecture == UnrealArch.X64)
-				{
-					BuildPlatform = UEBuildPlatform.GetBuildPlatform(UnrealTargetPlatform.Linux);
-				}
-				else if (TargetArchitecture == UnrealArch.Arm64)
-				{
-					BuildPlatform = UEBuildPlatform.GetBuildPlatform(UnrealTargetPlatform.LinuxArm64);
-				}
-				else
-				{
-					throw new ArgumentException($"Wrong Target.Architecture: {TargetArchitecture}");
-				}
-
+				UEBuildPlatform BuildPlatform = UEBuildPlatform.GetBuildPlatform(Target.Platform);
 				string PlatformSdkVersionString = UEBuildPlatformSDK.GetSDKForPlatform(BuildPlatform.GetPlatformName())!.GetInstalledVersion()!;
 				string Version = GetLinuxToolchainVersionFromFullString(PlatformSdkVersionString);
 
-				string? InternalSdkPath = UEBuildPlatform.GetSDK(UnrealTargetPlatform.Linux)!.GetInternalSDKPath();
+				string? InternalSdkPath = UEBuildPlatform.GetSDK(Target.Platform)!.GetInternalSDKPath();
 				if (InternalSdkPath != null)
 				{
 					Writer.WriteValue(Path.Combine(InternalSdkPath, "include"));
 					Writer.WriteValue(Path.Combine(InternalSdkPath, "usr/include"));
-					Writer.WriteValue(Path.Combine(InternalSdkPath, "lib/clang/" + Version + "/include/"));
+
+					string ClangIncludeDirectory = Path.Combine(InternalSdkPath, "lib/clang/" + Version + "/include/");
+					Writer.WriteValue(ClangIncludeDirectory);
+					if (!Directory.Exists(ClangIncludeDirectory))
+					{
+						Logger.LogWarning("Clang include directory doesn't exist on disk. VersionString={VersionString}, ClangDir={ClangDir}",
+							PlatformSdkVersionString, ClangIncludeDirectory);
+					}
 				}
 			}
 
@@ -827,27 +820,39 @@ namespace UnrealBuildTool
 		/// <summary>
 		/// Get clang toolchain version from full version string
 		/// v17_clang-10.0.1-centos7 -> 10.0.1
+		/// v17_clang-16.0.1-centos7 -> 16
 		/// </summary>
-		/// <param name="FullVersion">Full clang toolchain version string. example: "v17_clang-10.0.1-centos7"</param>
-		/// <returns>clang toolchain version. example: 10.0.1</returns>
-		private string GetLinuxToolchainVersionFromFullString(string FullVersion)
+		/// <param name="FullVersion">Full clang toolchain version string. Example: "v17_clang-10.0.1-centos7"</param>
+		/// <returns>Clang toolchain version. Example: 10.0.1 or 16</returns>
+		/// <remarks>Starting with clang 16.x the directory naming changed to include major version only</remarks>
+		private static string GetLinuxToolchainVersionFromFullString(string FullVersion)
 		{
-			string FullVersionPattern = @"^v[0-9]+_.*-([0-9]+\.[0-9]+\.[0-9]+)-.*$";
+			string FullVersionPattern = @"^v[0-9]+_.*-(([0-9]+)\.[0-9]+\.[0-9]+)-.*$";
 			Regex Regex = new Regex(FullVersionPattern);
-			Match m = Regex.Match(FullVersion);
-			if (!m.Success)
+			Match Match = Regex.Match(FullVersion);
+			if (!Match.Success)
+			{
+				throw new ArgumentException("Wrong full version string", FullVersion);
+			}
+
+			Group MajorVersionGroup = Match.Groups[2];
+			CaptureCollection MajorVersionCaptures = MajorVersionGroup.Captures;
+			if (MajorVersionCaptures.Count != 1)
 			{
-				throw new ArgumentException("Wrong full version string: {0}", FullVersion);
+				throw new ArgumentException("Multiple regex captures in major version string", FullVersion);
 			}
 
-			Group g = m.Groups[1]; // first and the last capture group 
-			CaptureCollection c = g.Captures;
-			if (c.Count != 1)
+			if (Int32.TryParse(MajorVersionCaptures[0].Value, out int MajorVersion))
 			{
-				throw new ArgumentException("Multiple regex capture in full version string: {0}", FullVersion);
+				if (MajorVersion >= 16)
+				{
+					return MajorVersionCaptures[0].Value;
+				}
 			}
 
-			return c[0].Value;
+			Group FullNumberVersionGroup = Match.Groups[1];
+			CaptureCollection FullNumberVersionCaptures = FullNumberVersionGroup.Captures;
+			return FullNumberVersionCaptures[0].Value;
 		}
 
 		private class XcrunRunner
