PR #11163: [Rider] Expose additional project properties from UBT to Rider

diff --git a/Engine/Source/Programs/UnrealBuildTool/ProjectFiles/Rider/RiderProjectFile.cs b/Engine/Source/Programs/UnrealBuildTool/ProjectFiles/Rider/RiderProjectFile.cs
--- a/Engine/Source/Programs/UnrealBuildTool/ProjectFiles/Rider/RiderProjectFile.cs
+++ b/Engine/Source/Programs/UnrealBuildTool/ProjectFiles/Rider/RiderProjectFile.cs
@@ -799,8 +799,11 @@ namespace UnrealBuildTool
 
 			if (CurrentTarget!.Platform.IsInGroup(UnrealPlatformGroup.Windows))
 			{
+				ToolchainInfo.bEnableAddressSanitizer = CurrentTarget.Rules.WindowsPlatform.bEnableAddressSanitizer;
 				WindowsCompiler WindowsPlatformCompiler = CurrentTarget.Rules.WindowsPlatform.Compiler;
 				ToolchainInfo.bStrictConformanceMode = WindowsPlatformCompiler.IsMSVC() && CurrentTarget.Rules.WindowsPlatform.bStrictConformanceMode;
+				ToolchainInfo.bStrictPreprocessorConformanceMode =
+					WindowsPlatformCompiler.IsMSVC() && CurrentTarget.Rules.WindowsPlatform.bStrictPreprocessorConformance;
 				ToolchainInfo.Compiler = WindowsPlatformCompiler.ToString();
 			}
 			else
diff --git a/Engine/Source/Programs/UnrealBuildTool/ProjectFiles/Rider/ToolchainInfo.cs b/Engine/Source/Programs/UnrealBuildTool/ProjectFiles/Rider/ToolchainInfo.cs
--- a/Engine/Source/Programs/UnrealBuildTool/ProjectFiles/Rider/ToolchainInfo.cs
+++ b/Engine/Source/Programs/UnrealBuildTool/ProjectFiles/Rider/ToolchainInfo.cs
@@ -28,7 +28,9 @@ namespace UnrealBuildTool
 		public List<string>? ForceIncludeFiles;
 		public string? Compiler;
 		public bool bStrictConformanceMode;
+		public bool bStrictPreprocessorConformanceMode;
 		public bool bEnableCoroutines;
+		public bool bEnableAddressSanitizer;
 
 		public IEnumerable<Tuple<string, object?>> GetDiff(ToolchainInfo Other)
 		{
@@ -85,7 +87,8 @@ namespace UnrealBuildTool
 				bUseDebugCRT == Other.bUseDebugCRT && bUseStaticCRT == Other.bUseStaticCRT &&
 				PrecompiledHeaderAction == Other.PrecompiledHeaderAction && PrecompiledHeaderFile == Other.PrecompiledHeaderFile &&
 				Equals(ForceIncludeFiles, Other.ForceIncludeFiles) && Compiler == Other.Compiler &&
-				bStrictConformanceMode == Other.bStrictConformanceMode && bEnableCoroutines == Other.bEnableCoroutines;
+				bStrictConformanceMode == Other.bStrictConformanceMode && bStrictPreprocessorConformanceMode == Other.bStrictPreprocessorConformanceMode &&
+				bEnableCoroutines == Other.bEnableCoroutines && bEnableAddressSanitizer == Other.bEnableAddressSanitizer;
 		}
 
 		public override bool Equals(object? Obj)
@@ -126,7 +129,9 @@ namespace UnrealBuildTool
 				HashCode = (HashCode * 397) ^ (ForceIncludeFiles != null ? ForceIncludeFiles.GetHashCode() : 0);
 				HashCode = (HashCode * 397) ^ (Compiler != null ? Compiler.GetHashCode() : 0);
 				HashCode = (HashCode * 397) ^ bStrictConformanceMode.GetHashCode();
+				HashCode = (HashCode * 397) ^ bStrictPreprocessorConformanceMode.GetHashCode();
 				HashCode = (HashCode * 397) ^ bEnableCoroutines.GetHashCode();
+				HashCode = (HashCode * 397) ^ bEnableAddressSanitizer.GetHashCode();
 				return HashCode;
 			}
 		}
