Report obsolete attributes from referenced modules.

diff --git a/Engine/Source/Programs/UnrealBuildTool/Configuration/UEBuildTarget.cs b/Engine/Source/Programs/UnrealBuildTool/Configuration/UEBuildTarget.cs
--- a/Engine/Source/Programs/UnrealBuildTool/Configuration/UEBuildTarget.cs
+++ b/Engine/Source/Programs/UnrealBuildTool/Configuration/UEBuildTarget.cs
@@ -1509,6 +1509,14 @@ namespace UnrealBuildTool
 		{
 			bool anyErrors = module.ValidateModule("Target", logger);
 
+			// Report any [Obsolete("...")] messages from referenced modules.
+			IEnumerable<ObsoleteAttribute> ObsoleteAttributes = module.Rules.GetType().GetCustomAttributes<ObsoleteAttribute>();
+			if (ObsoleteAttributes.Any())
+			{
+				string Message = string.Join(", ", ObsoleteAttributes.Select(x => x.Message ?? "<unknown reason>"));
+				logger.LogWarning($"Warning: Referenced Module '{module.Name}' is obsolete: '{Message}'");
+			}
+
 			Lazy<HashSet<UEBuildModule>> referencedModules = new(() => module.GetDependencies(bWithIncludePathModules: true, bWithDynamicallyLoadedModules: true));
 
 			// Check there aren't any engine binaries with dependencies on game modules. This can happen when game-specific plugins override engine plugins.
