UnrealBuildTool: bTreatAsEngine module creates a pch variant, clean up some unneeded logic

#rnx

diff --git a/Engine/Source/Programs/UnrealBuildTool/Configuration/UEBuildModule.cs b/Engine/Source/Programs/UnrealBuildTool/Configuration/UEBuildModule.cs
--- a/Engine/Source/Programs/UnrealBuildTool/Configuration/UEBuildModule.cs
+++ b/Engine/Source/Programs/UnrealBuildTool/Configuration/UEBuildModule.cs
@@ -702,12 +702,6 @@ namespace UnrealBuildTool
 			bool bWithLegacyParentIncludePaths
 			)
 		{
-			if (!Rules.bTreatAsEngineModule)
-			{
-				Definitions.Add("DEPRECATED_FORGAME=DEPRECATED");
-				Definitions.Add("UE_DEPRECATED_FORGAME=UE_DEPRECATED");
-			}
-
 			// Add this module's private include paths and definitions.
 			IncludePaths.UnionWith(PrivateIncludePaths);
 
diff --git a/Engine/Source/Programs/UnrealBuildTool/Configuration/UEBuildModuleCPP.cs b/Engine/Source/Programs/UnrealBuildTool/Configuration/UEBuildModuleCPP.cs
--- a/Engine/Source/Programs/UnrealBuildTool/Configuration/UEBuildModuleCPP.cs
+++ b/Engine/Source/Programs/UnrealBuildTool/Configuration/UEBuildModuleCPP.cs
@@ -1051,12 +1051,12 @@ namespace UnrealBuildTool
 			HashSet<string> ImmutableDefinitions = new();
 			foreach (string Definition in Definitions)
 			{
-				if (Definition.Contains("UE_IS_ENGINE_MODULE") ||
-					Definition.Contains("UE_VALIDATE_FORMAT_STRINGS") ||
-					Definition.Contains("UE_VALIDATE_INTERNAL_API") ||
-					Definition.Contains("DEPRECATED_FORGAME") ||
-					Definition.Contains("UE_DEPRECATED_FORGAME") ||
-					Definition.Contains("UE_ENABLE_INCLUDE_ORDER_DEPRECATED_IN_"))
+				if (Definition.Contains("UE_IS_ENGINE_MODULE", StringComparison.Ordinal) ||
+					Definition.Contains("UE_VALIDATE_FORMAT_STRINGS", StringComparison.Ordinal) ||
+					Definition.Contains("UE_VALIDATE_INTERNAL_API", StringComparison.Ordinal) ||
+					Definition.Contains("DEPRECATED_FORGAME", StringComparison.Ordinal) ||
+					Definition.Contains("UE_DEPRECATED_FORGAME", StringComparison.Ordinal) ||
+					Definition.Contains("UE_ENABLE_INCLUDE_ORDER_DEPRECATED_IN_", StringComparison.Ordinal))
 				{
 					continue;
 				}
@@ -1128,6 +1128,19 @@ namespace UnrealBuildTool
 					Definitions.AddRange(EngineIncludeOrderHelper.GetDeprecationDefines(ModuleCompileEnvironment.IncludeOrderVersion));
 				}
 
+				// Modify definitions if we need to create a new shared pch for project modules
+				if (ModuleCompileEnvironment.bTreatAsEngineModule != Template.BaseCompileEnvironment.bTreatAsEngineModule)
+				{
+					Definitions = new List<string>(Definitions);
+					Definitions.RemoveAll(x => x.Contains("UE_IS_ENGINE_MODULE", StringComparison.Ordinal));
+					Definitions.RemoveAll(x => x.Contains("DEPRECATED_FORGAME", StringComparison.Ordinal));
+					Definitions.RemoveAll(x => x.Contains("UE_DEPRECATED_FORGAME", StringComparison.Ordinal));
+
+					Definitions.Add($"UE_IS_ENGINE_MODULE={(Rules.bTreatAsEngineModule ? "1" : "0")}");
+					Definitions.Add($"DEPRECATED_FORGAME={(Rules.bTreatAsEngineModule ? String.Empty : "DEPRECATED")}");
+					Definitions.Add($"UE_DEPRECATED_FORGAME={(Rules.bTreatAsEngineModule ? String.Empty : "UE_DEPRECATED")}");
+				}
+
 				// Modify definitions if we need to create a new shared pch for validating format strings
 				if (ModuleCompileEnvironment.bValidateFormatStrings != Template.BaseCompileEnvironment.bValidateFormatStrings)
 				{
@@ -1144,14 +1157,13 @@ namespace UnrealBuildTool
 					Definitions.Add($"UE_VALIDATE_INTERNAL_API={(ModuleCompileEnvironment.bValidateInternalApi ? "1" : "0")}");
 				}
 
-				// Create a suffix to distinguish this shared PCH variant from any others. Currently only optimized and non-optimized shared PCHs are supported.
+				// Create a suffix to distinguish this shared PCH variant from any others.
 				string Variant = GetSuffixForSharedPCH(ModuleCompileEnvironment, Template.BaseCompileEnvironment);
 
 				FileReference SharedDefinitionsLocation = FileReference.Combine(Template.OutputDir, String.Format("SharedDefinitions.{0}{1}.h", Template.Module.Name, Variant));
 				List<string> NewDefinitions = new();
 				StringBuilder Writer = new StringBuilder();
-				Writer.AppendLine("#pragma once");
-				WriteDefinitions(Definitions, Writer);
+				WriteDefinitions($"Shared Definitions for {Template.Module.Name}{Variant}", Array.Empty<string>(), Definitions, Array.Empty<string>(), Writer);
 				FileItem SharedDefinitionsFileItem = Graph.CreateIntermediateTextFile(SharedDefinitionsLocation, Writer.ToString(), AllowAsync: false);
 
 				// Create the wrapper file, which sets all the definitions needed to compile it
@@ -1241,6 +1253,11 @@ namespace UnrealBuildTool
 		/// <returns>True if the two compile enviroments are compatible</returns>
 		internal static bool IsCompatibleForSharedPCH(CppCompileEnvironment ModuleCompileEnvironment, CppCompileEnvironment CompileEnvironment)
 		{
+			if (ModuleCompileEnvironment.bTreatAsEngineModule != CompileEnvironment.bTreatAsEngineModule)
+			{
+				return false;
+			}
+
 			if (ModuleCompileEnvironment.bOptimizeCode != CompileEnvironment.bOptimizeCode)
 			{
 				return false;
@@ -1284,6 +1301,17 @@ namespace UnrealBuildTool
 		private static string GetSuffixForSharedPCH(CppCompileEnvironment CompileEnvironment, CppCompileEnvironment BaseCompileEnvironment)
 		{
 			string Variant = "";
+			if (CompileEnvironment.bTreatAsEngineModule != BaseCompileEnvironment.bTreatAsEngineModule)
+			{
+				if (CompileEnvironment.bTreatAsEngineModule)
+				{
+					Variant += ".Engine";
+				}
+				else
+				{
+					Variant += ".Project";
+				}
+			}
 			if (CompileEnvironment.bOptimizeCode != BaseCompileEnvironment.bOptimizeCode)
 			{
 				if (CompileEnvironment.bOptimizeCode)
@@ -1384,6 +1412,7 @@ namespace UnrealBuildTool
 			CompileEnvironment.bCodeCoverage = ModuleCompileEnvironment.bCodeCoverage;
 			CompileEnvironment.bUseRTTI = ModuleCompileEnvironment.bUseRTTI;
 			CompileEnvironment.bEnableExceptions = ModuleCompileEnvironment.bEnableExceptions;
+			CompileEnvironment.bTreatAsEngineModule = ModuleCompileEnvironment.bTreatAsEngineModule;
 			CompileEnvironment.CppStandardEngine = ModuleCompileEnvironment.CppStandardEngine;
 			CompileEnvironment.CppStandard = ModuleCompileEnvironment.CppStandard;
 			CompileEnvironment.IncludeOrderVersion = ModuleCompileEnvironment.IncludeOrderVersion;
@@ -1541,9 +1570,7 @@ namespace UnrealBuildTool
 						FileString = File.Location.MakeRelativeTo(Unreal.EngineSourceDirectory);
 					}
 					FileString = FileString.Replace('\\', '/');
-					WrapperContents.AppendLine("// Dedicated PCH for {0}", FileString);
-					WrapperContents.AppendLine();
-					WriteDefinitions(CompileEnvironment.Definitions, WrapperContents);
+					WriteDefinitions($"Dedicated PCH for {FileString}", Array.Empty<string>(), CompileEnvironment.Definitions, Array.Empty<string>(), WrapperContents);
 					WrapperContents.AppendLine();
 					using (StreamReader Reader = new StreamReader(File.Location.FullName))
 					{
@@ -1622,31 +1649,9 @@ namespace UnrealBuildTool
 
 						FileItem PrivateDefinitionsFileItem;
 						{
-							StringBuilder Writer = new();
-
-							Writer.AppendLine($"#include \"{GetIncludeString(Instance.DefinitionsFile)}\"");
-
-							// Games may choose to use shared PCHs from the engine, so allow them to change the value of these macros
-							if (!Rules.bTreatAsEngineModule)
-							{
-								Writer.AppendLine("#undef UE_IS_ENGINE_MODULE");
-								Writer.AppendLine("#undef UE_VALIDATE_FORMAT_STRINGS");
-								Writer.AppendLine("#undef DEPRECATED_FORGAME");
-								Writer.AppendLine("#define DEPRECATED_FORGAME DEPRECATED");
-								Writer.AppendLine("#undef UE_DEPRECATED_FORGAME");
-								Writer.AppendLine("#define UE_DEPRECATED_FORGAME UE_DEPRECATED");
-								foreach (string DeprecationDefine in EngineIncludeOrderHelper.GetAllDeprecationDefines())
-								{
-									Writer.AppendLine("#undef " + DeprecationDefine);
-								}
-							}
-							if (Rules.bValidateInternalApi)
-							{
-								Writer.AppendLine("#undef UE_VALIDATE_INTERNAL_API");
-							}
-
 							// Only add new definitions that are not already existing in the shared pch
 							List<string> NewDefinitions = new();
+							List<string> Undefintions = new();
 							bool ModuleApiUndef = false;
 							foreach (string Definition in CompileEnvironment.Definitions)
 							{
@@ -1658,13 +1663,15 @@ namespace UnrealBuildTool
 								NewDefinitions.Add(Definition);
 
 								// Remove the module _API definition for cases where there are circular dependencies between the shared PCH module and modules using it
-								if (!ModuleApiUndef && Definition.StartsWith(ModuleApiDefine))
+								if (!ModuleApiUndef && Definition.StartsWith(ModuleApiDefine, StringComparison.Ordinal))
 								{
 									ModuleApiUndef = true;
-									Writer.AppendLine("#undef {0}", ModuleApiDefine);
+									Undefintions.Add(ModuleApiDefine);
 								}
 							}
-							WriteDefinitions(NewDefinitions, Writer);
+
+							StringBuilder Writer = new();
+							WriteDefinitions($"Shared PCH Definitions for {Name}", new string[] { GetIncludeString(Instance.DefinitionsFile) }, NewDefinitions, Undefintions, Writer);
 							PrivateDefinitionsFileItem = Graph.CreateIntermediateTextFile(PrivateDefinitionsFile, Writer.ToString());
 						}
 
@@ -1705,8 +1712,7 @@ namespace UnrealBuildTool
 				FileReference PrivateDefinitionsFile = FileReference.Combine(IntermediateDirectory, PrivateDefinitionsName);
 				{
 					StringBuilder Writer = new();
-					Writer.AppendLine("#pragma once");
-					WriteDefinitions(CompileEnvironment.Definitions, Writer);
+					WriteDefinitions("Definitions", Array.Empty<string>(), CompileEnvironment.Definitions, Array.Empty<string>(), Writer);
 					CompileEnvironment.Definitions.Clear();
 
 					FileItem PrivateDefinitionsFileItem = Graph.CreateIntermediateTextFile(PrivateDefinitionsFile, Writer.ToString());
@@ -1780,14 +1786,28 @@ namespace UnrealBuildTool
 		/// <summary>
 		/// Write a list of macro definitions to an output file
 		/// </summary>
+		/// <param name="Context">Additional context to add to generated comment</param>
+		/// <param name="Includes">List of includes</param>
 		/// <param name="Definitions">List of definitions</param>
+		/// <param name="Undefinitions">List of definitions to undefine</param>
 		/// <param name="Writer">Writer to receive output</param>
-		static void WriteDefinitions(IEnumerable<string> Definitions, StringBuilder Writer)
+		static void WriteDefinitions(string Context, IEnumerable<string> Includes, IEnumerable<string> Definitions, IEnumerable<string> Undefinitions, StringBuilder Writer)
 		{
-			Writer.AppendLine("// Generated by UnrealBuildTool (UEBuildModuleCPP.cs)");
-			foreach (string Definition in Definitions)
+			Writer.AppendLine($"// Generated by UnrealBuildTool (UEBuildModuleCPP.cs) : {Context}");
+			Writer.AppendLine("#pragma once");
+			foreach (string Include in Includes)
+			{
+				Writer.Append("#include \"").Append(Include).Append('\"').AppendLine();
+			}
+			foreach (string Undefinition in Undefinitions)
+			{
+				Writer.Append("#undef ").Append(Undefinition).AppendLine();
+			}
+
+			HashSet<string> Processed = new();
+			foreach (string Definition in Definitions.Where(x => Processed.Add(x)))
 			{
-				int EqualsIdx = Definition.IndexOf('=');
+				int EqualsIdx = Definition.IndexOf('=', StringComparison.Ordinal);
 				if (EqualsIdx == -1)
 				{
 					Writer.Append("#define ").Append(Definition).AppendLine(" 1");
@@ -2011,6 +2031,7 @@ namespace UnrealBuildTool
 			Result.StaticAnalyzerDisabledCheckers = Rules.StaticAnalyzerDisabledCheckers;
 			Result.StaticAnalyzerAdditionalCheckers = Rules.StaticAnalyzerAdditionalCheckers;
 			Result.bEnableUndefinedIdentifierWarnings = Rules.bEnableUndefinedIdentifierWarnings;
+			Result.bTreatAsEngineModule = Rules.bTreatAsEngineModule;
 			Result.IncludeOrderVersion = Rules.IncludeOrderVersion;
 			Result.DeterministicWarningLevel = Rules.DeterministicWarningLevel;
 			Result.bValidateFormatStrings = Rules.bValidateFormatStrings;
@@ -2097,13 +2118,12 @@ namespace UnrealBuildTool
 			}
 
 			// Add a macro for when we're compiling an engine module, to enable additional compiler diagnostics through code.
-			if (Rules.bTreatAsEngineModule || Target.bWarnAboutMonolithicHeadersIncluded)
-			{
-				Result.Definitions.Add("UE_IS_ENGINE_MODULE=1");
-			}
-			else
+			Result.Definitions.Add($"UE_IS_ENGINE_MODULE={(Rules.bTreatAsEngineModule || Target.bWarnAboutMonolithicHeadersIncluded ? "1" : "0")}");
+
+			if (!Rules.bTreatAsEngineModule)
 			{
-				Result.Definitions.Add("UE_IS_ENGINE_MODULE=0");
+				Result.Definitions.Add("DEPRECATED_FORGAME=DEPRECATED");
+				Result.Definitions.Add("UE_DEPRECATED_FORGAME=UE_DEPRECATED");
 			}
 
 			Result.Definitions.Add($"UE_VALIDATE_FORMAT_STRINGS={(Rules.bValidateFormatStrings ? "1" : "0")}");
@@ -2155,6 +2175,7 @@ namespace UnrealBuildTool
 			// Use the default optimization setting for
 			CompileEnvironment.bOptimizeCode = ShouldEnableOptimization(ModuleRules.CodeOptimization.Default, Target.Configuration, Rules.bTreatAsEngineModule, Rules.bCodeCoverage);
 			CompileEnvironment.bCodeCoverage = Rules.bCodeCoverage;
+			CompileEnvironment.bTreatAsEngineModule = Rules.bTreatAsEngineModule;
 			CompileEnvironment.bValidateFormatStrings = Rules.bValidateFormatStrings;
 			CompileEnvironment.bValidateInternalApi = Rules.bValidateInternalApi;
 
@@ -2163,13 +2184,12 @@ namespace UnrealBuildTool
 			CompileEnvironment.bIsBuildingLibrary = false;
 
 			// Add a macro for when we're compiling an engine module, to enable additional compiler diagnostics through code.
-			if (Rules.bTreatAsEngineModule)
-			{
-				CompileEnvironment.Definitions.Add("UE_IS_ENGINE_MODULE=1");
-			}
-			else
+			CompileEnvironment.Definitions.Add($"UE_IS_ENGINE_MODULE={(Rules.bTreatAsEngineModule ? "1" : "0")}");
+
+			if (!Rules.bTreatAsEngineModule)
 			{
-				CompileEnvironment.Definitions.Add("UE_IS_ENGINE_MODULE=0");
+				CompileEnvironment.Definitions.Add("DEPRECATED_FORGAME=DEPRECATED");
+				CompileEnvironment.Definitions.Add("UE_DEPRECATED_FORGAME=UE_DEPRECATED");
 			}
 
 			CompileEnvironment.Definitions.Add($"UE_VALIDATE_FORMAT_STRINGS={(Rules.bValidateFormatStrings ? "1" : "0")}");
diff --git a/Engine/Source/Programs/UnrealBuildTool/Configuration/UEBuildTarget.cs b/Engine/Source/Programs/UnrealBuildTool/Configuration/UEBuildTarget.cs
--- a/Engine/Source/Programs/UnrealBuildTool/Configuration/UEBuildTarget.cs
+++ b/Engine/Source/Programs/UnrealBuildTool/Configuration/UEBuildTarget.cs
@@ -13,7 +13,6 @@ using System.Text;
 using System.Text.Json.Serialization;
 using System.Threading.Tasks;
 using EpicGames.Core;
-using EpicGames.Serialization;
 using EpicGames.UHT.Utils;
 using Microsoft.Extensions.Logging;
 using OpenTracing.Util;
@@ -3229,7 +3228,7 @@ namespace UnrealBuildTool
 		/// <param name="HeaderFile">Path to the header to output</param>
 		/// <param name="GlobalCompileEnvironment">The global compile environment for this target</param>
 		/// <param name="Logger"></param>
-		void WritePublicHeader(UEBuildBinary Binary, FileReference HeaderFile, CppCompileEnvironment GlobalCompileEnvironment, ILogger Logger)
+		static void WritePublicHeader(UEBuildBinary Binary, FileReference HeaderFile, CppCompileEnvironment GlobalCompileEnvironment, ILogger Logger)
 		{
 			DirectoryReference.CreateDirectory(HeaderFile.Directory);
 
@@ -3243,11 +3242,13 @@ namespace UnrealBuildTool
 			// Write the header
 			using (StreamWriter Writer = new StreamWriter(HeaderFile.FullName))
 			{
+				Writer.WriteLine($"// Definitions for {HeaderFile.GetFileNameWithoutAnyExtensions()}");
+				Writer.WriteLine();
 				Writer.WriteLine("#pragma once");
 				Writer.WriteLine();
 				foreach (string Definition in Definitions)
 				{
-					int EqualsIdx = Definition.IndexOf('=');
+					int EqualsIdx = Definition.IndexOf('=', StringComparison.Ordinal);
 					if (EqualsIdx == -1)
 					{
 						Writer.WriteLine(String.Format("#define {0} 1", Definition));
diff --git a/Engine/Source/Programs/UnrealBuildTool/System/CppCompileEnvironment.cs b/Engine/Source/Programs/UnrealBuildTool/System/CppCompileEnvironment.cs
--- a/Engine/Source/Programs/UnrealBuildTool/System/CppCompileEnvironment.cs
+++ b/Engine/Source/Programs/UnrealBuildTool/System/CppCompileEnvironment.cs
@@ -635,6 +635,11 @@ namespace UnrealBuildTool
 		/// </summary>
 		public bool bHideSymbolsByDefault = true;
 
+		/// <summary>
+		/// Whether this environment should be treated as an engine module.
+		/// </summary>
+		public bool bTreatAsEngineModule;
+
 		/// <summary>
 		/// Which C++ standard to support for engine modules. CppStandard will be set to this for engine modules and CppStandardEngine should not be checked in any toolchain. May not be compatible with all platforms.
 		/// </summary>
@@ -799,6 +804,7 @@ namespace UnrealBuildTool
 			ParentPCHInstance = Other.ParentPCHInstance;
 			bHackHeaderGenerator = Other.bHackHeaderGenerator;
 			bHideSymbolsByDefault = Other.bHideSymbolsByDefault;
+			bTreatAsEngineModule = Other.bTreatAsEngineModule;
 			CppStandardEngine = Other.CppStandardEngine;
 			CppStandard = Other.CppStandard;
 			CStandard = Other.CStandard;
