UnrealBuildTool: Quote response file in compile_commands.json

#rnx

diff --git a/Engine/Source/Programs/UnrealBuildTool/Modes/GenerateClangDatabase.cs b/Engine/Source/Programs/UnrealBuildTool/Modes/GenerateClangDatabase.cs
--- a/Engine/Source/Programs/UnrealBuildTool/Modes/GenerateClangDatabase.cs
+++ b/Engine/Source/Programs/UnrealBuildTool/Modes/GenerateClangDatabase.cs
@@ -372,7 +372,7 @@ namespace UnrealBuildTool
 			}
 
 			FileReference FileReference = new FileReference(OldPath);
-			return Line.Replace(OldPath, FileReference.FullName.Replace("\\", "/"));
+			return Line.Replace(OldPath, Utils.MakePathSafeToUseWithCommandLine(FileReference.FullName.Replace("\\", "/")));
 		}
 	}
 }
