PR #11012: Makefile generator is now platform agnostic

diff --git a/Engine/Source/Programs/UnrealBuildTool/ProjectFiles/Make/MakefileGenerator.cs b/Engine/Source/Programs/UnrealBuildTool/ProjectFiles/Make/MakefileGenerator.cs
--- a/Engine/Source/Programs/UnrealBuildTool/ProjectFiles/Make/MakefileGenerator.cs
+++ b/Engine/Source/Programs/UnrealBuildTool/ProjectFiles/Make/MakefileGenerator.cs
@@ -58,16 +58,26 @@ namespace UnrealBuildTool
 			string UnrealRootPath = Unreal.RootDirectory.FullName;
 
 			string EditorTarget = "UnrealEditor";
+			string CurrentPlatform = "Win64";
+
+			if (OperatingSystem.IsLinux())
+			{
+				CurrentPlatform = "Linux";
+			}
+			else if (OperatingSystem.IsMacOS())
+			{
+				CurrentPlatform = "Mac";
+			}
 
 			if (!String.IsNullOrEmpty(GameProjectName))
 			{
 				GameProjectFile = OnlyGameProject!.FullName;
 				MakeGameProjectFile = "GAMEPROJECTFILE =" + GameProjectFile + "\n";
-				ProjectBuildCommand = $"PROJECTBUILD = \"$(UNREALROOTPATH)/Engine/{Unreal.RelativeDotnetDirectory}/dotnet\" \"$(UNREALROOTPATH)/Engine/Binaries/DotNET/UnrealBuildTool/UnrealBuildTool.dll\"\n";
+				ProjectBuildCommand = $"PROJECTBUILD = \"$(UNREALROOTPATH)/Engine/{Unreal.RelativeDotnetDirectory}/dotnet{BuildHostPlatform.Current.BinarySuffix}\" \"$(UNREALROOTPATH)/Engine/Binaries/DotNET/UnrealBuildTool/UnrealBuildTool.dll\"\n";
 				EditorTarget = GameProjectName + "Editor";
 			}
 
-			BuildCommand = "BUILD = bash \"$(UNREALROOTPATH)/Engine/Build/BatchFiles/Linux/Build.sh\"\n";
+			BuildCommand = $"BUILD = \"$(UNREALROOTPATH)/Engine/{Unreal.RelativeDotnetDirectory}/dotnet{BuildHostPlatform.Current.BinarySuffix}\" \"$(UNREALROOTPATH)/Engine/Binaries/DotNET/UnrealBuildTool/UnrealBuildTool.dll\"\n";
 
 			string FileName = "Makefile"; // PrimaryProjectName + ".mk";
 			StringBuilder MakefileContent = new StringBuilder();
@@ -99,7 +109,7 @@ namespace UnrealBuildTool
 							if (InstalledPlatformInfo.IsValidConfiguration(CurConfiguration, EProjectType.Code))
 							{
 								string Confname = Enum.GetName(typeof(UnrealTargetConfiguration), CurConfiguration)!;
-								MakefileContent.Append(String.Format(" \\\n\t{0}-Linux-{1} ", Basename, Confname));
+								MakefileContent.Append(String.Format(" \\\n\t{0}-{1}-{2} ", Basename, CurrentPlatform, Confname));
 							}
 						}
 					}
@@ -110,9 +120,9 @@ namespace UnrealBuildTool
 
 			MakefileContent.Append("\n\n" + BuildCommand + ProjectBuildCommand + "\n" +
 				"all: StandardSet\n\n" +
-				"RequiredTools: CrashReportClient-Linux-Shipping CrashReportClientEditor-Linux-Shipping ShaderCompileWorker UnrealLightmass EpicWebHelper-Linux-Shipping\n\n" +
+				$"RequiredTools: CrashReportClient-{CurrentPlatform}-Shipping CrashReportClientEditor-{CurrentPlatform}-Shipping ShaderCompileWorker UnrealLightmass EpicWebHelper-{CurrentPlatform}-Shipping\n\n" +
 				$"StandardSet: RequiredTools UnrealFrontend {EditorTarget} UnrealInsights\n\n" +
-				$"DebugSet: RequiredTools UnrealFrontend-Linux-Debug {EditorTarget}-Linux-Debug\n\n"
+				$"DebugSet: RequiredTools UnrealFrontend-{CurrentPlatform}-Debug {EditorTarget}-{CurrentPlatform}-Debug\n\n"
 			);
 
 			foreach (ProjectFile Project in GeneratedProjectFiles)
@@ -154,11 +164,11 @@ namespace UnrealBuildTool
 							if (InstalledPlatformInfo.IsValidConfiguration(CurConfiguration, EProjectType.Code))
 							{
 								string Confname = Enum.GetName(typeof(UnrealTargetConfiguration), CurConfiguration)!;
-								MakefileContent.Append(String.Format("\n{1}-Linux-{2}:\n\t {0} {1} Linux {2} {3} $(ARGS)\n", MakeBuildCommand, Basename, Confname, MakeProjectCmdArg));
+								MakefileContent.Append(String.Format("\n{1}-{4}-{2}:\n\t {0} {1} {4} {2} {3} $(ARGS)\n", MakeBuildCommand, Basename, Confname, MakeProjectCmdArg, CurrentPlatform));
 							}
 						}
 					}
-					MakefileContent.Append(String.Format("\n{1}:\n\t {0} {1} Linux Development {2} $(ARGS)\n", MakeBuildCommand, Basename, MakeProjectCmdArg));
+					MakefileContent.Append(String.Format("\n{1}:\n\t {0} {1} {3} Development {2} $(ARGS)\n", MakeBuildCommand, Basename, MakeProjectCmdArg, CurrentPlatform));
 				}
 			}
 
diff --git a/Engine/Source/Programs/UnrealBuildTool/UnrealBuildTool.cs b/Engine/Source/Programs/UnrealBuildTool/UnrealBuildTool.cs
--- a/Engine/Source/Programs/UnrealBuildTool/UnrealBuildTool.cs
+++ b/Engine/Source/Programs/UnrealBuildTool/UnrealBuildTool.cs
@@ -276,7 +276,7 @@ namespace UnrealBuildTool
 
 			[CommandLine("-ProjectFiles", Value = "GenerateProjectFiles", Description = "Generate project files based on IDE preference. Equivalent to -Mode=GenerateProjectFiles")]
 			[CommandLine("-ProjectFileFormat=", Value = "GenerateProjectFiles", Description = "Generate project files in specified format. May be used multiple times.")]
-			[CommandLine("-Makefile", Value = "GenerateProjectFiles", Description = "Generate Linux Makefile")]
+			[CommandLine("-Makefile", Value = "GenerateProjectFiles", Description = "Generate Makefile")]
 			[CommandLine("-CMakefile", Value = "GenerateProjectFiles", Description = "Generate project files for CMake")]
 			[CommandLine("-QMakefile", Value = "GenerateProjectFiles", Description = "Generate project files for QMake")]
 			[CommandLine("-KDevelopfile", Value = "GenerateProjectFiles", Description = "Generate project files for KDevelop")]
