UnrealBuildTool: Always set msvc extended error info flags

#rnx

diff --git a/Engine/Source/Programs/UnrealBuildTool/Platform/Windows/UEBuildWindows.cs b/Engine/Source/Programs/UnrealBuildTool/Platform/Windows/UEBuildWindows.cs
--- a/Engine/Source/Programs/UnrealBuildTool/Platform/Windows/UEBuildWindows.cs
+++ b/Engine/Source/Programs/UnrealBuildTool/Platform/Windows/UEBuildWindows.cs
@@ -319,6 +319,7 @@ namespace UnrealBuildTool
 		/// <summary>
 		/// True if /fastfail should be passed to the msvc compiler and linker
 		/// </summary>
+		[RequiresUniqueBuildEnvironment]
 		[ConfigFile(ConfigHierarchyType.Engine, "/Script/WindowsTargetPlatform.WindowsTargetSettings", "bVCFastFail")]
 		[XmlConfigFile(Category = "WindowsPlatform")]
 		[CommandLine("-VCFastFail")]
@@ -327,10 +328,12 @@ namespace UnrealBuildTool
 		/// <summary>
 		/// True if /d2ExtendedWarningInfo should be passed to the compiler and /d2:-ExtendedWarningInfo to the linker
 		/// </summary>
+		[RequiresUniqueBuildEnvironment]
 		[ConfigFile(ConfigHierarchyType.Engine, "/Script/WindowsTargetPlatform.WindowsTargetSettings", "bVCExtendedWarningInfo")]
 		[XmlConfigFile(Category = "WindowsPlatform")]
 		[CommandLine("-VCExtendedWarningInfo")]
-		public bool bVCExtendedWarningInfo = false;
+		[CommandLine("-VCDisableExtendedWarningInfo", Value ="false")]
+		public bool bVCExtendedWarningInfo = true;
 
 		/// <summary>
 		/// True if optimizations to reduce the size of debug information should be disabled
