UnrealBuildTool: Downgrade msvc "C4702: unreachable code" to a warning when running LTCG or PGO builds

#rnx

diff --git a/Engine/Source/Programs/UnrealBuildTool/Platform/Windows/VCToolChain.cs b/Engine/Source/Programs/UnrealBuildTool/Platform/Windows/VCToolChain.cs
--- a/Engine/Source/Programs/UnrealBuildTool/Platform/Windows/VCToolChain.cs
+++ b/Engine/Source/Programs/UnrealBuildTool/Platform/Windows/VCToolChain.cs
@@ -1003,6 +1003,12 @@ namespace UnrealBuildTool
 				Arguments.Add("/wd4180"); // qualifier applied to function type has no meaning; ignored
 			}
 
+			// Downgrade C4702: unreachable code to a warning when running LTCG or PGO
+			if (CompileEnvironment.bPGOOptimize || CompileEnvironment.bPGOProfile || CompileEnvironment.bAllowLTCG)
+			{
+				Arguments.Add("/wd4702 /w44702");
+			}
+
 			if (CompileEnvironment.Architecture == UnrealArch.Arm64ec)
 			{
 				Arguments.Add("/arm64EC");
