UnrealBuildTool: msvc pch when running /analyze also creates a .pchast file

#rnx

diff --git a/Engine/Source/Programs/UnrealBuildTool/Platform/Windows/VCToolChain.cs b/Engine/Source/Programs/UnrealBuildTool/Platform/Windows/VCToolChain.cs
--- a/Engine/Source/Programs/UnrealBuildTool/Platform/Windows/VCToolChain.cs
+++ b/Engine/Source/Programs/UnrealBuildTool/Platform/Windows/VCToolChain.cs
@@ -1513,6 +1513,11 @@ namespace UnrealBuildTool
 
 				BaseCompileAction.UsingPchFile = CompileEnvironment.PrecompiledHeaderFile;
 				BaseCompileAction.PchThroughHeaderFile = IncludeHeader;
+
+				if (CompileEnvironment.PrecompiledHeaderFile != null && !CompileEnvironment.bUseHeaderUnitsForPch && Target.StaticAnalyzer == StaticAnalyzer.Default && !CompileEnvironment.bDisableStaticAnalysis)
+				{
+					BaseCompileAction.AdditionalPrerequisiteItems.Add(FileItem.GetItemByFileReference(new FileReference(CompileEnvironment.PrecompiledHeaderFile.FullName + "ast")));
+				}
 			}
 
 			// Generate the timing info
@@ -1610,6 +1615,11 @@ namespace UnrealBuildTool
 					CompileAction.CreatePchFile = FileItem.GetItemByFileReference(FileReference.Combine(OutputDir, SourceFile.Location.GetFileName() + PchExtension));
 					CompileAction.PchThroughHeaderFile = FileItem.GetItemByFileReference(CompileEnvironment.PrecompiledHeaderIncludeFilename);
 
+					if (!CompileEnvironment.bUseHeaderUnitsForPch && Target.StaticAnalyzer == StaticAnalyzer.Default && !CompileEnvironment.bDisableStaticAnalysis)
+					{
+						CompileAction.AdditionalProducedItems.Add(FileItem.GetItemByFileReference(FileReference.Combine(OutputDir, SourceFile.Location.GetFileName() + PchExtension + "ast")));
+					}
+
 					// If we're creating a PCH that will be used to compile source files for a library, we need
 					// the compiled modules to retain a reference to PCH's module, so that debugging information
 					// will be included in the library.  This is also required to avoid linker warning "LNK4206"
