Add support for HWAsan in NDK 26

diff --git a/Engine/Build/Android/ClangSanitizers/wrap.sh b/Engine/Build/Android/ClangSanitizers/wrap.sh
deleted file mode 100644
diff --git a/Engine/Config/Android/Android_SDK.json b/Engine/Config/Android/Android_SDK.json
--- a/Engine/Config/Android/Android_SDK.json
+++ b/Engine/Config/Android/Android_SDK.json
@@ -3,5 +3,5 @@
 "AutoSDKDirectory" : "-26",
 
 "MinVersion" : "r25b",
-"MaxVersion" : "r26",
+"MaxVersion" : "r26b"
 }
diff --git a/Engine/Source/Programs/UnrealBuildTool/Platform/Android/AndroidToolChain.cs b/Engine/Source/Programs/UnrealBuildTool/Platform/Android/AndroidToolChain.cs
--- a/Engine/Source/Programs/UnrealBuildTool/Platform/Android/AndroidToolChain.cs
+++ b/Engine/Source/Programs/UnrealBuildTool/Platform/Android/AndroidToolChain.cs
@@ -115,6 +115,11 @@ namespace UnrealBuildTool
 			return NDKVersionInt >= 220000;
 		}
 
+		public bool HasEmbeddedHWASanSupport()
+		{
+			return NDKVersionInt >= 260000;
+		}
+
 		public AndroidToolChain(FileReference? InProjectFile, ILogger InLogger)
 			: this(InProjectFile, ClangToolChainOptions.None, InLogger)
 		{
@@ -601,7 +606,7 @@ namespace UnrealBuildTool
 
 				if (Sanitizer == ClangSanitizer.Address || Sanitizer == ClangSanitizer.HwAddress)
 				{
-					Arguments.Add("-fno-omit-frame-pointer -DRUNNING_WITH_ASAN=1 -DFORCE_ANSI_ALLOCATOR=1");
+					Arguments.Add("-fno-omit-frame-pointer -DFORCE_ANSI_ALLOCATOR=1");
 				}
 			}
 
diff --git a/Engine/Source/Programs/UnrealBuildTool/Platform/Android/UEDeployAndroid.cs b/Engine/Source/Programs/UnrealBuildTool/Platform/Android/UEDeployAndroid.cs
--- a/Engine/Source/Programs/UnrealBuildTool/Platform/Android/UEDeployAndroid.cs
+++ b/Engine/Source/Programs/UnrealBuildTool/Platform/Android/UEDeployAndroid.cs
@@ -1498,7 +1498,7 @@ namespace UnrealBuildTool
 			}
 		}
 
-		void CopyClangSanitizerLib(string UnrealBuildPath, UnrealArch UnrealArch, string NDKArch, AndroidToolChain.ClangSanitizer Sanitizer)
+		void CopyClangSanitizerLib(AndroidToolChain ToolChain, string UnrealBuildPath, UnrealArch UnrealArch, string NDKArch, AndroidToolChain.ClangSanitizer Sanitizer)
 		{
 			string Architecture = "-aarch64";
 			switch (NDKArch)
@@ -1518,8 +1518,11 @@ namespace UnrealBuildTool
 			switch (Sanitizer)
 			{
 				case AndroidToolChain.ClangSanitizer.HwAddress:
-					LibName = "hwasan";
+				{
+					// no need to bundle asan .so in NDK r26b+
+					LibName = ToolChain.HasEmbeddedHWASanSupport() ? string.Empty : "hwasan";
 					break;
+				}
 				case AndroidToolChain.ClangSanitizer.UndefinedBehavior:
 					LibName = "ubsan_standalone";
 					break;
@@ -1531,9 +1534,12 @@ namespace UnrealBuildTool
 					break;
 			}
 
-			string SanitizerFullLibName = "libclang_rt." + LibName + Architecture + "-android.so";
+			string SanitizerFullLibName = string.IsNullOrEmpty(LibName) ? string.Empty : "libclang_rt." + LibName + Architecture + "-android.so";
+
+			// NDK r26b+ needs different wrap.sh script
+			string WrapShName = Sanitizer == AndroidToolChain.ClangSanitizer.HwAddress && ToolChain.HasEmbeddedHWASanSupport() ? "hwasan.sh" : "asan.sh";
 
-			string WrapSh = Path.Combine(Unreal.EngineDirectory.ToString(), "Build", "Android", "ClangSanitizers", "wrap.sh");
+			string WrapShFilePath = Path.Combine(Environment.ExpandEnvironmentVariables("%NDKROOT%"), "wrap.sh", WrapShName);
 
 			string PlatformHostName = GetPlatformNDKHostName();
 
@@ -1542,25 +1548,43 @@ namespace UnrealBuildTool
 			string LibsVersion = VersionFile.ReadLine()!;
 			VersionFile.Close();
 
-			string SanitizerLib = Path.Combine(Environment.ExpandEnvironmentVariables("%NDKROOT%"), "toolchains", "llvm", "prebuilt", PlatformHostName, "lib64", "clang", LibsVersion, "lib", "linux", SanitizerFullLibName);
-			if (File.Exists(SanitizerLib) && File.Exists(WrapSh))
+			string SanitizerLib = string.IsNullOrEmpty(SanitizerFullLibName) ?
+				string.Empty :
+				Path.Combine(Environment.ExpandEnvironmentVariables("%NDKROOT%"), "toolchains", "llvm", "prebuilt", PlatformHostName, (ToolChain.HasEmbeddedHWASanSupport() ? "lib" : "lib64"), "clang", LibsVersion, "lib", "linux", SanitizerFullLibName);
+
+			if (!string.IsNullOrEmpty(SanitizerLib))
+			{
+				if (File.Exists(SanitizerLib))
+				{
+					string LibDestDir = Path.Combine(UnrealBuildPath, "libs", NDKArch);
+					Directory.CreateDirectory(LibDestDir);
+
+					string LibDestFilePath = Path.Combine(LibDestDir, SanitizerFullLibName);
+					Logger.LogInformation("Copying asan lib from {SanitizerLib} to {LibDestFilePath}", SanitizerLib, LibDestFilePath);
+					File.Copy(SanitizerLib, LibDestFilePath, true);
+				}
+				else
+				{
+					throw new BuildException("No asan lib found in {0}", SanitizerLib);
+				}
+			}
+
+			if (File.Exists(WrapShFilePath))
 			{
-				string LibDestDir = Path.Combine(UnrealBuildPath, "libs", NDKArch);
-				Directory.CreateDirectory(LibDestDir);
-				Logger.LogInformation("Copying asan lib from {SanitizerLib} to {LibDestDir}", SanitizerLib, LibDestDir);
-				File.Copy(SanitizerLib, Path.Combine(LibDestDir, SanitizerFullLibName), true);
 				string WrapDestDir = Path.Combine(UnrealBuildPath, "resources", "lib", NDKArch);
 				Directory.CreateDirectory(WrapDestDir);
+
 				string WrapDestFilePath = Path.Combine(WrapDestDir, "wrap.sh");
-				Logger.LogInformation("Copying wrap.sh from {WrapSh} to {WrapDestFilePath}", WrapSh, WrapDestFilePath);
-				File.Copy(WrapSh, WrapDestFilePath, true);
+				Logger.LogInformation("Copying wrap.sh from {WrapShFilePath} to {WrapDestFilePath}", WrapShFilePath, WrapDestFilePath);
+				File.Copy(WrapShFilePath, WrapDestFilePath, true);
+
 				FileAttributes Attributes = File.GetAttributes(WrapDestFilePath);
 				Attributes &= ~FileAttributes.ReadOnly;
 				File.SetAttributes(WrapDestFilePath, Attributes);
 			}
 			else
 			{
-				throw new BuildException("No asan lib found in {0} or wrap.sh in {1}", SanitizerLib, WrapSh);
+				throw new BuildException("No asan wrap.sh found in {0}", WrapShFilePath);
 			}
 		}
 
@@ -2802,7 +2826,8 @@ namespace UnrealBuildTool
 			Text.AppendLine("\t             android:icon=\"@drawable/icon\"");
 
 			AndroidToolChain.ClangSanitizer Sanitizer = ToolChain.BuildWithSanitizer();
-			if ((Sanitizer != AndroidToolChain.ClangSanitizer.None && Sanitizer != AndroidToolChain.ClangSanitizer.HwAddress) || bEnableScudoMemoryTracing)
+			// hwasan on NDK r26b+ requires wrap.sh that needs to be unpacked
+			if ((Sanitizer != AndroidToolChain.ClangSanitizer.None && (Sanitizer != AndroidToolChain.ClangSanitizer.HwAddress || ToolChain.HasEmbeddedHWASanSupport())) || bEnableScudoMemoryTracing)
 			{
 				bExtractNativeLibs = true;
 			}
@@ -3920,7 +3945,8 @@ namespace UnrealBuildTool
 			bool bExtractNativeLibs = true;
 			Ini.GetBool("/Script/AndroidRuntimeSettings.AndroidRuntimeSettings", "bExtractNativeLibs", out bExtractNativeLibs);
 			AndroidToolChain.ClangSanitizer Sanitizer = ToolChain.BuildWithSanitizer();
-			if ((Sanitizer != AndroidToolChain.ClangSanitizer.None && Sanitizer != AndroidToolChain.ClangSanitizer.HwAddress) || bEnableScudoMemoryTracing)
+			// hwasan on NDK r26b+ requires wrap.sh that needs to be unpacked
+			if ((Sanitizer != AndroidToolChain.ClangSanitizer.None && (Sanitizer != AndroidToolChain.ClangSanitizer.HwAddress || ToolChain.HasEmbeddedHWASanSupport())) || bEnableScudoMemoryTracing)
 			{
 				bExtractNativeLibs = true;
 			}
@@ -4221,7 +4247,9 @@ namespace UnrealBuildTool
 			// see if last time matches the skipGradle setting
 			string BuildTypeFilename = Path.Combine(IntermediateAndroidPath, "BuildType.txt");
 			string BuildTypeID = bSkipGradleBuild ? "Embedded" : "Standalone";
-			if (Sanitizer != AndroidToolChain.ClangSanitizer.None && Sanitizer != AndroidToolChain.ClangSanitizer.HwAddress)
+			
+			// hwasan on NDK r26b+ requires wrap.sh that needs to be unpacked
+			if (Sanitizer != AndroidToolChain.ClangSanitizer.None && (Sanitizer != AndroidToolChain.ClangSanitizer.HwAddress || ToolChain.HasEmbeddedHWASanSupport()))
 			{
 				BuildTypeID += Sanitizer.ToString() + "Sanitizer";
 			}
@@ -5038,7 +5066,7 @@ popd
 
 				if (Sanitizer != AndroidToolChain.ClangSanitizer.None)
 				{
-					CopyClangSanitizerLib(UnrealBuildPath, Arch, NDKArch, Sanitizer);
+					CopyClangSanitizerLib(ToolChain, UnrealBuildPath, Arch, NDKArch, Sanitizer);
 
 					if (bEnableScudoMemoryTracing)
 					{
@@ -5110,7 +5138,7 @@ popd
 
 				CleanCopyDirectory(Path.Combine(UnrealBuildPath, "jni"), Path.Combine(UnrealBuildGradleMainPath, "jniLibs"), Excludes);  // has debug symbols
 				CleanCopyDirectory(Path.Combine(UnrealBuildPath, "libs"), Path.Combine(UnrealBuildGradleMainPath, "libs"), Excludes);
-				if ((Sanitizer != AndroidToolChain.ClangSanitizer.None && Sanitizer != AndroidToolChain.ClangSanitizer.HwAddress) || bEnableScudoMemoryTracing)
+				if ((Sanitizer != AndroidToolChain.ClangSanitizer.None && (Sanitizer != AndroidToolChain.ClangSanitizer.HwAddress || ToolChain.HasEmbeddedHWASanSupport())) || bEnableScudoMemoryTracing)
 				{
 					CleanCopyDirectory(Path.Combine(UnrealBuildPath, "resources"), Path.Combine(UnrealBuildGradleMainPath, "resources"), Excludes);
 				}
diff --git a/Engine/Source/Programs/UnrealBuildTool/ToolChain/ClangWarnings.cs b/Engine/Source/Programs/UnrealBuildTool/ToolChain/ClangWarnings.cs
--- a/Engine/Source/Programs/UnrealBuildTool/ToolChain/ClangWarnings.cs
+++ b/Engine/Source/Programs/UnrealBuildTool/ToolChain/ClangWarnings.cs
@@ -45,7 +45,8 @@ namespace UnrealBuildTool
 			}
 			if (ClangVersion >= new VersionNumber(17))
 			{
-				if (CompileEnvironment.CppStandard < CppStandardVersion.Latest)
+				bool bIsAndroidClang17 = ClangVersion == new VersionNumber(17, 0, 2) && CompileEnvironment.Platform == UnrealTargetPlatform.Android;
+				if (CompileEnvironment.CppStandard < CppStandardVersion.Latest && !bIsAndroidClang17) // Android clang 17.0.2 in NDK r26b is missing this warning
 				{
 					Arguments.Add("-Wno-invalid-unevaluated-string");   // https://clang.llvm.org/docs/DiagnosticsReference.html#winvalid-unevaluated-string			// new warning for clang 17
 				}
diff --git a/Engine/Source/Runtime/Core/Private/Android/AndroidPlatformMemory.cpp b/Engine/Source/Runtime/Core/Private/Android/AndroidPlatformMemory.cpp
--- a/Engine/Source/Runtime/Core/Private/Android/AndroidPlatformMemory.cpp
+++ b/Engine/Source/Runtime/Core/Private/Android/AndroidPlatformMemory.cpp
@@ -813,3 +813,41 @@ bool FAndroidPlatformMemory::GetLLMAllocFunctions(void*(*&OutAllocFunction)(size
 	return false;
 #endif
 }
+
+#if USING_HW_ADDRESS_SANITISER && PLATFORM_USED_NDK_VERSION_INTEGER >= 26
+
+// Using libc++_static on NDK r26b with HWAsan enabled is not officially supported
+// In practice it does work, but when used with C++ 17 dynamic linker fails to find "_ZnamSt11align_val_t", "_ZnwmSt11align_val_t", "_ZdlPvSt11align_val_t", "_ZdaPvSt11align_val_t" when loading "libUnreal.so".
+// It shouldn't be used by anything on UE side, so providing a stub as a temporary fix.
+
+extern void* operator new(std::size_t Size, std::align_val_t Alignment)
+{
+	void* Result;
+	if (UNLIKELY(posix_memalign(&Result, (std::size_t)Alignment, Size) != 0))
+	{
+		Result = nullptr;
+	}
+	return Result;
+}
+
+extern void* operator new[](std::size_t Size, std::align_val_t Alignment)
+{
+	void* Result;
+	if (UNLIKELY(posix_memalign(&Result, (std::size_t)Alignment, Size) != 0))
+	{
+		Result = nullptr;
+	}
+	return Result;
+}
+
+extern void operator delete(void* Ptr, std::align_val_t)
+{
+	free(Ptr);
+}
+
+extern void operator delete[](void* Ptr, std::align_val_t)
+{
+	free(Ptr);
+}
+
+#endif
\ No newline at end of file
diff --git a/Engine/Source/Runtime/Core/Public/Android/AndroidPlatform.h b/Engine/Source/Runtime/Core/Public/Android/AndroidPlatform.h
--- a/Engine/Source/Runtime/Core/Public/Android/AndroidPlatform.h
+++ b/Engine/Source/Runtime/Core/Public/Android/AndroidPlatform.h
@@ -69,10 +69,6 @@ typedef FAndroidTypes FPlatformTypes;
 
 #define PLATFORM_GLOBAL_LOG_CATEGORY				LogAndroid
 
-#ifndef RUNNING_WITH_ASAN
-	#define RUNNING_WITH_ASAN						0
-#endif
-
 #define PLATFORM_ENABLE_VECTORINTRINSICS			1
 #define PLATFORM_ENABLE_VECTORINTRINSICS_NEON		PLATFORM_ANDROID_ARM64
 
diff --git a/Engine/Source/Runtime/Core/Public/Clang/ClangPlatformCodeAnalysis.h b/Engine/Source/Runtime/Core/Public/Clang/ClangPlatformCodeAnalysis.h
--- a/Engine/Source/Runtime/Core/Public/Clang/ClangPlatformCodeAnalysis.h
+++ b/Engine/Source/Runtime/Core/Public/Clang/ClangPlatformCodeAnalysis.h
@@ -95,6 +95,12 @@
 	#define USING_ADDRESS_SANITISER 0
 #endif
 
+#if defined(__has_feature) && __has_feature(hwaddress_sanitizer)
+	#define USING_HW_ADDRESS_SANITISER 1
+#else
+	#define USING_HW_ADDRESS_SANITISER 0
+#endif
+
 #if defined(__has_feature) && __has_feature(thread_sanitizer)
 	#define USING_THREAD_SANITISER 1
 #else
diff --git a/Engine/Source/Runtime/Core/Public/HAL/Platform.h b/Engine/Source/Runtime/Core/Public/HAL/Platform.h
--- a/Engine/Source/Runtime/Core/Public/HAL/Platform.h
+++ b/Engine/Source/Runtime/Core/Public/HAL/Platform.h
@@ -118,6 +118,10 @@
 	#define USING_ADDRESS_SANITISER 0
 #endif
 
+#ifndef USING_HW_ADDRESS_SANITISER
+	#define USING_HW_ADDRESS_SANITISER 0
+#endif
+
 #ifndef PLATFORM_HAS_ASAN_INCLUDE
 	#define PLATFORM_HAS_ASAN_INCLUDE __has_include(<sanitizer/asan_interface.h>)
 #endif
