UBT: Clang: Moved the -fno-delete-null-pointer-checks flag to ClangToolchain.cs so that it affected all Clang-based platform.

We still have some chacks like "this != nullptr" and without this flag Clang is allowed to optimize it away as invalid.

#rnx

diff --git a/Engine/Source/Programs/UnrealBuildTool/Platform/Android/AndroidToolChain.cs b/Engine/Source/Programs/UnrealBuildTool/Platform/Android/AndroidToolChain.cs
--- a/Engine/Source/Programs/UnrealBuildTool/Platform/Android/AndroidToolChain.cs
+++ b/Engine/Source/Programs/UnrealBuildTool/Platform/Android/AndroidToolChain.cs
@@ -480,12 +480,6 @@ namespace UnrealBuildTool
 		{
 			// @todo the base has PGO options that are a little different than what we want - it would be nice to merge these
 
-			if (CompilerVersionGreaterOrEqual(12, 0, 0))
-			{
-				// We have 'this' vs nullptr comparisons that get optimized away for newer versions of Clang, which is undesirable until we refactor these checks.
-				Arguments.Add("-fno-delete-null-pointer-checks");
-			}
-
 			// optimization level
 			if (!CompileEnvironment.bOptimizeCode)
 			{
diff --git a/Engine/Source/Programs/UnrealBuildTool/Platform/IOS/IOSToolChain.cs b/Engine/Source/Programs/UnrealBuildTool/Platform/IOS/IOSToolChain.cs
--- a/Engine/Source/Programs/UnrealBuildTool/Platform/IOS/IOSToolChain.cs
+++ b/Engine/Source/Programs/UnrealBuildTool/Platform/IOS/IOSToolChain.cs
@@ -159,9 +159,6 @@ namespace UnrealBuildTool
 		{
 			base.GetCompileArguments_Optimizations(CompileEnvironment, Arguments);
 
-			// We have 'this' vs nullptr comparisons that get optimized away for newer versions of Clang, which is undesirable until we refactor these checks.
-			Arguments.Add("-fno-delete-null-pointer-checks");
-
 			// use LTO if desired (like VCToolchain does)
 			if (CompileEnvironment.bAllowLTCG)
 			{
diff --git a/Engine/Source/Programs/UnrealBuildTool/ToolChain/ClangToolChain.cs b/Engine/Source/Programs/UnrealBuildTool/ToolChain/ClangToolChain.cs
--- a/Engine/Source/Programs/UnrealBuildTool/ToolChain/ClangToolChain.cs
+++ b/Engine/Source/Programs/UnrealBuildTool/ToolChain/ClangToolChain.cs
@@ -639,6 +639,12 @@ namespace UnrealBuildTool
 			{
 				Arguments.Add("-fno-inline-functions");
 			}
+
+			if (CompilerVersionGreaterOrEqual(12, 0, 0))
+			{
+				// We have 'this' vs nullptr comparisons that get optimized away for newer versions of Clang, which is undesirable until we refactor these checks.
+				Arguments.Add("-fno-delete-null-pointer-checks");
+			}
 		}
 
 		/// <summary>
