UnrealBuildTool: Retry on any remote internal compiler error

#rnx

diff --git a/Engine/Source/Programs/UnrealBuildTool/Executors/Experimental/UBAExecutor.cs b/Engine/Source/Programs/UnrealBuildTool/Executors/Experimental/UBAExecutor.cs
--- a/Engine/Source/Programs/UnrealBuildTool/Executors/Experimental/UBAExecutor.cs
+++ b/Engine/Source/Programs/UnrealBuildTool/Executors/Experimental/UBAExecutor.cs
@@ -736,6 +736,11 @@ namespace UnrealBuildTool
                         RemoteActionFailedCrash(queue, action, e.ExitCode, e.ExecutingHost ?? "Unknown", "Stack buffer overflow");
                         return;
                     }
+					else if (e.ExitCode != 0 && e.LogLines.Any(x => x.Contains(" C1001: ", StringComparison.Ordinal)))
+					{
+						RemoteActionFailedCrash(queue, action, e.ExitCode, e.ExecutingHost ?? "Unknown", "C1001");
+						return;
+					}
 
                     string additionalDescription = $"[RemoteExecutor: {e.ExecutingHost}]";
 					TimeSpan processorTime = e.TotalProcessorTime;
