UnrealBuildTool: Don't block local execution when waiting for a coordinator to init

#rnx

diff --git a/Engine/Source/Programs/UnrealBuildTool/Executors/Experimental/UBAExecutor.cs b/Engine/Source/Programs/UnrealBuildTool/Executors/Experimental/UBAExecutor.cs
--- a/Engine/Source/Programs/UnrealBuildTool/Executors/Experimental/UBAExecutor.cs
+++ b/Engine/Source/Programs/UnrealBuildTool/Executors/Experimental/UBAExecutor.cs
@@ -333,11 +333,13 @@ namespace UnrealBuildTool
 				_rootDirRef = DirectoryReference.FromString(Environment.GetEnvironmentVariable("UBA_ROOT") ?? Environment.GetEnvironmentVariable("BOX_ROOT"));
 			}
 
+			List<Task> coordinatorInitTasks = new List<Task>();
 			if (!UBAConfig.bDisableRemote)
 			{
 				foreach (IUBAAgentCoordinator coordinator in _agentCoordinators)
 				{
-					await coordinator.InitAsync(this);
+					// Do not await here, it will block local tasks from running
+					coordinatorInitTasks.Add(coordinator.InitAsync(this));
 
 					if (_rootDirRef == null)
 					{
