UBT handling Privacy Manifest files for iOS and Mac

diff --git a/Engine/Config/BaseEngine.ini b/Engine/Config/BaseEngine.ini
--- a/Engine/Config/BaseEngine.ini
+++ b/Engine/Config/BaseEngine.ini
@@ -3179,6 +3179,8 @@ bUseEntitlementsForPrograms=false
 ; The following are used only with Manual codesigning
 MacSigningIdentity=Apple Development
 IOSSigningIdentity=Apple Development
+AdditionalPrivacyInfoMac=(FilePath="/Game/Build/Mac/Resources/PrivacyInfo.xcprivacy")
+AdditionalPrivacyInfoIOS=(FilePath="/Game/Build/IOS/Resources/PrivacyInfo.xcprivacy")
 
 [AutomationTestExcludelistSettings:Windows]
 +SupportedRHIs=DirectX 11
diff --git a/Engine/Source/Developer/Mac/MacPlatformEditor/Private/XcodeProjectSettingsDetailsCustomization.cpp b/Engine/Source/Developer/Mac/MacPlatformEditor/Private/XcodeProjectSettingsDetailsCustomization.cpp
--- a/Engine/Source/Developer/Mac/MacPlatformEditor/Private/XcodeProjectSettingsDetailsCustomization.cpp
+++ b/Engine/Source/Developer/Mac/MacPlatformEditor/Private/XcodeProjectSettingsDetailsCustomization.cpp
@@ -84,6 +84,10 @@ void FXcodeProjectSettingsDetailsCustomization::CustomizeDetails(IDetailLayoutBu
 			{
 				SortOrder = 3;
 			}
+			else if(CategoryName == "Privacy Manifests")
+			{
+				SortOrder = 4;
+			}
             else
             {
                 // Unknown category, should explicitly set order
diff --git a/Engine/Source/Developer/Mac/MacTargetPlatform/Classes/XcodeProjectSettings.h b/Engine/Source/Developer/Mac/MacTargetPlatform/Classes/XcodeProjectSettings.h
--- a/Engine/Source/Developer/Mac/MacTargetPlatform/Classes/XcodeProjectSettings.h
+++ b/Engine/Source/Developer/Mac/MacTargetPlatform/Classes/XcodeProjectSettings.h
@@ -163,4 +163,16 @@ public:
 	 */
 	UPROPERTY(EditAnywhere, config, Category="Code Signing", meta = (EditCondition="bUseModernXcode && bUseAutomaticCodeSigning && bUseAppStoreConnect", DisplayName = "App Store Connect: Key File"))
 	FFilePath AppStoreConnectKeyPath;
+	
+	/**
+	 * The path to the optional PrivacyInfo.xcprivacy file for your project
+	 */
+	UPROPERTY(EditAnywhere, config, Category="Privacy Manifests", meta = (EditCondition="bUseModernXcode", DisplayName = "Mac: Additional Privacy Info"))
+	FFilePath AdditionalPrivacyInfoMac;
+	
+	/**
+	 * The path to the optional PrivacyInfo.xcprivacy file for your project
+	 */
+	UPROPERTY(EditAnywhere, config, Category="Privacy Manifests", meta = (EditCondition="bUseModernXcode", DisplayName = "IOS/TVOS/VisionOS: Additional Privacy Info"))
+	FFilePath AdditionalPrivacyInfoIOS;
 };
diff --git a/Engine/Source/Programs/UnrealBuildTool/ProjectFiles/Xcode/XcodeProject.File.cs b/Engine/Source/Programs/UnrealBuildTool/ProjectFiles/Xcode/XcodeProject.File.cs
--- a/Engine/Source/Programs/UnrealBuildTool/ProjectFiles/Xcode/XcodeProject.File.cs
+++ b/Engine/Source/Programs/UnrealBuildTool/ProjectFiles/Xcode/XcodeProject.File.cs
@@ -125,6 +125,7 @@ namespace UnrealBuildTool.XcodeProjectXcconfig
 			".bundle",
 			".lproj",
 			".png",
+			".xcprivacy",
 			// no extensions will be folders, like "cookeddata" that we want to copy as a resource
 			"",
 		};
diff --git a/Engine/Source/Programs/UnrealBuildTool/ProjectFiles/Xcode/XcodeProject.cs b/Engine/Source/Programs/UnrealBuildTool/ProjectFiles/Xcode/XcodeProject.cs
--- a/Engine/Source/Programs/UnrealBuildTool/ProjectFiles/Xcode/XcodeProject.cs
+++ b/Engine/Source/Programs/UnrealBuildTool/ProjectFiles/Xcode/XcodeProject.cs
@@ -161,7 +161,7 @@ namespace UnrealBuildTool.XcodeProjectXcconfig
 
 		// Location: Can be key of setting entry in .ini, or the full file path
 		// CopyFromFolderIfNotFound: If the file at "Location" does not exist, try to copy the same named file from this folder
-		public MetadataItem(DirectoryReference ProductDirectory, DirectoryReference XcodeProject, ConfigHierarchy Ini, string Location, MetadataMode InMode, DirectoryReference CopyFromFolderIfNotFound)
+		public MetadataItem(DirectoryReference ProductDirectory, DirectoryReference XcodeProject, ConfigHierarchy Ini, string Location, MetadataMode InMode, DirectoryReference? CopyFromFolderIfNotFound)
 		{
 			// no extension means it's a .ini entry
 			if (Path.GetExtension(Location).Length == 0)
@@ -181,7 +181,7 @@ namespace UnrealBuildTool.XcodeProjectXcconfig
 			{
 				// Copy from source location if no such file exist
 				// Except UBTGenerated, they will be generated during UBT runs
-				if (!FileReference.Exists(File) && !File.ContainsName("UBTGenerated", 0))
+				if (!FileReference.Exists(File) && !File.ContainsName("UBTGenerated", 0) && CopyFromFolderIfNotFound != null)
 				{
 					FileReference SourceFile = FileReference.Combine(CopyFromFolderIfNotFound, File.GetFileName());
 					if (FileReference.Exists(SourceFile))
@@ -225,6 +225,7 @@ namespace UnrealBuildTool.XcodeProjectXcconfig
 		public Dictionary<MetadataPlatform, MetadataItem> PlistFiles = new();
 		public Dictionary<MetadataPlatform, MetadataItem> EntitlementsFiles = new();
 		public Dictionary<MetadataPlatform, MetadataItem> ShippingEntitlementsFiles = new();
+		public Dictionary<MetadataPlatform, MetadataItem> ProjectPrivacyInfoFiles = new();
 
 		public Metadata(DirectoryReference ProductDirectory, DirectoryReference XcodeProject, ConfigHierarchy Ini, bool bSupportsMac, bool bSupportsIOSOrTVOS, ILogger Logger)
 		{
@@ -263,6 +264,11 @@ namespace UnrealBuildTool.XcodeProjectXcconfig
 						MetadataMode.UpdateTemplate,
 						ResourceFolder);
 					}
+
+					ProjectPrivacyInfoFiles[MetadataPlatform.Mac] = new MetadataItem(ProductDirectory, XcodeProject, Ini,
+						"AdditionalPrivacyInfoMac",
+						MetadataMode.UsePremade,
+						null);
 				}
 
 				EntitlementsFiles[MetadataPlatform.MacEditor] = new MetadataItem(ProductDirectory, XcodeProject, Ini,
@@ -298,6 +304,14 @@ namespace UnrealBuildTool.XcodeProjectXcconfig
 						ResourceFolder);
 				}
 
+				if (ProductDirectory != Unreal.EngineDirectory)
+				{
+					ProjectPrivacyInfoFiles[MetadataPlatform.IOS] = new MetadataItem(ProductDirectory, XcodeProject, Ini,
+						"AdditionalPrivacyInfoIOS",
+						MetadataMode.UsePremade,
+						null);
+				}
+
 				EntitlementsFiles[MetadataPlatform.IOS] = new MetadataItem(ProductDirectory, XcodeProject, Ini,
 					"PremadeIOSEntitlements",
 					MetadataMode.UsePremade,
@@ -1522,6 +1536,31 @@ namespace UnrealBuildTool.XcodeProjectXcconfig
 					ResourcesBuildPhase.AddResource(new FileReference(LaunchImagePath));
 				}
 			}
+
+			if (Platform == UnrealTargetPlatform.Mac)
+			{
+				ResourcesBuildPhase.AddFolderResource(DirectoryReference.Combine(Unreal.EngineDirectory, "Build/Mac/Resources/UEMetadata"), "Resources");
+				if (UnrealData.Metadata?.ProjectPrivacyInfoFiles.ContainsKey(MetadataPlatform.Mac) == true)
+				{
+					MetadataItem PrivacyInfo = UnrealData.Metadata.ProjectPrivacyInfoFiles[MetadataPlatform.Mac];
+					if (PrivacyInfo.XcodeProjectRelative != null)
+					{
+						ResourcesBuildPhase.AddResource(FileReference.Combine(UnrealData.XcodeProjectFileLocation.ParentDirectory!, PrivacyInfo.XcodeProjectRelative));
+					}
+				}
+			}
+			else if (Platform == UnrealTargetPlatform.IOS || Platform == UnrealTargetPlatform.TVOS || Platform == UnrealTargetPlatform.VisionOS)
+			{
+				ResourcesBuildPhase.AddFolderResource(DirectoryReference.Combine(Unreal.EngineDirectory, "Build/IOS/Resources/UEMetadata"), "Resources");
+				if (UnrealData.Metadata?.ProjectPrivacyInfoFiles.ContainsKey(MetadataPlatform.IOS) == true)
+				{
+					MetadataItem PrivacyInfo = UnrealData.Metadata.ProjectPrivacyInfoFiles[MetadataPlatform.IOS];
+					if (PrivacyInfo.XcodeProjectRelative != null)
+					{
+						ResourcesBuildPhase.AddResource(FileReference.Combine(UnrealData.XcodeProjectFileLocation.ParentDirectory!, PrivacyInfo.XcodeProjectRelative));
+					}
+				}
+			}
 		}
 		protected void ProcessFrameworks(XcodeResourcesBuildPhase ResourcesBuildPhase, XcodeProject Project, XcodeProjectFile ProjectFile, ILogger Logger)
 		{
