Horde: Remove stacktrace when logging non-strict worker acquistion fails

diff --git a/Engine/Source/Programs/UnrealBuildTool/Executors/Experimental/UBAHordeSession.cs b/Engine/Source/Programs/UnrealBuildTool/Executors/Experimental/UBAHordeSession.cs
--- a/Engine/Source/Programs/UnrealBuildTool/Executors/Experimental/UBAHordeSession.cs
+++ b/Engine/Source/Programs/UnrealBuildTool/Executors/Experimental/UBAHordeSession.cs
@@ -447,7 +447,7 @@ namespace UnrealBuildTool
 
 				_workerId++;
 
-				workerLogger.LogInformation("Agent properties:");
+				workerLogger.LogDebug("Agent properties:");
 
 				int numLogicalCores = 24; // Assume 24 if something goes wrong here and property is not found
 				string computeIp = String.Empty;
@@ -457,7 +457,7 @@ namespace UnrealBuildTool
 					StringView propertyName = new(property, 0, equalsIdx);
 					if (s_logProperties.Contains(propertyName))
 					{
-						_logger.LogInformation("  {Property}", property);
+						_logger.LogDebug("  {Property}", property);
 
 						if (propertyName == ResourceLogicalCores && Int32.TryParse(property.AsSpan(equalsIdx + 1), out int value))
 						{
@@ -619,7 +619,7 @@ namespace UnrealBuildTool
 
 		async Task RunWorkerAsync(Worker self, IComputeLease lease, BlobLocator tool, string executable, ILogger logger, UnrealBuildAcceleratorHordeConfig hordeConfig, CancellationToken cancellationToken)
 		{
-			logger.LogInformation("Running worker task..");
+			logger.LogDebug("Running worker task..");
 			try
 			{
 				await using (_ = lease)
@@ -628,7 +628,7 @@ namespace UnrealBuildTool
 					const int PrimaryChannelId = 0;
 					using (AgentMessageChannel channel = lease.Socket.CreateAgentMessageChannel(PrimaryChannelId, 4 * 1024 * 1024))
 					{
-						logger.LogInformation("Waiting for attach...");
+						logger.LogDebug("Waiting for attach...");
 
 						TimeSpan attachTimeout = TimeSpan.FromSeconds(20.0);
 						try
@@ -642,7 +642,7 @@ namespace UnrealBuildTool
 							throw;
 						}
 
-						logger.LogInformation("Uploading files...");
+						logger.LogDebug("Uploading files...");
 						await channel.UploadFilesAsync("", tool, _storage.Backend, cancellationToken);
 
 						string hordeHost = _owner.UBAConfig.Host;
@@ -718,10 +718,10 @@ namespace UnrealBuildTool
 								shouldConnect = false;
 							}
 						}
-						logger.LogInformation("Shutting down process");
+						logger.LogDebug("Shutting down process");
 					}
 
-					logger.LogInformation("Closing channel");
+					logger.LogDebug("Closing channel");
 					await lease.CloseAsync(cancellationToken);
 				}
 			}
@@ -865,7 +865,14 @@ namespace UnrealBuildTool
 				{
 					if (!_cancellationSource!.IsCancellationRequested)
 					{
-						_logger.Log(_ubaConfig.bStrict ? LogLevel.Error : LogLevel.Information, KnownLogEvents.Systemic_Horde_Compute, ex, "Unable to get worker: {Ex}", ex.ToString());
+						if (_ubaConfig.bStrict)
+						{
+							_logger.Log(LogLevel.Error, KnownLogEvents.Systemic_Horde_Compute, ex, "Unable to get worker: {Ex}", ex.ToString());
+						}
+						else
+						{
+							_logger.Log(LogLevel.Information, KnownLogEvents.Systemic_Horde_Compute, "Unable to get worker: {Ex}", ex.ToString());
+						}
 					}
 				}
 
