Removed DEPRECATED_FORGAME from UBT, as it hasn't existed since CL# 19480377.

diff --git a/Engine/Source/Programs/IncludeTool/IncludeTool/Rules.cs b/Engine/Source/Programs/IncludeTool/IncludeTool/Rules.cs
--- a/Engine/Source/Programs/IncludeTool/IncludeTool/Rules.cs
+++ b/Engine/Source/Programs/IncludeTool/IncludeTool/Rules.cs
@@ -598,7 +598,7 @@ namespace IncludeTool
 			{
 				return true;
 			}			
-			if(Markup.Type == PreprocessorMarkupType.Define && (Markup.Tokens[0].Text == "UE_DEPRECATED_FORGAME" || Markup.Tokens[0].Text == "DEPRECATED_FORGAME"))
+			if(Markup.Type == PreprocessorMarkupType.Define && (Markup.Tokens[0].Text == "UE_DEPRECATED_FORGAME"))
 			{
 				return true;
 			}
diff --git a/Engine/Source/Programs/UnrealBuildTool/Configuration/UEBuildModuleCPP.cs b/Engine/Source/Programs/UnrealBuildTool/Configuration/UEBuildModuleCPP.cs
--- a/Engine/Source/Programs/UnrealBuildTool/Configuration/UEBuildModuleCPP.cs
+++ b/Engine/Source/Programs/UnrealBuildTool/Configuration/UEBuildModuleCPP.cs
@@ -1054,7 +1054,6 @@ namespace UnrealBuildTool
 				if (Definition.Contains("UE_IS_ENGINE_MODULE", StringComparison.Ordinal) ||
 					Definition.Contains("UE_VALIDATE_FORMAT_STRINGS", StringComparison.Ordinal) ||
 					Definition.Contains("UE_VALIDATE_INTERNAL_API", StringComparison.Ordinal) ||
-					Definition.Contains("DEPRECATED_FORGAME", StringComparison.Ordinal) ||
 					Definition.Contains("UE_DEPRECATED_FORGAME", StringComparison.Ordinal) ||
 					Definition.Contains("UE_ENABLE_INCLUDE_ORDER_DEPRECATED_IN_", StringComparison.Ordinal))
 				{
@@ -1133,11 +1132,9 @@ namespace UnrealBuildTool
 				{
 					Definitions = new List<string>(Definitions);
 					Definitions.RemoveAll(x => x.Contains("UE_IS_ENGINE_MODULE", StringComparison.Ordinal));
-					Definitions.RemoveAll(x => x.Contains("DEPRECATED_FORGAME", StringComparison.Ordinal));
 					Definitions.RemoveAll(x => x.Contains("UE_DEPRECATED_FORGAME", StringComparison.Ordinal));
 
 					Definitions.Add($"UE_IS_ENGINE_MODULE={(Rules.bTreatAsEngineModule ? "1" : "0")}");
-					Definitions.Add($"DEPRECATED_FORGAME={(Rules.bTreatAsEngineModule ? String.Empty : "DEPRECATED")}");
 					Definitions.Add($"UE_DEPRECATED_FORGAME={(Rules.bTreatAsEngineModule ? String.Empty : "UE_DEPRECATED")}");
 				}
 
@@ -2122,7 +2119,6 @@ namespace UnrealBuildTool
 
 			if (!Rules.bTreatAsEngineModule)
 			{
-				Result.Definitions.Add("DEPRECATED_FORGAME=DEPRECATED");
 				Result.Definitions.Add("UE_DEPRECATED_FORGAME=UE_DEPRECATED");
 			}
 
@@ -2188,7 +2184,6 @@ namespace UnrealBuildTool
 
 			if (!Rules.bTreatAsEngineModule)
 			{
-				CompileEnvironment.Definitions.Add("DEPRECATED_FORGAME=DEPRECATED");
 				CompileEnvironment.Definitions.Add("UE_DEPRECATED_FORGAME=UE_DEPRECATED");
 			}
 
