Added UE_DEPRECATED_FORENGINE deprecation warnings for engine modules, and TargetRules.bDisableEngineDeprecations to silence engine deprecations in non-engine modules, to allow incremental fixup.  Disabling engine deprecations should be regarded as short-term to allow incremental fixup.

diff --git a/Engine/Source/Programs/IncludeTool/IncludeTool/Rules.cs b/Engine/Source/Programs/IncludeTool/IncludeTool/Rules.cs
--- a/Engine/Source/Programs/IncludeTool/IncludeTool/Rules.cs
+++ b/Engine/Source/Programs/IncludeTool/IncludeTool/Rules.cs
@@ -602,7 +602,11 @@ namespace IncludeTool
 			{
 				return true;
 			}
-			if(Markup.Type == PreprocessorMarkupType.Undef && Markup.Tokens[0].Text == "TEXT")
+			if (Markup.Type == PreprocessorMarkupType.Define && Markup.Tokens[0].Text == "UE_DEPRECATED_FORENGINE")
+			{
+				return true;
+			}
+			if (Markup.Type == PreprocessorMarkupType.Undef && Markup.Tokens[0].Text == "TEXT")
 			{
 				return true;
 			}
diff --git a/Engine/Source/Programs/UnrealBuildTool/Configuration/ReadOnlyTargetRules.cs b/Engine/Source/Programs/UnrealBuildTool/Configuration/ReadOnlyTargetRules.cs
--- a/Engine/Source/Programs/UnrealBuildTool/Configuration/ReadOnlyTargetRules.cs
+++ b/Engine/Source/Programs/UnrealBuildTool/Configuration/ReadOnlyTargetRules.cs
@@ -494,6 +494,8 @@ namespace UnrealBuildTool
 
 		public bool bLegacyParentIncludePaths => Inner.bLegacyParentIncludePaths;
 
+		public bool bDisableEngineDeprecations => Inner.bDisableEngineDeprecations;
+
 		public CppStandardVersion CppStandardEngine => Inner.CppStandardEngine;
 
 		public CppStandardVersion CppStandard => Inner.CppStandard;
diff --git a/Engine/Source/Programs/UnrealBuildTool/Configuration/TargetRules.cs b/Engine/Source/Programs/UnrealBuildTool/Configuration/TargetRules.cs
--- a/Engine/Source/Programs/UnrealBuildTool/Configuration/TargetRules.cs
+++ b/Engine/Source/Programs/UnrealBuildTool/Configuration/TargetRules.cs
@@ -2210,6 +2210,12 @@ namespace UnrealBuildTool
 		/// </summary>
 		public bool bIgnoreBuildOutputs { get; set; }
 
+		/// <summary>
+		/// Whether to disable UE_DEPRECATED_FORENGINE deprecation warnings in non-engine modules.
+		/// This should only be used in the short-term as these deprecations will get removed.
+		/// </summary>
+		public bool bDisableEngineDeprecations { get; set; }
+
 		/// <summary>
 		/// Indicates that this is a formal build, intended for distribution. This flag is automatically set to true when Build.version has a changelist set and is a promoted build.
 		/// The only behavior currently bound to this flag is to compile the default resource file separately for each binary so that the OriginalFilename field is set correctly.
diff --git a/Engine/Source/Programs/UnrealBuildTool/Configuration/UEBuildModuleCPP.cs b/Engine/Source/Programs/UnrealBuildTool/Configuration/UEBuildModuleCPP.cs
--- a/Engine/Source/Programs/UnrealBuildTool/Configuration/UEBuildModuleCPP.cs
+++ b/Engine/Source/Programs/UnrealBuildTool/Configuration/UEBuildModuleCPP.cs
@@ -1055,6 +1055,7 @@ namespace UnrealBuildTool
 					Definition.Contains("UE_VALIDATE_FORMAT_STRINGS", StringComparison.Ordinal) ||
 					Definition.Contains("UE_VALIDATE_INTERNAL_API", StringComparison.Ordinal) ||
 					Definition.Contains("UE_DEPRECATED_FORGAME", StringComparison.Ordinal) ||
+					Definition.Contains("UE_DEPRECATED_FORENGINE", StringComparison.Ordinal) ||
 					Definition.Contains("UE_ENABLE_INCLUDE_ORDER_DEPRECATED_IN_", StringComparison.Ordinal))
 				{
 					continue;
@@ -1133,9 +1134,11 @@ namespace UnrealBuildTool
 					Definitions = new List<string>(Definitions);
 					Definitions.RemoveAll(x => x.Contains("UE_IS_ENGINE_MODULE", StringComparison.Ordinal));
 					Definitions.RemoveAll(x => x.Contains("UE_DEPRECATED_FORGAME", StringComparison.Ordinal));
+					Definitions.RemoveAll(x => x.Contains("UE_DEPRECATED_FORENGINE", StringComparison.Ordinal));
 
 					Definitions.Add($"UE_IS_ENGINE_MODULE={(Rules.bTreatAsEngineModule ? "1" : "0")}");
-					Definitions.Add($"UE_DEPRECATED_FORGAME={(Rules.bTreatAsEngineModule ? String.Empty : "UE_DEPRECATED")}");
+					Definitions.Add($"UE_DEPRECATED_FORGAME={(Rules.bTreatAsEngineModule ? "PREPROCESSOR_NOTHING_FUNCTION" : "UE_DEPRECATED")}");
+					Definitions.Add($"UE_DEPRECATED_FORENGINE={(Rules.bTreatAsEngineModule || Rules.Target.bDisableEngineDeprecations ? "PREPROCESSOR_NOTHING_FUNCTION" : "UE_DEPRECATED")}");
 				}
 
 				// Modify definitions if we need to create a new shared pch for validating format strings
@@ -2120,6 +2123,7 @@ namespace UnrealBuildTool
 			if (!Rules.bTreatAsEngineModule)
 			{
 				Result.Definitions.Add("UE_DEPRECATED_FORGAME=UE_DEPRECATED");
+				Result.Definitions.Add($"UE_DEPRECATED_FORENGINE={(Rules.Target.bDisableEngineDeprecations ? "PREPROCESSOR_NOTHING_FUNCTION" : "UE_DEPRECATED")}");
 			}
 
 			Result.Definitions.Add($"UE_VALIDATE_FORMAT_STRINGS={(Rules.bValidateFormatStrings ? "1" : "0")}");
@@ -2185,6 +2189,7 @@ namespace UnrealBuildTool
 			if (!Rules.bTreatAsEngineModule)
 			{
 				CompileEnvironment.Definitions.Add("UE_DEPRECATED_FORGAME=UE_DEPRECATED");
+				CompileEnvironment.Definitions.Add($"UE_DEPRECATED_FORENGINE={(Rules.Target.bDisableEngineDeprecations ? "PREPROCESSOR_NOTHING_FUNCTION" : "UE_DEPRECATED")}");
 			}
 
 			CompileEnvironment.Definitions.Add($"UE_VALIDATE_FORMAT_STRINGS={(Rules.bValidateFormatStrings ? "1" : "0")}");
diff --git a/Engine/Source/Runtime/Core/Public/HAL/PreprocessorHelpers.h b/Engine/Source/Runtime/Core/Public/HAL/PreprocessorHelpers.h
--- a/Engine/Source/Runtime/Core/Public/HAL/PreprocessorHelpers.h
+++ b/Engine/Source/Runtime/Core/Public/HAL/PreprocessorHelpers.h
@@ -45,6 +45,9 @@
 // Expands to nothing - used as a placeholder
 #define PREPROCESSOR_NOTHING
 
+// Expands to nothing when used as a function - used as a placeholder
+#define PREPROCESSOR_NOTHING_FUNCTION(...)
+
 #define UE_SOURCE_LOCATION TEXT(__FILE__ "(" PREPROCESSOR_TO_STRING(__LINE__) ")")
 
 // Removes a single layer of parentheses from a macro argument if they are present - used to allow
diff --git a/Engine/Source/Runtime/Core/Public/Misc/CoreMiscDefines.h b/Engine/Source/Runtime/Core/Public/Misc/CoreMiscDefines.h
--- a/Engine/Source/Runtime/Core/Public/Misc/CoreMiscDefines.h
+++ b/Engine/Source/Runtime/Core/Public/Misc/CoreMiscDefines.h
@@ -275,7 +275,11 @@ enum EInPlace {InPlace};
 #endif
 
 #ifndef UE_DEPRECATED_FORGAME
-	#define UE_DEPRECATED_FORGAME(...)
+	#define UE_DEPRECATED_FORGAME PREPROCESSOR_NOTHING_FUNCTION
+#endif
+
+#ifndef UE_DEPRECATED_FORENGINE
+	#define UE_DEPRECATED_FORENGINE UE_DEPRECATED
 #endif
 
 #if UE_VALIDATE_INTERNAL_API
