Horde: Fix static analysis warning.

diff --git a/Engine/Source/Programs/UnrealBuildTool/Matchers/CompileEventMatcher.cs b/Engine/Source/Programs/UnrealBuildTool/Matchers/CompileEventMatcher.cs
--- a/Engine/Source/Programs/UnrealBuildTool/Matchers/CompileEventMatcher.cs
+++ b/Engine/Source/Programs/UnrealBuildTool/Matchers/CompileEventMatcher.cs
@@ -114,7 +114,7 @@ namespace UnrealBuildTool.Matchers
 
 					// If warnings as errors is enabled, upgrade any following warnings to errors.
 					LogValue? code;
-					if (newEvent.Properties != null && newEvent.TryGetProperty("code", out code) && code.Text == "C2220")
+					if (newEvent.Properties != null && newEvent.TryGetProperty("code", out code) && code.Text.Equals("C2220", StringComparison.Ordinal))
 					{
 						ILogCursor nextCursor = builder.Next;
 						while (nextCursor.CurrentLine != null)
