Expose RuntimeDependencies to ToolChain ModifyBuildProducts and LinkFiles methods. This is required to properly bundle DLLs on some platforms.

#changelist validated
#virtualized

diff --git a/Engine/Source/Programs/UnrealBuildTool/Configuration/UEBuildBinary.cs b/Engine/Source/Programs/UnrealBuildTool/Configuration/UEBuildBinary.cs
--- a/Engine/Source/Programs/UnrealBuildTool/Configuration/UEBuildBinary.cs
+++ b/Engine/Source/Programs/UnrealBuildTool/Configuration/UEBuildBinary.cs
@@ -118,5 +118,10 @@
 		private List<FileReference>? DependentLinkLibraries;
 
 		/// <summary>
+		/// Expanded list of runtime dependencies generated by calling PrepareRuntimeDependencies.
+		/// </summary>
+		public readonly List<ModuleRules.RuntimeDependency> RuntimeDependencies = new List<ModuleRules.RuntimeDependency>();
+
+		/// <summary>
 		/// Create an instance of the class with the given configuration data
 		/// </summary>
@@ -161,5 +166,61 @@
 		}
 
 		/// <summary>
+		/// Expands all of the runtime dependencies.
+		/// </summary>
+		/// <param name="ExeDir">Output directory for the executable</param>
+		public void PrepareRuntimeDependencies(DirectoryReference ExeDir)
+		{
+			RuntimeDependencies.Clear();
+			foreach (UEBuildModule Module in Modules)
+			{
+				foreach (ModuleRules.RuntimeDependency Dependency in Module.Rules.RuntimeDependencies.Inner)
+				{
+					if (Dependency.SourcePath == null)
+					{
+						// Expand the target path
+						string ExpandedPath = Module.ExpandPathVariables(Dependency.Path, OutputDir, ExeDir);
+						if (FileFilter.FindWildcardIndex(ExpandedPath) == -1)
+						{
+							RuntimeDependencies.Add(Dependency);
+						}
+						else
+						{
+							RuntimeDependencies.AddRange(FileFilter.ResolveWildcard(ExpandedPath).Select(x => new ModuleRules.RuntimeDependency(x.FullName, Dependency.Type)));
+						}
+					}
+					else
+					{
+						// Parse the source and target patterns
+						FilePattern SourcePattern = new FilePattern(Unreal.EngineSourceDirectory, Module.ExpandPathVariables(Dependency.SourcePath, OutputDir, ExeDir));
+						FilePattern TargetPattern = new FilePattern(Unreal.EngineSourceDirectory, Module.ExpandPathVariables(Dependency.Path, OutputDir, ExeDir));
+
+						// Skip non-essential single files if they do not exist
+						if (Dependency.Type == StagedFileType.DebugNonUFS && !SourcePattern.ContainsWildcards() && !FileReference.Exists(SourcePattern.GetSingleFile()))
+						{
+							continue;
+						}
+
+						// Resolve all the wildcards between the source and target paths
+						Dictionary<FileReference, FileReference> Mapping;
+						try
+						{
+							Mapping = FilePattern.CreateMapping(null, ref SourcePattern, ref TargetPattern);
+						}
+						catch (FilePatternException Ex)
+						{
+							ExceptionUtils.AddContext(Ex, "while creating runtime dependencies for module '{0}'", Module.Name);
+							throw;
+						}
+						foreach (KeyValuePair<FileReference, FileReference> Pair in Mapping)
+						{
+							RuntimeDependencies.Add(new ModuleRules.RuntimeDependency(Pair.Key.FullName, Pair.Value.FullName, Dependency.Type));
+						}
+					}
+				}
+			}
+		}
+
+		/// <summary>
 		/// Builds the binary.
 		/// </summary>
@@ -253,66 +314,28 @@
 		/// <summary>
 		/// Gets all the runtime dependencies Copies all the runtime dependencies from any modules in 
 		/// </summary>
-		/// <param name="RuntimeDependencies">The output list of runtime dependencies, mapping target file to type</param>
+		/// <param name="OutRuntimeDependencies">The output list of runtime dependencies, mapping target file to type</param>
 		/// <param name="TargetFileToSourceFile">Map of target files to source files that need to be copied</param>
-		/// <param name="ExeDir">Output directory for the executable</param>
-		public void PrepareRuntimeDependencies(List<RuntimeDependency> RuntimeDependencies, Dictionary<FileReference, FileReference> TargetFileToSourceFile, DirectoryReference ExeDir)
+		public void CollectRuntimeDependencies(List<RuntimeDependency> OutRuntimeDependencies, Dictionary<FileReference, FileReference> TargetFileToSourceFile)
 		{
-			foreach (UEBuildModule Module in Modules)
+			foreach (ModuleRules.RuntimeDependency Dependency in RuntimeDependencies)
 			{
-				foreach (ModuleRules.RuntimeDependency Dependency in Module.Rules.RuntimeDependencies.Inner)
+				FileReference DepFile = new FileReference(Dependency.Path);
+				if (Dependency.SourcePath == null)
 				{
-					if (Dependency.SourcePath == null)
+					OutRuntimeDependencies.Add(new RuntimeDependency(DepFile, Dependency.Type));
+				}
+				else
+				{
+					FileReference? ExistingSourceFile;
+					if (!TargetFileToSourceFile.TryGetValue(DepFile, out ExistingSourceFile))
 					{
-						// Expand the target path
-						string ExpandedPath = Module.ExpandPathVariables(Dependency.Path, OutputDir, ExeDir);
-						if (FileFilter.FindWildcardIndex(ExpandedPath) == -1)
-						{
-							RuntimeDependencies.Add(new RuntimeDependency(new FileReference(ExpandedPath), Dependency.Type));
-						}
-						else
-						{
-							RuntimeDependencies.AddRange(FileFilter.ResolveWildcard(ExpandedPath).Select(x => new RuntimeDependency(x, Dependency.Type)));
-						}
+						TargetFileToSourceFile[DepFile] = new FileReference(Dependency.SourcePath);
+						OutRuntimeDependencies.Add(new RuntimeDependency(DepFile, Dependency.Type));
 					}
-					else
+					else if (ExistingSourceFile.FullName != Dependency.SourcePath)
 					{
-						// Parse the source and target patterns
-						FilePattern SourcePattern = new FilePattern(Unreal.EngineSourceDirectory, Module.ExpandPathVariables(Dependency.SourcePath, OutputDir, ExeDir));
-						FilePattern TargetPattern = new FilePattern(Unreal.EngineSourceDirectory, Module.ExpandPathVariables(Dependency.Path, OutputDir, ExeDir));
-
-						// Skip non-essential single files if they do not exist
-						if (Dependency.Type == StagedFileType.DebugNonUFS && !SourcePattern.ContainsWildcards() && !FileReference.Exists(SourcePattern.GetSingleFile()))
-						{
-							continue;
-						}
-
-						// Resolve all the wildcards between the source and target paths
-						Dictionary<FileReference, FileReference> Mapping;
-						try
-						{
-							Mapping = FilePattern.CreateMapping(null, ref SourcePattern, ref TargetPattern);
-						}
-						catch (FilePatternException Ex)
-						{
-							ExceptionUtils.AddContext(Ex, "while creating runtime dependencies for module '{0}'", Module.Name);
-							throw;
-						}
-
-						// Add actions to copy everything
-						foreach (KeyValuePair<FileReference, FileReference> Pair in Mapping)
-						{
-							FileReference? ExistingSourceFile;
-							if (!TargetFileToSourceFile.TryGetValue(Pair.Key, out ExistingSourceFile))
-							{
-								TargetFileToSourceFile[Pair.Key] = Pair.Value;
-								RuntimeDependencies.Add(new RuntimeDependency(Pair.Key, Dependency.Type));
-							}
-							else if (ExistingSourceFile != Pair.Value)
-							{
-								throw new BuildException("Runtime dependency '{0}' is configured to be staged from '{1}' and '{2}'", Pair.Key, Pair.Value, ExistingSourceFile);
-							}
-						}
+						throw new BuildException("Runtime dependency '{0}' is configured to be staged from '{1}' and '{2}'", Dependency.Path, Dependency.SourcePath, ExistingSourceFile);
 					}
 				}
 			}
@@ -842,6 +865,9 @@
 			// Add all the common resource files
 			BinaryLinkEnvironment.InputFiles.AddRange(BinaryLinkEnvironment.CommonResourceFiles);
 
+			// Provide list of runtime dependencies to link environment.
+			BinaryLinkEnvironment.RuntimeDependencies = RuntimeDependencies;
+
 			// Allow the platform to modify the binary link environment for platform-specific resources etc.
 			UEBuildPlatform.GetBuildPlatform(Target.Platform).ModifyBinaryLinkEnvironment(BinaryLinkEnvironment, BinaryCompileEnvironment, Target, ToolChain, Graph);
 
diff --git a/Engine/Source/Programs/UnrealBuildTool/Configuration/UEBuildTarget.cs b/Engine/Source/Programs/UnrealBuildTool/Configuration/UEBuildTarget.cs
--- a/Engine/Source/Programs/UnrealBuildTool/Configuration/UEBuildTarget.cs
+++ b/Engine/Source/Programs/UnrealBuildTool/Configuration/UEBuildTarget.cs
@@ -2628,6 +2628,7 @@ namespace UnrealBuildTool
 			{
 				foreach (UEBuildBinary Binary in Binaries)
 				{
+					Binary.PrepareRuntimeDependencies(ExeDir);
 					List<FileItem> BinaryOutputItems = Binary.Build(Rules, TargetToolChain, GlobalCompileEnvironment, GlobalLinkEnvironment, WorkingSet, ExeDir, MakefileBuilder, Logger);
 					Makefile.OutputItems.AddRange(BinaryOutputItems);
 				}
@@ -2651,7 +2652,7 @@ namespace UnrealBuildTool
 			Dictionary<FileReference, FileReference> RuntimeDependencyTargetFileToSourceFile = new Dictionary<FileReference, FileReference>();
 			foreach (UEBuildBinary Binary in Binaries)
 			{
-				Binary.PrepareRuntimeDependencies(RuntimeDependencies, RuntimeDependencyTargetFileToSourceFile, ExeDir);
+				Binary.CollectRuntimeDependencies(RuntimeDependencies, RuntimeDependencyTargetFileToSourceFile);
 			}
 			TargetToolChain.PrepareRuntimeDependencies(RuntimeDependencies, RuntimeDependencyTargetFileToSourceFile, ExeDir);
 
diff --git a/Engine/Source/Programs/UnrealBuildTool/System/LinkEnvironment.cs b/Engine/Source/Programs/UnrealBuildTool/System/LinkEnvironment.cs
--- a/Engine/Source/Programs/UnrealBuildTool/System/LinkEnvironment.cs
+++ b/Engine/Source/Programs/UnrealBuildTool/System/LinkEnvironment.cs
@@ -357,6 +357,11 @@ namespace UnrealBuildTool
 		/// </summary>
 		public List<ReceiptProperty> AdditionalProperties = new List<ReceiptProperty>();
 
+		/// <summary>
+		/// List of all of the runtime dependencies.
+		/// </summary>
+		public List<ModuleRules.RuntimeDependency> RuntimeDependencies = new List<ModuleRules.RuntimeDependency>();
+
 		/// <summary>
 		/// Default constructor.
 		/// </summary>
