Do not ignore default value on setInFromProperty in UPL

diff --git a/Engine/Source/Programs/UnrealBuildTool/System/UnrealPluginLanguage.cs b/Engine/Source/Programs/UnrealBuildTool/System/UnrealPluginLanguage.cs
--- a/Engine/Source/Programs/UnrealBuildTool/System/UnrealPluginLanguage.cs
+++ b/Engine/Source/Programs/UnrealBuildTool/System/UnrealPluginLanguage.cs
@@ -2117,12 +2117,16 @@ namespace UnrealBuildTool
 							string DefaultVal = GetAttribute(CurrentContext, Node, "default", true, false, "0");
 							if (Result != null && Ini != null && Section != null && Property != null)
 							{
-								int Value = StringToInt(CurrentContext, Node, DefaultVal);
+								int DefaultInt = StringToInt(CurrentContext, Node, DefaultVal);
+								int Value = DefaultInt;
 
 								ConfigCacheIni_UPL ConfigIni = GetConfigCacheIni_UPL(Ini);
 								if (ConfigIni != null)
 								{
-									ConfigIni.GetInt32(Section, Property, out Value);
+									if (!ConfigIni.GetInt32(Section, Property, out Value))
+									{
+										Value = DefaultInt;
+									}
 								}
 								CurrentContext.IntVariables[Result] = Value;
 							}
