- Updated the logging for mismatched build settings that require UniqueBuildEnvironment to show the values that don't match

diff --git a/Engine/Source/Programs/UnrealBuildTool/Configuration/TargetRules.cs b/Engine/Source/Programs/UnrealBuildTool/Configuration/TargetRules.cs
--- a/Engine/Source/Programs/UnrealBuildTool/Configuration/TargetRules.cs
+++ b/Engine/Source/Programs/UnrealBuildTool/Configuration/TargetRules.cs
@@ -2648,11 +2648,11 @@ namespace UnrealBuildTool
 		/// </summary>
 		/// <param name="rulesAssembly">Assembly containing the target</param>
 		/// <param name="arguments">Commandline options that may affect the target creation</param>
-		/// <param name="propNamesThatRequireUnique">If the target requires a unique build environment, this will contain the names of the field/property that require unique</param>
+		/// <param name="propNamesThatRequireUnique">If the target requires a unique build environment, this will contain the names of the field/property that require unique, mapped to a pair of base and current value</param>
 		/// <param name="baseTargetName">If the target requires a unique build environment, this will contain the name of the target this was based on (UnrealGame, UnrealEditor, etc)/property</param>
 		/// <returns>true if a property was set such that it requires a unique build environment</returns>
 		/// <exception cref="BuildException"></exception>
-		public bool RequiresUniqueEnvironment(RulesAssembly rulesAssembly, CommandLineArguments? arguments, List<string> propNamesThatRequireUnique, [NotNullWhen(true)] out string? baseTargetName)
+		public bool RequiresUniqueEnvironment(RulesAssembly rulesAssembly, CommandLineArguments? arguments, Dictionary<string, (string?, string?)> propNamesThatRequireUnique, [NotNullWhen(true)] out string? baseTargetName)
 		{
 			baseTargetName = null;
 
@@ -2703,7 +2703,7 @@ namespace UnrealBuildTool
 					object? baseValue = field.GetValue(baseObjects[idx]);
 					if (!CheckValuesMatch(field.FieldType, thisValue, baseValue))
 					{
-						propNamesThatRequireUnique.Add(field.Name);
+						propNamesThatRequireUnique.Add(field.Name, (thisValue?.ToString(), baseValue?.ToString()));
 					}
 				}
 
@@ -2713,7 +2713,7 @@ namespace UnrealBuildTool
 					object? baseValue = property.GetValue(baseObjects[idx]);
 					if (!CheckValuesMatch(property.PropertyType, thisValue, baseValue))
 					{
-						propNamesThatRequireUnique.Add(property.Name);
+						propNamesThatRequireUnique.Add(property.Name, (thisValue?.ToString(), baseValue?.ToString()));
 					}
 				}
 			}
@@ -2757,12 +2757,12 @@ namespace UnrealBuildTool
 			// if we didn't set it above, check the properties
 			if (BuildEnvironmentOverride.Value == TargetBuildEnvironment.UniqueIfNeeded)
 			{
-				List<string> propNames = new();
+				Dictionary<string, (string?, string?)> propNames = new();
 				string? baseTargetName;
 				if (RequiresUniqueEnvironment(rulesAssembly, arguments, propNames, out baseTargetName))
 				{
 					logger.LogInformation("Setting {Target}'s BuildEnvironment to Unique, because it had changed the values of the propertues [ {Props} ] away from the values specified in {BaseTarget}",
-						Name, String.Join(", ", propNames), baseTargetName);
+						Name, String.Join(", ", propNames.Select(x => $"{x.Key}: {x.Value.Item1} != {x.Value.Item2}")), baseTargetName);
 					BuildEnvironment = TargetBuildEnvironment.Unique;
 				}
 				else
diff --git a/Engine/Source/Programs/UnrealBuildTool/Configuration/UEBuildTarget.cs b/Engine/Source/Programs/UnrealBuildTool/Configuration/UEBuildTarget.cs
--- a/Engine/Source/Programs/UnrealBuildTool/Configuration/UEBuildTarget.cs
+++ b/Engine/Source/Programs/UnrealBuildTool/Configuration/UEBuildTarget.cs
@@ -1400,12 +1400,12 @@ namespace UnrealBuildTool
 		/// </summary>
 		static void ValidateSharedEnvironment(RulesAssembly RulesAssembly, string ThisTargetName, CommandLineArguments Arguments, TargetRules ThisRules, ILogger Logger)
 		{
-			List<string> PropNamesThatRequiredUnique = new();
+			Dictionary<string, (string?, string?)> PropNamesThatRequiredUnique = new();
 			string? BaseTargetName;
 			if (ThisRules.RequiresUniqueEnvironment(RulesAssembly, Arguments, PropNamesThatRequiredUnique, out BaseTargetName))
 			{
 				throw new BuildException("{0} modifies the values of properties: [ {1} ]. This is not allowed, as {0} has build products in common with {2}.\nRemove the modified setting, change {0} to use a unique build environment by setting 'BuildEnvironment = TargetBuildEnvironment.Unique;' in the {3} constructor, or set bOverrideBuildEnvironment = true to force this setting on.", 
-					ThisTargetName, String.Join(", ", PropNamesThatRequiredUnique), BaseTargetName, ThisRules.GetType().Name);
+					ThisTargetName, String.Join(", ", PropNamesThatRequiredUnique.Select(x => $"{x.Key}: {x.Value.Item1} != {x.Value.Item2}")), BaseTargetName, ThisRules.GetType().Name);
 			}
 
 			// Make sure that we don't explicitly enable or disable any plugins through the target rules. We can't do this with the shared build environment because it requires recompiling the "Projects" engine module.
