Add back new/delete operator global visibility for Android shipping builds

#android

diff --git a/Engine/Source/Programs/UnrealBuildTool/Platform/Android/AndroidToolChain.cs b/Engine/Source/Programs/UnrealBuildTool/Platform/Android/AndroidToolChain.cs
--- a/Engine/Source/Programs/UnrealBuildTool/Platform/Android/AndroidToolChain.cs
+++ b/Engine/Source/Programs/UnrealBuildTool/Platform/Android/AndroidToolChain.cs
@@ -757,8 +757,8 @@ namespace UnrealBuildTool
 				string VersionScriptFile = GetVersionScriptFilename(LinkEnvironment);
 				using (StreamWriter Writer = File.CreateText(VersionScriptFile))
 				{
-					// Make all symbols (except ones called from Java) hidden
-					Writer.WriteLine("{ global: Java_*; ANativeActivity_onCreate; JNI_OnLoad; local: *; };");
+					// Make all symbols hidden (except new/delete operators and ones called from Java)
+					Writer.WriteLine("{ global: _Znwm*; _Znam*; _ZdlPv*; _ZdaPv*; Java_*; ANativeActivity_onCreate; JNI_OnLoad; local: *; };");
 				}
 				Result += " -Wl,--version-script=\"" + VersionScriptFile + "\"";
 			}
