[UBA]
* Added -UBAActionsOutputFile=<file.yaml> that can be used to have UBT output a yaml file with actions that can then be used to run inside UbaCli.

Example
1. RunUbt.bat UnrealEditor win64 development -UBAActionsOutputFile=c:\temp\file.yaml
2. UbaCli.exe local c:\temp\file.yaml

diff --git a/Engine/Source/Programs/UnrealBuildTool/Executors/Experimental/UBAExecutor.cs b/Engine/Source/Programs/UnrealBuildTool/Executors/Experimental/UBAExecutor.cs
--- a/Engine/Source/Programs/UnrealBuildTool/Executors/Experimental/UBAExecutor.cs
+++ b/Engine/Source/Programs/UnrealBuildTool/Executors/Experimental/UBAExecutor.cs
@@ -235,6 +235,13 @@ namespace UnrealBuildTool
 		[XmlConfigFile(Category = "UnrealBuildAccelerator")]
 		[CommandLine("-UBAUseKnownInputs", Value = "true")]
 		public bool bUseKnownInputs { get; set; } = false;
+
+		/// <summary>
+		/// Write yaml file with all actions that are queued for build. This can be used to replay using "UbaCli.exe local file.yaml"
+		/// </summary>
+		[XmlConfigFile(Category = "UnrealBuildAccelerator")]
+		[CommandLine("-UBAActionsOutputFile")]
+		public string ActionsOutputFile { get; set; } = String.Empty;
 	}
 
 	class UBAExecutor : ParallelExecutor
@@ -310,6 +317,51 @@ namespace UnrealBuildTool
 				return true;
 			}
 
+			if (!String.IsNullOrEmpty(UBAConfig.ActionsOutputFile))
+			{
+				if (!UBAConfig.ActionsOutputFile.EndsWith(".yaml"))
+					_threadedLogger.LogError("UBA actions output file needs to have extension .yaml for UbaCli to understand it");
+				using var writer1 = new System.IO.StreamWriter(UBAConfig.ActionsOutputFile);
+				using var writer = new System.CodeDom.Compiler.IndentedTextWriter(writer1, "  ");
+				writer.Write("environment: ");
+				writer.WriteLine(Environment.GetEnvironmentVariable("PATH"));
+				writer.WriteLine("processes:");
+				writer.Indent++;
+				int index = 0;
+				foreach (var action in inputActions)
+				{
+					action.SortIndex = index++;
+					writer.WriteLine($"- id: {action.SortIndex}");
+					writer.Indent++;
+					writer.WriteLine($"app: {action.CommandPath}");
+					writer.WriteLine($"arg: {action.CommandArguments}");
+					writer.WriteLine($"dir: {action.WorkingDirectory}");
+					writer.WriteLine($"desc: {action.StatusDescription}");
+					if (action.Weight != 1.0f)
+						writer.WriteLine($"weight: {action.Weight}");
+					if (!action.bCanExecuteInUBA)
+						writer.WriteLine("detour: false");
+					else if (!action.bCanExecuteRemotely)
+						writer.WriteLine("remote: false");
+					if (action.PrerequisiteActions.Any())
+					{
+						writer.Write("dep: [");
+						bool isFirst = true;
+						foreach (var dep in action.PrerequisiteActions)
+						{
+							if (!isFirst)
+								writer.Write(", ");
+							isFirst = false;
+							writer.Write(dep.SortIndex.ToString());
+						}
+						writer.WriteLine("]");
+					}
+					writer.Indent--;
+					writer.WriteLineNoTabs(null);
+				}
+			}
+
+			string remoteDisabledReason = "";
 			if (inputActions.Count() < NumParallelProcesses && !UBAConfig.bForceBuildAllRemote)
 			{
 				UBAConfig.bDisableRemote = true;
