UnrealBuildTool: Write to separate intermediate directory when generating project files

#rnx

diff --git a/Engine/Source/Programs/UnrealBuildTool/ProjectFiles/ProjectFileGenerator.cs b/Engine/Source/Programs/UnrealBuildTool/ProjectFiles/ProjectFileGenerator.cs
--- a/Engine/Source/Programs/UnrealBuildTool/ProjectFiles/ProjectFileGenerator.cs
+++ b/Engine/Source/Programs/UnrealBuildTool/ProjectFiles/ProjectFileGenerator.cs
@@ -2159,9 +2159,10 @@ namespace UnrealBuildTool
 
 							// Create the target descriptor
 							TargetDescriptor TargetDesc = new TargetDescriptor(CurTarget.UnrealProjectFilePath, CurTarget.Name, IntellisensePlatform, UnrealTargetConfiguration.Development, DefaultArchitecture, new CommandLineArguments(NewArguments.ToArray()));
+							TargetDesc.IntermediateEnvironment = UnrealIntermediateEnvironment.GenerateProjectFiles;
 
 							// Create the target
-							UEBuildTarget Target = UEBuildTarget.Create(TargetDesc, false, false, bUsePrecompiled, Logger);
+							UEBuildTarget Target = UEBuildTarget.Create(TargetDesc, false, false, bUsePrecompiled, UnrealIntermediateEnvironment.GenerateProjectFiles, Logger);
 
 							AddTargetForIntellisense(Target, Logger);
 
diff --git a/Engine/Source/Programs/UnrealBuildTool/ProjectFiles/VisualStudioWorkspace/VSWorkspaceProjectFile.cs b/Engine/Source/Programs/UnrealBuildTool/ProjectFiles/VisualStudioWorkspace/VSWorkspaceProjectFile.cs
--- a/Engine/Source/Programs/UnrealBuildTool/ProjectFiles/VisualStudioWorkspace/VSWorkspaceProjectFile.cs
+++ b/Engine/Source/Programs/UnrealBuildTool/ProjectFiles/VisualStudioWorkspace/VSWorkspaceProjectFile.cs
@@ -90,13 +90,14 @@ namespace UnrealBuildTool
 
 						TargetDescriptor TargetDesc = new(ProjectTarget.UnrealProjectFilePath, ProjectTarget.Name,
 							Platform, Configuration, ProjectArchitectures, Arguments);
+						TargetDesc.IntermediateEnvironment = UnrealIntermediateEnvironment.GenerateProjectFiles;
 
 						try
 						{
 							FileReference OutputFile = FileReference.Combine(ProjectRootFolder, 
 								$"{ProjectTarget.TargetFilePath.GetFileNameWithoutAnyExtensions()}_{Configuration}_{Platform}.json");
 
-							UEBuildTarget BuildTarget = UEBuildTarget.Create(TargetDesc, false, false, false, Logger);
+							UEBuildTarget BuildTarget = UEBuildTarget.Create(TargetDesc, false, false, false, UnrealIntermediateEnvironment.GenerateProjectFiles, Logger);
 							FileToTarget.Add(OutputFile, (BuildTarget, bBuildByDefault));
 						}
 						catch (Exception Ex)
