[Switch] Add automatic --defragment-size calculation for Switch patches when --defragment is in use. Uses PatchDefragmentThreshold from [Packaging] in SwitchEngine.ini as the target patch fragmentation to hit. Tests --defragment-sizes of 0,32,64,128,256,512. Picks the largest one that generates fragmentation below PatchDefragmentationThreshold. Can be run using the UpdateSwitchPatchParameters stage of NightlyBuild, which will update the PatchDefragmentBlockSize setting in the [Packaging] section of SwitchEngine.ini. Also can be run when patch is generated if -AutoDefragment is supplied as an additional option when creating a patch.

#rnx

diff --git a/Engine/Source/Programs/UnrealBuildTool/System/ConfigCache.cs b/Engine/Source/Programs/UnrealBuildTool/System/ConfigCache.cs
--- a/Engine/Source/Programs/UnrealBuildTool/System/ConfigCache.cs
+++ b/Engine/Source/Programs/UnrealBuildTool/System/ConfigCache.cs
@@ -645,6 +645,17 @@ namespace UnrealBuildTool
 			return FileReference.Combine(ProjectDir, "Config", $"Default{ConfigType}.ini");
 		}
 
+		/// <summary>
+		/// Calculates the path to where the project's platform specific config of the type given (ie DefaultEngine.ini)
+		/// </summary>
+		/// <param name="ConfigType">Game, Engine, etc</param>
+		/// <param name="ProjectDir">Project directory, used to find Config/Default[Type].ini</param>
+		/// <param name="Platform">Platform name</param>
+		public static FileReference GetPlatformConfigFileReference(ConfigHierarchyType ConfigType, DirectoryReference ProjectDir, string Platform)
+		{
+			return FileReference.Combine(ProjectDir, "Platforms", Platform, "Config", $"{Platform}{ConfigType}.ini");
+		}
+
 		/// <summary>
 		/// Updates a section in a Default***.ini, and will write it out. If the file is not writable, p4 can attempt to check it out.
 		/// </summary>
@@ -659,16 +670,49 @@ namespace UnrealBuildTool
 		public static bool WriteSettingToDefaultConfig(ConfigHierarchyType ConfigType, DirectoryReference ProjectDir, ConfigDefaultUpdateType UpdateType, string Section, string Key, string Value, ILogger Logger)
 		{
 			FileReference DefaultConfigFile = GetDefaultConfigFileReference(ConfigType, ProjectDir);
+			return WriteSettingToConfigFile(DefaultConfigFile, UpdateType, Section, Key, Value, Logger);
+		}
+
+		/// <summary>
+		/// Updates a section in a Default***.ini, and will write it out. If the file is not writable, p4 can attempt to check it out.
+		/// </summary>
+		/// <param name="ConfigType">Game, Engine, etc</param>
+		/// <param name="ProjectDir">Project directory, used to find Config/Default[Type].ini</param>
+		/// <param name="Platform">Platform name. Must have a directory under projects Platforms subdirectory</param>
+		/// <param name="UpdateType">How to modify the secion</param>
+		/// <param name="Section">Name of the section with the Key in it</param>
+		/// <param name="Key">Key to update</param>
+		/// <param name="Value">Value to write for te Key</param>
+		/// <param name="Logger">Logger for output</param>
+		/// <returns></returns>		
+		public static bool WriteSettingToPlatformConfigFile(ConfigHierarchyType ConfigType, DirectoryReference ProjectDir, string Platform, ConfigDefaultUpdateType UpdateType, string Section, string Key, string Value, ILogger Logger)
+		{
+			FileReference PlatformConfigFile = GetPlatformConfigFileReference(ConfigType, ProjectDir, Platform);
+			return WriteSettingToConfigFile(PlatformConfigFile, UpdateType, Section, Key, Value, Logger);
+		}
+
+		/// <summary>
+		/// Updates a section in config file, and will write it out. If the file is not writable, p4 can attempt to check it out.
+		/// </summary>
+		/// <param name="ConfigFile">Config file (.ini) name</param>
+		/// <param name="UpdateType">How to modify the secion</param>
+		/// <param name="Section">Name of the section with the Key in it</param>
+		/// <param name="Key">Key to update</param>
+		/// <param name="Value">Value to write for te Key</param>
+		/// <param name="Logger">Logger for output</param>
+		/// <returns></returns>		
+		public static bool WriteSettingToConfigFile(FileReference ConfigFile, ConfigDefaultUpdateType UpdateType, string Section, string Key, string Value, ILogger Logger)
+		{
 
-			if (!FileReference.Exists(DefaultConfigFile))
+			if (!FileReference.Exists(ConfigFile))
 			{
-				Logger.LogWarning("Failed to find config file '{DefaultConfigFile}' to update", DefaultConfigFile);
+				Logger.LogWarning("Failed to find config file '{DefaultConfigFile}' to update", ConfigFile);
 				return false;
 			}
 
-			if (File.GetAttributes(DefaultConfigFile.FullName).HasFlag(FileAttributes.ReadOnly))
+			if (File.GetAttributes(ConfigFile.FullName).HasFlag(FileAttributes.ReadOnly))
 			{
-				Logger.LogWarning("Config file '{ConfigFile}' is read-only, unable to write setting {Key}", DefaultConfigFile.FullName, Key);
+				Logger.LogWarning("Config file '{ConfigFile}' is read-only, unable to write setting {Key}", ConfigFile.FullName, Key);
 				return false;
 			}
 
@@ -681,7 +725,7 @@ namespace UnrealBuildTool
 			LineToWrite += KeyWithEquals + Value;
 
 			// read in all the lines so we can insert or replace one
-			List<string> Lines = File.ReadAllLines(DefaultConfigFile.FullName).ToList();
+			List<string> Lines = File.ReadAllLines(ConfigFile.FullName).ToList();
 
 			// look for the section
 			int SectionIndex = -1;
@@ -700,7 +744,7 @@ namespace UnrealBuildTool
 				Lines.Add(SectionString);
 				Lines.Add(LineToWrite);
 
-				File.WriteAllLines(DefaultConfigFile.FullName, Lines);
+				File.WriteAllLines(ConfigFile.FullName, Lines);
 				return true;
 			}
 
@@ -761,7 +805,7 @@ namespace UnrealBuildTool
 			}
 
 			// now the lines are updated, we can overwrite the file
-			File.WriteAllLines(DefaultConfigFile.FullName, Lines);
+			File.WriteAllLines(ConfigFile.FullName, Lines);
 			return true;
 		}
 
