PR #11342: Use integrated terminal for VSCode

diff --git a/Engine/Source/Programs/UnrealBuildTool/ProjectFiles/VisualStudioCode/VSCodeProjectFileGenerator.cs b/Engine/Source/Programs/UnrealBuildTool/ProjectFiles/VisualStudioCode/VSCodeProjectFileGenerator.cs
--- a/Engine/Source/Programs/UnrealBuildTool/ProjectFiles/VisualStudioCode/VSCodeProjectFileGenerator.cs
+++ b/Engine/Source/Programs/UnrealBuildTool/ProjectFiles/VisualStudioCode/VSCodeProjectFileGenerator.cs
@@ -1334,7 +1334,7 @@ namespace UnrealBuildTool
 					if (HostPlatform == UnrealTargetPlatform.Win64)
 					{
 						OutFile.AddField("stopAtEntry", false);
-						OutFile.AddField("console", "externalTerminal");
+						OutFile.AddField("console", "integratedTerminal");
 
 						OutFile.AddField("type", "cppvsdbg");
 						OutFile.AddField("visualizerFile", MakeUnquotedPathString(FileReference.Combine(ProjectRoot, "Engine", "Extras", "VisualStudioDebugging", "Unreal.natvis"), EPathType.Absolute));
@@ -1412,7 +1412,7 @@ namespace UnrealBuildTool
 				OutFile.EndArray();
 				if (HostPlatform == UnrealTargetPlatform.Win64)
 				{
-					OutFile.AddField("console", "externalTerminal");
+					OutFile.AddField("console", "integratedTerminal");
 				}
 				else
 				{
