PR #10731: UnrealBuildTool: fix GenerateClangDatabase

diff --git a/Engine/Source/Programs/UnrealBuildTool/Modes/GenerateClangDatabase.cs b/Engine/Source/Programs/UnrealBuildTool/Modes/GenerateClangDatabase.cs
--- a/Engine/Source/Programs/UnrealBuildTool/Modes/GenerateClangDatabase.cs
+++ b/Engine/Source/Programs/UnrealBuildTool/Modes/GenerateClangDatabase.cs
@@ -262,7 +262,7 @@ namespace UnrealBuildTool
 										// Create the command
 										StringBuilder CommandBuilder = new StringBuilder();
 										string CommandArguments = Action.CommandArguments.Replace(".rsp", ".rsp.gcd").Replace(".response", ".response.gcd");
-										CommandBuilder.AppendFormat("\"{0}\" {1}", Action.CommandPath, CommandArguments);
+										CommandBuilder.AppendFormat("{0} {1}", Action.CommandPath, CommandArguments);
 
 										foreach (string ExtraArgument in GetExtraPlatformArguments(TargetToolChain))
 										{
@@ -333,14 +333,19 @@ namespace UnrealBuildTool
 					continue;
 				}
 				// The file that is going to compile
-				else if (!Line.StartsWith("-") && !Line.StartsWith("/"))
+				else if (!Line.StartsWith("-") && !Line.StartsWith("/") && !Line.StartsWith("@"))
 				{
 					Line = ConvertPath(Line, Line);
 				}
 				// Arguments
 				else
 				{
+					 // TODO this needs a clean way to handle path separators that is used in all instances
 					int StrIndex = Line.IndexOf("..\\");
+					if (StrIndex == -1)
+					{
+						StrIndex = Line.IndexOf("../");
+					}
 					if (StrIndex != -1)
 					{
 						Line = ConvertPath(Line, Line.Substring(StrIndex));
@@ -354,6 +359,11 @@ namespace UnrealBuildTool
 						}
 					}
 				}
+				// TODO DIRTY DIRTY HACK because we don't support recursion in response files
+				if (Line.StartsWith("@"))
+				{
+					Line = Line.Replace(".rsp", ".rsp.gcd").Replace(".response", ".response.gcd");
+				}
 
 				NewFileContents[i] = Line;
 			}
diff --git a/Engine/Source/Programs/UnrealBuildTool/Platform/Linux/LinuxToolChain.cs b/Engine/Source/Programs/UnrealBuildTool/Platform/Linux/LinuxToolChain.cs
--- a/Engine/Source/Programs/UnrealBuildTool/Platform/Linux/LinuxToolChain.cs
+++ b/Engine/Source/Programs/UnrealBuildTool/Platform/Linux/LinuxToolChain.cs
@@ -1701,7 +1701,8 @@ namespace UnrealBuildTool
 			{
 				// starting with clang 16.x the directory naming changed to include major version only
 				string ClangVersionString = (Info.ClangVersion.Major >= 16) ? Info.ClangVersion.Major.ToString() : Info.ClangVersion.ToString();
-				ExtraArguments.Add(String.Format("-isystem\"{0}\"", System.IO.Path.Combine(InternalSdkPath, "lib/clang/" + ClangVersionString + "/include/").Replace("\\", "/")));
+				ExtraArguments.Add(String.Format("-isystem {0}", System.IO.Path.Combine(InternalSdkPath, "lib", "clang", ClangVersionString, "include").Replace("\\", "/")));
+				ExtraArguments.Add(String.Format("-isystem {0}", System.IO.Path.Combine(InternalSdkPath, "usr", "include").Replace("\\", "/")));
 			}
 		}
 	}
