Clang fixes
- Getting the UnrealVS "Preprocess single file" working with clang compiles.
- Making the "Preprocess single file" result filename match the source file's name to match VSSpecificFileAction.
- Making sure .h files are compiled correctly when using things like -IncludeHeaders

diff --git a/Engine/Source/Programs/UnrealBuildTool/Modes/BuildMode.cs b/Engine/Source/Programs/UnrealBuildTool/Modes/BuildMode.cs
--- a/Engine/Source/Programs/UnrealBuildTool/Modes/BuildMode.cs
+++ b/Engine/Source/Programs/UnrealBuildTool/Modes/BuildMode.cs
@@ -806,18 +806,33 @@ namespace UnrealBuildTool
 						}
 
 						// these are parsed by external tools wishing to open this file directly
-						foreach (LinkedAction BuildAction in MergedActionsToExecute.Where(BuildAction => BuildAction.ActionType == ActionType.Compile && (BuildAction.Inner as VCCompileAction)?.bIsAnalyzing != true))
+						foreach (LinkedAction BuildAction in MergedActionsToExecute.Where(BuildAction => BuildAction.ActionType == ActionType.Compile))
 						{
-							FileItem? PreprocessedFile = BuildAction.ProducedItems.FirstOrDefault(x => x.HasExtension(".i"));
-							if (PreprocessedFile != null)
+							if (BuildAction.Inner is VCCompileAction)
 							{
-								Logger.LogInformation("PreProcessPath: {File}", PreprocessedFile);
-							}
+								if ((BuildAction.Inner as VCCompileAction)?.bIsAnalyzing != true)
+								{
+									FileItem? PreprocessedFile = BuildAction.ProducedItems.FirstOrDefault(x => x.HasExtension(".i"));
+									if (PreprocessedFile != null)
+									{
+										Logger.LogInformation("PreProcessPath: {File}", PreprocessedFile);
+									}
 
-							FileItem? AssemblyPath = BuildAction.ProducedItems.FirstOrDefault(x => x.HasExtension(".asm"));
-							if (AssemblyPath != null)
+									FileItem? AssemblyPath = BuildAction.ProducedItems.FirstOrDefault(x => x.HasExtension(".asm"));
+									if (AssemblyPath != null)
+									{
+										Logger.LogInformation("AssemblyPath: {File}", AssemblyPath);
+									}
+								}
+							}
+							else
 							{
-								Logger.LogInformation("AssemblyPath: {File}", AssemblyPath);
+								FileItem? PreprocessedFile = BuildAction.ProducedItems.FirstOrDefault(x => x.HasExtension(".n"));
+								if (PreprocessedFile != null)
+								{
+									string PreprocessedFilePath = PreprocessedFile.ToString();
+									Logger.LogInformation($"PreProcessPath: {PreprocessedFilePath.Remove(PreprocessedFilePath.Length - 2)}.i");
+								}
 							}
 						}
 					}
diff --git a/Engine/Source/Programs/UnrealBuildTool/ToolChain/ClangSpecificFileAction.cs b/Engine/Source/Programs/UnrealBuildTool/ToolChain/ClangSpecificFileAction.cs
--- a/Engine/Source/Programs/UnrealBuildTool/ToolChain/ClangSpecificFileAction.cs
+++ b/Engine/Source/Programs/UnrealBuildTool/ToolChain/ClangSpecificFileAction.cs
@@ -17,9 +17,14 @@ namespace UnrealBuildTool
 		DirectoryReference SourceDir;
 		DirectoryReference OutputDir;
 		IEnumerable<string> RspLines;
-		int SingleFileCounter;
 
-		internal ClangSpecificFileAction(DirectoryReference Source, DirectoryReference Output, Action Action, IEnumerable<string> ContentLines) : base(Action)
+		// TODO: Entire CppCompileEnvironment needs to be saved with BinaryArchiveWriter, some options may still be unused
+		bool bPreprocessOnly;
+		bool bWithAssembly;
+
+		Dictionary<string, List<FileItem>> SingleFiles = new();
+
+		internal ClangSpecificFileAction(DirectoryReference Source, DirectoryReference Output, Action Action, CppCompileEnvironment? CompileEnvironment, IEnumerable<string> ContentLines) : base(Action)
 		{
 			ProducedItems.Clear();
 			DependencyListFile = null;
@@ -27,6 +32,8 @@ namespace UnrealBuildTool
 			SourceDir = Source;
 			OutputDir = Output;
 			RspLines = ContentLines;
+			bPreprocessOnly = CompileEnvironment?.bPreprocessOnly ?? false;
+			bWithAssembly = CompileEnvironment?.bWithAssembly ?? false;
 		}
 
 		public ClangSpecificFileAction(BinaryArchiveReader Reader) : base(Reader)
@@ -34,6 +41,8 @@ namespace UnrealBuildTool
 			SourceDir = Reader.ReadCompactDirectoryReference();
 			OutputDir = Reader.ReadCompactDirectoryReference();
 			RspLines = Reader.ReadList(() => Reader.ReadString())!;
+			bPreprocessOnly = Reader.ReadBool();
+			bWithAssembly = Reader.ReadBool();
 		}
 
 		public new void Write(BinaryArchiveWriter Writer)
@@ -42,26 +51,42 @@ namespace UnrealBuildTool
 			Writer.WriteCompactDirectoryReference(SourceDir);
 			Writer.WriteCompactDirectoryReference(OutputDir);
 			Writer.WriteList(RspLines.ToList(), (Str) => Writer.WriteString(Str));
+			Writer.WriteBool(bPreprocessOnly);
+			Writer.WriteBool(bWithAssembly);
 		}
 
 		public DirectoryReference RootDirectory => SourceDir;
 
 		public IExternalAction? CreateAction(FileItem SourceFile, ILogger Logger)
 		{
+			// Keep track of all specific files, so the output file can be renamed if there's a naming conflict
+			string Filename = SourceFile.Name;
+			if (!SingleFiles.ContainsKey(SourceFile.Name))
+			{
+				SingleFiles[SourceFile.Name] = new();
+			}
+			else
+			{
+				Filename = $"{Path.GetFileNameWithoutExtension(SourceFile.Name)}{SingleFiles[Filename].Count}{Path.GetExtension(SourceFile.Name)}";
+			}
+			SingleFiles[SourceFile.Name].Add(SourceFile);
+
 			string DummyName = "SingleFile.cpp";
-			string UniqueDummyName = $"SingleFile{SingleFileCounter}.cpp";
-			++SingleFileCounter;
 
 			int FileNameIndex = CommandArguments.IndexOf(DummyName);
 			string DummyPath = CommandArguments.Substring(2, FileNameIndex + DummyName.Length - 2);
 
 			if (SourceFile.HasExtension(".h"))
 			{
-				FileItem DummyFile = FileItem.GetItemByFileReference(FileReference.Combine(OutputDir, "SingleFile", SourceFile.Name));
+				FileItem DummyFile = FileItem.GetItemByFileReference(FileReference.Combine(OutputDir, "SingleFile", $"{Filename}.cpp"));
 				Directory.CreateDirectory(DummyFile.Directory.FullName);
 				File.WriteAllText(DummyFile.FullName, $"#include \"{SourceFile.FullName.Replace('\\', '/')}\"");
 				SourceFile = DummyFile;
 			}
+			else if (!SourceFile.HasExtension(".cpp"))
+			{
+				return null;
+			}
 
 			List<string> NewRspLines = new();
 			foreach (string L in RspLines)
@@ -73,7 +98,7 @@ namespace UnrealBuildTool
 					Line.Contains(".cpp.json", System.StringComparison.Ordinal) ||
 					Line.Contains(".cpp.o", System.StringComparison.Ordinal))
 				{
-					Line = Line.Replace("SingleFile.cpp", UniqueDummyName);
+					Line = Line.Replace("SingleFile.cpp", Filename);
 				}
 				else
 				{
@@ -83,15 +108,32 @@ namespace UnrealBuildTool
 			}
 
 			Action Action = new Action(this);
-			Action.CommandArguments = CommandArguments.Replace(DummyName, UniqueDummyName);
+			Action.CommandArguments = CommandArguments.Replace(DummyName, Filename);
 			Action.DependencyListFile = null;
-			Action.StatusDescription = SourceFile.Name;
+			Action.StatusDescription = Filename;
+
+			if (bPreprocessOnly)
+			{
+				// We have to add a produced item so this action is not skipped.
+				// Note we on purpose use a different extension than what the compiler produce because otherwise up-to-date checker might see it as up-to-date
+				// even though we want it to always be built
+				FileItem ProducedItem = FileItem.GetItemByFileReference(FileReference.Combine(OutputDir, $"{Filename}.n"));
+				Action.ProducedItems.Add(ProducedItem);
+			}
+			else
+			{
+				// We have to add a produced item so this action is not skipped.
+				// Note we on purpose use a different extension than what the compiler produce because otherwise up-to-date checker might see it as up-to-date
+				// even though we want it to always be built
+				FileItem ProducedItem = FileItem.GetItemByFileReference(FileReference.Combine(OutputDir, $"{Filename}.o"));
+				Action.ProducedItems.Add(ProducedItem);
 
-			// We have to add a produced item so this action is not skipped.
-			// Note we on purpose use a different extension than what the compiler produce because otherwise up-to-date checker might see it as up-to-date
-			// even though we want it to always be built
-			FileItem ProducedItem = FileItem.GetItemByFileReference(FileReference.Combine(OutputDir, UniqueDummyName + ".n"));
-			Action.ProducedItems.Add(ProducedItem);
+				if (bWithAssembly)
+				{
+					FileItem AssemblyFile = FileItem.GetItemByFileReference(FileReference.Combine(OutputDir, $"{Filename}.asm"));
+					Action.ProducedItems.Add(AssemblyFile);
+				}
+			}
 
 			FileItem ResponseFile = FileItem.GetItemByPath(Action.CommandArguments.Substring(1).Trim('"'));
 			File.WriteAllLines(ResponseFile.FullName, NewRspLines);
diff --git a/Engine/Source/Programs/UnrealBuildTool/ToolChain/ClangToolChain.cs b/Engine/Source/Programs/UnrealBuildTool/ToolChain/ClangToolChain.cs
--- a/Engine/Source/Programs/UnrealBuildTool/ToolChain/ClangToolChain.cs
+++ b/Engine/Source/Programs/UnrealBuildTool/ToolChain/ClangToolChain.cs
@@ -444,6 +444,12 @@ namespace UnrealBuildTool
 			GetCppStandardCompileArgument(CompileEnvironment, Arguments);
 		}
 
+		protected virtual void GetCompileArguments_H(CppCompileEnvironment CompileEnvironment, List<string> Arguments)
+		{
+			Arguments.Add("-x c++-header");
+			GetCppStandardCompileArgument(CompileEnvironment, Arguments);
+		}
+
 		// Conditionally enable (default disabled) generation of information about every class with virtual functions for use by the C++ runtime type identification features
 		// (`dynamic_cast' and `typeid'). If you don't use those parts of the language, you can save some space by using -fno-rtti.
 		// Note that exception handling uses the same information, but it will generate it as needed.
@@ -892,6 +898,10 @@ namespace UnrealBuildTool
 				// Compile the file as Objective-C code.
 				GetCompileArguments_M(CompileEnvironment, Arguments);
 			}
+			else if (Extension == ".H")
+			{
+				GetCompileArguments_H(CompileEnvironment, Arguments);
+			}
 			else
 			{
 				// Compile the file as C++ code.
@@ -1082,7 +1092,7 @@ namespace UnrealBuildTool
 			Action.bCanExecuteRemotelyWithSNDBS = Action.bCanExecuteRemotely && !CompileEnvironment.bBuildLocallyWithSNDBS;
 			Action.Weight = CompileActionWeight;
 
-			Graph.AddAction(new ClangSpecificFileAction(SourceDir, OutputDir, Action, GraphBuilder.ContentLines));
+			Graph.AddAction(new ClangSpecificFileAction(SourceDir, OutputDir, Action, CompileEnvironment, GraphBuilder.ContentLines));
 		}
 
 		protected virtual Action CompileCPPFile(CppCompileEnvironment CompileEnvironment, FileItem SourceFile, DirectoryReference OutputDir, string ModuleName, IActionGraphBuilder Graph, IReadOnlyCollection<string> GlobalArguments, CPPOutput Result)
