- Redid how some of the "non-supported platform SDKs" are managed to avoid edge cases.
- Disabled VisionOS on Windows to reduce confusion that it might work for launch on

diff --git a/Engine/Platforms/VisionOS/Config/DataDrivenPlatformInfo.ini b/Engine/Platforms/VisionOS/Config/DataDrivenPlatformInfo.ini
--- a/Engine/Platforms/VisionOS/Config/DataDrivenPlatformInfo.ini
+++ b/Engine/Platforms/VisionOS/Config/DataDrivenPlatformInfo.ini
@@ -9,7 +9,7 @@ LargeIconPath=/Platforms/VisionOS/Content/Editor/Slate/Launcher/Platform_VisionO
 XLargeIconPath=
 AutoSDKPath=
 TutorialPath=SharingAndReleasing/Mobile/iOS
-Windows:bIsEnabled=true
+Windows:bIsEnabled=false
 Mac:bIsEnabled=true
 Linux:bIsEnabled=false
 bUsesHostCompiler=false
diff --git a/Engine/Source/Programs/Shared/EpicGames.Build/UEBuildPlatformSDK.cs b/Engine/Source/Programs/Shared/EpicGames.Build/UEBuildPlatformSDK.cs
--- a/Engine/Source/Programs/Shared/EpicGames.Build/UEBuildPlatformSDK.cs
+++ b/Engine/Source/Programs/Shared/EpicGames.Build/UEBuildPlatformSDK.cs
@@ -325,7 +325,14 @@ namespace EpicGames.Core
 			bIsSdkAllowedOnHost = bInIsSdkAllowedOnHost;
 
 			// load the SDK config file
-			LoadJsonFile(PlatformName);
+			if (bIsSdkAllowedOnHost)
+			{
+				LoadJsonFile(PlatformName);
+			}
+			else
+			{
+				Console.WriteLine("SKIPPING JSON FOR {0}", PlatformName);
+			}
 
 			// if the parent set up autosdk, the env vars will be wrong, but we can still get the manual SDK version from before it was setup
 			string? ParentManualSDKVersions = Environment.GetEnvironmentVariable(GetPlatformManualSDKSetupEnvVar());
@@ -387,7 +394,9 @@ namespace EpicGames.Core
 					object[] parameter = new object[1];
 					parameter[0] = Log.Logger;
 					SDK = (UEBuildPlatformSDK)Activator.CreateInstance(typeof(T), parameter)!;
-					SDK.bIsSdkAllowedOnHost = true;
+					// by setting this to false, we don't require any of the RequiredVersions to exist, but if they do, they will be read
+					// this is useful on other platforms that 
+					SDK.bIsSdkAllowedOnHost = false;
 					SDK.LoadJsonFile(PlatformName);
 					TempSDKRegistry.Add(PlatformName, SDK);
 				}
diff --git a/Engine/Source/Programs/UnrealBuildTool/Platform/Windows/MicrosoftPlatformSDK.Versions.cs b/Engine/Source/Programs/UnrealBuildTool/Platform/Windows/MicrosoftPlatformSDK.Versions.cs
--- a/Engine/Source/Programs/UnrealBuildTool/Platform/Windows/MicrosoftPlatformSDK.Versions.cs
+++ b/Engine/Source/Programs/UnrealBuildTool/Platform/Windows/MicrosoftPlatformSDK.Versions.cs
@@ -19,46 +19,46 @@ namespace UnrealBuildTool
 		/// <summary>
 		/// The default compiler version to be used, if installed. 
 		/// </summary>
-		static readonly VersionNumberRange[] PreferredClangVersions = SDK.GetVersionNumberRangeArrayFromConfig("PreferredClangVersions");
+		static VersionNumberRange[] PreferredClangVersions = SDK.GetVersionNumberRangeArrayFromConfig("PreferredClangVersions");
 
 		/// <summary>
 		/// The minimum supported Clang compiler
 		/// </summary>
-		static readonly VersionNumber MinimumClangVersion = SDK.GetRequiredVersionNumberFromConfig("MinimumClangVersion");
+		static VersionNumber MinimumClangVersion => SDK.GetRequiredVersionNumberFromConfig("MinimumClangVersion");
 
 		/// <summary>
 		/// Ranges of tested compiler toolchains to be used, in order of preference. If multiple toolchains in a range are present, the latest version will be preferred.
 		/// Note that the numbers here correspond to the installation *folders* rather than precise executable versions.
 		/// </summary>
 		/// <seealso href="https://learn.microsoft.com/en-us/lifecycle/products/visual-studio-2022"/>
-		static readonly VersionNumberRange[] PreferredVisualCppVersions = SDK.GetVersionNumberRangeArrayFromConfig("PreferredVisualCppVersions");
+		static VersionNumberRange[] PreferredVisualCppVersions => SDK.GetVersionNumberRangeArrayFromConfig("PreferredVisualCppVersions");
 
 		/// <summary>
 		/// Minimum Clang version required for MSVC toolchain versions
 		/// </summary>
-		static readonly Tuple<VersionNumber, VersionNumber>[] MinimumRequiredClangVersion = SDK.GetVersionNumberRangeArrayFromConfig("MinimumRequiredClangVersion").
+		static Tuple<VersionNumber, VersionNumber>[] MinimumRequiredClangVersion => SDK.GetVersionNumberRangeArrayFromConfig("MinimumRequiredClangVersion").
 			Select(x => new Tuple<VersionNumber, VersionNumber>(x.Min, x.Max)).ToArray();
 
 		/// <summary>
 		/// Tested compiler toolchains that should not be allowed.
 		/// </summary>
-		static readonly VersionNumberRange[] BannedVisualCppVersions = SDK.GetVersionNumberRangeArrayFromConfig("BannedVisualCppVersions");
+		static VersionNumberRange[] BannedVisualCppVersions => SDK.GetVersionNumberRangeArrayFromConfig("BannedVisualCppVersions");
 
 		/// <summary>
 		/// The minimum supported MSVC compiler
 		/// </summary>
-		static readonly VersionNumber MinimumVisualCppVersion = SDK.GetRequiredVersionNumberFromConfig("MinimumVisualCppVersion");
+		static VersionNumber MinimumVisualCppVersion => SDK.GetRequiredVersionNumberFromConfig("MinimumVisualCppVersion");
 
 		/// <summary>
 		/// The default compiler version to be used, if installed. 
 		/// https://www.intel.com/content/www/us/en/developer/articles/tool/oneapi-standalone-components.html#dpcpp-cpp
 		/// </summary>
-		static readonly VersionNumberRange[] PreferredIntelOneApiVersions = SDK.GetVersionNumberRangeArrayFromConfig("PreferredIntelOneApiVersions");
+		static VersionNumberRange[] PreferredIntelOneApiVersions => SDK.GetVersionNumberRangeArrayFromConfig("PreferredIntelOneApiVersions");
 
 		/// <summary>
 		/// The minimum supported Intel compiler
 		/// </summary>
-		static readonly VersionNumber MinimumIntelOneApiVersion = SDK.GetRequiredVersionNumberFromConfig("MinimumIntelOneApiVersion");
+		static VersionNumber MinimumIntelOneApiVersion => SDK.GetRequiredVersionNumberFromConfig("MinimumIntelOneApiVersion");
 
 		/// <summary>
 		/// If a toolchain version is a preferred version
