- Fixed Unique BuildEnvironment editors to work with Modern Xcode correctly

#changelist validated
#virtualized

diff --git a/Engine/Source/Programs/UnrealBuildTool/ProjectFiles/Xcode/XcodeProject.cs b/Engine/Source/Programs/UnrealBuildTool/ProjectFiles/Xcode/XcodeProject.cs
--- a/Engine/Source/Programs/UnrealBuildTool/ProjectFiles/Xcode/XcodeProject.cs
+++ b/Engine/Source/Programs/UnrealBuildTool/ProjectFiles/Xcode/XcodeProject.cs
@@ -1727,7 +1727,14 @@ namespace UnrealBuildTool.XcodeProjectXcconfig
 			DirectoryReference ProjectOrEngineDir = UnrealData.UProjectFileLocation?.Directory ?? Unreal.EngineDirectory;
 
 			// point to the shader Engine/Binaries for content only project (sadly, the TargetRules.OutputFile is not filled out)
-			DirectoryReference ConfigBuildDir = (TargetRules.Type == TargetType.Editor || UnrealData.bIsContentOnlyProject) ? Unreal.EngineDirectory : ProjectOrEngineDir;
+			bool bUseEngineDirectory = TargetRules.Type == TargetType.Editor || UnrealData.bIsContentOnlyProject;
+			if (TargetRules.BuildEnvironment == TargetBuildEnvironment.Unique)
+			{
+				// unique build environment editors will put their binaries in the project
+				bUseEngineDirectory = false;
+			}
+
+			DirectoryReference ConfigBuildDir = bUseEngineDirectory ? Unreal.EngineDirectory : ProjectOrEngineDir;
 			if (TargetRules.Type == TargetType.Program && TargetRules.File!.IsUnderDirectory(Unreal.EngineDirectory))
 			{
 				ConfigBuildDir = BuildConfig.RootDirectory;
