Bundle llvm-symbolizer into ASAN/TSAN/etc internal builds.

diff --git a/Engine/Source/Programs/UnrealBuildTool/Platform/Linux/LinuxToolChain.cs b/Engine/Source/Programs/UnrealBuildTool/Platform/Linux/LinuxToolChain.cs
--- a/Engine/Source/Programs/UnrealBuildTool/Platform/Linux/LinuxToolChain.cs
+++ b/Engine/Source/Programs/UnrealBuildTool/Platform/Linux/LinuxToolChain.cs
@@ -643,6 +643,40 @@ namespace UnrealBuildTool
 				: String.Format("{0}={1}", Key, Value);
 		}
 
+		public override void PrepareRuntimeDependencies(List<RuntimeDependency> RuntimeDependencies, Dictionary<FileReference, FileReference> TargetFileToSourceFile, DirectoryReference ExeDir)
+		{
+			// If ASan is enabled we need to copy the companion helper libraries from the MSVC tools bin folder to the
+			// target executable folder.
+			if (Options.HasFlag(ClangToolChainOptions.EnableAddressSanitizer) ||
+				Options.HasFlag(ClangToolChainOptions.EnableThreadSanitizer) ||
+				Options.HasFlag(ClangToolChainOptions.EnableUndefinedBehaviorSanitizer) ||
+				Options.HasFlag(ClangToolChainOptions.EnableMemorySanitizer) ||
+				Options.HasFlag(ClangToolChainOptions.EnableLibFuzzer))
+			{
+				bool bInternalBuild = false;
+				BuildVersion? Version;
+				if (BuildVersion.TryRead(BuildVersion.GetDefaultFileName(), out Version))
+				{
+					bInternalBuild = !Version.IsLicenseeVersion;
+				}
+
+				if (bInternalBuild)
+				{
+					string? InternalSdkPath = UEBuildPlatform.GetSDK(UnrealTargetPlatform.Linux)!.GetInternalSDKPath();
+					if (InternalSdkPath != null)
+					{
+						DirectoryReference InternalSdkPathRef = new DirectoryReference(InternalSdkPath);
+
+						FileReference SymbolizerSourcePath = FileReference.Combine(InternalSdkPathRef, "bin/llvm-symbolizer");
+						FileReference SymbolizerTargetPath = FileReference.Combine(ExeDir, "llvm-symbolizer");
+
+						RuntimeDependencies.Add(new RuntimeDependency(SymbolizerSourcePath, StagedFileType.NonUFS));
+						TargetFileToSourceFile[SymbolizerSourcePath] = SymbolizerTargetPath;
+					}
+				}
+			}
+		}
+
 		protected virtual void GetLinkArguments(LinkEnvironment LinkEnvironment, List<string> Arguments)
 		{
 			// always select the driver g++ in-case we are using a different binary for clang, such as clang/clang-cl
