PR #11428: Create separate arm64 and x64 target platforms for use with AGDE

#android

diff --git a/Engine/Source/Programs/UnrealBuildTool/Platform/Android/AndroidProjectGenerator.cs b/Engine/Source/Programs/UnrealBuildTool/Platform/Android/AndroidProjectGenerator.cs
--- a/Engine/Source/Programs/UnrealBuildTool/Platform/Android/AndroidProjectGenerator.cs
+++ b/Engine/Source/Programs/UnrealBuildTool/Platform/Android/AndroidProjectGenerator.cs
@@ -87,7 +87,8 @@ namespace UnrealBuildTool
 
 			if (InVSSettings.Platform == UnrealTargetPlatform.Android && AGDEInstalled)
 			{
-				PlatformName = "Android-arm64-v8a";
+				string longAbi = GetLongAbi(InVSSettings);
+				PlatformName = $"Android-{longAbi}";
 			}
 
 			return PlatformName;
@@ -102,22 +103,65 @@ namespace UnrealBuildTool
 			}
 		}
 
+		private string GetShortAbi(VSSettings InVSSettings)
+		{
+			if (InVSSettings.Architecture == null)
+			{
+				throw new BuildException("Architecture cannot be null");
+			}
+			else if (InVSSettings.Architecture == UnrealArch.Arm64)
+			{
+				return "arm64";
+			}
+			else if (InVSSettings.Architecture == UnrealArch.X64)
+			{
+				return "x64";
+			}
+			else
+			{
+				throw new BuildException($"Unexpected architecture: {InVSSettings.Architecture}");
+			}
+		}
+
+		private string GetLongAbi(VSSettings InVSSettings)
+		{
+			if (InVSSettings.Architecture == null)
+			{
+				throw new BuildException("Architecture cannot be null");
+			}
+			else if (InVSSettings.Architecture == UnrealArch.Arm64)
+			{
+				return "arm64-v8a";
+			}
+			else if (InVSSettings.Architecture == UnrealArch.X64)
+			{
+				return "x86_64";
+			}
+			else
+			{
+				throw new BuildException($"Unexpected architecture: {InVSSettings.Architecture}");
+			}
+		}
+
 		/// <inheritdoc/>
 		public override void GetVisualStudioPathsEntries(VSSettings InVSSettings, TargetType TargetType, FileReference TargetRulesPath, FileReference ProjectFilePath, FileReference NMakeOutputPath, StringBuilder ProjectFileBuilder)
 		{
 			if (AGDEInstalled)
 			{
+				string shortAbi = GetShortAbi(InVSSettings);
+				string longAbi = GetLongAbi(InVSSettings);
+
 				string apkLocation = Path.Combine(
 					Path.GetDirectoryName(NMakeOutputPath.FullName)!,
-					Path.GetFileNameWithoutExtension(NMakeOutputPath.FullName) + "-arm64.apk");
+					Path.GetFileNameWithoutExtension(NMakeOutputPath.FullName) + $"-{shortAbi}.apk");
 
 				ProjectFileBuilder.AppendLine($"    <AndroidApkLocation>{apkLocation}</AndroidApkLocation>");
 				string intermediateRootPath = Path.GetFullPath(Path.GetDirectoryName(NMakeOutputPath.FullName) + @"\..\..\Intermediate\Android\");
 				List<string> symbolLocations = new List<string>
 				{
-					Path.Combine(intermediateRootPath, "arm64", "jni", "arm64-v8a"),
-					Path.Combine(intermediateRootPath, "arm64", "libs", "arm64-v8a"),
-					Path.Combine(intermediateRootPath, "LLDBSymbolsLibs", "arm64") // support bDontBundleLibrariesInAPK
+					Path.Combine(intermediateRootPath, shortAbi, "jni", longAbi),
+					Path.Combine(intermediateRootPath, shortAbi, "libs", longAbi),
+					Path.Combine(intermediateRootPath, "LLDBSymbolsLibs", shortAbi) // support bDontBundleLibrariesInAPK
 				};
 				ProjectFileBuilder.AppendLine($"    <AndroidSymbolDirectories>{string.Join(";", symbolLocations)}</AndroidSymbolDirectories>");
 			}
@@ -129,8 +173,15 @@ namespace UnrealBuildTool
 
 		public override string GetExtraBuildArguments(VSSettings InVSSettings)
 		{
-			// do not need to check InPlatform since it will always be UnrealTargetPlatform.Android
-			return (AGDEInstalled ? " -Architectures=arm64 -ForceAPKGeneration" : "") + base.GetExtraBuildArguments(InVSSettings);
+			if (AGDEInstalled)
+			{
+				// do not need to check InPlatform since it will always be UnrealTargetPlatform.Android
+				return $" -ForceAPKGeneration" + base.GetExtraBuildArguments(InVSSettings);
+			}
+			else
+			{
+				return base.GetExtraBuildArguments(InVSSettings);
+			}
 		}
 
 		public override string GetVisualStudioUserFileStrings(VisualStudioUserFileSettings VCUserFileSettings, VSSettings InVSSettings, string InConditionString, TargetRules InTargetRules, FileReference TargetRulesPath, FileReference ProjectFilePath, FileReference? NMakeOutputPath, string ProjectName, string? ForeignUProjectPath)
@@ -152,9 +203,10 @@ namespace UnrealBuildTool
 					Out.AppendLine("    <AndroidDebugTarget></AndroidDebugTarget>");
 
 					// At this stage we don't know if bDontBundleLibrariesInAPK is enabled or not, so make a fail-safe check.
+					string shortAbi = GetShortAbi(InVSSettings);
 					string PushSOScript = Path.Combine(
 						Path.GetDirectoryName(NMakeOutputPath.FullName)!,
-						"Push_" + Path.GetFileNameWithoutExtension(NMakeOutputPath.FullName) + "-arm64_so.bat");
+						"Push_" + Path.GetFileNameWithoutExtension(NMakeOutputPath.FullName) + $"-{shortAbi}_so.bat");
 
 					// AGDE specifies current debug target in AndroidDebugTarget property in a form of "model:serial:arch".
 					// AndroidDebugTarget is a special property and needs to be evaluated in-line. And the push script needs the device serial as first argument to push to the correct device. 
diff --git a/Engine/Source/Programs/UnrealBuildTool/ProjectFiles/VisualStudio/VCProject.cs b/Engine/Source/Programs/UnrealBuildTool/ProjectFiles/VisualStudio/VCProject.cs
--- a/Engine/Source/Programs/UnrealBuildTool/ProjectFiles/VisualStudio/VCProject.cs
+++ b/Engine/Source/Programs/UnrealBuildTool/ProjectFiles/VisualStudio/VCProject.cs
@@ -2299,6 +2299,10 @@ namespace UnrealBuildTool
 					string UBTPlatformName = Platform.ToString();
 					string UBTConfigurationName = Configuration.ToString();
 					VSSettings VSSettings = new(Platform, Configuration, ProjectFileFormat, null);
+					if (Platform == UnrealTargetPlatform.Android && Combination.Architecture != null)
+					{
+						VSSettings.Architecture = Combination.Architecture;
+					}
 
 					// Setup output path
 					UEBuildPlatform BuildPlatform = UEBuildPlatform.GetBuildPlatform(Platform);
diff --git a/Engine/Source/Programs/UnrealBuildTool/ProjectFiles/VisualStudio/VCProjectFileGenerator.cs b/Engine/Source/Programs/UnrealBuildTool/ProjectFiles/VisualStudio/VCProjectFileGenerator.cs
--- a/Engine/Source/Programs/UnrealBuildTool/ProjectFiles/VisualStudio/VCProjectFileGenerator.cs
+++ b/Engine/Source/Programs/UnrealBuildTool/ProjectFiles/VisualStudio/VCProjectFileGenerator.cs
@@ -261,7 +261,8 @@ namespace UnrealBuildTool
 		/// </summary>
 		public static UnrealArchitectures? GetPlatformArchitecturesToGenerate(UEBuildPlatform BuildPlatform)
 		{
-			return BuildPlatform.ArchitectureConfig.Mode == UnrealArchitectureMode.OneTargetPerArchitecture ?
+			return (BuildPlatform.ArchitectureConfig.Mode == UnrealArchitectureMode.OneTargetPerArchitecture ||
+				BuildPlatform.ArchitectureConfig.Mode == UnrealArchitectureMode.SingleTargetLinkSeparately) ?
 				BuildPlatform.ArchitectureConfig.AllSupportedArchitectures : null;
 		}
 
