UnrealBuildTool: Correctly add project source dir to include path when building a generic target (UnrealEditor, UnrealGame) with a -Project

diff --git a/Engine/Source/Programs/UnrealBuildTool/Configuration/UEBuildTarget.cs b/Engine/Source/Programs/UnrealBuildTool/Configuration/UEBuildTarget.cs
--- a/Engine/Source/Programs/UnrealBuildTool/Configuration/UEBuildTarget.cs
+++ b/Engine/Source/Programs/UnrealBuildTool/Configuration/UEBuildTarget.cs
@@ -5495,7 +5495,7 @@ namespace UnrealBuildTool
 					}
 
 					// If it's a game module (plugin or otherwise), add the root source directory to the include paths.
-					if (RulesObject.File.IsUnderDirectory(TargetRulesFile.Directory) || (RulesObject.Plugin != null && RulesObject.Plugin.LoadedFrom == PluginLoadedFrom.Project))
+					if (RulesObject.File.IsUnderDirectory(BaseSourceDirectory) || (RulesObject.Plugin != null && RulesObject.Plugin.LoadedFrom == PluginLoadedFrom.Project))
 					{
 						if (DirectoryReference.Exists(BaseSourceDirectory))
 						{
