UnrealBuildTool: Fix some bad quotes in the clangdb generation that snuck in with a github pr

#rnx

diff --git a/Engine/Source/Programs/UnrealBuildTool/Modes/GenerateClangDatabase.cs b/Engine/Source/Programs/UnrealBuildTool/Modes/GenerateClangDatabase.cs
--- a/Engine/Source/Programs/UnrealBuildTool/Modes/GenerateClangDatabase.cs
+++ b/Engine/Source/Programs/UnrealBuildTool/Modes/GenerateClangDatabase.cs
@@ -261,8 +261,9 @@ namespace UnrealBuildTool
 									{
 										// Create the command
 										StringBuilder CommandBuilder = new StringBuilder();
+										string CommandPath = Action.CommandPath.FullName.Contains(' ') ? Utils.MakePathSafeToUseWithCommandLine(Action.CommandPath) : Action.CommandPath.FullName;
 										string CommandArguments = Action.CommandArguments.Replace(".rsp", ".rsp.gcd").Replace(".response", ".response.gcd");
-										CommandBuilder.AppendFormat("{0} {1}", Action.CommandPath, CommandArguments);
+										CommandBuilder.AppendFormat("{0} {1}", CommandPath, CommandArguments);
 
 										foreach (string ExtraArgument in GetExtraPlatformArguments(TargetToolChain))
 										{
