Remove some old references to mono

#rnx

diff --git a/Engine/Source/Programs/AutomationTool/Win/WinPlatform.Automation.cs b/Engine/Source/Programs/AutomationTool/Win/WinPlatform.Automation.cs
--- a/Engine/Source/Programs/AutomationTool/Win/WinPlatform.Automation.cs
+++ b/Engine/Source/Programs/AutomationTool/Win/WinPlatform.Automation.cs
@@ -227,7 +227,6 @@ public class Win64Platform : Platform
 			CommandUtils.CopyFile(InputFile.FullName, IntermediateFile.FullName);
 			CommandUtils.SetFileAttributes(IntermediateFile.FullName, ReadOnly: false);
 	
-			// currently the icon updating doesn't run under mono
 			if (UnrealBuildTool.BuildHostPlatform.Current.Platform == UnrealTargetPlatform.Win64)
 			{
 				Logger.LogInformation("Patching bootstrap executable; {Arg0}", IntermediateFile.FullName);
diff --git a/Engine/Source/Programs/Shared/EpicGames.Core/Log.cs b/Engine/Source/Programs/Shared/EpicGames.Core/Log.cs
--- a/Engine/Source/Programs/Shared/EpicGames.Core/Log.cs
+++ b/Engine/Source/Programs/Shared/EpicGames.Core/Log.cs
@@ -1271,8 +1271,6 @@ namespace EpicGames.Core
 					catch (IOException)
 					{
 						// Potential file access/sharing issue on std out
-						// This can occur on some versions of mono (e.g. macOS 6.12.0) if writing to a full pipe
-						// during IPC when the reader isn't consuming it quick enough
 					}
 					finally
 					{
diff --git a/Engine/Source/Programs/UnrealBuildTool/Platform/Mac/MacExports.cs b/Engine/Source/Programs/UnrealBuildTool/Platform/Mac/MacExports.cs
--- a/Engine/Source/Programs/UnrealBuildTool/Platform/Mac/MacExports.cs
+++ b/Engine/Source/Programs/UnrealBuildTool/Platform/Mac/MacExports.cs
@@ -1,4 +1,4 @@
-﻿// Copyright Epic Games, Inc. All Rights Reserved.
+// Copyright Epic Games, Inc. All Rights Reserved.
 
 using EpicGames.Core;
 using Microsoft.Extensions.Logging;
@@ -45,7 +45,7 @@ namespace UnrealBuildTool
 		}
 
 		/// <summary>
-		/// Returns true if we're running on Apple architecture (either natively which mono/dotnet will do, or under Rosetta)
+		/// Returns true if we're running on Apple architecture (either natively which dotnet will do, or under Rosetta)
 		/// </summary>
 		/// <returns></returns>
 		public static bool IsRunningOnAppleArchitecture
diff --git a/Engine/Source/Programs/UnrealBuildTool/ProjectFiles/VisualStudioCode/VSCodeProjectFileGenerator.cs b/Engine/Source/Programs/UnrealBuildTool/ProjectFiles/VisualStudioCode/VSCodeProjectFileGenerator.cs
--- a/Engine/Source/Programs/UnrealBuildTool/ProjectFiles/VisualStudioCode/VSCodeProjectFileGenerator.cs
+++ b/Engine/Source/Programs/UnrealBuildTool/ProjectFiles/VisualStudioCode/VSCodeProjectFileGenerator.cs
@@ -1670,14 +1670,6 @@ namespace UnrealBuildTool
 					{
 						WorkspaceFile.AddUnnamedField("epic.vscode-ue");
 					}
-
-					// If the platform we run the generator on uses mono, there are additional debugging extensions to add.
-					if (!RuntimePlatform.IsWindows)
-					{
-						WorkspaceFile.AddUnnamedField("vadimcn.vscode-lldb");
-						WorkspaceFile.AddUnnamedField("ms-vscode.mono-debug");
-						WorkspaceFile.AddUnnamedField("dfarley1.file-picker");
-					}
 				}
 				WorkspaceFile.EndArray();
 			}
diff --git a/Engine/Source/Programs/UnrealBuildTool/System/BuildHostPlatform.cs b/Engine/Source/Programs/UnrealBuildTool/System/BuildHostPlatform.cs
--- a/Engine/Source/Programs/UnrealBuildTool/System/BuildHostPlatform.cs
+++ b/Engine/Source/Programs/UnrealBuildTool/System/BuildHostPlatform.cs
@@ -281,7 +281,7 @@ namespace UnrealBuildTool
 		public override string BinarySuffix => String.Empty;
 
 		/// <summary>
-		/// Currently Mono returns incomplete process names in Process.GetProcesses() so we need to parse 'ps' output.
+		/// (needs confirmation) Currently returns incomplete process names in Process.GetProcesses() so we need to parse 'ps' output.
 		/// </summary>
 		/// <returns></returns>
 		public override ProcessInfo[] GetProcesses()
@@ -328,7 +328,7 @@ namespace UnrealBuildTool
 		}
 
 		/// <summary>
-		/// Currently Mono returns incomplete list of modules for Process.Modules so we need to parse vmmap output.
+		/// (needs confirmation) Currently returns incomplete list of modules for Process.Modules so we need to parse vmmap output.
 		/// </summary>
 		/// <param name="PID"></param>
 		/// <param name="Filename"></param>
@@ -401,8 +401,8 @@ namespace UnrealBuildTool
 		public override string BinarySuffix => String.Empty;
 
 		/// <summary>
-		/// Currently Mono returns incomplete process names in Process.GetProcesses() so we need to use /proc
-		/// (also, Mono locks up during process traversal sometimes, trying to open /dev/snd/pcm*)
+		/// (needs confirmation) Currently returns incomplete process names in Process.GetProcesses() so we need to use /proc
+		/// (also, locks up during process traversal sometimes, trying to open /dev/snd/pcm*)
 		/// </summary>
 		/// <returns></returns>
 		public override ProcessInfo[] GetProcesses()
@@ -412,8 +412,8 @@ namespace UnrealBuildTool
 		}
 
 		/// <summary>
-		/// Currently Mono returns incomplete list of modules for Process.Modules so we need to parse /proc/PID/maps.
-		/// (also, Mono locks up during process traversal sometimes, trying to open /dev/snd/pcm*)
+		/// (needs confirmation) Currently returns incomplete list of modules for Process.Modules so we need to parse /proc/PID/maps.
+		/// (also, locks up during process traversal sometimes, trying to open /dev/snd/pcm*)
 		/// </summary>
 		/// <param name="PID"></param>
 		/// <param name="Filename"></param>
