- Always pass the project to building SCW from VisualStudio

diff --git a/Engine/Source/Programs/UnrealBuildTool/ProjectFiles/VisualStudio/VCProject.cs b/Engine/Source/Programs/UnrealBuildTool/ProjectFiles/VisualStudio/VCProject.cs
--- a/Engine/Source/Programs/UnrealBuildTool/ProjectFiles/VisualStudio/VCProject.cs
+++ b/Engine/Source/Programs/UnrealBuildTool/ProjectFiles/VisualStudio/VCProject.cs
@@ -2137,7 +2137,8 @@ namespace UnrealBuildTool
 				{
 					if (TargetRulesObject.Type == TargetType.Editor && bEditorDependsOnShaderCompileWorker && !Unreal.IsEngineInstalled())
 					{
-						ExtraTargets.Add("ShaderCompileWorker Win64 Development");
+						string ProjParam = UProjectPath.Length > 0 ? $" -Project={UProjectPath}" : "";
+						ExtraTargets.Add($"ShaderCompileWorker Win64 Development{ProjParam}");
 					}
 					if (TargetRulesObject.bWithLiveCoding && bBuildLiveCodingConsole && !Unreal.IsEngineInstalled() && TargetRulesObject.Name != "LiveCodingConsole")
 					{
