- Split Mac target/deployment versions to have editor and non-editor versions
- Moved them into SDK.json
- Fixed up some deprecation warnings when tarrgeting macOS 13
- IOS is still using old .ini project setting method

diff --git a/Engine/Config/Mac/Mac_SDK.json b/Engine/Config/Mac/Mac_SDK.json
--- a/Engine/Config/Mac/Mac_SDK.json
+++ b/Engine/Config/Mac/Mac_SDK.json
@@ -1,6 +1,12 @@
 {
 	"ParentSDKFile": "../Apple/Apple_SDK.json",
 
+	"EditorDeploymentTarget": "13.0.0",
+	"EditorBuildTarget": "13.00",
+
+	"DeploymentTarget": "11.7.9",
+	"BuildTarget": "11.00",
+
 	// macOS Monterey 12.5 (Released 2022/07/22)
 	"MinSoftwareVersion": "12.5.0"
 }
diff --git a/Engine/Source/Developer/DesktopPlatform/Private/Mac/DesktopPlatformMac.cpp b/Engine/Source/Developer/DesktopPlatform/Private/Mac/DesktopPlatformMac.cpp
--- a/Engine/Source/Developer/DesktopPlatform/Private/Mac/DesktopPlatformMac.cpp
+++ b/Engine/Source/Developer/DesktopPlatform/Private/Mac/DesktopPlatformMac.cpp
@@ -9,6 +9,7 @@
 #include "Misc/ConfigCacheIni.h"
 #include "Misc/Guid.h"
 #include "HAL/FileManager.h"
+#import <UniformTypeIdentifiers/UniformTypeIdentifiers.h>
 
 #define LOCTEXT_NAMESPACE "DesktopPlatform"
 
@@ -127,8 +128,8 @@ private:
 	}
 	else
 	{
-		// Allow all file types
-		[DialogPanel setAllowedFileTypes:nil];
+		// Allow all file types (empty array means all types)
+		DialogPanel.allowedContentTypes = @[];
 	}
 }
 
@@ -146,7 +147,8 @@ private:
 	NSString* ExtsToParse = [AllowedFileTypes objectAtIndex:index * 2 + 1];
 	if( [ExtsToParse compare:@"*.*"] == NSOrderedSame )
 	{
-		[DialogPanel setAllowedFileTypes: nil];
+		// Allow all file types (empty array means all types)
+		DialogPanel.allowedContentTypes = @[];
 	}
 	else
 	{
@@ -156,10 +158,10 @@ private:
 		for( int32 Index = 0; Index < [ExtensionsWildcards count]; ++Index )
 		{
 			NSString* Temp = [[ExtensionsWildcards objectAtIndex:Index] stringByTrimmingCharactersInSet:[NSCharacterSet characterSetWithCharactersInString:@"*."]];
-			[Extensions addObject: Temp];
+			[Extensions addObject: [UTType typeWithFilenameExtension:Temp]];
 		}
 
-		[DialogPanel setAllowedFileTypes: Extensions];
+		DialogPanel.allowedContentTypes = Extensions;
 	}
 }
 
diff --git a/Engine/Source/Developer/DirectoryWatcher/Private/Mac/DirectoryWatchRequestMac.cpp b/Engine/Source/Developer/DirectoryWatcher/Private/Mac/DirectoryWatchRequestMac.cpp
--- a/Engine/Source/Developer/DirectoryWatcher/Private/Mac/DirectoryWatchRequestMac.cpp
+++ b/Engine/Source/Developer/DirectoryWatcher/Private/Mac/DirectoryWatchRequestMac.cpp
@@ -3,6 +3,7 @@
 #include "DirectoryWatchRequestMac.h"
 #include "HAL/PlatformFileManager.h"
 #include "GenericPlatform/GenericPlatformFile.h"
+#include "Mac/CocoaThread.h"
 
 void DirectoryWatchMacCallback( ConstFSEventStreamRef StreamRef, void* WatchRequestPtr, size_t EventCount, void* EventPaths, const FSEventStreamEventFlags EventFlags[], const FSEventStreamEventId EventIDs[] )
 {
@@ -10,7 +11,9 @@ void DirectoryWatchMacCallback( ConstFSEventStreamRef StreamRef, void* WatchRequ
 	check(WatchRequest);
 	check(WatchRequest->EventStream == StreamRef);
 
-	WatchRequest->ProcessChanges( EventCount, EventPaths, EventFlags);
+	GameThreadCall(^{
+		WatchRequest->ProcessChanges( EventCount, EventPaths, EventFlags);
+	});
 }
 
 // ============================================================================================================================
@@ -33,7 +36,6 @@ void FDirectoryWatchRequestMac::Shutdown( void )
 		check(EventStream);
 
 		FSEventStreamStop(EventStream);
-		FSEventStreamUnscheduleFromRunLoop(EventStream, CFRunLoopGetCurrent(), kCFRunLoopDefaultMode);
 		FSEventStreamInvalidate(EventStream);
 		FSEventStreamRelease(EventStream);
 
@@ -89,7 +91,7 @@ bool FDirectoryWatchRequestMac::Init(const FString& InDirectory)
 		return false;
 	}
 
-	FSEventStreamScheduleWithRunLoop( EventStream, CFRunLoopGetCurrent(), kCFRunLoopDefaultMode );
+	FSEventStreamSetDispatchQueue( EventStream, dispatch_get_main_queue() );
 	FSEventStreamStart( EventStream );
 
 	bRunning = true;
diff --git a/Engine/Source/Developer/TraceAnalysis/Private/Store/AsioStore.cpp b/Engine/Source/Developer/TraceAnalysis/Private/Store/AsioStore.cpp
--- a/Engine/Source/Developer/TraceAnalysis/Private/Store/AsioStore.cpp
+++ b/Engine/Source/Developer/TraceAnalysis/Private/Store/AsioStore.cpp
@@ -15,6 +15,7 @@
 	#include <sys/inotify.h>
 #elif PLATFORM_MAC
 	#include "Misc/Paths.h"
+	#include "Mac/CocoaThread.h"
 #endif
 
 namespace UE {
@@ -60,7 +61,6 @@ public:
 		if (bIsRunning)
 		{
 			FSEventStreamStop(EventStream);
-			FSEventStreamUnscheduleFromRunLoop(EventStream, CFRunLoopGetCurrent(), kCFRunLoopDefaultMode);
 			FSEventStreamInvalidate(EventStream);
 			FSEventStreamRelease(EventStream);
 			bIsRunning = false;
@@ -91,7 +91,9 @@ void MacCallback(ConstFSEventStreamRef StreamRef,
 	check(DirWatcherPtr);
 	check(DirWatcherPtr->EventStream == StreamRef);
 
-	DirWatcherPtr->ProcessChanges(EventCount, EventPaths, EventFlags);
+	GameThreadCall(^{
+		DirWatcherPtr->ProcessChanges(EventCount, EventPaths, EventFlags);
+	});
 }
 
 void FAsioStore::FDirWatcher::async_wait(HandlerType InHandler)
@@ -126,7 +128,7 @@ void FAsioStore::FDirWatcher::async_wait(HandlerType InHandler)
 		kFSEventStreamCreateFlagUseCFTypes | kFSEventStreamCreateFlagNoDefer
 	);
 
-	FSEventStreamScheduleWithRunLoop(EventStream, CFRunLoopGetCurrent(), kCFRunLoopDefaultMode);
+	FSEventStreamSetDispatchQueue(EventStream, dispatch_get_main_queue());
 	FSEventStreamStart(EventStream);
 	bIsRunning = true;
 
diff --git a/Engine/Source/Programs/UnrealBuildTool/Platform/Mac/ApplePlatformSDK.Versions.cs b/Engine/Source/Programs/UnrealBuildTool/Platform/Mac/ApplePlatformSDK.Versions.cs
--- a/Engine/Source/Programs/UnrealBuildTool/Platform/Mac/ApplePlatformSDK.Versions.cs
+++ b/Engine/Source/Programs/UnrealBuildTool/Platform/Mac/ApplePlatformSDK.Versions.cs
@@ -1,8 +1,60 @@
 // Copyright Epic Games, Inc. All Rights Reserved.
 
+using EpicGames.Core;
+
 namespace UnrealBuildTool
 {
 	/////////////////////////////////////////////////////////////////////////////////////
 	// If you are looking for version numbers, see Engine/Config/Apple/Apple_SDK.json
 	/////////////////////////////////////////////////////////////////////////////////////
+	
+	// NOTE: These are currently only used for Mac targets
+
+	partial class ApplePlatformSDK : UEBuildPlatformSDK
+	{
+		/// <summary>
+		/// Get the default deployment target version for the given target type. This will be put into the .plist file for runtime check
+		/// Will be in the format A.B.C
+		/// </summary>
+		/// <param name="Type"></param>
+		/// <returns></returns>
+		public string GetDeploymentTargetVersion(TargetType Type)
+		{
+			string? DeploymentTarget = null;
+			if (Type == TargetType.Editor || Type == TargetType.Program)
+			{
+				DeploymentTarget = GetVersionFromConfig("EditorDeploymentTarget");
+			}
+
+			if (DeploymentTarget == null)
+			{
+				DeploymentTarget = GetRequiredVersionFromConfig("DeploymentTarget");
+			}
+
+			return DeploymentTarget;
+		}
+
+
+		/// <summary>
+		/// Get the default build target version for the given target type. This will be passed to clang when compiling/linking
+		/// Will be in the format AA.BB
+		/// </summary>
+		/// <param name="Type"></param>
+		/// <returns></returns>
+		public string GetBuildTargetVersion(TargetType Type)
+		{
+			string? DeploymentTarget = null;
+			if (Type == TargetType.Editor || Type == TargetType.Program)
+			{
+				DeploymentTarget = GetVersionFromConfig("EditorBuildTarget");
+			}
+
+			if (DeploymentTarget == null)
+			{
+				DeploymentTarget = GetRequiredVersionFromConfig("BuildTarget");
+			}
+
+			return DeploymentTarget;
+		}
+	}
 }
diff --git a/Engine/Source/Programs/UnrealBuildTool/Platform/Mac/MacToolChain.cs b/Engine/Source/Programs/UnrealBuildTool/Platform/Mac/MacToolChain.cs
--- a/Engine/Source/Programs/UnrealBuildTool/Platform/Mac/MacToolChain.cs
+++ b/Engine/Source/Programs/UnrealBuildTool/Platform/Mac/MacToolChain.cs
@@ -17,12 +17,18 @@ namespace UnrealBuildTool
 		/// <summary>
 		/// Which version of the Mac OS X to allow at run time
 		/// </summary>
-		public string MacOSVersion = "11.00";
+		public string MinMacBuildVersion(TargetType TargetType)
+		{
+			return ((ApplePlatformSDK)UEBuildPlatformSDK.GetSDKForPlatform("Mac")!).GetBuildTargetVersion(TargetType);
+		}
 
 		/// <summary>
 		/// Minimum version of Mac OS X to actually run on, running on earlier versions will display the system minimum version error dialog and exit.
 		/// </summary>
-		public string MinMacOSVersion = "11.7.9";
+		public string MinMacDeploymentVersion(TargetType TargetType)
+		{
+			return ((ApplePlatformSDK)UEBuildPlatformSDK.GetSDKForPlatform("Mac")!).GetDeploymentTargetVersion(TargetType);
+		}
 
 		/// <summary>
 		/// Constructor
@@ -199,7 +205,7 @@ namespace UnrealBuildTool
 			// Pass through architecture and OS info
 			Arguments.Add("" + FormatArchitectureArg(CompileEnvironment.Architectures));
 			Arguments.Add($"-isysroot \"{Settings.GetSDKPath()}\"");
-			Arguments.Add("-mmacosx-version-min=" + (CompileEnvironment.bEnableOSX109Support ? "10.9" : Settings.MacOSVersion));
+			Arguments.Add("-mmacosx-version-min=" + (CompileEnvironment.bEnableOSX109Support ? "10.9" : Settings.MinMacBuildVersion(Target!.Type)));
 
 			List<string> FrameworksSearchPaths = new List<string>();
 			foreach (UEBuildFramework Framework in CompileEnvironment.AdditionalFrameworks)
@@ -231,7 +237,7 @@ namespace UnrealBuildTool
 			// Pass through architecture and OS info		
 			Arguments.Add(FormatArchitectureArg(LinkEnvironment.Architectures));
 			Arguments.Add(String.Format("-isysroot \"{0}\"", ToolChainSettings.Value.GetSDKPath(LinkEnvironment.Architecture)));
-			Arguments.Add("-mmacosx-version-min=" + Settings.MacOSVersion);
+			Arguments.Add("-mmacosx-version-min=" + Settings.MinMacBuildVersion(Target!.Type));
 			Arguments.Add("-dead_strip");
 
 			// Temporary workaround for linker warning with Xcode 14:
@@ -628,7 +634,7 @@ namespace UnrealBuildTool
 				// Fix contents of Info.plist
 				AppendMacLine(FinalizeAppBundleScript, "/usr/bin/sed -i \"\" -e \"s/\\${0}/{1}/g\" \"{2}\"", "{EXECUTABLE_NAME}", ExeName, TempInfoPlist);
 				AppendMacLine(FinalizeAppBundleScript, "/usr/bin/sed -i \"\" -e \"s/\\${0}/{1}/g\" \"{2}\"", "{APP_NAME}", bBuildingEditor ? ("com.epicgames." + GameName) : (BundleIdentifier.Replace("[PROJECT_NAME]", GameName).Replace("_", "")), TempInfoPlist);
-				AppendMacLine(FinalizeAppBundleScript, "/usr/bin/sed -i \"\" -e \"s/\\${0}/{1}/g\" \"{2}\"", "{MACOSX_DEPLOYMENT_TARGET}", Settings.MinMacOSVersion, TempInfoPlist);
+				AppendMacLine(FinalizeAppBundleScript, "/usr/bin/sed -i \"\" -e \"s/\\${0}/{1}/g\" \"{2}\"", "{MACOSX_DEPLOYMENT_TARGET}", Settings.MinMacDeploymentVersion(Target!.Type), TempInfoPlist);
 				AppendMacLine(FinalizeAppBundleScript, "/usr/bin/sed -i \"\" -e \"s/\\${0}/{1}/g\" \"{2}\"", "{ICON_NAME}", GameName, TempInfoPlist);
 				AppendMacLine(FinalizeAppBundleScript, "/usr/bin/sed -i \"\" -e \"s/\\${0}/{1}/g\" \"{2}\"", "{BUNDLE_VERSION}", BundleVersion, TempInfoPlist);
 
@@ -745,6 +751,7 @@ namespace UnrealBuildTool
 			if (bBuildImportLibraryOnly)
 			{
 				LinkCommand += " -undefined dynamic_lookup";
+				LinkCommand += " -Wl,-no_fixup_chains";
 			}
 			else if (!LinkEnvironment.bIsBuildingLibrary)
 			{
diff --git a/Engine/Source/Programs/UnrealBuildTool/Platform/Mac/UEDeployMac.cs b/Engine/Source/Programs/UnrealBuildTool/Platform/Mac/UEDeployMac.cs
--- a/Engine/Source/Programs/UnrealBuildTool/Platform/Mac/UEDeployMac.cs
+++ b/Engine/Source/Programs/UnrealBuildTool/Platform/Mac/UEDeployMac.cs
@@ -19,7 +19,7 @@ namespace UnrealBuildTool
 			return base.PrepTargetForDeployment(Receipt);
 		}
 
-		public static bool GeneratePList(string ProjectDirectory, bool bIsUnrealGame, string GameName, string ProjectName, string InEngineDir, string ExeName)
+		public static bool GeneratePList(string ProjectDirectory, bool bIsUnrealGame, string GameName, string ProjectName, string InEngineDir, string ExeName, TargetType TargetType)
 		{
 			string IntermediateDirectory = (bIsUnrealGame ? InEngineDir : ProjectDirectory) + "/Intermediate/Mac";
 			string DestPListFile = IntermediateDirectory + "/" + ExeName + "-Info.plist";
@@ -57,7 +57,7 @@ namespace UnrealBuildTool
 			PListData = PListData.Replace("${EXECUTABLE_NAME}", ExeName).
 				Replace("${APP_NAME}", FinalBundleID).
 				Replace("${ICON_NAME}", GameName).
-				Replace("${MACOSX_DEPLOYMENT_TARGET}", MacToolChain.Settings.MinMacOSVersion).
+				Replace("${MACOSX_DEPLOYMENT_TARGET}", MacToolChain.Settings.MinMacDeploymentVersion(TargetType)).
 				Replace("${BUNDLE_VERSION}", BundleVersion);
 
 			if (!Directory.Exists(IntermediateDirectory))
diff --git a/Engine/Source/Programs/UnrealBuildTool/ProjectFiles/Xcode/XcodeProject.cs b/Engine/Source/Programs/UnrealBuildTool/ProjectFiles/Xcode/XcodeProject.cs
--- a/Engine/Source/Programs/UnrealBuildTool/ProjectFiles/Xcode/XcodeProject.cs
+++ b/Engine/Source/Programs/UnrealBuildTool/ProjectFiles/Xcode/XcodeProject.cs
@@ -1793,7 +1793,7 @@ namespace UnrealBuildTool.XcodeProjectXcconfig
 				SDKRoot = "macosx";
 				SupportedPlatforms = "macosx";
 				DeploymentTargetKey = "MACOSX_DEPLOYMENT_TARGET";
-				DeploymentTarget = MacToolChain.Settings.MacOSVersion;
+				DeploymentTarget = MacToolChain.Settings.MinMacDeploymentVersion(UnrealData.TargetRules.Type);
 				BundleIdentifier = bIsEditor ? "com.epicgames.UnrealEditor" : UnrealData.BundleIdentifier;
 
 				// @todo: get a version for  games, like IOS has
@@ -2660,7 +2660,7 @@ namespace UnrealBuildTool.XcodeProjectXcconfig
 
 					// the Build target used some ini settings to compile, and Run target must match, so we override a few settings, at
 					// whatever level they were already specified at (Projet and/or Target)
-					XcodeUtils.PlistSetUpdate($":objects:{ConfigGuid}:buildSettings:MACOSX_DEPLOYMENT_TARGET", MacToolChain.Settings.MacOSVersion);
+					XcodeUtils.PlistSetUpdate($":objects:{ConfigGuid}:buildSettings:MACOSX_DEPLOYMENT_TARGET", MacToolChain.Settings.MinMacDeploymentVersion(UnrealData.TargetRules.Type));
 					if (UnrealData.IOSProjectSettings != null)
 					{
 						XcodeUtils.PlistSetUpdate($":objects:{ConfigGuid}:buildSettings:IPHONEOS_DEPLOYMENT_TARGET", UnrealData.IOSProjectSettings.RuntimeVersion);
diff --git a/Engine/Source/Programs/UnrealBuildTool/ProjectFiles/Xcode/XcodeProjectLegacy.cs b/Engine/Source/Programs/UnrealBuildTool/ProjectFiles/Xcode/XcodeProjectLegacy.cs
--- a/Engine/Source/Programs/UnrealBuildTool/ProjectFiles/Xcode/XcodeProjectLegacy.cs
+++ b/Engine/Source/Programs/UnrealBuildTool/ProjectFiles/Xcode/XcodeProjectLegacy.cs
@@ -1237,7 +1237,7 @@ namespace UnrealBuildTool.XcodeProjectLegacy
 					if (bCreateMacInfoPlist)
 					{
 						Directory.CreateDirectory(Path.GetDirectoryName(MacInfoPlistPath)!);
-						UEDeployMac.GeneratePList(ProjectPath.FullName, bIsUnrealGame, GameName, Config.BuildTarget, EngineDir.FullName, MacExecutableFileName);
+						UEDeployMac.GeneratePList(ProjectPath.FullName, bIsUnrealGame, GameName, Config.BuildTarget, EngineDir.FullName, MacExecutableFileName, Config.ProjectTarget!.TargetRules!.Type);
 					}
 					if (bCreateIOSInfoPlist)
 					{
@@ -1269,7 +1269,7 @@ namespace UnrealBuildTool.XcodeProjectLegacy
 			// #jira UE-143619: Pre Monterey macOS requires this option for a packaged app to run on iOS15 due to new code signature format. Could be removed once Monterey is minimum OS.
 			Content.Append("\t\t\t\tOTHER_CODE_SIGN_FLAGS = \"--generate-entitlement-der --deep\";" + ProjectFileGenerator.NewLine);
 
-			Content.Append("\t\t\t\tMACOSX_DEPLOYMENT_TARGET = " + MacToolChain.Settings.MacOSVersion + ";" + ProjectFileGenerator.NewLine);
+			Content.Append("\t\t\t\tMACOSX_DEPLOYMENT_TARGET = " + MacToolChain.Settings.MinMacDeploymentVersion(Config.ProjectTarget!.TargetRules!.Type) + ";" + ProjectFileGenerator.NewLine);
 			Content.Append("\t\t\t\tINFOPLIST_OUTPUT_FORMAT = xml;" + ProjectFileGenerator.NewLine);
 			Content.Append("\t\t\t\tCOMBINE_HIDPI_IMAGES = YES;" + ProjectFileGenerator.NewLine);
 
diff --git a/Engine/Source/Runtime/Core/Core.Build.cs b/Engine/Source/Runtime/Core/Core.Build.cs
--- a/Engine/Source/Runtime/Core/Core.Build.cs
+++ b/Engine/Source/Runtime/Core/Core.Build.cs
@@ -102,7 +102,7 @@ public class Core : ModuleRules
 				"zlib",
 				"PLCrashReporter"
 				);
-			PublicFrameworks.AddRange(new string[] { "Cocoa", "Carbon", "IOKit", "Security" });
+			PublicFrameworks.AddRange(new string[] { "Cocoa", "Carbon", "IOKit", "Security", "UniformTypeIdentifiers" });
 
 			PrivateDependencyModuleNames.Add("mimalloc");
 			PrivateDefinitions.Add("PLATFORM_BUILDS_MIMALLOC=1");
diff --git a/Engine/Source/Runtime/Core/Private/Mac/MacPlatformMisc.cpp b/Engine/Source/Runtime/Core/Private/Mac/MacPlatformMisc.cpp
--- a/Engine/Source/Runtime/Core/Private/Mac/MacPlatformMisc.cpp
+++ b/Engine/Source/Runtime/Core/Private/Mac/MacPlatformMisc.cpp
@@ -61,6 +61,11 @@ extern CORE_API bool GIsGPUCrashed;
 #define MAC_GRAPHICS_INI GEngineIni
 #endif
 
+#if MAC_OS_X_VERSION_MAX_ALLOWED < 120000
+	// use old constant on old macos, since kIOMainPortDefault was introduced in macOS 12
+	#define kIOMainPortDefault kIOMasterPortDefault
+#endif
+
 /*------------------------------------------------------------------------------
  Console variables.
  ------------------------------------------------------------------------------*/
@@ -287,7 +292,7 @@ struct FMacApplicationInfo
 		ParentProcess = TempSysCtlBuffer;
 		
 		MachineUUID = TEXT("00000000-0000-0000-0000-000000000000");
-		io_service_t PlatformExpert = IOServiceGetMatchingService(kIOMasterPortDefault,IOServiceMatching("IOPlatformExpertDevice"));
+		io_service_t PlatformExpert = IOServiceGetMatchingService(kIOMainPortDefault,IOServiceMatching("IOPlatformExpertDevice"));
 		if(PlatformExpert)
 		{
 			CFTypeRef SerialNumberAsCFString = IORegistryEntryCreateCFProperty(PlatformExpert,CFSTR(kIOPlatformUUIDKey),kCFAllocatorDefault, 0);
@@ -677,7 +682,7 @@ TArray<uint8> FMacPlatformMisc::GetMacAddress()
 		CFDictionarySetValue(MatchingDict, CFSTR(kIOPropertyMatchKey), PropertyMatchDict);
 		CFRelease(PropertyMatchDict);
 
-		if (IOServiceGetMatchingServices(kIOMasterPortDefault, MatchingDict, &InterfaceIterator) != KERN_SUCCESS)
+		if (IOServiceGetMatchingServices(kIOMainPortDefault, MatchingDict, &InterfaceIterator) != KERN_SUCCESS)
 		{
 			UE_LOG(LogMac, Warning, TEXT("GetMacAddress failed - error getting matching services"));
 			return Result;
@@ -1210,7 +1215,7 @@ public:
 		// Enumerate the GPUs via IOKit to avoid dragging in OpenGL
 		io_iterator_t Iterator;
 		CFMutableDictionaryRef MatchDictionary = IOServiceMatching(GetIOServiceMatchingName());
-		if(IOServiceGetMatchingServices(kIOMasterPortDefault, MatchDictionary, &Iterator) == kIOReturnSuccess)
+		if(IOServiceGetMatchingServices(kIOMainPortDefault, MatchDictionary, &Iterator) == kIOReturnSuccess)
 		{
 			uint32 Index = 0;
 			io_registry_entry_t ServiceEntry;
@@ -1296,7 +1301,7 @@ public:
 				CFMutableDictionaryRef MatchDictionary = IORegistryEntryIDMatching(DeviceRegistryID);
 				if(MatchDictionary)
 				{
-					io_registry_entry_t ServiceEntry = IOServiceGetMatchingService(kIOMasterPortDefault, MatchDictionary);
+					io_registry_entry_t ServiceEntry = IOServiceGetMatchingService(kIOMainPortDefault, MatchDictionary);
 					if(ServiceEntry)
 					{
 						io_iterator_t ParentIterator;
@@ -1791,7 +1796,7 @@ int32 FMacPlatformMisc::MacOSXVersionCompare(uint8 Major, uint8 Minor, uint8 Rev
 FString FMacPlatformMisc::GetOperatingSystemId()
 {
 	FString Result;
-	io_service_t Entry = IOServiceGetMatchingService(kIOMasterPortDefault, IOServiceMatching("IOPlatformExpertDevice"));
+	io_service_t Entry = IOServiceGetMatchingService(kIOMainPortDefault, IOServiceMatching("IOPlatformExpertDevice"));
 	if (Entry)
 	{
 		CFTypeRef UUID = IORegistryEntryCreateCFProperty(Entry, CFSTR(kIOPlatformUUIDKey), kCFAllocatorDefault, 0);
