Removing exclusions from using matching header files first.

diff --git a/Engine/Source/Programs/UnrealBuildTool/Configuration/UEBuildModuleCPP.cs b/Engine/Source/Programs/UnrealBuildTool/Configuration/UEBuildModuleCPP.cs
--- a/Engine/Source/Programs/UnrealBuildTool/Configuration/UEBuildModuleCPP.cs
+++ b/Engine/Source/Programs/UnrealBuildTool/Configuration/UEBuildModuleCPP.cs
@@ -1853,7 +1853,7 @@ namespace UnrealBuildTool
 								if (String.Compare(IncludeName, ExpectedName, StringComparison.OrdinalIgnoreCase) != 0)
 								{
 									FileReference? HeaderFile;
-									if (NameToHeaderFile.TryGetValue(ExpectedName, out HeaderFile) && !IgnoreMismatchedHeader(ExpectedName))
+									if (NameToHeaderFile.TryGetValue(ExpectedName, out HeaderFile))
 									{
 										InvalidIncludeDirectives.Add(new InvalidIncludeDirective(CppFile.Location, HeaderFile));
 									}
@@ -1865,20 +1865,6 @@ namespace UnrealBuildTool
 			}
 		}
 
-		private bool IgnoreMismatchedHeader(string ExpectedName)
-		{
-			switch (Name)
-			{
-				case "D3D11RHI":
-				case "D3D12RHI":
-				case "VulkanRHI":
-				case "OpenGLDrv":
-				case "MetalRHI":
-					return true;
-			}
-			return false;
-		}
-
 		private void CompileEnvironmentDebugInfoSettings(ReadOnlyTargetRules Target, CppCompileEnvironment Result)
 		{
 			// If bCreateDebugInfo is disabled for the whole target do not adjust any settings
