- Added a method for Shared build environment builds to have certain modules marked as caring about SDK version overrides for some platforms - if any of those platforms are overridden, the DLL will be written to the project Binaries dir instead of Engine (set bAllowSDKOverrideModulesWithSharedEnvironment = true in a Target.cs)
- Added ability to optionally skip SDK validation when making a Target instance even if we don't want to skip Target validation (used when we make the base UnrealEditor target to compare to project's Editor target - we don't need to manage SDK for that temporary UnrealEditor target).
- ShaderCompileWorker is now BuildEnvironment.UniqueIfNeeded, which will build into the project binaries dir if needed (based on per-project SDK override)
- UnrealPak is marked as not caring about an SDK versions so it can be built along with SCW that does care
- Added $(BinaryDir) variable to UEBuildTarget, etc, for the location of the first Binary in the output binaries

diff --git a/Engine/Source/Programs/ShaderCompileWorker/ShaderCompileWorker.Target.cs b/Engine/Source/Programs/ShaderCompileWorker/ShaderCompileWorker.Target.cs
--- a/Engine/Source/Programs/ShaderCompileWorker/ShaderCompileWorker.Target.cs
+++ b/Engine/Source/Programs/ShaderCompileWorker/ShaderCompileWorker.Target.cs
@@ -12,6 +12,7 @@ public class ShaderCompileWorkerTarget : TargetRules
 		Type = TargetType.Program;
 		LinkType = TargetLinkType.Modular;
 		IncludeOrderVersion = EngineIncludeOrderVersion.Latest;
+		BuildEnvironment = TargetBuildEnvironment.UniqueIfNeeded;
 
 		LaunchModuleName = "ShaderCompileWorker";
 
@@ -19,8 +20,8 @@ public class ShaderCompileWorkerTarget : TargetRules
         {
             // The interception interface in XGE requires that the parent and child processes have different filenames on disk.
             // To avoid building an entire separate worker just for this, we duplicate the ShaderCompileWorker in a post build step.
-            const string SrcPath  = "$(EngineDir)\\Binaries\\$(TargetPlatform)\\$(TargetName).exe";
-            const string DestPath = "$(EngineDir)\\Binaries\\$(TargetPlatform)\\XGEControlWorker.exe";
+            const string SrcPath  = "$(BinaryDir)\\$(TargetName).exe";
+            const string DestPath = "$(BinaryDir)\\XGEControlWorker.exe";
 
             PostBuildSteps.Add(string.Format("echo Copying {0} to {1}", SrcPath, DestPath));
             PostBuildSteps.Add(string.Format("copy /Y /B \"{0}\" /B \"{1}\" >nul:", SrcPath, DestPath));
diff --git a/Engine/Source/Programs/UnrealBuildTool/Configuration/ModuleRules.cs b/Engine/Source/Programs/UnrealBuildTool/Configuration/ModuleRules.cs
--- a/Engine/Source/Programs/UnrealBuildTool/Configuration/ModuleRules.cs
+++ b/Engine/Source/Programs/UnrealBuildTool/Configuration/ModuleRules.cs
@@ -791,6 +791,11 @@ namespace UnrealBuildTool
 		/// </summary>
 		public bool bTreatAsEngineModule { get; set; }
 
+		/// <summary>
+		/// If this contains a platform and the project has had it's SDK version overriden from default, this module will be compiled as a project module, instead of shared engine module
+		/// </summary>
+		public List<UnrealTargetPlatform> SDKVersionRelevantPlatforms { get; } = new();
+
 		/// <summary>
 		/// Emits compilation errors for incorrect UE_LOG format strings.
 		/// </summary>
diff --git a/Engine/Source/Programs/UnrealBuildTool/Configuration/TargetRules.cs b/Engine/Source/Programs/UnrealBuildTool/Configuration/TargetRules.cs
--- a/Engine/Source/Programs/UnrealBuildTool/Configuration/TargetRules.cs
+++ b/Engine/Source/Programs/UnrealBuildTool/Configuration/TargetRules.cs
@@ -2470,6 +2470,14 @@ namespace UnrealBuildTool
 			set => BuildEnvironmentOverride = value;
 		}
 
+		/// <summary>
+		/// If this true, a Shared build environment target will allow for any modules that are SDK version sensitive to have
+		/// a project side module when an SDK is overridden. For instance, if IOSTargetPlatform, which is marked as IOS SDK 
+		/// version relevant, is compiled for a Target with this set to true, and that target overrides the IOS SDK, then
+		/// this would compile to e.g. MyProject/Binaries/Win64/IOS/MyProject-IOSTargetPlatform.dll
+		/// </summary>
+		public bool bAllowSDKOverrideModulesWithSharedEnvironment { get; protected set; } = false;
+
 		/// <summary>
 		/// Whether to ignore violations to the shared build environment (eg. editor targets modifying definitions)
 		/// </summary>
@@ -2634,7 +2642,7 @@ namespace UnrealBuildTool
 		public bool AllowsPerProjectSDKVersion()
 		{
 			// modular target with TargetBuildEnvironment.Shared build type cannot allow per-project SDKs
-			return LinkType == TargetLinkType.Monolithic || BuildEnvironment == TargetBuildEnvironment.Unique;
+			return LinkType == TargetLinkType.Monolithic || BuildEnvironment == TargetBuildEnvironment.Unique || bAllowSDKOverrideModulesWithSharedEnvironment;
 		}
 
 		/// <summary>
@@ -2683,7 +2691,8 @@ namespace UnrealBuildTool
 			}
 
 			// Create the target rules for it
-			TargetRules baseRules = rulesAssembly.CreateTargetRules(baseTargetName, thisRules.Platform, thisRules.Configuration, thisRules.Architectures, null, arguments, Logger, IntermediateEnvironment: thisRules.IntermediateEnvironment);
+			// we need to validate the target, but we don't care about SDK versions at all
+			TargetRules baseRules = rulesAssembly.CreateTargetRules(baseTargetName, thisRules.Platform, thisRules.Configuration, thisRules.Architectures, null, arguments, Logger, IntermediateEnvironment: thisRules.IntermediateEnvironment, ValidationOptions: TargetRulesValidationOptions.ValidateTargetOnly);
 
 			// Get all the configurable objects
 			object[] baseObjects = baseRules.GetConfigurableObjects().ToArray();
diff --git a/Engine/Source/Programs/UnrealBuildTool/Configuration/UEBuildTarget.cs b/Engine/Source/Programs/UnrealBuildTool/Configuration/UEBuildTarget.cs
--- a/Engine/Source/Programs/UnrealBuildTool/Configuration/UEBuildTarget.cs
+++ b/Engine/Source/Programs/UnrealBuildTool/Configuration/UEBuildTarget.cs
@@ -3540,6 +3540,7 @@ namespace UnrealBuildTool
 			{
 				Variables.Add("PluginDir", Plugin.Directory.FullName);
 			}
+			Variables.Add("BinaryDir", Binaries[0].OutputDir.FullName);
 			return Variables;
 		}
 
@@ -4065,7 +4066,21 @@ namespace UnrealBuildTool
 			// Get the root output directory and base name (target name/app name) for this binary
 			DirectoryReference BaseOutputDirectory;
 
-			if (bUseSharedBuildEnvironment && (ModuleRules.Plugin == null || (ModuleRules.Plugin.Type != PluginType.External || ModuleRules.Plugin.bExplicitPluginTarget)))
+			// if this module cares about any SDK overrides, and at least one was overridden, then treat the module as a project module
+			bool bHasImportantSDKOverride = false;
+			if (UEBuildPlatformSDK.bHasAnySDKOverride)
+			{
+				foreach (UnrealTargetPlatform Platform in ModuleRules.SDKVersionRelevantPlatforms)
+				{
+					UEBuildPlatformSDK? SDK = UEBuildPlatformSDK.GetSDKForPlatform(Platform.ToString());
+					if (SDK != null && SDK.bHasSDKOverride)
+					{
+						bHasImportantSDKOverride = true;
+					}
+				}
+			}
+
+			if (bUseSharedBuildEnvironment && !bHasImportantSDKOverride && (ModuleRules.Plugin == null || (ModuleRules.Plugin.Type != PluginType.External || ModuleRules.Plugin.bExplicitPluginTarget)))
 			{
 				BaseOutputDirectory = ModuleRules.Context.DefaultOutputBaseDir;
 			}
@@ -4726,9 +4741,17 @@ namespace UnrealBuildTool
 
 			// Construct the output paths for this target's executable
 			DirectoryReference OutputDirectory;
-			if (ProjectFile != null && (bCompileMonolithic || !bUseSharedBuildEnvironment) && Rules.File.IsUnderDirectory(ProjectDirectory))
+			// if we are building a program with a project file, and we are making a unique build environment, output the program to that project's binaries
+			bool bIsUniqueBuildProgram = Rules.Type == TargetType.Program && Rules.BuildEnvironment == TargetBuildEnvironment.Unique;
+			if (ProjectFile != null && (bCompileMonolithic || !bUseSharedBuildEnvironment) && (Rules.File.IsUnderDirectory(ProjectDirectory) || bIsUniqueBuildProgram))
 			{
 				OutputDirectory = GetOutputDirectoryForExecutable(ProjectDirectory, Rules.File);
+				// engine programs generally won't be under the directories searched in GetOutputDirectoryForExecutable, so it will return the EngineDir, even tho unique programs
+				// built with a ProjectFile want to be under the Project dir, so just force it into the ProjectDir if one wasn't found otherwise
+				if (bIsUniqueBuildProgram && Rules.File.IsUnderDirectory(Unreal.EngineDirectory) && OutputDirectory == Unreal.EngineDirectory)
+				{
+					OutputDirectory = ProjectDirectory;
+				}
 			}
 			else
 			{
diff --git a/Engine/Source/Programs/UnrealBuildTool/Modes/QueryTargetsMode.cs b/Engine/Source/Programs/UnrealBuildTool/Modes/QueryTargetsMode.cs
--- a/Engine/Source/Programs/UnrealBuildTool/Modes/QueryTargetsMode.cs
+++ b/Engine/Source/Programs/UnrealBuildTool/Modes/QueryTargetsMode.cs
@@ -124,7 +124,7 @@ namespace UnrealBuildTool
 					try
 					{
 						UnrealArchitectures Architectures = UnrealArchitectureConfig.ForPlatform(BuildHostPlatform.Current.Platform).ActiveArchitectures(ProjectFile, TargetName);
-						TargetRules = Assembly.CreateTargetRules(TargetName, BuildHostPlatform.Current.Platform, UnrealTargetConfiguration.Development, Architectures, ProjectFile, Arguments, Logger, bSkipValidation: true);
+						TargetRules = Assembly.CreateTargetRules(TargetName, BuildHostPlatform.Current.Platform, UnrealTargetConfiguration.Development, Architectures, ProjectFile, Arguments, Logger, ValidationOptions: TargetRulesValidationOptions.ValidateNothing);
 					}
 					catch (Exception Ex)
 					{
diff --git a/Engine/Source/Programs/UnrealBuildTool/Platform/Mac/UEBuildMac.cs b/Engine/Source/Programs/UnrealBuildTool/Platform/Mac/UEBuildMac.cs
--- a/Engine/Source/Programs/UnrealBuildTool/Platform/Mac/UEBuildMac.cs
+++ b/Engine/Source/Programs/UnrealBuildTool/Platform/Mac/UEBuildMac.cs
@@ -178,7 +178,7 @@ namespace UnrealBuildTool
 					// CreateTargetRules here needs to have an UnrealArchitectures object, because otherwise with 'null', it will call
 					// back to this function to get the ActiveArchitectures! in this case the arch is unimportant
 					UnrealArchitectures DummyArchitectures = new(UnrealArch.X64);
-					TargetRules? Rules = RulesAsm.CreateTargetRules(TargetName, UnrealTargetPlatform.Mac, UnrealTargetConfiguration.Development, DummyArchitectures, ProjectFile, null, Log.Logger, bSkipValidation:true);
+					TargetRules? Rules = RulesAsm.CreateTargetRules(TargetName, UnrealTargetPlatform.Mac, UnrealTargetConfiguration.Development, DummyArchitectures, ProjectFile, null, Log.Logger, ValidationOptions: TargetRulesValidationOptions.ValidateNothing);
 					bIsEditor = Rules.Type == TargetType.Editor;
 
 					// the projectfile passed in may be a game's uproject file that we are compiling a program in the context of, 
diff --git a/Engine/Source/Programs/UnrealBuildTool/System/RulesAssembly.cs b/Engine/Source/Programs/UnrealBuildTool/System/RulesAssembly.cs
--- a/Engine/Source/Programs/UnrealBuildTool/System/RulesAssembly.cs
+++ b/Engine/Source/Programs/UnrealBuildTool/System/RulesAssembly.cs
@@ -12,6 +12,34 @@ using UnrealBuildBase;
 
 namespace UnrealBuildTool
 {
+	/// <summary>
+	/// Controls how the target and SDK are validated in CreateTargetRulesInstance
+	/// </summary>
+	public enum TargetRulesValidationOptions
+	{
+		/// <summary>
+		///  This will perform full target validation (allows the platform to update the target to make sure it's valid), and
+		///  SDK validation (makes sure SDK overrides are not conflicting).
+		///  This is the standard, default mode.
+		/// </summary>
+		ValidateTargetAndSDK,
+
+		/// <summary>
+		/// This will validate the target only (see ValidateTargetAndSDK)
+		/// </summary>
+		ValidateTargetOnly,
+
+		/// <summary>
+		/// This will validate the SDK only (see ValidateTargetAndSDK)
+		/// </summary>
+		ValidateSDKOnly,
+
+		/// <summary>
+		/// This will perform neither of the validations
+		/// </summary>
+		ValidateNothing
+	}
+
 	/// <summary>
 	/// Stores information about a compiled rules assembly and the types it contains
 	/// </summary>
@@ -586,9 +614,9 @@ namespace UnrealBuildTool
 		/// <param name="TargetInfo">Target configuration information to pass to the constructor</param>
 		/// <param name="Logger">Logger for output</param>
 		/// <param name="IsTestTarget">If building a low level tests target</param>
-		/// <param name="bSkipValidation">If validation should be skipped (QueryTargetMode)</param>
+		/// <param name="ValidationOptions">Controls validation of target and SDK</param>
 		/// <returns>Instance of the corresponding TargetRules or null if requested type name does not exist</returns>
-		protected TargetRules? CreateTargetRulesInstance(string TypeName, TargetInfo TargetInfo, ILogger Logger, bool IsTestTarget = false, bool bSkipValidation = false)
+		protected TargetRules? CreateTargetRulesInstance(string TypeName, TargetInfo TargetInfo, ILogger Logger, bool IsTestTarget = false, TargetRulesValidationOptions ValidationOptions = TargetRulesValidationOptions.ValidateTargetAndSDK)
 		{
 			// The build module must define a type named '<TargetName>Target' that derives from our 'TargetRules' type.  
 			Type? BaseRulesType = CompiledAssembly?.GetType(TypeName);
@@ -632,6 +660,8 @@ namespace UnrealBuildTool
 			FileReference BaseFile = TargetNameToTargetFile[TargetInfo.Name];
 			FileReference PlatformFile = TargetNameToTargetFile.TryGetValue(PlatformRulesName, out FileReference? PlatformTargetFile) ? PlatformTargetFile : BaseFile;
 			TargetRules Rules = TargetRules.Create(RulesType, TargetInfo, BaseFile, PlatformFile, TargetNameToTargetFile.Values, DefaultBuildSettings, Logger);
+			bool bValidateTarget = ValidationOptions == TargetRulesValidationOptions.ValidateTargetAndSDK || ValidationOptions == TargetRulesValidationOptions.ValidateTargetOnly;
+			bool bValidateSDK = ValidationOptions == TargetRulesValidationOptions.ValidateTargetAndSDK || ValidationOptions == TargetRulesValidationOptions.ValidateSDKOnly;
 
 			// Set the default overriddes for the configured target type
 			Rules.SetOverridesForTargetType();
@@ -642,7 +672,7 @@ namespace UnrealBuildTool
 				throw new CompilationResultException(CompilationResult.RulesError, "TargetRules.LinkType should be inferred from TargetType");
 			}
 
-			if (!bSkipValidation)
+			if (bValidateTarget)
 			{
 				// Delayed-fixup of TargetBuildEnvironment.UniqueIfNeeded
 				Rules.UpdateBuildEnvironmentIfNeeded(this, arguments: null, Logger);
@@ -730,11 +760,15 @@ namespace UnrealBuildTool
 			}
 
 			// Allow the platform to finalize the settings
-			if (!bSkipValidation)
+			if (bValidateTarget)
 			{
 				UEBuildPlatform Platform = UEBuildPlatform.GetBuildPlatform(Rules.Platform);
 				Platform.ValidateTarget(Rules);
+			}
 
+			// make sure any SDK overrides are valid
+			if (bValidateSDK)
+			{ 
 				ValidateSDKs(Rules);
 			}
 
@@ -801,10 +835,10 @@ namespace UnrealBuildTool
 		/// <param name="Arguments">Command line arguments for this target</param>
 		/// <param name="Logger"></param>
 		/// <param name="IsTestTarget">If building a low level test target</param>
-		/// <param name="bSkipValidation">If validation should be skipped (QueryTargetMode)</param>
+		/// <param name="ValidationOptions">Controls validation of target and SDK</param>
 		/// <param name="IntermediateEnvironment">Intermediate environment to use</param>
 		/// <returns>The build target rules for the specified target</returns>
-		public TargetRules CreateTargetRules(string TargetName, UnrealTargetPlatform Platform, UnrealTargetConfiguration Configuration, UnrealArchitectures? Architectures, FileReference? ProjectFile, CommandLineArguments? Arguments, ILogger Logger, bool IsTestTarget = false, bool bSkipValidation = false, UnrealIntermediateEnvironment IntermediateEnvironment = UnrealIntermediateEnvironment.Default)
+		public TargetRules CreateTargetRules(string TargetName, UnrealTargetPlatform Platform, UnrealTargetConfiguration Configuration, UnrealArchitectures? Architectures, FileReference? ProjectFile, CommandLineArguments? Arguments, ILogger Logger, bool IsTestTarget = false, TargetRulesValidationOptions ValidationOptions = TargetRulesValidationOptions.ValidateTargetAndSDK, UnrealIntermediateEnvironment IntermediateEnvironment = UnrealIntermediateEnvironment.Default)
 		{
 			if (IsTestTarget)
 			{
@@ -837,7 +871,7 @@ namespace UnrealBuildTool
 				}
 				else
 				{
-					return Parent.CreateTargetRules(TargetName, Platform, Configuration, Architectures, ProjectFile, Arguments, Logger, IsTestTarget, bSkipValidation, IntermediateEnvironment);
+					return Parent.CreateTargetRules(TargetName, Platform, Configuration, Architectures, ProjectFile, Arguments, Logger, IsTestTarget, ValidationOptions, IntermediateEnvironment);
 				}
 			}
 
@@ -845,7 +879,7 @@ namespace UnrealBuildTool
 			string TargetTypeName = TargetName + "Target";
 
 			// The build module must define a type named '<TargetName>Target' that derives from our 'TargetRules' type.  
-			TargetRules? TargetRules = CreateTargetRulesInstance(TargetTypeName, new TargetInfo(TargetName, Platform, Configuration, Architectures, ProjectFile, Arguments, IntermediateEnvironment), Logger, IsTestTarget, bSkipValidation);
+			TargetRules? TargetRules = CreateTargetRulesInstance(TargetTypeName, new TargetInfo(TargetName, Platform, Configuration, Architectures, ProjectFile, Arguments, IntermediateEnvironment), Logger, IsTestTarget, ValidationOptions);
 
 			if (TargetRules == null)
 			{
diff --git a/Engine/Source/Programs/UnrealPak/UnrealPak.Target.cs b/Engine/Source/Programs/UnrealPak/UnrealPak.Target.cs
--- a/Engine/Source/Programs/UnrealPak/UnrealPak.Target.cs
+++ b/Engine/Source/Programs/UnrealPak/UnrealPak.Target.cs
@@ -34,6 +34,7 @@ public class UnrealPakTarget : TargetRules
 		bIsBuildingConsoleApplication		= true;
 		bForceBuildTargetPlatforms			= true;
 		bNeedsExtraShaderFormats			= true;
+		bAreTargetSDKVersionsRelevantOverride = false;
 
 		WindowsPlatform.bUseBundledDbgHelp = false;
 		WindowsPlatform.bPixProfilingEnabled = false;
