UnrealBuildTool: Skip generating invalid platforms for client targets

#rnx

diff --git a/Engine/Source/Programs/UnrealBuildTool/ProjectFiles/Rider/RiderProjectFile.cs b/Engine/Source/Programs/UnrealBuildTool/ProjectFiles/Rider/RiderProjectFile.cs
--- a/Engine/Source/Programs/UnrealBuildTool/ProjectFiles/Rider/RiderProjectFile.cs
+++ b/Engine/Source/Programs/UnrealBuildTool/ProjectFiles/Rider/RiderProjectFile.cs
@@ -72,6 +72,9 @@ namespace UnrealBuildTool
 			string ProjectName = ProjectFilePath.GetFileNameWithoutAnyExtensions();
 			DirectoryReference ProjectRootFolder = RootPath;
 			List<TargetEntry> FileToTarget = new List<TargetEntry>();
+
+			HashSet<UnrealTargetPlatform> ServerPlatforms = Utils.GetPlatformsInClass(UnrealPlatformClass.Server).ToHashSet();
+
 			foreach (UnrealTargetPlatform Platform in InPlatforms)
 			{
 				foreach (UnrealTargetConfiguration Configuration in InConfigurations)
@@ -97,6 +100,12 @@ namespace UnrealBuildTool
 							continue;
 						}
 
+						// Skip Server for all invalid platforms
+						if (ProjectTarget.TargetRules.Type == TargetType.Server && !ServerPlatforms.Contains(Platform))
+						{
+							continue;
+						}
+
 						bool bBuildByDefault = ShouldBuildByDefaultForSolutionTargets && ProjectTarget.SupportedPlatforms.Contains(Platform);
 
 						DirectoryReference ConfigurationFolder = DirectoryReference.Combine(ProjectRootFolder, Platform.ToString(), Configuration.ToString());
diff --git a/Engine/Source/Programs/UnrealBuildTool/ProjectFiles/VisualStudioWorkspace/VSWorkspaceProjectFile.cs b/Engine/Source/Programs/UnrealBuildTool/ProjectFiles/VisualStudioWorkspace/VSWorkspaceProjectFile.cs
--- a/Engine/Source/Programs/UnrealBuildTool/ProjectFiles/VisualStudioWorkspace/VSWorkspaceProjectFile.cs
+++ b/Engine/Source/Programs/UnrealBuildTool/ProjectFiles/VisualStudioWorkspace/VSWorkspaceProjectFile.cs
@@ -55,6 +55,8 @@ namespace UnrealBuildTool
 
 			Dictionary<FileReference, (UEBuildTarget BuildTarget, bool bBuildByDefault)> FileToTarget = new();
 
+			HashSet<UnrealTargetPlatform> ServerPlatforms = Utils.GetPlatformsInClass(UnrealPlatformClass.Server).ToHashSet();
+
 			foreach (UnrealTargetPlatform Platform in InPlatforms)
 			{
 				foreach (UnrealTargetConfiguration Configuration in InConfigurations)
@@ -82,6 +84,12 @@ namespace UnrealBuildTool
 							continue;
 						}
 
+						// Skip Server for all invalid platforms
+						if (ProjectTarget.TargetRules.Type == TargetType.Server && !ServerPlatforms.Contains(Platform))
+						{
+							continue;
+						}
+
 						bool bBuildByDefault = ShouldBuildByDefaultForSolutionTargets && ProjectTarget.SupportedPlatforms.Contains(Platform);
 
 						UnrealArchitectures ProjectArchitectures = UEBuildPlatform
