UnrealBuildTool: Move recommended Visual Studio components to json sdk config

#rnx

diff --git a/Engine/Config/Windows/Windows_SDK.json b/Engine/Config/Windows/Windows_SDK.json
--- a/Engine/Config/Windows/Windows_SDK.json
+++ b/Engine/Config/Windows/Windows_SDK.json
@@ -33,7 +33,23 @@
 	"PreferredIntelOneApiVersions": [
 		"2024.0.0-2024.9999"
 	],
-	"MinimumIntelOneApiVersion": "2024.0.0"
+	"MinimumIntelOneApiVersion": "2024.0.0",
 
+	"VisualStudioSuggestedComponents": [
+		"Microsoft.Net.Component.4.6.2.TargetingPack",
+		"Microsoft.VisualStudio.Component.VC.14.38.17.8.x86.x64",
+		"Microsoft.VisualStudio.Component.VC.Tools.x86.x64",
+		"Microsoft.VisualStudio.Component.Windows10SDK.22621",
+		"Microsoft.VisualStudio.Workload.CoreEditor",
+		"Microsoft.VisualStudio.Workload.ManagedDesktop",
+		"Microsoft.VisualStudio.Workload.NativeDesktop",
+		"Microsoft.VisualStudio.Workload.NativeGame"
+	],
+
+	"VisualStudioSuggestedLinuxComponents": [
+		"Microsoft.VisualStudio.Workload.NativeCrossPlat"
+	],
 
+	"VisualStudio2022SuggestedComponents": [
+	]
 }
diff --git a/Engine/Source/Programs/Shared/EpicGames.Build/UEBuildPlatformSDK.cs b/Engine/Source/Programs/Shared/EpicGames.Build/UEBuildPlatformSDK.cs
--- a/Engine/Source/Programs/Shared/EpicGames.Build/UEBuildPlatformSDK.cs
+++ b/Engine/Source/Programs/Shared/EpicGames.Build/UEBuildPlatformSDK.cs
@@ -1233,6 +1233,17 @@ namespace EpicGames.Core
 			return Ranges.ToArray();
 		}
 
+		public string[] GetStringArrayFromConfig(string Name)
+		{
+			string[]? Results;
+			if (ConfigSDKVersionArrays.TryGetValue(GetHostSpecificVersionName(Name), out Results) || ConfigSDKVersionArrays.TryGetValue(Name, out Results))
+			{
+				return Results;
+			}
+
+			return Array.Empty<string>();
+		}
+
 		#endregion
 
 		// AutoSDKs handling portion
diff --git a/Engine/Source/Programs/UnrealBuildTool/Platform/Windows/MicrosoftPlatformSDK.VisualStudio.cs b/Engine/Source/Programs/UnrealBuildTool/Platform/Windows/MicrosoftPlatformSDK.VisualStudio.cs
--- a/Engine/Source/Programs/UnrealBuildTool/Platform/Windows/MicrosoftPlatformSDK.VisualStudio.cs
+++ b/Engine/Source/Programs/UnrealBuildTool/Platform/Windows/MicrosoftPlatformSDK.VisualStudio.cs
@@ -1,6 +1,5 @@
 ﻿// Copyright Epic Games, Inc. All Rights Reserved.
 
-using System;
 using System.Collections.Generic;
 using EpicGames.Core;
 
@@ -12,31 +11,17 @@ namespace UnrealBuildTool
 		/// The default set of components that should be suggested to be installed for Visual Studio
 		/// This or the 2022 specific components should be updated if the preferred visual cpp version changes
 		/// </summary>
-		static readonly string[] VisualStudioSuggestedComponents = new string[]
-		{
-			"Microsoft.Net.Component.4.6.2.TargetingPack",
-			"Microsoft.VisualStudio.Component.VC.14.38.17.8.x86.x64",
-			"Microsoft.VisualStudio.Component.VC.Tools.x86.x64",
-			"Microsoft.VisualStudio.Component.Windows10SDK.22000",
-			"Microsoft.VisualStudio.Workload.CoreEditor",
-			"Microsoft.VisualStudio.Workload.ManagedDesktop",
-			"Microsoft.VisualStudio.Workload.NativeDesktop",
-			"Microsoft.VisualStudio.Workload.NativeGame",
-		};
+		static readonly string[] s_visualStudioSuggestedComponents = SDK.GetStringArrayFromConfig("VisualStudioSuggestedComponents");
 
 		/// <summary>
 		/// Additional set of components that should be suggested to be installed for Visual Studio
 		/// to support the Linux platform.
 		/// </summary>
-		static readonly string[] VisualStudioSuggestedLinuxComponents = new string[]
-		{
-			"Microsoft.VisualStudio.Workload.NativeCrossPlat",
-		};
-
+		static readonly string[] s_visualStudioSuggestedLinuxComponents = SDK.GetStringArrayFromConfig("VisualStudioSuggestedLinuxComponents");
 		/// <summary>
 		/// Additional set of components that should be suggested to be installed for Visual Studio 2022.
 		/// </summary>
-		static readonly string[] VisualStudio2022SuggestedComponents = Array.Empty<string>();
+		static readonly string[] s_visualStudio2022SuggestedComponents = SDK.GetStringArrayFromConfig("VisualStudio2022SuggestedComponents");
 
 		/// <summary>
 		/// Returns the list of suggested of components that should be suggested to be installed for Visual Studio.
@@ -47,12 +32,12 @@ namespace UnrealBuildTool
 			bool platformLinuxValid = (InstalledPlatformInfo.IsValidPlatform(UnrealTargetPlatform.Linux) && UEBuildPlatform.IsPlatformAvailable(UnrealTargetPlatform.Linux))
 				|| (InstalledPlatformInfo.IsValidPlatform(UnrealTargetPlatform.LinuxArm64) && UEBuildPlatform.IsPlatformAvailable(UnrealTargetPlatform.LinuxArm64));
 
-			SortedSet<string> components = new(VisualStudioSuggestedComponents);
+			SortedSet<string> components = new(s_visualStudioSuggestedComponents);
 
 			switch (format)
 			{
 				case VCProjectFileFormat.VisualStudio2022:
-					components.UnionWith(VisualStudio2022SuggestedComponents);
+					components.UnionWith(s_visualStudio2022SuggestedComponents);
 					break;
 				default:
 					throw new BuildException("Unsupported Visual Studio version {0}", format);
@@ -60,7 +45,7 @@ namespace UnrealBuildTool
 
 			if (platformLinuxValid)
 			{
-				components.UnionWith(VisualStudioSuggestedLinuxComponents);
+				components.UnionWith(s_visualStudioSuggestedLinuxComponents);
 			}
 
 			return components;
