UnrealBuildTool: Option to disable anaylzing engine modules or generated code

diff --git a/Engine/Source/Programs/UnrealBuildTool/Configuration/ReadOnlyTargetRules.cs b/Engine/Source/Programs/UnrealBuildTool/Configuration/ReadOnlyTargetRules.cs
--- a/Engine/Source/Programs/UnrealBuildTool/Configuration/ReadOnlyTargetRules.cs
+++ b/Engine/Source/Programs/UnrealBuildTool/Configuration/ReadOnlyTargetRules.cs
@@ -403,6 +403,10 @@ namespace UnrealBuildTool
 
 		public StaticAnalyzerMode StaticAnalyzerMode => Inner.StaticAnalyzerMode;
 
+		public bool bStaticAnalyzerProjectOnly => Inner.bStaticAnalyzerProjectOnly;
+
+		public bool bStaticAnalyzerSkipGenerated => Inner.bStaticAnalyzerSkipGenerated;
+
 		public int MinFilesUsingPrecompiledHeader => Inner.MinFilesUsingPrecompiledHeader;
 
 		public bool bForcePrecompiledHeaderForGameModules => Inner.bForcePrecompiledHeaderForGameModules;
diff --git a/Engine/Source/Programs/UnrealBuildTool/Configuration/TargetRules.cs b/Engine/Source/Programs/UnrealBuildTool/Configuration/TargetRules.cs
--- a/Engine/Source/Programs/UnrealBuildTool/Configuration/TargetRules.cs
+++ b/Engine/Source/Programs/UnrealBuildTool/Configuration/TargetRules.cs
@@ -1977,6 +1977,20 @@ namespace UnrealBuildTool
 		[XmlConfigFile(Category = "BuildConfiguration")]
 		public StaticAnalyzerMode StaticAnalyzerMode { get; set; } = StaticAnalyzerMode.Deep;
 
+		/// <summary>
+		/// Only run static analysis against project modules, skipping engine modules
+		/// </summary>
+		[CommandLine("-StaticAnalyzerProjectOnly")]
+		[XmlConfigFile(Category = "BuildConfiguration")]
+		public bool bStaticAnalyzerProjectOnly = false;
+
+		/// <summary>
+		/// Skip analyzing generated source files
+		/// </summary>
+		[CommandLine("-StaticAnalyzerSkipGenerated")]
+		[XmlConfigFile(Category = "BuildConfiguration")]
+		public bool bStaticAnalyzerSkipGenerated = false;
+
 		/// <summary>
 		/// The minimum number of files that must use a pre-compiled header before it will be created and used.
 		/// </summary>
diff --git a/Engine/Source/Programs/UnrealBuildTool/Configuration/UEBuildModuleCPP.cs b/Engine/Source/Programs/UnrealBuildTool/Configuration/UEBuildModuleCPP.cs
--- a/Engine/Source/Programs/UnrealBuildTool/Configuration/UEBuildModuleCPP.cs
+++ b/Engine/Source/Programs/UnrealBuildTool/Configuration/UEBuildModuleCPP.cs
@@ -561,6 +561,13 @@ namespace UnrealBuildTool
 						bMergeUnityFiles = false;
 					}
 
+					if (Target.StaticAnalyzer != StaticAnalyzer.None && Target.bStaticAnalyzerSkipGenerated)
+					{
+						GeneratedCPPCompileEnvironment = new CppCompileEnvironment(GeneratedCPPCompileEnvironment);
+						GeneratedCPPCompileEnvironment.bDisableStaticAnalysis = true;
+						bMergeUnityFiles = false;
+					}
+
 					// Always force include the PCH, even if PCHs are disabled, for generated code. Legacy code can rely on PCHs being included to compile correctly, and this used to be done by UHT manually including it.
 					if (Target.bForceIncludePCHHeadersForGenCppFilesWhenPCHIsDisabled && GeneratedCPPCompileEnvironment.bHasPrecompiledHeader == false && Rules.PrivatePCHHeaderFile != null && Rules.PCHUsage != ModuleRules.PCHUsageMode.UseExplicitOrSharedPCHs)
 					{
@@ -2010,7 +2017,7 @@ namespace UnrealBuildTool
 			Result.bWarningsAsErrors |= Rules.bWarningsAsErrors;
 			Result.ShadowVariableWarningLevel = Rules.ShadowVariableWarningLevel;
 			Result.UnsafeTypeCastWarningLevel = Rules.UnsafeTypeCastWarningLevel;
-			Result.bDisableStaticAnalysis = Rules.bDisableStaticAnalysis;
+			Result.bDisableStaticAnalysis = Rules.bDisableStaticAnalysis || (Target.bStaticAnalyzerProjectOnly && Rules.bTreatAsEngineModule);
 			Result.bStaticAnalyzerExtensions = Rules.bStaticAnalyzerExtensions;
 			Result.StaticAnalyzerRulesets = Rules.StaticAnalyzerRulesets;
 			Result.StaticAnalyzerCheckers = Rules.StaticAnalyzerCheckers;
diff --git a/Engine/Source/Programs/UnrealBuildTool/Platform/Android/AndroidToolChain.cs b/Engine/Source/Programs/UnrealBuildTool/Platform/Android/AndroidToolChain.cs
--- a/Engine/Source/Programs/UnrealBuildTool/Platform/Android/AndroidToolChain.cs
+++ b/Engine/Source/Programs/UnrealBuildTool/Platform/Android/AndroidToolChain.cs
@@ -1020,7 +1020,7 @@ namespace UnrealBuildTool
 
 		protected override CPPOutput CompileCPPFiles(CppCompileEnvironment CompileEnvironment, IEnumerable<FileItem> InputFiles, DirectoryReference OutputDir, string ModuleName, IActionGraphBuilder Graph)
 		{
-			if (ShouldSkipModule(ModuleName, CompileEnvironment.Architecture))
+			if (ShouldSkipCompile(CompileEnvironment) || ShouldSkipModule(ModuleName, CompileEnvironment.Architecture))
 			{
 				return new CPPOutput();
 			}
diff --git a/Engine/Source/Programs/UnrealBuildTool/Platform/Linux/LinuxToolChain.cs b/Engine/Source/Programs/UnrealBuildTool/Platform/Linux/LinuxToolChain.cs
--- a/Engine/Source/Programs/UnrealBuildTool/Platform/Linux/LinuxToolChain.cs
+++ b/Engine/Source/Programs/UnrealBuildTool/Platform/Linux/LinuxToolChain.cs
@@ -970,6 +970,11 @@ namespace UnrealBuildTool
 
 		protected override CPPOutput CompileCPPFiles(CppCompileEnvironment CompileEnvironment, IEnumerable<FileItem> InputFiles, DirectoryReference OutputDir, string ModuleName, IActionGraphBuilder Graph)
 		{
+			if (ShouldSkipCompile(CompileEnvironment))
+			{
+				return new CPPOutput();
+			}
+
 			List<string> GlobalArguments = new();
 			GetCompileArguments_Global(CompileEnvironment, GlobalArguments);
 
diff --git a/Engine/Source/Programs/UnrealBuildTool/Platform/Windows/VCToolChain.cs b/Engine/Source/Programs/UnrealBuildTool/Platform/Windows/VCToolChain.cs
--- a/Engine/Source/Programs/UnrealBuildTool/Platform/Windows/VCToolChain.cs
+++ b/Engine/Source/Programs/UnrealBuildTool/Platform/Windows/VCToolChain.cs
@@ -1591,6 +1591,11 @@ namespace UnrealBuildTool
 
 		protected override CPPOutput CompileCPPFiles(CppCompileEnvironment CompileEnvironment, IEnumerable<FileItem> InputFiles, DirectoryReference OutputDir, string ModuleName, IActionGraphBuilder Graph)
 		{
+			if (Target.StaticAnalyzer != StaticAnalyzer.None && CompileEnvironment.bDisableStaticAnalysis)
+			{
+				return new CPPOutput();
+			}
+
 			VCCompileAction BaseCompileAction = CreateBaseCompileAction(CompileEnvironment);
 
 			// Create a compile action for each source file.
diff --git a/Engine/Source/Programs/UnrealBuildTool/ToolChain/AppleToolChain.cs b/Engine/Source/Programs/UnrealBuildTool/ToolChain/AppleToolChain.cs
--- a/Engine/Source/Programs/UnrealBuildTool/ToolChain/AppleToolChain.cs
+++ b/Engine/Source/Programs/UnrealBuildTool/ToolChain/AppleToolChain.cs
@@ -290,6 +290,11 @@ namespace UnrealBuildTool
 
 		protected override CPPOutput CompileCPPFiles(CppCompileEnvironment CompileEnvironment, IEnumerable<FileItem> InputFiles, DirectoryReference OutputDir, string ModuleName, IActionGraphBuilder Graph)
 		{
+			if (ShouldSkipCompile(CompileEnvironment))
+			{
+				return new CPPOutput();
+			}
+
 			List<string> GlobalArguments = new();
 
 			GetCompileArguments_Global(CompileEnvironment, GlobalArguments);
diff --git a/Engine/Source/Programs/UnrealBuildTool/ToolChain/ClangToolChain.cs b/Engine/Source/Programs/UnrealBuildTool/ToolChain/ClangToolChain.cs
--- a/Engine/Source/Programs/UnrealBuildTool/ToolChain/ClangToolChain.cs
+++ b/Engine/Source/Programs/UnrealBuildTool/ToolChain/ClangToolChain.cs
@@ -337,6 +337,11 @@ namespace UnrealBuildTool
 			&& CompileEnvironment.PrecompiledHeaderAction != PrecompiledHeaderAction.Create
 			&& !CompileEnvironment.bDisableStaticAnalysis;
 
+		protected bool ShouldSkipCompile(CppCompileEnvironment CompileEnvironment) =>
+				StaticAnalyzer == StaticAnalyzer.Default
+			&& CompileEnvironment.PrecompiledHeaderAction != PrecompiledHeaderAction.Create
+			&& CompileEnvironment.bDisableStaticAnalysis;
+
 		protected virtual void GetCppStandardCompileArgument(CppCompileEnvironment CompileEnvironment, List<string> Arguments)
 		{
 			// https://clang.llvm.org/cxx_status.html
@@ -1216,6 +1221,11 @@ namespace UnrealBuildTool
 
 		protected override CPPOutput CompileCPPFiles(CppCompileEnvironment CompileEnvironment, IEnumerable<FileItem> InputFiles, DirectoryReference OutputDir, string ModuleName, IActionGraphBuilder Graph)
 		{
+			if (ShouldSkipCompile(CompileEnvironment))
+			{
+				return new CPPOutput();
+			}
+
 			List<string> GlobalArguments = new();
 
 			if (!CompileEnvironment.bHasSharedResponseFile)
