UnrealBuildTool: Don't analyze generated files by default

#rnx

diff --git a/Engine/Source/Programs/UnrealBuildTool/Configuration/ReadOnlyTargetRules.cs b/Engine/Source/Programs/UnrealBuildTool/Configuration/ReadOnlyTargetRules.cs
--- a/Engine/Source/Programs/UnrealBuildTool/Configuration/ReadOnlyTargetRules.cs
+++ b/Engine/Source/Programs/UnrealBuildTool/Configuration/ReadOnlyTargetRules.cs
@@ -405,7 +405,7 @@ namespace UnrealBuildTool
 
 		public bool bStaticAnalyzerProjectOnly => Inner.bStaticAnalyzerProjectOnly;
 
-		public bool bStaticAnalyzerSkipGenerated => Inner.bStaticAnalyzerSkipGenerated;
+		public bool bStaticAnalyzerIncludeGenerated => Inner.bStaticAnalyzerIncludeGenerated;
 
 		public int MinFilesUsingPrecompiledHeader => Inner.MinFilesUsingPrecompiledHeader;
 
diff --git a/Engine/Source/Programs/UnrealBuildTool/Configuration/TargetRules.cs b/Engine/Source/Programs/UnrealBuildTool/Configuration/TargetRules.cs
--- a/Engine/Source/Programs/UnrealBuildTool/Configuration/TargetRules.cs
+++ b/Engine/Source/Programs/UnrealBuildTool/Configuration/TargetRules.cs
@@ -1985,11 +1985,11 @@ namespace UnrealBuildTool
 		public bool bStaticAnalyzerProjectOnly = false;
 
 		/// <summary>
-		/// Skip analyzing generated source files
+		/// When enabled, generated source files will be analyzed
 		/// </summary>
-		[CommandLine("-StaticAnalyzerSkipGenerated")]
+		[CommandLine("-StaticAnalyzerIncludeGenerated")]
 		[XmlConfigFile(Category = "BuildConfiguration")]
-		public bool bStaticAnalyzerSkipGenerated = false;
+		public bool bStaticAnalyzerIncludeGenerated = false;
 
 		/// <summary>
 		/// The minimum number of files that must use a pre-compiled header before it will be created and used.
diff --git a/Engine/Source/Programs/UnrealBuildTool/Configuration/UEBuildModuleCPP.cs b/Engine/Source/Programs/UnrealBuildTool/Configuration/UEBuildModuleCPP.cs
--- a/Engine/Source/Programs/UnrealBuildTool/Configuration/UEBuildModuleCPP.cs
+++ b/Engine/Source/Programs/UnrealBuildTool/Configuration/UEBuildModuleCPP.cs
@@ -561,7 +561,7 @@ namespace UnrealBuildTool
 						bMergeUnityFiles = false;
 					}
 
-					if (Target.StaticAnalyzer != StaticAnalyzer.None && Target.bStaticAnalyzerSkipGenerated)
+					if (Target.StaticAnalyzer != StaticAnalyzer.None && !Target.bStaticAnalyzerIncludeGenerated)
 					{
 						GeneratedCPPCompileEnvironment = new CppCompileEnvironment(GeneratedCPPCompileEnvironment);
 						GeneratedCPPCompileEnvironment.bDisableStaticAnalysis = true;
