When compiling with -IncludeHeaders, don't include headers from modules that are marked as not supporting IWYU.

diff --git a/Engine/Source/Programs/UnrealBuildTool/Configuration/UEBuildModuleCPP.cs b/Engine/Source/Programs/UnrealBuildTool/Configuration/UEBuildModuleCPP.cs
--- a/Engine/Source/Programs/UnrealBuildTool/Configuration/UEBuildModuleCPP.cs
+++ b/Engine/Source/Programs/UnrealBuildTool/Configuration/UEBuildModuleCPP.cs
@@ -786,15 +786,15 @@ namespace UnrealBuildTool
 			// IWYU needs to build all headers separate from cpp files to produce proper recommendations for includes
 			if (Target.bIncludeHeaders)
 			{
+				if (Target.bHeadersOnly)
+				{
+					LinkInputFiles.Clear();
+				}
+
 				// Collect the headers that should be built
 				List<FileItem> HeaderFileItems = GetCompilableHeaders(InputFiles, CompileEnvironment);
 				if (HeaderFileItems.Count > 0)
 				{
-					if (Target.bHeadersOnly)
-					{
-						LinkInputFiles.Clear();
-					}
-
 					// Add the compile actions
 					LinkInputFiles.AddRange(ToolChain.CompileAllCPPFiles(CompileEnvironment, HeaderFileItems, IntermediateDirectory, Name, Graph).ObjectFiles);
 				}
@@ -805,6 +805,11 @@ namespace UnrealBuildTool
 
 		List<FileItem> GetCompilableHeaders(InputFileCollection InputFiles, CppCompileEnvironment CompileEnvironment)
 		{
+			if (Rules.IWYUSupport == IWYUSupport.None)
+			{
+				return new List<FileItem>();
+			}
+
 			// Find FileItems for module's pch files
 			FileItem? PrivatePchFileItem = null;
 			if (Rules.PrivatePCHHeaderFile != null)
