UnrealBuildTool: Speedup VSWorkspace geneation

#rnx

diff --git a/Engine/Source/Programs/UnrealBuildTool/ProjectFiles/VisualStudioWorkspace/VSWorkspaceProjectFile.cs b/Engine/Source/Programs/UnrealBuildTool/ProjectFiles/VisualStudioWorkspace/VSWorkspaceProjectFile.cs
--- a/Engine/Source/Programs/UnrealBuildTool/ProjectFiles/VisualStudioWorkspace/VSWorkspaceProjectFile.cs
+++ b/Engine/Source/Programs/UnrealBuildTool/ProjectFiles/VisualStudioWorkspace/VSWorkspaceProjectFile.cs
@@ -1,11 +1,12 @@
 // Copyright Epic Games, Inc. All Rights Reserved.
 
 using System;
+using System.Collections.Concurrent;
 using System.Collections.Generic;
 using System.IO;
 using System.Linq;
-using System.Text;
 using System.Text.Json;
+using System.Threading.Tasks;
 using EpicGames.Core;
 using Microsoft.Extensions.Logging;
 using UnrealBuildBase;
@@ -53,13 +54,12 @@ namespace UnrealBuildTool
 		{
 			DirectoryReference ProjectRootFolder = RootPath;
 
-			Dictionary<FileReference, (UEBuildTarget BuildTarget, bool bBuildByDefault)> FileToTarget = new();
-
 			HashSet<UnrealTargetPlatform> ServerPlatforms = Utils.GetPlatformsInClass(UnrealPlatformClass.Server).ToHashSet();
 
-			foreach (UnrealTargetPlatform Platform in InPlatforms)
+			foreach (UnrealTargetConfiguration Configuration in InConfigurations)
 			{
-				foreach (UnrealTargetConfiguration Configuration in InConfigurations)
+				List<Task> Tasks = new List<Task>();
+				foreach (UnrealTargetPlatform Platform in InPlatforms)
 				{
 					foreach (ProjectTarget ProjectTarget in ProjectTargets.OfType<ProjectTarget>())
 					{
@@ -90,61 +90,53 @@ namespace UnrealBuildTool
 							continue;
 						}
 
-						bool bBuildByDefault = ShouldBuildByDefaultForSolutionTargets && ProjectTarget.SupportedPlatforms.Contains(Platform);
-
-						UnrealArchitectures ProjectArchitectures = UEBuildPlatform
-							.GetBuildPlatform(Platform)
-							.ArchitectureConfig.ActiveArchitectures(ProjectTarget.UnrealProjectFilePath, ProjectTarget.Name);
-
-						TargetDescriptor TargetDesc = new(ProjectTarget.UnrealProjectFilePath, ProjectTarget.Name,
-							Platform, Configuration, ProjectArchitectures, Arguments);
-						TargetDesc.IntermediateEnvironment = UnrealIntermediateEnvironment.GenerateProjectFiles;
-
-						try
-						{
-							FileReference OutputFile = FileReference.Combine(ProjectRootFolder, 
-								$"{ProjectTarget.TargetFilePath.GetFileNameWithoutAnyExtensions()}_{Configuration}_{Platform}.json");
-
-							UEBuildTarget BuildTarget = UEBuildTarget.Create(TargetDesc, false, false, false, UnrealIntermediateEnvironment.GenerateProjectFiles, Logger);
-							FileToTarget.Add(OutputFile, (BuildTarget, bBuildByDefault));
-						}
-						catch (Exception Ex)
+						Task CreateTask = Task.Run(async () =>
 						{
-							Logger.LogWarning("Exception while generating include data for Target:{Target}, Platform: {Platform}, Configuration: {Configuration}", TargetDesc.Name, Platform.ToString(), Configuration.ToString());
-							Logger.LogWarning("{Ex}", Ex.ToString());
-						}
+							bool bBuildByDefault = ShouldBuildByDefaultForSolutionTargets && ProjectTarget.SupportedPlatforms.Contains(Platform);
+
+							UnrealArchitectures ProjectArchitectures = UEBuildPlatform
+								.GetBuildPlatform(Platform)
+								.ArchitectureConfig.ActiveArchitectures(ProjectTarget.UnrealProjectFilePath, ProjectTarget.Name);
+
+							TargetDescriptor TargetDesc = new(ProjectTarget.UnrealProjectFilePath, ProjectTarget.Name,
+								Platform, Configuration, ProjectArchitectures, Arguments);
+							TargetDesc.IntermediateEnvironment = UnrealIntermediateEnvironment.GenerateProjectFiles;
+
+							try
+							{
+								FileReference OutputFile = FileReference.Combine(ProjectRootFolder,
+									$"{ProjectTarget.TargetFilePath.GetFileNameWithoutAnyExtensions()}_{Configuration}_{Platform}.json");
+
+								UEBuildTarget BuildTarget = UEBuildTarget.Create(TargetDesc, false, false, false, UnrealIntermediateEnvironment.GenerateProjectFiles, Logger);
+								BuildTarget.PreBuildSetup(Logger);
+
+								ExportedTargetInfo TargetInfo = ExportTarget(BuildTarget, bBuildByDefault, PlatformProjectGenerators, Logger);
+
+								DirectoryReference.CreateDirectory(OutputFile.Directory);
+								using FileStream Stream = new(OutputFile.FullName, FileMode.Create, FileAccess.Write);
+								await JsonSerializer.SerializeAsync(Stream, TargetInfo, options: new JsonSerializerOptions()
+								{
+									PropertyNamingPolicy = JsonNamingPolicy.CamelCase,
+									WriteIndented = Minimize == JsonWriterStyle.Readable,
+								});
+
+								lock (ExportedTargetProjects)
+								{
+									ExportedTargetProjects.Add(TargetInfo);
+								}
+							}
+							catch (Exception Ex)
+							{
+								Logger.LogWarning("Exception while generating include data for Target:{Target}, Platform: {Platform}, Configuration: {Configuration}", TargetDesc.Name, Platform.ToString(), Configuration.ToString());
+								Logger.LogWarning("{Ex}", Ex.ToString());
+							}
+						});
+
+						Tasks.Add(CreateTask);
 					}
 				}
-			}
 
-			foreach (var Entry in FileToTarget)
-			{
-				var OutputFile = Entry.Key;
-				var BuildTarget = Entry.Value.BuildTarget;
-				var bBuildByDefault = Entry.Value.bBuildByDefault;
-
-				try
-				{
-					BuildTarget.PreBuildSetup(Logger);
-
-					ExportedTargetInfo TargetInfo = ExportTarget(BuildTarget, bBuildByDefault, PlatformProjectGenerators, Logger);
-
-					DirectoryReference.CreateDirectory(OutputFile.Directory);
-					using FileStream Stream = new(OutputFile.FullName, FileMode.Create, FileAccess.Write);
-					JsonSerializer.Serialize(Stream, TargetInfo, options: new JsonSerializerOptions()
-					{
-						PropertyNamingPolicy = JsonNamingPolicy.CamelCase,
-						WriteIndented = Minimize == JsonWriterStyle.Readable,
-					});
-
-					ExportedTargetProjects.Add(TargetInfo);
-				}
-				catch (Exception Ex)
-				{
-					Logger.LogWarning("Exception while generating include data for Target:{Target}, Platform: {Platform}, Configuration: {Configuration}",
-						BuildTarget.AppName, BuildTarget.Platform.ToString(), BuildTarget.Configuration.ToString());
-					Logger.LogWarning("{Ex}", Ex.ToString());
-				}
+				Task.WaitAll(Tasks.ToArray());
 			}
 
 			return true;
diff --git a/Engine/Source/Programs/UnrealBuildTool/System/RulesCompiler.cs b/Engine/Source/Programs/UnrealBuildTool/System/RulesCompiler.cs
--- a/Engine/Source/Programs/UnrealBuildTool/System/RulesCompiler.cs
+++ b/Engine/Source/Programs/UnrealBuildTool/System/RulesCompiler.cs
@@ -81,32 +81,37 @@ namespace UnrealBuildTool
 		/// <returns>New rules assembly</returns>
 		public static RulesAssembly CreateEngineRulesAssembly(bool bUsePrecompiled, bool bSkipCompile, bool bForceCompile, ILogger Logger)
 		{
-			if (EngineRulesAssembly == null)
+			// Prevent multiple conflicting processes building rules assembly at the same time
+			string MutexName = SingleInstanceMutex.GetUniqueMutexForPath("UnrealBuildTool_CreateTargetRulesAssembly", Unreal.EngineDirectory.FullName);
+			using (new SingleInstanceMutex(MutexName, true))
 			{
-				List<PluginInfo> EnginePlugins = new List<PluginInfo>();
-				List<PluginInfo> MarketplacePlugins = new List<PluginInfo>();
-
-				DirectoryReference MarketplaceDirectory = DirectoryReference.Combine(Unreal.EngineDirectory, "Plugins", "Marketplace");
-				foreach (PluginInfo PluginInfo in Plugins.ReadEnginePlugins(Unreal.EngineDirectory))
+				if (EngineRulesAssembly == null)
 				{
-					if (PluginInfo.File.IsUnderDirectory(MarketplaceDirectory))
-					{
-						MarketplacePlugins.Add(PluginInfo);
-					}
-					else
+					List<PluginInfo> EnginePlugins = new List<PluginInfo>();
+					List<PluginInfo> MarketplacePlugins = new List<PluginInfo>();
+
+					DirectoryReference MarketplaceDirectory = DirectoryReference.Combine(Unreal.EngineDirectory, "Plugins", "Marketplace");
+					foreach (PluginInfo PluginInfo in Plugins.ReadEnginePlugins(Unreal.EngineDirectory))
 					{
-						EnginePlugins.Add(PluginInfo);
+						if (PluginInfo.File.IsUnderDirectory(MarketplaceDirectory))
+						{
+							MarketplacePlugins.Add(PluginInfo);
+						}
+						else
+						{
+							EnginePlugins.Add(PluginInfo);
+						}
 					}
-				}
 
-				EngineRulesAssembly = CreateEngineRulesAssemblyInternal(Unreal.GetExtensionDirs(Unreal.EngineDirectory), ProjectFileGenerator.EngineRulesAssemblyName, EnginePlugins, Unreal.IsEngineInstalled() || bUsePrecompiled, bSkipCompile, bForceCompile, null, Logger);
+					EngineRulesAssembly = CreateEngineRulesAssemblyInternal(Unreal.GetExtensionDirs(Unreal.EngineDirectory), ProjectFileGenerator.EngineRulesAssemblyName, EnginePlugins, Unreal.IsEngineInstalled() || bUsePrecompiled, bSkipCompile, bForceCompile, null, Logger);
 
-				if (MarketplacePlugins.Count > 0)
-				{
-					EngineRulesAssembly = CreateMarketplaceRulesAssembly(MarketplacePlugins, Unreal.IsEngineInstalled() || bUsePrecompiled, bSkipCompile, bForceCompile, EngineRulesAssembly, Logger);
+					if (MarketplacePlugins.Count > 0)
+					{
+						EngineRulesAssembly = CreateMarketplaceRulesAssembly(MarketplacePlugins, Unreal.IsEngineInstalled() || bUsePrecompiled, bSkipCompile, bForceCompile, EngineRulesAssembly, Logger);
+					}
 				}
+				return EngineRulesAssembly;
 			}
-			return EngineRulesAssembly;
 		}
 
 		/// <summary>
@@ -234,96 +239,101 @@ namespace UnrealBuildTool
 		/// <returns>New rules assembly</returns>
 		public static RulesAssembly CreateProjectRulesAssembly(FileReference ProjectFileName, bool bUsePrecompiled, bool bSkipCompile, bool bForceCompile, ILogger Logger)
 		{
-			// Check if there's an existing assembly for this project
-			return LoadedAssemblyMap.GetOrAdd(ProjectFileName, _ =>
+			// Prevent multiple conflicting processes building rules assembly at the same time
+			string MutexName = SingleInstanceMutex.GetUniqueMutexForPath("UnrealBuildTool_CreateTargetRulesAssembly", ProjectFileName.FullName);
+			using (new SingleInstanceMutex(MutexName, true))
 			{
-				Logger.LogTrace("Creating project rules assembly for {Project}...", ProjectFileName);
+				// Check if there's an existing assembly for this project
+				return LoadedAssemblyMap.GetOrAdd(ProjectFileName, _ =>
+				{
+					Logger.LogTrace("Creating project rules assembly for {Project}...", ProjectFileName);
 
-				ProjectDescriptor Project = ProjectDescriptor.FromFile(ProjectFileName);
+					ProjectDescriptor Project = ProjectDescriptor.FromFile(ProjectFileName);
 
-				// Create the parent assembly
-				RulesAssembly Parent = CreateEngineRulesAssembly(bUsePrecompiled, bSkipCompile, bForceCompile, Logger);
+					// Create the parent assembly
+					RulesAssembly Parent = CreateEngineRulesAssembly(bUsePrecompiled, bSkipCompile, bForceCompile, Logger);
 
-				DirectoryReference MainProjectDirectory = ProjectFileName.Directory;
-				//DirectoryReference MainProjectSourceDirectory = DirectoryReference.Combine(MainProjectDirectory, "Source");
+					DirectoryReference MainProjectDirectory = ProjectFileName.Directory;
+					//DirectoryReference MainProjectSourceDirectory = DirectoryReference.Combine(MainProjectDirectory, "Source");
 
-				// Create a scope for things in this assembly
-				RulesScope Scope = new RulesScope("Project", Parent.Scope);
+					// Create a scope for things in this assembly
+					RulesScope Scope = new RulesScope("Project", Parent.Scope);
 
-				// Create a new context for modules created by this assembly
-				ModuleRulesContext DefaultModuleContext = new ModuleRulesContext(Scope, MainProjectDirectory);
-				DefaultModuleContext.bCanBuildDebugGame = true;
-				DefaultModuleContext.bCanHotReload = true;
-				DefaultModuleContext.bClassifyAsGameModuleForUHT = true;
-				DefaultModuleContext.bCanUseForSharedPCH = false;
+					// Create a new context for modules created by this assembly
+					ModuleRulesContext DefaultModuleContext = new ModuleRulesContext(Scope, MainProjectDirectory);
+					DefaultModuleContext.bCanBuildDebugGame = true;
+					DefaultModuleContext.bCanHotReload = true;
+					DefaultModuleContext.bClassifyAsGameModuleForUHT = true;
+					DefaultModuleContext.bCanUseForSharedPCH = false;
 
-				// gather modules from project and platforms
-				Dictionary<FileReference, ModuleRulesContext> ModuleFiles = new Dictionary<FileReference, ModuleRulesContext>();
-				List<FileReference> TargetFiles = new List<FileReference>();
+					// gather modules from project and platforms
+					Dictionary<FileReference, ModuleRulesContext> ModuleFiles = new Dictionary<FileReference, ModuleRulesContext>();
+					List<FileReference> TargetFiles = new List<FileReference>();
 
-				// Find all the project directories
-				List<DirectoryReference> ProjectDirectories = Unreal.GetExtensionDirs(ProjectFileName.Directory);
-				if (Project.AdditionalRootDirectories != null)
-				{
-					ProjectDirectories.AddRange(Project.AdditionalRootDirectories);
-				}
+					// Find all the project directories
+					List<DirectoryReference> ProjectDirectories = Unreal.GetExtensionDirs(ProjectFileName.Directory);
+					if (Project.AdditionalRootDirectories != null)
+					{
+						ProjectDirectories.AddRange(Project.AdditionalRootDirectories);
+					}
 
-				// Find all the rules/plugins under the project source directories
-				foreach (DirectoryReference ProjectDirectory in ProjectDirectories)
-				{
-					DirectoryReference ProjectSourceDirectory = DirectoryReference.Combine(ProjectDirectory, "Source");
+					// Find all the rules/plugins under the project source directories
+					foreach (DirectoryReference ProjectDirectory in ProjectDirectories)
+					{
+						DirectoryReference ProjectSourceDirectory = DirectoryReference.Combine(ProjectDirectory, "Source");
 
-					AddModuleRulesWithContext(ProjectSourceDirectory, DefaultModuleContext, ModuleFiles);
-					TargetFiles.AddRange(Rules.FindAllRulesFiles(ProjectSourceDirectory, Rules.RulesFileType.Target));
-				}
+						AddModuleRulesWithContext(ProjectSourceDirectory, DefaultModuleContext, ModuleFiles);
+						TargetFiles.AddRange(Rules.FindAllRulesFiles(ProjectSourceDirectory, Rules.RulesFileType.Target));
+					}
 
-				// Find all the project plugins
-				List<PluginInfo> ProjectPlugins = new List<PluginInfo>();
-				ProjectPlugins.AddRange(Plugins.ReadProjectPlugins(MainProjectDirectory));
+					// Find all the project plugins
+					List<PluginInfo> ProjectPlugins = new List<PluginInfo>();
+					ProjectPlugins.AddRange(Plugins.ReadProjectPlugins(MainProjectDirectory));
 
-				// Add the project's additional plugin directories plugins too
-				if (Project.AdditionalPluginDirectories != null)
-				{
-					foreach (DirectoryReference AdditionalPluginDirectory in Project.AdditionalPluginDirectories)
+					// Add the project's additional plugin directories plugins too
+					if (Project.AdditionalPluginDirectories != null)
 					{
-						ProjectPlugins.AddRange(Plugins.ReadAdditionalPlugins(AdditionalPluginDirectory, Logger));
+						foreach (DirectoryReference AdditionalPluginDirectory in Project.AdditionalPluginDirectories)
+						{
+							ProjectPlugins.AddRange(Plugins.ReadAdditionalPlugins(AdditionalPluginDirectory, Logger));
+						}
 					}
-				}
 
-				Logger.LogTrace(" Found {Count} Plugins:", ProjectPlugins.Count);
-				ProjectPlugins.ForEach(x => { Logger.LogTrace("  {Plugin}", x.File); });
+					Logger.LogTrace(" Found {Count} Plugins:", ProjectPlugins.Count);
+					ProjectPlugins.ForEach(x => { Logger.LogTrace("  {Plugin}", x.File); });
 
-				// Find all the plugin module rules as well as plugin test target and module rules
-				FindModuleRulesForPlugins(ProjectPlugins, DefaultModuleContext, ModuleFiles);
-				FindTestRulesForPlugins(ProjectPlugins, DefaultModuleContext, ModuleFiles, TargetFiles);
+					// Find all the plugin module rules as well as plugin test target and module rules
+					FindModuleRulesForPlugins(ProjectPlugins, DefaultModuleContext, ModuleFiles);
+					FindTestRulesForPlugins(ProjectPlugins, DefaultModuleContext, ModuleFiles, TargetFiles);
 
-				Logger.LogTrace(" Found {Count} Modules:", ModuleFiles.Count);
-				foreach (KeyValuePair<FileReference, ModuleRulesContext> Item in ModuleFiles)
-				{
-					Logger.LogTrace("  {Module}", Item.Key);
-				}
+					Logger.LogTrace(" Found {Count} Modules:", ModuleFiles.Count);
+					foreach (KeyValuePair<FileReference, ModuleRulesContext> Item in ModuleFiles)
+					{
+						Logger.LogTrace("  {Module}", Item.Key);
+					}
 
-				// Add the games project's intermediate source folder
-				DirectoryReference ProjectIntermediateSourceDirectory = DirectoryReference.Combine(MainProjectDirectory, "Intermediate", "Source");
-				if (DirectoryReference.Exists(ProjectIntermediateSourceDirectory))
-				{
-					AddModuleRulesWithContext(ProjectIntermediateSourceDirectory, DefaultModuleContext, ModuleFiles);
-					TargetFiles.AddRange(Rules.FindAllRulesFiles(ProjectIntermediateSourceDirectory, Rules.RulesFileType.Target));
-				}
+					// Add the games project's intermediate source folder
+					DirectoryReference ProjectIntermediateSourceDirectory = DirectoryReference.Combine(MainProjectDirectory, "Intermediate", "Source");
+					if (DirectoryReference.Exists(ProjectIntermediateSourceDirectory))
+					{
+						AddModuleRulesWithContext(ProjectIntermediateSourceDirectory, DefaultModuleContext, ModuleFiles);
+						TargetFiles.AddRange(Rules.FindAllRulesFiles(ProjectIntermediateSourceDirectory, Rules.RulesFileType.Target));
+					}
 
-				RulesAssembly ProjectRulesAssembly;
-				// Compile the assembly. If there are no module or target files, just use the parent assembly.
-				FileReference AssemblyFileName = FileReference.Combine(MainProjectDirectory, "Intermediate", "Build", "BuildRules", ProjectFileName.GetFileNameWithoutExtension() + "ModuleRules" + FrameworkAssemblyExtension);
-				if (ModuleFiles.Count == 0 && TargetFiles.Count == 0 && ProjectPlugins.Count == 0)
-				{
-					ProjectRulesAssembly = Parent;
-				}
-				else
-				{
-					ProjectRulesAssembly = new RulesAssembly(Scope, new List<DirectoryReference> { MainProjectDirectory }, ProjectPlugins, ModuleFiles, TargetFiles, AssemblyFileName, bContainsEngineModules: false, DefaultBuildSettings: null, bReadOnly: UnrealBuildTool.IsProjectInstalled(), bSkipCompile: bSkipCompile, bForceCompile: bForceCompile, Parent: Parent, Logger: Logger);
-				}
-				return ProjectRulesAssembly;
-			});
+					RulesAssembly ProjectRulesAssembly;
+					// Compile the assembly. If there are no module or target files, just use the parent assembly.
+					FileReference AssemblyFileName = FileReference.Combine(MainProjectDirectory, "Intermediate", "Build", "BuildRules", ProjectFileName.GetFileNameWithoutExtension() + "ModuleRules" + FrameworkAssemblyExtension);
+					if (ModuleFiles.Count == 0 && TargetFiles.Count == 0 && ProjectPlugins.Count == 0)
+					{
+						ProjectRulesAssembly = Parent;
+					}
+					else
+					{
+						ProjectRulesAssembly = new RulesAssembly(Scope, new List<DirectoryReference> { MainProjectDirectory }, ProjectPlugins, ModuleFiles, TargetFiles, AssemblyFileName, bContainsEngineModules: false, DefaultBuildSettings: null, bReadOnly: UnrealBuildTool.IsProjectInstalled(), bSkipCompile: bSkipCompile, bForceCompile: bForceCompile, Parent: Parent, Logger: Logger);
+					}
+					return ProjectRulesAssembly;
+				});
+			}
 		}
 
 		/// <summary>
@@ -339,55 +349,60 @@ namespace UnrealBuildTool
 		/// <returns>The new rules assembly</returns>
 		public static RulesAssembly CreatePluginRulesAssembly(FileReference PluginFileName, bool bSkipCompile, bool bForceCompile, RulesAssembly Parent, bool bBuildPluginAsLocal, bool bContainsEngineModules, ILogger Logger)
 		{
-			// Check if there's an existing assembly for this project
-			return LoadedAssemblyMap.GetOrAdd(PluginFileName, _ =>
+			// Prevent multiple conflicting processes building rules assembly at the same time
+			string MutexName = SingleInstanceMutex.GetUniqueMutexForPath("UnrealBuildTool_CreateTargetRulesAssembly", PluginFileName.FullName);
+			using (new SingleInstanceMutex(MutexName, true))
 			{
-				Logger.LogTrace("Creating plugin rules assembly for {Plugin}...", PluginFileName);
+				// Check if there's an existing assembly for this project
+				return LoadedAssemblyMap.GetOrAdd(PluginFileName, _ =>
+				{
+					Logger.LogTrace("Creating plugin rules assembly for {Plugin}...", PluginFileName);
 
-				// Find all the rules source files
-				Dictionary<FileReference, ModuleRulesContext> ModuleFiles = new Dictionary<FileReference, ModuleRulesContext>();
-				List<FileReference> TargetFiles = new List<FileReference>();
+					// Find all the rules source files
+					Dictionary<FileReference, ModuleRulesContext> ModuleFiles = new Dictionary<FileReference, ModuleRulesContext>();
+					List<FileReference> TargetFiles = new List<FileReference>();
 
-				// Create a list of plugins for this assembly. We need to override the parent plugin, if it exists, due to overriding the
-				// setting for bClassifyAsGameModuleForUHT below.
-				List<PluginInfo> ForeignPlugins = new List<PluginInfo>();
-				if (!Parent.EnumeratePlugins().Any(x => x.ChoiceVersion != null && x.ChoiceVersion.File == PluginFileName && x.ChoiceVersion.Type == PluginType.Engine))
-				{
-					PluginInfo ForeignPluginInfo = new PluginInfo(PluginFileName, PluginType.External)
+					// Create a list of plugins for this assembly. We need to override the parent plugin, if it exists, due to overriding the
+					// setting for bClassifyAsGameModuleForUHT below.
+					List<PluginInfo> ForeignPlugins = new List<PluginInfo>();
+					if (!Parent.EnumeratePlugins().Any(x => x.ChoiceVersion != null && x.ChoiceVersion.File == PluginFileName && x.ChoiceVersion.Type == PluginType.Engine))
 					{
-						bExplicitPluginTarget = true
-					};
-					ForeignPlugins.Add(ForeignPluginInfo);
-				}
+						PluginInfo ForeignPluginInfo = new PluginInfo(PluginFileName, PluginType.External)
+						{
+							bExplicitPluginTarget = true
+						};
+						ForeignPlugins.Add(ForeignPluginInfo);
+					}
 
-				Logger.LogTrace(" Found {Count} ForeignPlugins:", ForeignPlugins.Count);
-				ForeignPlugins.ForEach(x => { Logger.LogTrace("  {ForeignPlugin}", x.File); });
+					Logger.LogTrace(" Found {Count} ForeignPlugins:", ForeignPlugins.Count);
+					ForeignPlugins.ForEach(x => { Logger.LogTrace("  {ForeignPlugin}", x.File); });
 
-				// Create a new scope for the plugin. It should not reference anything else.
-				RulesScope Scope = new RulesScope("Plugin", Parent.Scope);
+					// Create a new scope for the plugin. It should not reference anything else.
+					RulesScope Scope = new RulesScope("Plugin", Parent.Scope);
 
-				// Find all the modules
-				ModuleRulesContext PluginModuleContext = new ModuleRulesContext(Scope, PluginFileName.Directory);
-				PluginModuleContext.bClassifyAsGameModuleForUHT = !bContainsEngineModules;
-				if (bBuildPluginAsLocal)
-				{
-					PluginModuleContext.bCanBuildDebugGame = true;
-					PluginModuleContext.bCanHotReload = true;
-					PluginModuleContext.bClassifyAsGameModuleForUHT = true;
-					PluginModuleContext.bCanUseForSharedPCH = false;
-				}
-				FindModuleRulesForPlugins(ForeignPlugins, PluginModuleContext, ModuleFiles);
+					// Find all the modules
+					ModuleRulesContext PluginModuleContext = new ModuleRulesContext(Scope, PluginFileName.Directory);
+					PluginModuleContext.bClassifyAsGameModuleForUHT = !bContainsEngineModules;
+					if (bBuildPluginAsLocal)
+					{
+						PluginModuleContext.bCanBuildDebugGame = true;
+						PluginModuleContext.bCanHotReload = true;
+						PluginModuleContext.bClassifyAsGameModuleForUHT = true;
+						PluginModuleContext.bCanUseForSharedPCH = false;
+					}
+					FindModuleRulesForPlugins(ForeignPlugins, PluginModuleContext, ModuleFiles);
 
-				Logger.LogTrace(" Found {Count} Modules:", ModuleFiles.Count);
-				foreach (KeyValuePair<FileReference, ModuleRulesContext> Item in ModuleFiles)
-				{
-					Logger.LogTrace("  {Module}", Item.Key);
-				}
+					Logger.LogTrace(" Found {Count} Modules:", ModuleFiles.Count);
+					foreach (KeyValuePair<FileReference, ModuleRulesContext> Item in ModuleFiles)
+					{
+						Logger.LogTrace("  {Module}", Item.Key);
+					}
 
-				// Compile the assembly
-				FileReference AssemblyFileName = FileReference.Combine(PluginFileName.Directory, "Intermediate", "Build", "BuildRules", Path.GetFileNameWithoutExtension(PluginFileName.FullName) + "ModuleRules" + FrameworkAssemblyExtension);
-				return new RulesAssembly(Scope, new List<DirectoryReference> { PluginFileName.Directory }, ForeignPlugins, ModuleFiles, TargetFiles, AssemblyFileName, bContainsEngineModules, DefaultBuildSettings: null, bReadOnly: false, bSkipCompile: bSkipCompile, bForceCompile: bForceCompile, Parent: Parent, Logger: Logger);
-			});
+					// Compile the assembly
+					FileReference AssemblyFileName = FileReference.Combine(PluginFileName.Directory, "Intermediate", "Build", "BuildRules", Path.GetFileNameWithoutExtension(PluginFileName.FullName) + "ModuleRules" + FrameworkAssemblyExtension);
+					return new RulesAssembly(Scope, new List<DirectoryReference> { PluginFileName.Directory }, ForeignPlugins, ModuleFiles, TargetFiles, AssemblyFileName, bContainsEngineModules, DefaultBuildSettings: null, bReadOnly: false, bSkipCompile: bSkipCompile, bForceCompile: bForceCompile, Parent: Parent, Logger: Logger);
+				});
+			}
 		}
 
 		/// <summary>
