Fix SlateUGS binary in wrong location

diff --git a/Engine/Source/Programs/UnrealBuildTool/Configuration/UEBuildTarget.cs b/Engine/Source/Programs/UnrealBuildTool/Configuration/UEBuildTarget.cs
--- a/Engine/Source/Programs/UnrealBuildTool/Configuration/UEBuildTarget.cs
+++ b/Engine/Source/Programs/UnrealBuildTool/Configuration/UEBuildTarget.cs
@@ -4750,7 +4750,10 @@ namespace UnrealBuildTool
 				// built with a ProjectFile want to be under the Project dir, so just force it into the ProjectDir if one wasn't found otherwise
 				if (bIsUniqueBuildProgram && Rules.File.IsUnderDirectory(Unreal.EngineDirectory) && OutputDirectory == Unreal.EngineDirectory)
 				{
-					OutputDirectory = ProjectDirectory;
+					if (ProjectFile != null && ProjectDirectory.ParentDirectory!.GetDirectoryName() != "Programs")
+					{
+						OutputDirectory = ProjectDirectory;
+					}
 				}
 			}
 			else
