PR #11054: [PVS-Studio] Add new PVS-Studio.exe silent exit code mode.

diff --git a/Engine/Source/Programs/UnrealBuildTool/Platform/Windows/PVSToolChain.cs b/Engine/Source/Programs/UnrealBuildTool/Platform/Windows/PVSToolChain.cs
--- a/Engine/Source/Programs/UnrealBuildTool/Platform/Windows/PVSToolChain.cs
+++ b/Engine/Source/Programs/UnrealBuildTool/Platform/Windows/PVSToolChain.cs
@@ -801,6 +801,11 @@ namespace UnrealBuildTool
 				int Timeout = (int)(Settings.AnalysisTimeoutFlag == AnalysisTimeoutFlags.No_timeout ? 0 : Settings.AnalysisTimeoutFlag);
 				ConfigFileContents.AppendFormat("timeout={0}\n", Timeout);
 
+				if (AnalyzerVersion.CompareTo(new Version("7.28")) >= 0)
+				{
+					ConfigFileContents.Append("silent-exit-code-mode=yes\n");
+				}
+
 				string BaseFileName = PreprocessedFileItem.Location.GetFileName();
 
 				FileReference ConfigFileLocation = FileReference.Combine(OutputDir, BaseFileName + ".cfg");
