UBA: clean up some intellisense messages

#rnx

diff --git a/Engine/Source/Programs/UnrealBuildTool/Executors/Experimental/UBAExecutor.cs b/Engine/Source/Programs/UnrealBuildTool/Executors/Experimental/UBAExecutor.cs
--- a/Engine/Source/Programs/UnrealBuildTool/Executors/Experimental/UBAExecutor.cs
+++ b/Engine/Source/Programs/UnrealBuildTool/Executors/Experimental/UBAExecutor.cs
@@ -301,6 +301,56 @@ namespace UnrealBuildTool
 			_threadedLogger.LogInformation("  Storage capacity {StoreCapacityGb}Gb", UBAConfig.StoreCapacityGb);
 		}
 
+		private async Task WriteActionOutputFile(IEnumerable<LinkedAction> inputActions)
+		{
+			if (String.IsNullOrEmpty(UBAConfig.ActionsOutputFile))
+			{
+				return;
+			}
+
+			if (!UBAConfig.ActionsOutputFile.EndsWith(".yaml", StringComparison.OrdinalIgnoreCase))
+			{
+				_threadedLogger.LogError("UBA actions output file needs to have extension .yaml for UbaCli to understand it");
+			}
+			using System.IO.StreamWriter streamWriter = new System.IO.StreamWriter(UBAConfig.ActionsOutputFile);
+			using System.CodeDom.Compiler.IndentedTextWriter writer = new System.CodeDom.Compiler.IndentedTextWriter(streamWriter, "  ");
+			await writer.WriteAsync("environment: ");
+			await writer.WriteLineAsync(Environment.GetEnvironmentVariable("PATH"));
+			await writer.WriteLineAsync("processes:");
+			writer.Indent++;
+			int index = 0;
+			foreach (LinkedAction action in inputActions)
+			{
+				action.SortIndex = index++;
+				await writer.WriteLineAsync($"- id: {action.SortIndex}");
+				writer.Indent++;
+				await writer.WriteLineAsync($"app: {action.CommandPath}");
+				await writer.WriteLineAsync($"arg: {action.CommandArguments}");
+				await writer.WriteLineAsync($"dir: {action.WorkingDirectory}");
+				await writer.WriteLineAsync($"desc: {action.StatusDescription}");
+				if (action.Weight != 1.0f)
+				{
+					await writer.WriteLineAsync($"weight: {action.Weight}");
+				}
+				if (!action.bCanExecuteInUBA)
+				{
+					await writer.WriteLineAsync("detour: false");
+				}
+				else if (!action.bCanExecuteRemotely)
+				{
+					await writer.WriteLineAsync("remote: false");
+				}
+				if (action.PrerequisiteActions.Any())
+				{
+					await writer.WriteAsync("dep: [");
+					await writer.WriteAsync(String.Join(", ", action.PrerequisiteActions.Select(x => x.SortIndex)));
+					await writer.WriteLineAsync("]");
+				}
+				writer.Indent--;
+				await writer.WriteLineNoTabsAsync(null);
+			}
+		}
+
 		[System.Diagnostics.CodeAnalysis.SuppressMessage("Globalization", "CA1308:Normalize strings to uppercase", Justification = "lowercase crypto string")]
 		public static string CreateCrypto()
 		{
@@ -317,56 +367,13 @@ namespace UnrealBuildTool
 				return true;
 			}
 
-			if (!String.IsNullOrEmpty(UBAConfig.ActionsOutputFile))
-			{
-				if (!UBAConfig.ActionsOutputFile.EndsWith(".yaml"))
-					_threadedLogger.LogError("UBA actions output file needs to have extension .yaml for UbaCli to understand it");
-				using var writer1 = new System.IO.StreamWriter(UBAConfig.ActionsOutputFile);
-				using var writer = new System.CodeDom.Compiler.IndentedTextWriter(writer1, "  ");
-				writer.Write("environment: ");
-				writer.WriteLine(Environment.GetEnvironmentVariable("PATH"));
-				writer.WriteLine("processes:");
-				writer.Indent++;
-				int index = 0;
-				foreach (var action in inputActions)
-				{
-					action.SortIndex = index++;
-					writer.WriteLine($"- id: {action.SortIndex}");
-					writer.Indent++;
-					writer.WriteLine($"app: {action.CommandPath}");
-					writer.WriteLine($"arg: {action.CommandArguments}");
-					writer.WriteLine($"dir: {action.WorkingDirectory}");
-					writer.WriteLine($"desc: {action.StatusDescription}");
-					if (action.Weight != 1.0f)
-						writer.WriteLine($"weight: {action.Weight}");
-					if (!action.bCanExecuteInUBA)
-						writer.WriteLine("detour: false");
-					else if (!action.bCanExecuteRemotely)
-						writer.WriteLine("remote: false");
-					if (action.PrerequisiteActions.Any())
-					{
-						writer.Write("dep: [");
-						bool isFirst = true;
-						foreach (var dep in action.PrerequisiteActions)
-						{
-							if (!isFirst)
-								writer.Write(", ");
-							isFirst = false;
-							writer.Write(dep.SortIndex.ToString());
-						}
-						writer.WriteLine("]");
-					}
-					writer.Indent--;
-					writer.WriteLineNoTabs(null);
-				}
-			}
-
 			if (inputActions.Count() < NumParallelProcesses && !UBAConfig.bForceBuildAllRemote)
 			{
 				UBAConfig.bDisableRemote = true;
 			}
 
 			PrintConfiguration();
+			await WriteActionOutputFile(inputActions);
 
 			logger = _threadedLogger;
 
