[Uba]
* Changed lots of verbose logging to go through "Detail" instead of "Info" inside Uba code
* Added -UBADetailedLog to UBT and changed so unless UBADetailedLog is set all detail logging goes through LogDebug (which means they end up in log file instead of on screen)

diff --git a/Engine/Source/Programs/Shared/EpicGames.UBA/ILogger.cs b/Engine/Source/Programs/Shared/EpicGames.UBA/ILogger.cs
--- a/Engine/Source/Programs/Shared/EpicGames.UBA/ILogger.cs
+++ b/Engine/Source/Programs/Shared/EpicGames.UBA/ILogger.cs
@@ -59,10 +59,11 @@ namespace EpicGames.UBA
 		/// Create a ILogger object
 		/// </summary>
 		/// <param name="logger">The Microsoft.Extensions.Logging.ILogger to wrap</param>
+		/// <param name="showDetail">Will log out detailed log entries from native uba to logger</param>
 		/// <returns>The ILogger</returns>
-		public static ILogger CreateLogger(Microsoft.Extensions.Logging.ILogger logger)
+		public static ILogger CreateLogger(Microsoft.Extensions.Logging.ILogger logger, bool showDetail)
 		{
-			return new LoggerImpl(logger);
+			return new LoggerImpl(logger, showDetail);
 		}
 	}
 }
\ No newline at end of file
diff --git a/Engine/Source/Programs/Shared/EpicGames.UBA/Impl/LoggerImpl.cs b/Engine/Source/Programs/Shared/EpicGames.UBA/Impl/LoggerImpl.cs
--- a/Engine/Source/Programs/Shared/EpicGames.UBA/Impl/LoggerImpl.cs
+++ b/Engine/Source/Programs/Shared/EpicGames.UBA/Impl/LoggerImpl.cs
@@ -29,7 +29,7 @@ namespace EpicGames.UBA
 		static extern void DestroyCallbackLogWriter(IntPtr logger);
 		#endregion
 
-		public LoggerImpl(Microsoft.Extensions.Logging.ILogger logger)
+		public LoggerImpl(Microsoft.Extensions.Logging.ILogger logger, bool showDetail)
 		{
 			_logger = logger;
 			_beginScopeCallbackDelegate = BeginScope;
@@ -80,7 +80,7 @@ namespace EpicGames.UBA
 				case LogEntryType.Error: _logger.LogError("{Message}", message); break;
 				case LogEntryType.Warning: _logger.LogWarning("{Message}", message); break;
 				case LogEntryType.Info: _logger.LogInformation("{Message}", message); break;
-				case LogEntryType.Detail: _logger.LogInformation("{Message}", message); break;
+				case LogEntryType.Detail: _logger.LogDebug("{Message}", message); break;
 				case LogEntryType.Debug: _logger.LogDebug("{Message}", message); break;
 			}
 		}
diff --git a/Engine/Source/Programs/UnrealBuildTool/Executors/Experimental/UBAExecutor.cs b/Engine/Source/Programs/UnrealBuildTool/Executors/Experimental/UBAExecutor.cs
--- a/Engine/Source/Programs/UnrealBuildTool/Executors/Experimental/UBAExecutor.cs
+++ b/Engine/Source/Programs/UnrealBuildTool/Executors/Experimental/UBAExecutor.cs
@@ -242,6 +242,13 @@ namespace UnrealBuildTool
 		[XmlConfigFile(Category = "UnrealBuildAccelerator")]
 		[CommandLine("-UBAActionsOutputFile")]
 		public string ActionsOutputFile { get; set; } = String.Empty;
+
+		/// <summary>
+		/// Set to true to see more info about what is happening inside uba and also log output from agents
+		/// </summary>
+		[XmlConfigFile(Category = "UnrealBuildAccelerator")]
+		[CommandLine("-UBADetailedLog", Value = "true")]
+		public bool bDetailedLog { get; set; } = false;
 	}
 
 	class UBAExecutor : ParallelExecutor
@@ -461,7 +468,7 @@ namespace UnrealBuildTool
 					_ = Task.Run(LaunchVisualizer);
 				}
 
-				using EpicGames.UBA.ILogger ubaLogger = EpicGames.UBA.ILogger.CreateLogger(logger);
+				using EpicGames.UBA.ILogger ubaLogger = EpicGames.UBA.ILogger.CreateLogger(logger, UBAConfig.bDetailedLog);
 				using (Server = IServer.CreateServer(UBAConfig.MaxWorkers, UBAConfig.SendSize, ubaLogger, UBAConfig.bUseQuic))
 				{
 					using IStorageServer ubaStorageServer = IStorageServer.CreateStorageServer(Server, ubaLogger, new StorageServerCreateInfo(_rootDirRef.FullName, ((ulong)UBAConfig.StoreCapacityGb) * 1000 * 1000 * 1000, !UBAConfig.bStoreRaw, UBAConfig.Zone));
diff --git a/Engine/Source/Programs/UnrealBuildTool/Executors/Experimental/UBAHordeSession.cs b/Engine/Source/Programs/UnrealBuildTool/Executors/Experimental/UBAHordeSession.cs
--- a/Engine/Source/Programs/UnrealBuildTool/Executors/Experimental/UBAHordeSession.cs
+++ b/Engine/Source/Programs/UnrealBuildTool/Executors/Experimental/UBAHordeSession.cs
@@ -328,6 +328,15 @@ namespace UnrealBuildTool
 			else if (OperatingSystem.IsLinux())
 			{
 				agentFiles.Add("UbaAgent");
+				bool bIsDebug = false;
+				if (bIsDebug)
+				{
+					agentFiles.Add("UbaAgent.debug");
+					agentFiles.Add("UbaAgent.sym");
+					agentFiles.Add("libclang_rt.tsan.so"); // Needs to be copied from autosdk
+					agentFiles.Add("llvm-symbolizer"); // Needs to be copied from autosdk
+				}
+				
 				if (RuntimeInformation.ProcessArchitecture == Architecture.X64)
 				{
 					ubaDir = DirectoryReference.Combine(Unreal.EngineDirectory, "Binaries", "Linux", "UnrealBuildAccelerator");
@@ -697,16 +706,19 @@ namespace UnrealBuildTool
 							arguments.Add("-Log");
 						}
 
-						logger.LogInformation("Executing child process: {Executable} {Arguments}", executable, CommandLineArguments.Join(arguments));
+						LogLevel logLevel = _owner.UBAConfig.bDetailedLog ? LogLevel.Information : LogLevel.Debug;
+
+						logger.Log(logLevel, "Executing child process: {Executable} {Arguments}", executable, CommandLineArguments.Join(arguments));
 
 						ExecuteProcessFlags execFlags = _allowWine ? ExecuteProcessFlags.UseWine : ExecuteProcessFlags.None;
 						await using AgentManagedProcess process = await channel.ExecuteAsync(executable, arguments, null, null, execFlags, cancellationToken);
 						bool shouldConnect = !useListen;
 						self.Started = true;
 						string? line;
+
 						while ((line = await process.ReadLineAsync(cancellationToken)) != null)
 						{
-							logger.LogInformation("{Line}", line);
+							logger.Log(logLevel, "{Line}", line);
 
 							if (shouldConnect && line.Contains("Listening on", StringComparison.OrdinalIgnoreCase)) // This log entry means that the agent is ready for connections.
 							{
