UnrealBuildTool: Simplify GenerateClangDatabase, support -IncludeHeaders, generate ispc headers first by default

#rnx

diff --git a/Engine/Source/Programs/UnrealBuildTool/Configuration/UEBuildTarget.cs b/Engine/Source/Programs/UnrealBuildTool/Configuration/UEBuildTarget.cs
--- a/Engine/Source/Programs/UnrealBuildTool/Configuration/UEBuildTarget.cs
+++ b/Engine/Source/Programs/UnrealBuildTool/Configuration/UEBuildTarget.cs
@@ -1111,7 +1111,7 @@ namespace UnrealBuildTool
 		/// </summary>
 		Default,
 		/// <summary>
-		/// Generate clang database environment is non unity but uses the standard intermediate directory
+		/// Generate clang database environment is non unity and needs it own environment
 		/// </summary>
 		GenerateClangDatabase,
 		/// <summary>
@@ -1902,6 +1902,9 @@ namespace UnrealBuildTool
 				case UnrealIntermediateEnvironment.Query:
 					TargetFolderName += "QRY";
 					break;
+				case UnrealIntermediateEnvironment.GenerateClangDatabase:
+					TargetFolderName += "GCD";
+					break;
 				case UnrealIntermediateEnvironment.GenerateProjectFiles:
 					TargetFolderName += "GPF";
 					break;
diff --git a/Engine/Source/Programs/UnrealBuildTool/Modes/GenerateClangDatabase.cs b/Engine/Source/Programs/UnrealBuildTool/Modes/GenerateClangDatabase.cs
--- a/Engine/Source/Programs/UnrealBuildTool/Modes/GenerateClangDatabase.cs
+++ b/Engine/Source/Programs/UnrealBuildTool/Modes/GenerateClangDatabase.cs
@@ -1,13 +1,12 @@
 // Copyright Epic Games, Inc. All Rights Reserved.
 
+using EpicGames.Core;
+using Microsoft.Extensions.Logging;
 using System;
 using System.Collections.Generic;
 using System.Linq;
 using System.Text;
-using System.Text.RegularExpressions;
 using System.Threading.Tasks;
-using EpicGames.Core;
-using Microsoft.Extensions.Logging;
 using UnrealBuildBase;
 
 namespace UnrealBuildTool
@@ -18,8 +17,6 @@ namespace UnrealBuildTool
 	[ToolMode("GenerateClangDatabase", ToolModeOptions.XmlConfig | ToolModeOptions.BuildPlatforms | ToolModeOptions.SingleInstance | ToolModeOptions.StartPrefetchingEngine | ToolModeOptions.ShowExecutionTime)]
 	class GenerateClangDatabase : ToolMode
 	{
-		static Regex ArgumentRegex = new Regex(@"^(\-include|\-I|\/I|\/imsvc|\-isystem|\/FI|\/Fo)\s*(.*)");
-
 		/// <summary>
 		/// Set of filters for files to include in the database. Relative to the root directory, or to the project file.
 		/// </summary>
@@ -30,87 +27,8 @@ namespace UnrealBuildTool
 		/// Execute any actions which result in code generation (eg. ISPC compilation)
 		/// </summary>
 		[CommandLine("-ExecCodeGenActions")]
-		public bool bExecCodeGenActions = false;
-
-		/// <summary>
-		/// Used for generating a compilation database against the legacy/new VM for Verse. We default to `true`, since
-		/// in the general case the compilation database generated against the legacy VM is not very useful.
-		/// </summary>
-		[CommandLine("-NoUseVerseBPVM")]
-		public bool bNoUseVerseBPVM = true;
-
-		/// <summary>
-		/// This ActionGraphBuilder captures the build output from a UEBuildModuleCPP so it can be consumed later.
-		/// </summary>
-		private class CaptureActionGraphBuilder : IActionGraphBuilder
-		{
-			private readonly ILogger Logger;
-			public List<IExternalAction> CapturedActions = new List<IExternalAction>();
-			public List<Tuple<FileItem, IEnumerable<string>>> CapturedTextFiles = new List<Tuple<FileItem, IEnumerable<string>>>();
-
-			/// <summary>
-			/// Constructor
-			/// </summary>
-			/// <param name="InLogger"></param>
-			public CaptureActionGraphBuilder(ILogger InLogger)
-			{
-				Logger = InLogger;
-			}
-
-			/// <inheritdoc/>
-			public void AddAction(IExternalAction Action)
-			{
-				CapturedActions.Add(Action);
-			}
-
-			/// <inheritdoc/>
-			public void CreateIntermediateTextFile(FileItem FileItem, string Contents, bool AllowAsync)
-			{
-				Utils.WriteFileIfChanged(FileItem, Contents, Logger);
-			}
-
-			/// <inheritdoc/>
-			public void CreateIntermediateTextFile(FileItem FileItem, IEnumerable<string> ContentLines, bool AllowAsync = true)
-			{
-				Utils.WriteFileIfChanged(FileItem, ContentLines, Logger);
-				CapturedTextFiles.Add(new Tuple<FileItem, IEnumerable<string>>(FileItem, ContentLines));
-			}
-
-			/// <inheritdoc/>
-			public void AddSourceDir(DirectoryItem SourceDir)
-			{
-			}
-
-			/// <inheritdoc/>
-			public void AddSourceFiles(DirectoryItem SourceDir, FileItem[] SourceFiles)
-			{
-			}
-
-			/// <inheritdoc/>
-			public void AddHeaderFiles(FileItem[] HeaderFiles)
-			{
-			}
-
-			/// <inheritdoc/>
-			public void AddFileToWorkingSet(FileItem File)
-			{
-			}
-
-			/// <inheritdoc/>
-			public void AddCandidateForWorkingSet(FileItem File)
-			{
-			}
-
-			/// <inheritdoc/>
-			public void AddDiagnostic(string Message)
-			{
-			}
-
-			/// <inheritdoc/>
-			public void SetOutputItemsForModule(string ModuleName, FileItem[] OutputItems)
-			{
-			}
-		}
+		[CommandLine("-NoExecCodeGenActions", Value = "false")]
+		public bool bExecCodeGenActions = true;
 
 		/// <summary>
 		/// Execute the command
@@ -148,7 +66,7 @@ namespace UnrealBuildTool
 			using (ISourceFileWorkingSet WorkingSet = new EmptySourceFileWorkingSet())
 			{
 				// Find the compile commands for each file in the target
-				Dictionary<FileReference, string> FileToCommand = new Dictionary<FileReference, string>();
+				Dictionary<Tuple<string, string>, string> FileToCommand = new();
 				foreach (TargetDescriptor TargetDescriptor in TargetDescriptors)
 				{
 					// Disable PCHs and unity builds for the target
@@ -160,141 +78,79 @@ namespace UnrealBuildTool
 					{
 						TargetDescriptor.AdditionalArguments = TargetDescriptor.AdditionalArguments.Append(new string[] { "-Compiler=Clang" });
 					}
-					if (bNoUseVerseBPVM)
-					{
-						TargetDescriptor.AdditionalArguments = TargetDescriptor.AdditionalArguments.Append(new string[] { "-NoUseVerseBPVM" });
-					}
 
 					// Create a makefile for the target
+					Logger.LogInformation("Creating target...");
 					UEBuildTarget Target = UEBuildTarget.Create(TargetDescriptor, BuildConfiguration, Logger);
 					UEToolChain TargetToolChain = Target.CreateToolchain(Target.Platform, Logger);
 
-					// Execute code generation actions
+					// Create the makefile
+					TargetMakefile Makefile = await Target.BuildAsync(BuildConfiguration, WorkingSet, TargetDescriptor, Logger);
+					List<LinkedAction> Actions = Makefile.Actions.ConvertAll(x => new LinkedAction(x, TargetDescriptor));
+					ActionGraph.Link(Actions, Logger);
+
 					if (bExecCodeGenActions)
 					{
-						// Create the makefile
-						TargetMakefile Makefile = await Target.BuildAsync(BuildConfiguration, WorkingSet, TargetDescriptor, Logger);
-						List<LinkedAction> Actions = Makefile.Actions.ConvertAll(x => new LinkedAction(x, TargetDescriptor));
-						ActionGraph.Link(Actions, Logger);
-
 						// Filter all the actions to execute
-						HashSet<FileItem> PrerequisiteItems = new HashSet<FileItem>(Makefile.Actions.SelectMany(x => x.ProducedItems).Where(x => x.HasExtension(".h") || x.HasExtension(".cpp")));
+						HashSet<FileItem> PrerequisiteItems = new HashSet<FileItem>(Makefile.Actions.SelectMany(x => x.ProducedItems).Where(x => x.HasExtension(".h") || x.HasExtension(".cpp") || x.HasExtension(".cc") || x.HasExtension(".c")));
 						List<LinkedAction> PrerequisiteActions = ActionGraph.GatherPrerequisiteActions(Actions, PrerequisiteItems);
 
-						// Execute these actions
-						if (PrerequisiteActions.Count > 0)
+						Utils.ExecuteCustomBuildSteps(Makefile.PreBuildScripts, Logger);
+
+						// Execute code generation actions
+						if (PrerequisiteActions.Any())
 						{
 							Logger.LogInformation("Executing actions that produce source files...");
 							await ActionGraph.ExecuteActionsAsync(BuildConfiguration, PrerequisiteActions, new List<TargetDescriptor> { TargetDescriptor }, Logger);
 						}
 					}
 
-					// Create all the binaries and modules
-					CppCompileEnvironment GlobalCompileEnvironment = Target.CreateCompileEnvironmentForProjectFiles(Logger);
-					// Some of our platform-specific wrappers check against this preprocessor macro, which `clang`/`clang-cl` will not set by default.
-					if (!GlobalCompileEnvironment.Architectures.bIsMultiArch)  // For now, this will only handle single-architecture builds since setting both macros at once doesn't really make sense.
-					{
-						if (GlobalCompileEnvironment.Architecture.bIsX64)
-						{
-							GlobalCompileEnvironment.AdditionalArguments = "-DPLATFORM_CPU_X86_FAMILY";
-						}
-						else  // Currently, this assumes that non-x64 arch means we're running on ARM.
-						{
-							GlobalCompileEnvironment.AdditionalArguments = "-DPLATFORM_CPU_ARM_FAMILY";
-						}
-					}
-					foreach (UEBuildBinary Binary in Target.Binaries)
+					Logger.LogInformation("Filtering compile actions...");
+
+					IEnumerable<IExternalAction> CompileActions = Actions
+						.Where(x => x.ActionType == ActionType.Compile)
+						.Where(x => x.PrerequisiteItems.Any());
+
+					if (CompileActions.Any())
 					{
-						CppCompileEnvironment BinaryCompileEnvironment = Binary.CreateBinaryCompileEnvironment(GlobalCompileEnvironment);
-						foreach (UEBuildModuleCPP Module in Binary.Modules.OfType<UEBuildModuleCPP>())
+						foreach (IExternalAction Action in CompileActions)
 						{
-							if (!Module.Rules.bUsePrecompiled)
+							FileItem? SourceFile = Action.PrerequisiteItems.FirstOrDefault(x => x.HasExtension(".cpp") || x.HasExtension(".cc") || x.HasExtension(".c")) ?? Action.PrerequisiteItems.FirstOrDefault(x => x.HasExtension(".h"));
+							FileItem? OutputFile = Action.ProducedItems.FirstOrDefault(x => x.HasExtension(".obj") | x.HasExtension(".o"));
+							if (SourceFile == null || OutputFile == null)
 							{
-								// Gather all the files we care about
-								UEBuildModuleCPP.InputFileCollection InputFileCollection = Module.FindInputFiles(Target.Platform, new Dictionary<DirectoryItem, FileItem[]>(), Logger);
-								List<FileItem> InputFiles = new List<FileItem>();
-								InputFiles.AddRange(InputFileCollection.CPPFiles);
-								InputFiles.AddRange(InputFileCollection.CCFiles);
-
-								HashSet<FileItem> FilterFileList = new();
-								foreach (FileItem InputFile in InputFiles)
-								{
-									if (FileFilter == null || FileFilter.Matches(InputFile.Location.MakeRelativeTo(Unreal.RootDirectory)))
-									{
-										FilterFileList.Add(InputFile);
-									}
-								}
-
-								CaptureActionGraphBuilder ActionGraphBuilder = new CaptureActionGraphBuilder(Logger);
-
-								Module.Compile(Target.Rules, TargetToolChain, BinaryCompileEnvironment, WorkingSet, ActionGraphBuilder, Logger);
-
-								List<IExternalAction> ValidActions = new();
-								foreach (IExternalAction Action in ActionGraphBuilder.CapturedActions)
-								{
-									if (Action.ActionType == ActionType.Compile && Action.ProducedItems.Any())
-									{
-										foreach (FileItem Prereq in Action.PrerequisiteItems)
-										{
-											if (FilterFileList.Contains(Prereq))
-											{
-												ValidActions.Add(Action);
-											}
-										}
-									}
-								}
-
-								if (ValidActions.Count != 0)
-								{
-									// convert any rsp files
-									Dictionary<string, string> UpdatedResFiles = new Dictionary<string, string>();
-									foreach (Tuple<FileItem, IEnumerable<string>> FileAndContents in ActionGraphBuilder.CapturedTextFiles)
-									{
-										if (FileAndContents.Item1.AbsolutePath.EndsWith(".rsp") || FileAndContents.Item1.AbsolutePath.EndsWith(".response"))
-										{
-											string NewResPath = ConvertResponseFile(FileAndContents.Item1, FileAndContents.Item2, Logger);
-											UpdatedResFiles[FileAndContents.Item1.AbsolutePath] = NewResPath;
-										}
-									}
-
-									foreach (IExternalAction Action in ValidActions)
-									{
-										// Create the command
-										StringBuilder CommandBuilder = new StringBuilder();
-										string CommandPath = Action.CommandPath.FullName.Contains(' ') ? Utils.MakePathSafeToUseWithCommandLine(Action.CommandPath) : Action.CommandPath.FullName;
-										string CommandArguments = Action.CommandArguments.Replace(".rsp", ".rsp.gcd").Replace(".response", ".response.gcd");
-										CommandBuilder.AppendFormat("{0} {1}", CommandPath, CommandArguments);
-
-										foreach (string ExtraArgument in GetExtraPlatformArguments(TargetToolChain))
-										{
-											CommandBuilder.AppendFormat(" {0}", ExtraArgument);
-										}
+								continue;
+							}
+							// Create the command
+							StringBuilder CommandBuilder = new StringBuilder();
+							string CommandPath = Action.CommandPath.FullName.Contains(' ') ? Utils.MakePathSafeToUseWithCommandLine(Action.CommandPath) : Action.CommandPath.FullName;
+							CommandBuilder.AppendFormat("{0} {1}", CommandPath, Action.CommandArguments);
 
-										// find source file
-										FileItem? SourceFile = Action.PrerequisiteItems.FirstOrDefault(fi => fi.HasExtension(".cpp") || fi.HasExtension(".c") || fi.HasExtension(".c"));
-										if (SourceFile != null)
-										{
-											FileToCommand[SourceFile.Location] = CommandBuilder.ToString();
-										}
-									}
-								}
+							foreach (string ExtraArgument in GetExtraPlatformArguments(TargetToolChain))
+							{
+								CommandBuilder.AppendFormat(" {0}", ExtraArgument);
 							}
+
+							FileToCommand[Tuple.Create(SourceFile.FullName, OutputFile.FullName)] = CommandBuilder.ToString();
 						}
 					}
 				}
 
+				Logger.LogInformation("Writing database...");
+
 				// Write the compile database
 				DirectoryReference DatabaseDirectory = Arguments.GetDirectoryReferenceOrDefault("-OutputDir=", Unreal.RootDirectory);
 				FileReference DatabaseFile = FileReference.Combine(DatabaseDirectory, "compile_commands.json");
 				using (JsonWriter Writer = new JsonWriter(DatabaseFile))
 				{
 					Writer.WriteArrayStart();
-					foreach (KeyValuePair<FileReference, string> FileCommandPair in FileToCommand.OrderBy(x => x.Key.FullName))
+					foreach (KeyValuePair<Tuple<string, string>, string> FileCommandPair in FileToCommand.OrderBy(x => x.Key.Item1))
 					{
 						Writer.WriteObjectStart();
-						Writer.WriteValue("file", FileCommandPair.Key.FullName);
-						Writer.WriteValue("command", FileCommandPair.Value);
-						Writer.WriteValue("directory", Unreal.EngineSourceDirectory.ToString());
+						Writer.WriteValue("file", FileCommandPair.Key.Item1.Replace('\\', '/'));
+						Writer.WriteValue("command", FileCommandPair.Value.Replace('\\', '/'));
+						Writer.WriteValue("directory", Unreal.EngineSourceDirectory.FullName.Replace('\\', '/'));
+						Writer.WriteValue("output", FileCommandPair.Key.Item2.Replace('\\', '/'));
 						Writer.WriteObjectEnd();
 					}
 					Writer.WriteArrayEnd();
@@ -318,72 +174,5 @@ namespace UnrealBuildTool
 
 			return ExtraPlatformArguments;
 		}
-
-		private static string ConvertResponseFile(FileItem OriginalFileItem, IEnumerable<string> FileContents, ILogger Logger)
-		{
-			List<string> NewFileContents = new List<string>(FileContents);
-			FileItem NewFileItem = FileItem.GetItemByFileReference(new FileReference(OriginalFileItem.AbsolutePath + ".gcd"));
-
-			for (int i = 0; i < NewFileContents.Count; i++)
-			{
-				string Line = NewFileContents[i].TrimStart();
-
-				// Ignore empty strings
-				if (String.IsNullOrEmpty(Line))
-				{
-					continue;
-				}
-				// The file that is going to compile
-				else if (!Line.StartsWith("-") && !Line.StartsWith("/") && !Line.StartsWith("@"))
-				{
-					Line = ConvertPath(Line, Line);
-				}
-				// Arguments
-				else
-				{
-					 // TODO this needs a clean way to handle path separators that is used in all instances
-					int StrIndex = Line.IndexOf("..\\");
-					if (StrIndex == -1)
-					{
-						StrIndex = Line.IndexOf("../");
-					}
-					if (StrIndex != -1)
-					{
-						Line = ConvertPath(Line, Line.Substring(StrIndex));
-					}
-					else
-					{
-						Match LineMatch = ArgumentRegex.Match(Line);
-						if (LineMatch.Success)
-						{
-							Line = ConvertPath(Line, LineMatch.Groups[2].Value);
-						}
-					}
-				}
-				// TODO DIRTY DIRTY HACK because we don't support recursion in response files
-				if (Line.StartsWith("@"))
-				{
-					Line = Line.Replace(".rsp", ".rsp.gcd").Replace(".response", ".response.gcd");
-				}
-
-				NewFileContents[i] = Line;
-			}
-
-			Utils.WriteFileIfChanged(NewFileItem, NewFileContents, Logger);
-
-			return NewFileItem.AbsolutePath;
-		}
-
-		private static string ConvertPath(string Line, string OldPath)
-		{
-			OldPath = OldPath.Replace("\"", "");
-			if (OldPath[^1] == '\\' || OldPath[^1] == '/')
-			{
-				OldPath = OldPath.Remove(OldPath.Length - 1, 1);
-			}
-
-			FileReference FileReference = new FileReference(OldPath);
-			return Line.Replace(OldPath, Utils.MakePathSafeToUseWithCommandLine(FileReference.FullName.Replace("\\", "/")));
-		}
 	}
 }
