Only add Android architectures to VS project enabled in project

#android

diff --git a/Engine/Source/Programs/UnrealBuildTool/ProjectFiles/VisualStudio/VCProject.cs b/Engine/Source/Programs/UnrealBuildTool/ProjectFiles/VisualStudio/VCProject.cs
--- a/Engine/Source/Programs/UnrealBuildTool/ProjectFiles/VisualStudio/VCProject.cs
+++ b/Engine/Source/Programs/UnrealBuildTool/ProjectFiles/VisualStudio/VCProject.cs
@@ -168,7 +168,7 @@ namespace UnrealBuildTool
 		/// <param name="Architecture">The target architecture</param>
 		/// <param name="Logger">Logger for output</param>
 		/// <returns>Project context matching the given solution context</returns>
-		public abstract MSBuildProjectContext GetMatchingProjectContext(TargetType SolutionTarget, UnrealTargetConfiguration SolutionConfiguration, UnrealTargetPlatform SolutionPlatform, PlatformProjectGeneratorCollection PlatformProjectGenerators, UnrealArch? Architecture, ILogger Logger);
+		public abstract MSBuildProjectContext? GetMatchingProjectContext(TargetType SolutionTarget, UnrealTargetConfiguration SolutionConfiguration, UnrealTargetPlatform SolutionPlatform, PlatformProjectGeneratorCollection PlatformProjectGenerators, UnrealArch? Architecture, ILogger Logger);
 
 		/// <summary>
 		/// Checks to see if the specified solution platform and configuration is able to map to this project
@@ -315,7 +315,7 @@ namespace UnrealBuildTool
 		}
 
 		/// <inheritdoc/>
-		public override MSBuildProjectContext GetMatchingProjectContext(TargetType SolutionTarget, UnrealTargetConfiguration SolutionConfiguration, UnrealTargetPlatform SolutionPlatform, PlatformProjectGeneratorCollection PlatformProjectGenerators, UnrealArch? Architecture, ILogger Logger)
+		public override MSBuildProjectContext? GetMatchingProjectContext(TargetType SolutionTarget, UnrealTargetConfiguration SolutionConfiguration, UnrealTargetPlatform SolutionPlatform, PlatformProjectGeneratorCollection PlatformProjectGenerators, UnrealArch? Architecture, ILogger Logger)
 		{
 			// Stub projects always build in the same configuration
 			if (IsStubProject)
@@ -429,6 +429,11 @@ namespace UnrealBuildTool
 
 						if (Combination == null)
 						{
+							if (ProjectPlatform == UnrealTargetPlatform.Android)
+							{
+								// ok for Android not to find the architecture
+								return null;
+							}
 							throw new BuildException("Could not find the project config/platform combination in the generated list.");
 						}
 
@@ -734,7 +739,7 @@ namespace UnrealBuildTool
 										ProjectConfigAndTargetCombinations.Add(new ProjectConfigAndTargetCombination(Platform, Configuration, ProjectPlatformName, ProjectConfigurationName, ProjectTarget, Arch));
 									};
 
-									UnrealArchitectures? Architectures = VCProjectFileGenerator.GetPlatformArchitecturesToGenerate(BuildPlatform);
+									UnrealArchitectures? Architectures = VCProjectFileGenerator.GetPlatformArchitecturesToGenerate(BuildPlatform, ProjectTarget);
 									if (Architectures == null)
 									{
 										AddProjectAndTargetCombination(null);
@@ -2601,7 +2606,7 @@ namespace UnrealBuildTool
 		}
 
 		/// <inheritdoc/>
-		public override MSBuildProjectContext GetMatchingProjectContext(TargetType SolutionTarget, UnrealTargetConfiguration SolutionConfiguration, UnrealTargetPlatform SolutionPlatform, PlatformProjectGeneratorCollection PlatformProjectGenerators, UnrealArch? Architecture, ILogger Logger)
+		public override MSBuildProjectContext? GetMatchingProjectContext(TargetType SolutionTarget, UnrealTargetConfiguration SolutionConfiguration, UnrealTargetPlatform SolutionPlatform, PlatformProjectGeneratorCollection PlatformProjectGenerators, UnrealArch? Architecture, ILogger Logger)
 		{
 			// Find the matching platform name
 			string ProjectPlatformName;
diff --git a/Engine/Source/Programs/UnrealBuildTool/ProjectFiles/VisualStudio/VCProjectFileGenerator.cs b/Engine/Source/Programs/UnrealBuildTool/ProjectFiles/VisualStudio/VCProjectFileGenerator.cs
--- a/Engine/Source/Programs/UnrealBuildTool/ProjectFiles/VisualStudio/VCProjectFileGenerator.cs
+++ b/Engine/Source/Programs/UnrealBuildTool/ProjectFiles/VisualStudio/VCProjectFileGenerator.cs
@@ -256,13 +256,49 @@ namespace UnrealBuildTool
 			VCProjectFileContent.AppendLine("    <PlatformToolset>{0}</PlatformToolset>", PlatformToolsetVersionString);
 		}
 
+		// parses project ini for Android to get architecture(s) enabled
+		private static UnrealArchitectures GetAndroidProjectArchitectures(FileReference? ProjectFile, bool bGetAllSupported)
+		{
+			List<string> ActiveArches = new();
+
+			// look in ini settings for what platforms to compile for
+			ConfigHierarchy Ini = ConfigCache.ReadHierarchy(ConfigHierarchyType.Engine, DirectoryReference.FromFile(ProjectFile), UnrealTargetPlatform.Android);
+			bool bBuild;
+
+			if (Ini.GetBool("/Script/AndroidRuntimeSettings.AndroidRuntimeSettings", "bBuildForArm64", out bBuild) && bBuild)
+			{
+				ActiveArches.Add("arm64");
+			}
+			if (Ini.GetBool("/Script/AndroidRuntimeSettings.AndroidRuntimeSettings", "bBuildForx8664", out bBuild) && bBuild)
+			{
+				ActiveArches.Add("x64");
+			}
+
+			// we expect one to be specified
+			if (ActiveArches.Count == 0)
+			{
+				ActiveArches.Add("arm64");
+			}
+
+			return new UnrealArchitectures(ActiveArches);
+		}
+
 		/// <summary>
 		/// Returns a list of architectures to generate unique VS platforms for.
 		/// </summary>
-		public static UnrealArchitectures? GetPlatformArchitecturesToGenerate(UEBuildPlatform BuildPlatform)
+		public static UnrealArchitectures? GetPlatformArchitecturesToGenerate(UEBuildPlatform BuildPlatform, ProjectTarget InProjectTarget)
 		{
-			return (BuildPlatform.ArchitectureConfig.Mode == UnrealArchitectureMode.OneTargetPerArchitecture ||
-				BuildPlatform.ArchitectureConfig.Mode == UnrealArchitectureMode.SingleTargetLinkSeparately) ?
+			if (BuildPlatform.ArchitectureConfig.Mode == UnrealArchitectureMode.SingleTargetLinkSeparately)
+			{
+				// this should only be Android at the moment
+				if (BuildPlatform.Platform == UnrealTargetPlatform.Android)
+				{
+					return InProjectTarget.UnrealProjectFilePath == null ? BuildPlatform.ArchitectureConfig.AllSupportedArchitectures
+						: GetAndroidProjectArchitectures(InProjectTarget.UnrealProjectFilePath, true);
+				}
+				return BuildPlatform.ArchitectureConfig.AllSupportedArchitectures;
+			}
+			return (BuildPlatform.ArchitectureConfig.Mode == UnrealArchitectureMode.OneTargetPerArchitecture) ?
 				BuildPlatform.ArchitectureConfig.AllSupportedArchitectures : null;
 		}
 
@@ -815,7 +851,12 @@ namespace UnrealBuildTool
 							foreach (VCSolutionConfigCombination SolutionConfigCombination in SolutionConfigCombinations)
 							{
 								// Get the context for the current solution context
-								MSBuildProjectContext ProjectContext = CurProject.GetMatchingProjectContext(SolutionConfigCombination.TargetConfigurationName, SolutionConfigCombination.Configuration, SolutionConfigCombination.Platform, PlatformProjectGenerators, SolutionConfigCombination.Architecture, Logger);
+								MSBuildProjectContext? ProjectContext = CurProject.GetMatchingProjectContext(SolutionConfigCombination.TargetConfigurationName, SolutionConfigCombination.Configuration, SolutionConfigCombination.Platform, PlatformProjectGenerators, SolutionConfigCombination.Architecture, Logger);
+
+								if (ProjectContext == null)
+								{
+									continue;
+								}
 
 								// Override the configuration to build for UBT
 								if (Settings.bBuildUBTInDebug && CurProject == UBTProject)
@@ -1065,7 +1106,7 @@ namespace UnrealBuildTool
 							);
 						};
 
-						UnrealArchitectures? Architectures = GetPlatformArchitecturesToGenerate(BuildPlatform);
+						UnrealArchitectures? Architectures = GetPlatformArchitecturesToGenerate(BuildPlatform, ProjectTarget);
 						if (Architectures == null)
 						{
 							AddSolutionConfig(null, OutSolutionConfigs);
