Fixing compile error with UBT_TARGET_BUILD_PLUGINS undefined in LiveLinkHub

diff --git a/Engine/Source/Programs/UnrealBuildTool/Configuration/UEBuildTarget.cs b/Engine/Source/Programs/UnrealBuildTool/Configuration/UEBuildTarget.cs
--- a/Engine/Source/Programs/UnrealBuildTool/Configuration/UEBuildTarget.cs
+++ b/Engine/Source/Programs/UnrealBuildTool/Configuration/UEBuildTarget.cs
@@ -4438,11 +4438,11 @@ namespace UnrealBuildTool
 		/// <param name="moduleRules">the module's rules</param>
 		private void ModifyModuleRulesForTarget(string moduleName, ModuleRules moduleRules)
 		{
-			if (moduleName == "Projects" && (Rules.Type == TargetType.Editor))
+			if (moduleName == "Projects" && ((Rules.Type == TargetType.Editor || Rules.Type == TargetType.Program)))
 			{
 				// Monolithic and non-shared environment builds compile in the build plugins. Non-monolithic editor builds save them in the BuildPlugins receipt to avoid invalidating the shared build environment.
 				// See Projects.Build.cs
-				bool bUsingTargetReceipt = (Rules.BuildEnvironment == TargetBuildEnvironment.Shared) && (Rules.LinkType != TargetLinkType.Monolithic);
+				bool bUsingTargetReceipt = ((Rules.Type == TargetType.Editor) && (Rules.BuildEnvironment == TargetBuildEnvironment.Shared) && (Rules.LinkType != TargetLinkType.Monolithic));
 				if (!bUsingTargetReceipt)
 				{
 					// Compile in the build plugins list for the runtime plugin manager to filter optional plugin references against. Only plugins that were actually built
