UnrealBuildTool: Fix several fields that should be properties in TargetRules

#rnx

diff --git a/Engine/Source/Programs/UnrealBuildTool/Configuration/TargetRules.cs b/Engine/Source/Programs/UnrealBuildTool/Configuration/TargetRules.cs
--- a/Engine/Source/Programs/UnrealBuildTool/Configuration/TargetRules.cs
+++ b/Engine/Source/Programs/UnrealBuildTool/Configuration/TargetRules.cs
@@ -1982,14 +1982,14 @@ namespace UnrealBuildTool
 		/// </summary>
 		[CommandLine("-StaticAnalyzerProjectOnly")]
 		[XmlConfigFile(Category = "BuildConfiguration")]
-		public bool bStaticAnalyzerProjectOnly = false;
+		public bool bStaticAnalyzerProjectOnly { get; set; } = false;
 
 		/// <summary>
 		/// When enabled, generated source files will be analyzed
 		/// </summary>
 		[CommandLine("-StaticAnalyzerIncludeGenerated")]
 		[XmlConfigFile(Category = "BuildConfiguration")]
-		public bool bStaticAnalyzerIncludeGenerated = false;
+		public bool bStaticAnalyzerIncludeGenerated { get; set; } = false;
 
 		/// <summary>
 		/// The minimum number of files that must use a pre-compiled header before it will be created and used.
