UnrealBuildTool: Setting to control print level of PVS-Studio analyzer warnings (-StaticAnalyzerPVSPrintLevel)

diff --git a/Engine/Source/Programs/UnrealBuildTool/Configuration/ReadOnlyTargetRules.cs b/Engine/Source/Programs/UnrealBuildTool/Configuration/ReadOnlyTargetRules.cs
--- a/Engine/Source/Programs/UnrealBuildTool/Configuration/ReadOnlyTargetRules.cs
+++ b/Engine/Source/Programs/UnrealBuildTool/Configuration/ReadOnlyTargetRules.cs
@@ -403,6 +403,8 @@ namespace UnrealBuildTool
 
 		public StaticAnalyzerMode StaticAnalyzerMode => Inner.StaticAnalyzerMode;
 
+		public int StaticAnalyzerPVSPrintLevel => Inner.StaticAnalyzerPVSPrintLevel;
+
 		public bool bStaticAnalyzerProjectOnly => Inner.bStaticAnalyzerProjectOnly;
 
 		public bool bStaticAnalyzerIncludeGenerated => Inner.bStaticAnalyzerIncludeGenerated;
diff --git a/Engine/Source/Programs/UnrealBuildTool/Configuration/TargetRules.cs b/Engine/Source/Programs/UnrealBuildTool/Configuration/TargetRules.cs
--- a/Engine/Source/Programs/UnrealBuildTool/Configuration/TargetRules.cs
+++ b/Engine/Source/Programs/UnrealBuildTool/Configuration/TargetRules.cs
@@ -1977,6 +1977,13 @@ namespace UnrealBuildTool
 		[XmlConfigFile(Category = "BuildConfiguration")]
 		public StaticAnalyzerMode StaticAnalyzerMode { get; set; } = StaticAnalyzerMode.Deep;
 
+		/// <summary>
+		/// The level of warnings to print when analyzing using PVS-Studio
+		/// </summary>
+		[CommandLine("-StaticAnalyzerPVSPrintLevel")]
+		[XmlConfigFile(Category = "BuildConfiguration")]
+		public int StaticAnalyzerPVSPrintLevel { get; set; } = 1;
+
 		/// <summary>
 		/// Only run static analysis against project modules, skipping engine modules
 		/// </summary>
diff --git a/Engine/Source/Programs/UnrealBuildTool/Platform/Windows/PVSToolChain.cs b/Engine/Source/Programs/UnrealBuildTool/Platform/Windows/PVSToolChain.cs
--- a/Engine/Source/Programs/UnrealBuildTool/Platform/Windows/PVSToolChain.cs
+++ b/Engine/Source/Programs/UnrealBuildTool/Platform/Windows/PVSToolChain.cs
@@ -389,6 +389,12 @@ namespace UnrealBuildTool
 		[CommandLine("-Ignored")]
 		FileReference? IgnoredFile = null;
 
+		/// <summary>
+		/// The maximum level of warnings to print
+		/// </summary>
+		[CommandLine("-PrintLevel")]
+		public int PrintLevel { get; set; } = 1;
+
 		/// <summary>
 		/// If all ThirdParty code should be ignored
 		/// </summary>
@@ -466,7 +472,7 @@ namespace UnrealBuildTool
 									RawWriter.WriteLine(Line);
 
 									// Output the line to the log
-									if (!bFalseAlarm && Level == 1)
+									if (!bFalseAlarm && Level <= PrintLevel)
 									{
 										Logger.LogWarning(KnownLogEvents.Compiler, "{Path}({LineNumber}): warning {WarningCode}: {WarningMessage}", LogValue.SourceFile(file, FileName), LineNumber, WarningCode, WarningMessage);
 									}
@@ -883,7 +889,7 @@ namespace UnrealBuildTool
 			AnalyzeAction.ActionType = ActionType.PostBuildStep;
 			AnalyzeAction.CommandDescription = "Process PVS-Studio Results";
 			AnalyzeAction.CommandPath = Unreal.DotnetPath;
-			AnalyzeAction.CommandArguments = $"\"{Unreal.UnrealBuildToolDllPath}\" -Mode=PVSGather -Input=\"{InputFileListItem.Location}\" -Output=\"{OutputFile}\" -Ignored=\"{IgnoredFileListeItem.Location}\" ";
+			AnalyzeAction.CommandArguments = $"\"{Unreal.UnrealBuildToolDllPath}\" -Mode=PVSGather -Input=\"{InputFileListItem.Location}\" -Output=\"{OutputFile}\" -Ignored=\"{IgnoredFileListeItem.Location}\" -PrintLevel={Target.StaticAnalyzerPVSPrintLevel} ";
 			AnalyzeAction.WorkingDirectory = Unreal.EngineSourceDirectory;
 			AnalyzeAction.PrerequisiteItems.Add(InputFileListItem);
 			AnalyzeAction.PrerequisiteItems.Add(IgnoredFileListeItem);
