Allow to set a custom debug keystore instead of always using default one

diff --git a/Engine/Config/BaseGame.ini b/Engine/Config/BaseGame.ini
--- a/Engine/Config/BaseGame.ini
+++ b/Engine/Config/BaseGame.ini
@@ -136,6 +136,8 @@ WriteBackMetadataToAssetRegistry=Disabled
 
 +IniKeyDenylist=KeyStorePassword
 +IniKeyDenylist=KeyPassword
++IniKeyDenylist=DebugKeyStorePassword
++IniKeyDenylist=DebugKeyPassword
 +IniKeyDenylist=rsa.privateexp
 +IniKeyDenylist=rsa.modulus
 +IniKeyDenylist=rsa.publicexp
diff --git a/Engine/Source/Developer/Android/AndroidPlatformEditor/Private/AndroidTargetSettingsCustomization.cpp b/Engine/Source/Developer/Android/AndroidPlatformEditor/Private/AndroidTargetSettingsCustomization.cpp
--- a/Engine/Source/Developer/Android/AndroidPlatformEditor/Private/AndroidTargetSettingsCustomization.cpp
+++ b/Engine/Source/Developer/Android/AndroidPlatformEditor/Private/AndroidTargetSettingsCustomization.cpp
@@ -106,8 +106,10 @@ void FAndroidTargetSettingsCustomization::BuildAppManifestSection(IDetailLayoutB
 	AdvancedBuildCategory.InitiallyCollapsed(true);
 	IDetailCategoryBuilder& SDKConfigCategory = DetailLayout.EditCategory(TEXT("Project SDK Override"));
 	SDKConfigCategory.InitiallyCollapsed(true);
+	IDetailCategoryBuilder& DebugSigningCategory = DetailLayout.EditCategory(TEXT("DebugSigning"));
+	DebugSigningCategory.InitiallyCollapsed(true);
 
-	IDetailCategoryBuilder& SigningCategory = DetailLayout.EditCategory(TEXT("DistributionSigning"));
+	IDetailCategoryBuilder& DistributionSigningCategory = DetailLayout.EditCategory(TEXT("DistributionSigning"));
 
 	TSharedRef<SPlatformSetupMessage> PlatformSetupMessage = SNew(SPlatformSetupMessage, GameProjectPropertiesPath)
 		.PlatformName(LOCTEXT("AndroidPlatformName", "Android"))
@@ -216,7 +218,7 @@ void FAndroidTargetSettingsCustomization::BuildAppManifestSection(IDetailLayoutB
 	SDKConfigCategory.AddProperty(NDKAPILevelOverrideProperty);
 
 	// Signing category
-	SigningCategory.AddCustomRow(LOCTEXT("SigningHyperlink", "Signing Hyperlink"), false)
+	DistributionSigningCategory.AddCustomRow(LOCTEXT("SigningHyperlink", "Signing Hyperlink"), false)
 		.WholeRowWidget
 		[
 			SNew(SBox)
diff --git a/Engine/Source/Programs/UnrealBuildTool/Platform/Android/UEDeployAndroid.cs b/Engine/Source/Programs/UnrealBuildTool/Platform/Android/UEDeployAndroid.cs
--- a/Engine/Source/Programs/UnrealBuildTool/Platform/Android/UEDeployAndroid.cs
+++ b/Engine/Source/Programs/UnrealBuildTool/Platform/Android/UEDeployAndroid.cs
@@ -3968,6 +3968,45 @@ namespace UnrealBuildTool
 				GradleBuildAdditionsContent.AppendLine("\t}");
 			}
 
+			string DebugKeyAlias, DebugKeyStore, DebugKeyStorePassword, DebugKeyPassword;
+			Ini.GetString("/Script/AndroidRuntimeSettings.AndroidRuntimeSettings", "DebugKeyStore", out DebugKeyStore);
+			Ini.GetString("/Script/AndroidRuntimeSettings.AndroidRuntimeSettings", "DebugKeyAlias", out DebugKeyAlias);
+			Ini.GetString("/Script/AndroidRuntimeSettings.AndroidRuntimeSettings", "DebugKeyStorePassword", out DebugKeyStorePassword);
+			Ini.GetString("/Script/AndroidRuntimeSettings.AndroidRuntimeSettings", "DebugKeyPassword", out DebugKeyPassword);
+
+			bool bHasCustomDebugKeystore = false;
+			if (!String.IsNullOrEmpty(DebugKeyStore) && !String.IsNullOrEmpty(DebugKeyAlias) && !String.IsNullOrEmpty(DebugKeyStorePassword))
+			{
+				if (String.IsNullOrEmpty(DebugKeyPassword) || DebugKeyPassword == "_sameaskeystore_")
+				{
+					DebugKeyPassword = DebugKeyStorePassword;
+				}
+
+				string DebugKeyStoreFilename = Path.Combine(UnrealBuildPath, DebugKeyStore);
+				if (!File.Exists(DebugKeyStoreFilename))
+				{
+					Logger?.LogWarning("Debug keystore file is missing. Check the DistributionSettings section in the Android tab of Project Settings");
+				}
+				else
+				{
+					bHasCustomDebugKeystore = true;
+
+					GradleProperties.AppendLine(String.Format("DEBUG_STORE_FILE={0}", DebugKeyStoreFilename.Replace("\\", "/")));
+					GradleProperties.AppendLine(String.Format("DEBUG_STORE_PASSWORD={0}", DebugKeyStorePassword));
+					GradleProperties.AppendLine(String.Format("DEBUG_KEY_ALIAS={0}", DebugKeyAlias));
+					GradleProperties.AppendLine(String.Format("DEBUG_KEY_PASSWORD={0}", DebugKeyPassword));
+
+					GradleBuildAdditionsContent.AppendLine("\tsigningConfigs {");
+					GradleBuildAdditionsContent.AppendLine("\t\tdebug {");
+					GradleBuildAdditionsContent.AppendLine(String.Format("\t\t\tstoreFile file('{0}')", DebugKeyStoreFilename.Replace("\\", "/")));
+					GradleBuildAdditionsContent.AppendLine(String.Format("\t\t\tstorePassword '{0}'", DebugKeyStorePassword));
+					GradleBuildAdditionsContent.AppendLine(String.Format("\t\t\tkeyAlias '{0}'", DebugKeyAlias));
+					GradleBuildAdditionsContent.AppendLine(String.Format("\t\t\tkeyPassword '{0}'", DebugKeyPassword));
+					GradleBuildAdditionsContent.AppendLine("\t\t}");
+					GradleBuildAdditionsContent.AppendLine("\t}");
+				}
+			}
+
 			if (bForDistribution)
 			{
 				string KeyAlias, KeyStore, KeyStorePassword, KeyPassword;
@@ -4050,6 +4089,10 @@ namespace UnrealBuildTool
 			GradleBuildAdditionsContent.AppendLine("\t\t}");
 			GradleBuildAdditionsContent.AppendLine("\t\tdebug {");
 			GradleBuildAdditionsContent.AppendLine("\t\t\tdebuggable true");
+			if (bHasCustomDebugKeystore)
+			{
+				GradleBuildAdditionsContent.AppendLine("\t\t\tsigningConfig signingConfigs.debug");
+			}
 			GradleBuildAdditionsContent.AppendLine("\t\t}");
 			GradleBuildAdditionsContent.AppendLine("\t}");
 			GradleBuildAdditionsContent.AppendLine("}");
@@ -4691,6 +4734,36 @@ namespace UnrealBuildTool
 
 						KeyStore = Path.Combine(UnrealBuildPath, KeyStore);
 					}
+					else
+					{
+						string DebugKeyAlias, DebugKeyStore, DebugKeyStorePassword, DebugKeyPassword;
+						Ini.GetString("/Script/AndroidRuntimeSettings.AndroidRuntimeSettings", "DebugKeyStore", out DebugKeyStore);
+						Ini.GetString("/Script/AndroidRuntimeSettings.AndroidRuntimeSettings", "DebugKeyAlias", out DebugKeyAlias);
+						Ini.GetString("/Script/AndroidRuntimeSettings.AndroidRuntimeSettings", "DebugKeyStorePassword", out DebugKeyStorePassword);
+						Ini.GetString("/Script/AndroidRuntimeSettings.AndroidRuntimeSettings", "DebugKeyPassword", out DebugKeyPassword);
+
+						if (!String.IsNullOrEmpty(DebugKeyStore) && !String.IsNullOrEmpty(DebugKeyAlias) && !String.IsNullOrEmpty(DebugKeyStorePassword))
+						{
+							Logger.LogInformation("Using custom debug keystore {KeyStore}", DebugKeyStore);
+
+							KeyStore = Path.Combine(UnrealBuildPath, DebugKeyStore);
+							KeyAlias = DebugKeyAlias;
+							KeyStorePassword = DebugKeyStorePassword;
+							KeyPassword = DebugKeyPassword;
+							if (String.IsNullOrEmpty(KeyPassword) || KeyPassword == "_sameaskeystore_")
+							{
+								KeyPassword = KeyStorePassword;
+							}
+						}
+						else if (!String.IsNullOrEmpty(DebugKeyStore) || !String.IsNullOrEmpty(DebugKeyAlias) || !String.IsNullOrEmpty(DebugKeyStorePassword))
+						{
+							throw new BuildException("Debug signing settings are not all or none set. Check the DistributionSettings section in the Android tab of Project Settings");
+						}
+						else
+						{
+							Logger.LogInformation("Using default keystore");
+						}
+					}
 
 					// Make sure the keystore file exists
 					if (!File.Exists(KeyStore))
@@ -5468,6 +5541,11 @@ popd
 										continue;
 									}
 
+									if (SourceProperties[Index].StartsWith("DEBUG_STORE_FILE=") && DestProperties[Index].StartsWith("DEBUG_STORE_FILE="))
+									{
+										continue;
+									}
+
 									if (SourceProperties[Index].StartsWith("OUTPUT_FILENAME=") && DestProperties[Index].StartsWith("OUTPUT_FILENAME="))
 									{
 										continue;
diff --git a/Engine/Source/Runtime/Android/AndroidRuntimeSettings/Classes/AndroidRuntimeSettings.h b/Engine/Source/Runtime/Android/AndroidRuntimeSettings/Classes/AndroidRuntimeSettings.h
--- a/Engine/Source/Runtime/Android/AndroidRuntimeSettings/Classes/AndroidRuntimeSettings.h
+++ b/Engine/Source/Runtime/Android/AndroidRuntimeSettings/Classes/AndroidRuntimeSettings.h
@@ -407,6 +407,22 @@ public:
 	UPROPERTY(GlobalConfig, EditAnywhere, Category = DistributionSigning, Meta = (DisplayName = "Key Password (leave blank to use Key Store Password)"))
 	FString KeyPassword;
 
+	// This is the file that keytool outputs, specified with the -keystore parameter (file should be in <Project>/Build/Android)
+	UPROPERTY(GlobalConfig, EditAnywhere, Category = DebugSigning, Meta = (DisplayName = "Debug Key Store (output of keytool, placed in <Project>/Build/Android)"))
+	FString DebugKeyStore;
+	
+	// This is the name of the key that you specified with the -alias parameter to keytool
+	UPROPERTY(GlobalConfig, EditAnywhere, Category = DebugSigning, Meta = (DisplayName = "Debug Key Alias (-alias parameter to keytool)"))
+	FString DebugKeyAlias;
+	
+	// This is the password that you specified FOR THE KEYSTORE NOT THE KEY, when running keytool (either with -storepass or by typing it in).
+	UPROPERTY(GlobalConfig, EditAnywhere, Category = DebugSigning, Meta = (DisplayName = "Debug Key Store Password (-storepass parameter to keytool)"))
+	FString DebugKeyStorePassword;
+	
+	// This is the password for the key that you may have specified with keytool, if it's different from the keystore password. Leave blank to use same as Debug Keystore
+	UPROPERTY(GlobalConfig, EditAnywhere, Category = DebugSigning, Meta = (DisplayName = "Debug Key Password (leave blank to use Debug Key Store Password)"))
+	FString DebugKeyPassword;
+
 	// Enable Arm64 support?
 	UPROPERTY(GlobalConfig, EditAnywhere, Category = Build, meta = (DisplayName = "Support arm64 [aka arm64-v8a]"))
 	bool bBuildForArm64;
