UnrealBuildTool: Read extra command line args from environment variable

Useful for other apps (such as Horde) to modify parameters to UBT when invoked through BuildGraph.

diff --git a/Engine/Source/Programs/UnrealBuildTool/UnrealBuildTool.cs b/Engine/Source/Programs/UnrealBuildTool/UnrealBuildTool.cs
--- a/Engine/Source/Programs/UnrealBuildTool/UnrealBuildTool.cs
+++ b/Engine/Source/Programs/UnrealBuildTool/UnrealBuildTool.cs
@@ -8,6 +8,7 @@ using System.Linq;
 using System.Reflection;
 using System.Text;
 using System.Text.Json;
+using System.Text.RegularExpressions;
 using EpicGames.Core;
 using Microsoft.Extensions.Logging;
 using OpenTracing.Util;
@@ -406,6 +407,17 @@ namespace UnrealBuildTool
 			}
 		}
 
+		/// <summary>
+		/// Read extra command-line arguments from an environment variable
+		/// Double-quote any argument containing whitespace, as they are split by just that.
+		/// </summary>
+		/// <returns>Extra arguments</returns>
+		private static string[] GetExtraArgsFromEnvVar()
+		{
+			string? extraArgs = Environment.GetEnvironmentVariable("UBT_EXTRA_ARGS");
+			return String.IsNullOrEmpty(extraArgs) ? Array.Empty<string>() : CommandLineArguments.Split(extraArgs);
+		}
+
 		/// <summary>
 		/// Main entry point. Parses any global options and initializes the logging system, then invokes the appropriate command.
 		/// NB: That the entry point is deliberately NOT async, since we have a single-instance mutex that cannot be disposed from a different thread.
@@ -442,6 +454,7 @@ namespace UnrealBuildTool
 				// Start capturing performance info
 				Timeline.Start();
 				Tracer = JsonTracer.TryRegisterAsGlobalTracer();
+				ArgumentsArray = ArgumentsArray.Concat(GetExtraArgsFromEnvVar()).ToArray();
 
 				// Parse the command line arguments
 				CommandLineArguments Arguments = new CommandLineArguments(ArgumentsArray);
