- Fixed an issue with SubmitTool having wrong output dir

diff --git a/Engine/Source/Programs/UnrealBuildTool/Configuration/UEBuildTarget.cs b/Engine/Source/Programs/UnrealBuildTool/Configuration/UEBuildTarget.cs
--- a/Engine/Source/Programs/UnrealBuildTool/Configuration/UEBuildTarget.cs
+++ b/Engine/Source/Programs/UnrealBuildTool/Configuration/UEBuildTarget.cs
@@ -4773,19 +4773,22 @@ namespace UnrealBuildTool
 			// Construct the output paths for this target's executable
 			DirectoryReference OutputDirectory;
 			// if we are building a program with a project file, and we are making a unique build environment, output the program to that project's binaries
-			bool bIsUniqueBuildProgram = Rules.Type == TargetType.Program && Rules.BuildEnvironment == TargetBuildEnvironment.Unique;
-			if (ProjectFile != null && (bCompileMonolithic || !bUseSharedBuildEnvironment) && (Rules.File.IsUnderDirectory(ProjectDirectory) || bIsUniqueBuildProgram))
+			bool bIsUniqueBuildProgramWithExternalProject = 
+					Rules.Type == TargetType.Program && Rules.BuildEnvironment == TargetBuildEnvironment.Unique && 
+					Rules.File.IsUnderDirectory(Unreal.EngineDirectory) && 
+					ProjectFile != null && !ProjectFile.IsUnderDirectory(Unreal.EngineDirectory);
+
+			if (bIsUniqueBuildProgramWithExternalProject)
 			{
-				OutputDirectory = GetOutputDirectoryForExecutable(ProjectDirectory, Rules.File);
 				// engine programs generally won't be under the directories searched in GetOutputDirectoryForExecutable, so it will return the EngineDir, even tho unique programs
 				// built with a ProjectFile want to be under the Project dir, so just force it into the ProjectDir if one wasn't found otherwise
-				if (bIsUniqueBuildProgram && Rules.File.IsUnderDirectory(Unreal.EngineDirectory) && OutputDirectory == Unreal.EngineDirectory)
-				{
-					if (ProjectFile != null && ProjectDirectory.ParentDirectory!.GetDirectoryName() != "Programs")
-					{
-						OutputDirectory = ProjectDirectory;
-					}
-				}
+				Logger.LogInformation("Overriding output directory for unique build program {Name} from {Output} to {Project}, because it it is being built with external project {Project}", Rules.Name, GetOutputDirectoryForExecutable(ProjectDirectory, Rules.File), ProjectDirectory, ProjectFile);
+				OutputDirectory = ProjectDirectory;
+			}
+			else if (ProjectFile != null && (bCompileMonolithic || !bUseSharedBuildEnvironment) && (Rules.File.IsUnderDirectory(ProjectDirectory)))
+			{
+				OutputDirectory = GetOutputDirectoryForExecutable(ProjectDirectory, Rules.File);
+				
 			}
 			else
 			{
