Fix llvm-symbolizer staging issue.

diff --git a/Engine/Source/Programs/UnrealBuildTool/Platform/Linux/LinuxToolChain.cs b/Engine/Source/Programs/UnrealBuildTool/Platform/Linux/LinuxToolChain.cs
--- a/Engine/Source/Programs/UnrealBuildTool/Platform/Linux/LinuxToolChain.cs
+++ b/Engine/Source/Programs/UnrealBuildTool/Platform/Linux/LinuxToolChain.cs
@@ -670,8 +670,8 @@ namespace UnrealBuildTool
 						FileReference SymbolizerSourcePath = FileReference.Combine(InternalSdkPathRef, "bin/llvm-symbolizer");
 						FileReference SymbolizerTargetPath = FileReference.Combine(ExeDir, "llvm-symbolizer");
 
-						RuntimeDependencies.Add(new RuntimeDependency(SymbolizerSourcePath, StagedFileType.NonUFS));
-						TargetFileToSourceFile[SymbolizerSourcePath] = SymbolizerTargetPath;
+						RuntimeDependencies.Add(new RuntimeDependency(SymbolizerTargetPath, StagedFileType.NonUFS));
+						TargetFileToSourceFile[SymbolizerTargetPath] = SymbolizerSourcePath;
 					}
 				}
 			}
