Fix OnlineSubsystemGooglePlay module being skipped for x86_64

diff --git a/Engine/Source/Programs/UnrealBuildTool/Platform/Android/AndroidToolChain.cs b/Engine/Source/Programs/UnrealBuildTool/Platform/Android/AndroidToolChain.cs
--- a/Engine/Source/Programs/UnrealBuildTool/Platform/Android/AndroidToolChain.cs
+++ b/Engine/Source/Programs/UnrealBuildTool/Platform/Android/AndroidToolChain.cs
@@ -82,7 +82,7 @@ namespace UnrealBuildTool
 
 		private static Dictionary<UnrealArch, string[]> ModulesToSkip = new() {
 			{ UnrealArch.Arm64, new string[] {  } },
-			{ UnrealArch.X64,   new string[] { "OnlineSubsystemOculus", "OculusHMD", "OculusMR", "OnlineSubsystemGooglePlay" } }
+			{ UnrealArch.X64,   new string[] { "OnlineSubsystemOculus", "OculusHMD", "OculusMR" } }
 		};
 
 		private static Dictionary<UnrealArch, string[]> GeneratedModulesToSkip = new() {
