UnrealBuildTool: Share UHT generated files between intermediate invariants

#rnx

diff --git a/Engine/Source/Programs/UnrealBuildTool/Configuration/UEBuildTarget.cs b/Engine/Source/Programs/UnrealBuildTool/Configuration/UEBuildTarget.cs
--- a/Engine/Source/Programs/UnrealBuildTool/Configuration/UEBuildTarget.cs
+++ b/Engine/Source/Programs/UnrealBuildTool/Configuration/UEBuildTarget.cs
@@ -5527,8 +5527,8 @@ namespace UnrealBuildTool
 						GeneratedCodeDirectory = ProjectDirectory;
 					}
 
-					// Get the subfolder containing generated code - we don't need architeceture information since these are shared between all arches for a platform
-					GeneratedCodeDirectory = DirectoryReference.Combine(GeneratedCodeDirectory, PlatformIntermediateFolderNoArch, GetTargetIntermediateFolderName(AppName, IntermediateEnvironment), "Inc");
+					// Get the subfolder containing generated code - we don't need architecture information since these are shared between all arches for a platform, as well as shared between all intermediate environment variants
+					GeneratedCodeDirectory = DirectoryReference.Combine(GeneratedCodeDirectory, PlatformIntermediateFolderNoArch, AppName, "Inc");
 
 					// Append the binaries subfolder, if present. We rely on this to ensure that build products can be filtered correctly.
 					if (RulesObject.BinariesSubFolder != null)
