Fix missing quotes around project directory

diff --git a/Engine/Source/Programs/UnrealBuildTool/ToolChain/AppleToolChain.cs b/Engine/Source/Programs/UnrealBuildTool/ToolChain/AppleToolChain.cs
--- a/Engine/Source/Programs/UnrealBuildTool/ToolChain/AppleToolChain.cs
+++ b/Engine/Source/Programs/UnrealBuildTool/ToolChain/AppleToolChain.cs
@@ -363,7 +363,7 @@ namespace UnrealBuildTool
 
 			// make path to the script
 			FileItem BundleScript = FileItem.GetItemByFileReference(FileReference.Combine(Unreal.EngineDirectory, "Build/BatchFiles/Mac/UpdateVersionAfterBuild.sh"));
-			UpdateVersionAction.CommandArguments = $"\"{BundleScript.AbsolutePath}\" {ProductDirectory} {LinkEnvironment.Platform}";
+			UpdateVersionAction.CommandArguments = $"\"{BundleScript.AbsolutePath}\" \"{ProductDirectory}\" {LinkEnvironment.Platform}";
 			UpdateVersionAction.PrerequisiteItems.Add(Prerequisite);
 			UpdateVersionAction.PrerequisiteItems.Add(BundleScript);
 			UpdateVersionAction.ProducedItems.Add(DestFile);
