Fixing typo. ".d" files are for Clang outputs.

diff --git a/Engine/Source/Programs/UnrealBuildTool/Platform/Windows/VCSpecificFileAction.cs b/Engine/Source/Programs/UnrealBuildTool/Platform/Windows/VCSpecificFileAction.cs
--- a/Engine/Source/Programs/UnrealBuildTool/Platform/Windows/VCSpecificFileAction.cs
+++ b/Engine/Source/Programs/UnrealBuildTool/Platform/Windows/VCSpecificFileAction.cs
@@ -99,7 +99,7 @@ namespace UnrealBuildTool
 			{
 				Action.DependencyListFile = FileItem.GetItemByFileReference(FileReference.Combine(IntermediateDirectory, $"{Filename}.dep.json"));
 			}
-			else if (BaseAction.CompilerType.IsMSVC())
+			else if (BaseAction.CompilerType.IsClang())
 			{
 				Action.DependencyListFile = FileItem.GetItemByFileReference(FileReference.Combine(IntermediateDirectory, $"{Filename}.d"));
 			}
