Fixing duplicate named logger arguments

#rnx

diff --git a/Engine/Source/Programs/UnrealBuildTool/Configuration/UEBuildTarget.cs b/Engine/Source/Programs/UnrealBuildTool/Configuration/UEBuildTarget.cs
--- a/Engine/Source/Programs/UnrealBuildTool/Configuration/UEBuildTarget.cs
+++ b/Engine/Source/Programs/UnrealBuildTool/Configuration/UEBuildTarget.cs
@@ -4782,7 +4782,7 @@ namespace UnrealBuildTool
 			{
 				// engine programs generally won't be under the directories searched in GetOutputDirectoryForExecutable, so it will return the EngineDir, even tho unique programs
 				// built with a ProjectFile want to be under the Project dir, so just force it into the ProjectDir if one wasn't found otherwise
-				Logger.LogInformation("Overriding output directory for unique build program {Name} from {Output} to {Project}, because it it is being built with external project {Project}", Rules.Name, GetOutputDirectoryForExecutable(ProjectDirectory, Rules.File), ProjectDirectory, ProjectFile);
+				Logger.LogInformation("Overriding output directory for unique build program {Name} from {Output} to {Project}, because it it is being built with external project {ProjectFile}", Rules.Name, GetOutputDirectoryForExecutable(ProjectDirectory, Rules.File), ProjectDirectory, ProjectFile);
 				OutputDirectory = ProjectDirectory;
 			}
 			else if (ProjectFile != null && (bCompileMonolithic || !bUseSharedBuildEnvironment) && (Rules.File.IsUnderDirectory(ProjectDirectory)))
