Add missing quotes breaking builds with spaces in the path

#tests Created new project with spaces in path, compiled fine

diff --git a/Engine/Source/Programs/UnrealBuildTool/ProjectFiles/VisualStudio/VCProject.cs b/Engine/Source/Programs/UnrealBuildTool/ProjectFiles/VisualStudio/VCProject.cs
--- a/Engine/Source/Programs/UnrealBuildTool/ProjectFiles/VisualStudio/VCProject.cs
+++ b/Engine/Source/Programs/UnrealBuildTool/ProjectFiles/VisualStudio/VCProject.cs
@@ -2142,7 +2142,7 @@ namespace UnrealBuildTool
 				{
 					if (TargetRulesObject.Type == TargetType.Editor && bEditorDependsOnShaderCompileWorker && !Unreal.IsEngineInstalled())
 					{
-						string ProjParam = UProjectPath.Length > 0 ? $" -Project={UProjectPath}" : "";
+						string ProjParam = UProjectPath.Length > 0 ? $" -Project=\\\"{UProjectPath.Trim('"')}\\\"" : "";
 						ExtraTargets.Add($"ShaderCompileWorker Win64 Development{ProjParam}");
 					}
 					if (TargetRulesObject.bWithLiveCoding && bBuildLiveCodingConsole && !Unreal.IsEngineInstalled() && TargetRulesObject.Name != "LiveCodingConsole")
