Fix postbuild command from TestTargetRules

diff --git a/Engine/Source/Programs/UnrealBuildTool/Configuration/TestTargetRules.cs b/Engine/Source/Programs/UnrealBuildTool/Configuration/TestTargetRules.cs
--- a/Engine/Source/Programs/UnrealBuildTool/Configuration/TestTargetRules.cs
+++ b/Engine/Source/Programs/UnrealBuildTool/Configuration/TestTargetRules.cs
@@ -181,7 +181,7 @@ namespace UnrealBuildTool
 
 			if (target.Platform == UnrealTargetPlatform.Win64)
 			{
-				string outputName = "$(TargetName)";
+				string outputName = target.Name;
 				if (target.Configuration != UndecoratedConfiguration)
 				{
 					outputName = outputName + "-" + target.Platform + "-" + target.Configuration;
