Fix issue with Android env vars not always set on Mac

#android

diff --git a/Engine/Source/Programs/UnrealBuildTool/Platform/Android/AndroidPlatformSDK.cs b/Engine/Source/Programs/UnrealBuildTool/Platform/Android/AndroidPlatformSDK.cs
--- a/Engine/Source/Programs/UnrealBuildTool/Platform/Android/AndroidPlatformSDK.cs
+++ b/Engine/Source/Programs/UnrealBuildTool/Platform/Android/AndroidPlatformSDK.cs
@@ -34,6 +34,13 @@ namespace UnrealBuildTool
 
 			if (OperatingSystem.IsMacOS())
 			{
+				Dictionary<string, string> AndroidEnv = new Dictionary<string, string>();
+				Dictionary<string, string> EnvVarNames = new Dictionary<string, string> {
+														 {"ANDROID_HOME", "SDKPath"},
+														 {"NDKROOT", "NDKPath"},
+														 {"JAVA_HOME", "JavaPath"}
+														 };
+
 				string BashProfilePath = Path.Combine(Environment.GetFolderPath(Environment.SpecialFolder.Personal), ".bash_profile");
 				if (!File.Exists(BashProfilePath))
 				{
@@ -45,16 +52,34 @@ namespace UnrealBuildTool
 					string[] BashProfileContents = File.ReadAllLines(BashProfilePath);
 
 					// Walk backwards so we keep the last export setting instead of the first
-					string SdkKey = "NDKROOT";
 					for (int LineIndex = BashProfileContents.Length - 1; LineIndex >= 0; --LineIndex)
 					{
-						if (BashProfileContents[LineIndex].StartsWith("export " + SdkKey + "="))
+						foreach (KeyValuePair<string, string> kvp in EnvVarNames)
 						{
-							string PathVar = BashProfileContents[LineIndex].Split('=')[1].Replace("\"", "");
-							return PathVar;
+							if (AndroidEnv.ContainsKey(kvp.Key))
+							{
+								continue;
+							}
+							if (BashProfileContents[LineIndex].StartsWith("export " + kvp.Key + "="))
+							{
+								string PathVar = BashProfileContents[LineIndex].Split('=')[1].Replace("\"", "");
+								AndroidEnv.Add(kvp.Key, PathVar);
+							}
 						}
 					}
 				}
+
+				// Set for the process
+				foreach (KeyValuePair<string, string> kvp in AndroidEnv)
+				{
+					Environment.SetEnvironmentVariable(kvp.Key, kvp.Value);
+				}
+
+				NDKPath = Environment.GetEnvironmentVariable("NDKROOT");
+				if (!String.IsNullOrEmpty(NDKPath))
+				{
+					return NDKPath.Replace("\"", "");
+				}
 			}
 
 			return null;
diff --git a/Engine/Source/Programs/UnrealBuildTool/Platform/Android/AndroidToolChain.cs b/Engine/Source/Programs/UnrealBuildTool/Platform/Android/AndroidToolChain.cs
--- a/Engine/Source/Programs/UnrealBuildTool/Platform/Android/AndroidToolChain.cs
+++ b/Engine/Source/Programs/UnrealBuildTool/Platform/Android/AndroidToolChain.cs
@@ -131,7 +131,7 @@ namespace UnrealBuildTool
 			Options = ToolchainOptions;
 			ProjectFile = InProjectFile;
 
-			string? NDKPath = Environment.GetEnvironmentVariable("NDKROOT");
+			string? NDKPath = AndroidPlatformSDK.GetNDKRoot();
 
 			// don't register if we don't have an NDKROOT specified
 			if (String.IsNullOrEmpty(NDKPath))
