Handle PVS-Studio notice

#rnx

diff --git a/Engine/Source/Programs/UnrealBuildTool/Platform/Windows/PVSToolChain.cs b/Engine/Source/Programs/UnrealBuildTool/Platform/Windows/PVSToolChain.cs
--- a/Engine/Source/Programs/UnrealBuildTool/Platform/Windows/PVSToolChain.cs
+++ b/Engine/Source/Programs/UnrealBuildTool/Platform/Windows/PVSToolChain.cs
@@ -455,6 +455,12 @@ namespace UnrealBuildTool
 								{
 									bCanParse = true;
 
+									if (String.IsNullOrWhiteSpace(FileName))
+									{
+										Logger.LogInformation("PVS-Studio Notification {WarningCode}: {WarningMessage}", WarningCode, WarningMessage);
+										continue;
+									}
+
 									FileReference file = new FileReference(FileName);
 
 									// Ignore anything in the IgnoredDirectories folders
