PVS license renew message should be a warning

#rnx

diff --git a/Engine/Source/Programs/UnrealBuildTool/Platform/Windows/PVSToolChain.cs b/Engine/Source/Programs/UnrealBuildTool/Platform/Windows/PVSToolChain.cs
--- a/Engine/Source/Programs/UnrealBuildTool/Platform/Windows/PVSToolChain.cs
+++ b/Engine/Source/Programs/UnrealBuildTool/Platform/Windows/PVSToolChain.cs
@@ -411,8 +411,6 @@ namespace UnrealBuildTool
 			Arguments.ApplyTo(this);
 			Arguments.CheckAllArgumentsUsed();
 
-			Logger.LogInformation("{File}", OutputFile!.GetFileName());
-
 			// Read the input files
 			string[] InputFileLines = FileReference.ReadAllLines(InputFileList!);
 			FileReference[] InputFiles = InputFileLines.Select(x => x.Trim()).Where(x => x.Length > 0).Select(x => new FileReference(x)).ToArray();
@@ -433,35 +431,35 @@ namespace UnrealBuildTool
 					for (int LineIdx = 0; LineIdx < Lines.Length; LineIdx++)
 					{
 						string Line = Lines[LineIdx];
-						if (!String.IsNullOrWhiteSpace(Line) && UniqueItems.Add(Line))
+						if (String.IsNullOrWhiteSpace(Line) || !UniqueItems.Add(Line))
 						{
-							bool bCanParse = false;
+							continue;
+						}
+
+						bool bCanParse = false;
 
-							string[] Tokens = Line.Split(new string[] { "<#~>" }, StringSplitOptions.None);
-							if (Tokens.Length >= 9)
+						string[] Tokens = Line.Split(new string[] { "<#~>" }, StringSplitOptions.None);
+						if (Tokens.Length >= 9)
+						{
+							//string Trial = Tokens[1];
+							string LineNumberStr = Tokens[2];
+							string FileName = Tokens[3];
+							string WarningCode = Tokens[5];
+							string WarningMessage = Tokens[6];
+							string FalseAlarmStr = Tokens[7];
+							string LevelStr = Tokens[8];
+
+							if (Int32.TryParse(LineNumberStr, out int LineNumber) && Boolean.TryParse(FalseAlarmStr, out bool bFalseAlarm) && Int32.TryParse(LevelStr, out int Level))
 							{
-								//string Trial = Tokens[1];
-								string LineNumberStr = Tokens[2];
-								string FileName = Tokens[3];
-								string WarningCode = Tokens[5];
-								string WarningMessage = Tokens[6];
-								string FalseAlarmStr = Tokens[7];
-								string LevelStr = Tokens[8];
-
-								int LineNumber;
-								bool bFalseAlarm;
-								int Level;
-								if (Int32.TryParse(LineNumberStr, out LineNumber) && Boolean.TryParse(FalseAlarmStr, out bFalseAlarm) && Int32.TryParse(LevelStr, out Level))
-								{
-									bCanParse = true;
+								bCanParse = true;
 
-									if (String.IsNullOrWhiteSpace(FileName))
-									{
-										Logger.LogInformation("PVS-Studio Notification {WarningCode}: {WarningMessage}", WarningCode, WarningMessage);
-										continue;
-									}
+								// Output the line to the raw output file
+								RawWriter.WriteLine(Line);
 
-									FileReference file = new FileReference(FileName);
+								FileReference file;
+								if (!String.IsNullOrWhiteSpace(FileName))
+								{
+									file = new FileReference(FileName);
 
 									// Ignore anything in the IgnoredDirectories folders
 									if (IgnoredDirectories.Any() && IgnoredDirectories.Any(x => file.IsUnderDirectory(x)))
@@ -473,22 +471,25 @@ namespace UnrealBuildTool
 									{
 										continue;
 									}
+								}
+								else
+								{
+									file = InputFile;
+									FileName = InputFile.FullName;
+									LineNumber = LineIdx + 1;
+								}
 
-									// Output the line to the raw output file
-									RawWriter.WriteLine(Line);
-
-									// Output the line to the log
-									if (!bFalseAlarm && Level <= PrintLevel)
-									{
-										Logger.LogWarning(KnownLogEvents.Compiler, "{Path}({LineNumber}): warning {WarningCode}: {WarningMessage}", LogValue.SourceFile(file, FileName), LineNumber, WarningCode, WarningMessage);
-									}
+								// Output the line to the log
+								if (!bFalseAlarm && Level <= PrintLevel)
+								{
+									Logger.LogWarning(KnownLogEvents.Compiler, "{Path}({LineNumber}): warning {WarningCode}: {WarningMessage}", LogValue.SourceFile(file, FileName), LineNumber, WarningCode, WarningMessage);
 								}
 							}
+						}
 
-							if (!bCanParse)
-							{
-								Logger.LogWarning("{Path}({LineNumber}): warning: Unable to parse PVS output line '{Line}' (tokens=|{Tokens}|)", InputFile, LineIdx + 1, Line, String.Join("|", Tokens));
-							}
+						if (!bCanParse)
+						{
+							Logger.LogWarning(KnownLogEvents.Compiler, "{Path}({LineNumber}): warning: Unable to parse PVS output line '{Line}' (tokens=|{Tokens}|)", LogValue.SourceFile(InputFile, InputFile.GetFileName()), LineIdx + 1, Line, String.Join("|", Tokens));
 						}
 					}
 				}
