UnrealBuildTool: Fix analysis warning

#rnx

diff --git a/Engine/Source/Programs/UnrealBuildTool/Platform/Windows/PVSToolChain.cs b/Engine/Source/Programs/UnrealBuildTool/Platform/Windows/PVSToolChain.cs
--- a/Engine/Source/Programs/UnrealBuildTool/Platform/Windows/PVSToolChain.cs
+++ b/Engine/Source/Programs/UnrealBuildTool/Platform/Windows/PVSToolChain.cs
@@ -423,7 +423,7 @@ namespace UnrealBuildTool
 			HashSet<string> UniqueItems = new HashSet<string>();
 			List<string> OutputLines = new List<string>();
 
-			using (StreamWriter RawWriter = new StreamWriter(OutputFile.FullName))
+			using (StreamWriter RawWriter = new StreamWriter(OutputFile!.FullName))
 			{
 				foreach (FileReference InputFile in InputFiles)
 				{
