Added PGO artifact merging step and sync of required dependencies

#rnx

diff --git a/Engine/Source/Programs/UnrealBuildTool/Platform/Windows/VCToolChain.cs b/Engine/Source/Programs/UnrealBuildTool/Platform/Windows/VCToolChain.cs
--- a/Engine/Source/Programs/UnrealBuildTool/Platform/Windows/VCToolChain.cs
+++ b/Engine/Source/Programs/UnrealBuildTool/Platform/Windows/VCToolChain.cs
@@ -2037,6 +2037,8 @@ namespace UnrealBuildTool
 				string[] PGOFiles = {
 					"pgort140.dll",
 					"pgosweep.exe",
+					"mspdbcore.dll",
+					"pgomgr.exe"
 				};
 				foreach (string PGOFile in PGOFiles)
 				{
