PR #11565: [PVS-Studio] Added a new format of the analyzer report

diff --git a/Engine/Source/Programs/UnrealBuildAccelerator/Core/Private/UbaApplicationRules.cpp b/Engine/Source/Programs/UnrealBuildAccelerator/Core/Private/UbaApplicationRules.cpp
--- a/Engine/Source/Programs/UnrealBuildAccelerator/Core/Private/UbaApplicationRules.cpp
+++ b/Engine/Source/Programs/UnrealBuildAccelerator/Core/Private/UbaApplicationRules.cpp
@@ -341,13 +341,15 @@ namespace uba
 	{
 		virtual bool IsOutputFile(const tchar* file, u64 fileLen) override
 		{
-			return EndsWith(file, fileLen, TC(".pvslog"))
+			return EndsWith(file, fileLen, TC(".PVS-Studio.log"))
+				|| EndsWith(file, fileLen, TC(".pvslog"))
 				|| EndsWith(file, fileLen, TC(".stacktrace.txt"));
 		}
 		
 		virtual bool IsRarelyRead(const StringBufferBase& file) override
 		{
 			return file.EndsWith(TC(".i"))
+				|| file.EndsWith(TC(".PVS-Studio.log"))
 				|| file.EndsWith(TC(".pvslog"))
 				|| file.EndsWith(TC(".stacktrace.txt"));
 		}
diff --git a/Engine/Source/Programs/UnrealBuildTool/Platform/Windows/PVSToolChain.cs b/Engine/Source/Programs/UnrealBuildTool/Platform/Windows/PVSToolChain.cs
--- a/Engine/Source/Programs/UnrealBuildTool/Platform/Windows/PVSToolChain.cs
+++ b/Engine/Source/Programs/UnrealBuildTool/Platform/Windows/PVSToolChain.cs
@@ -1,5 +1,8 @@
 // Copyright Epic Games, Inc. All Rights Reserved.
 
+using EpicGames.Core;
+using Microsoft.Extensions.Logging;
+using Newtonsoft.Json;
 using System;
 using System.Collections.Generic;
 using System.Collections.Immutable;
@@ -10,8 +13,6 @@ using System.Text;
 using System.Text.RegularExpressions;
 using System.Threading.Tasks;
 using System.Xml.Serialization;
-using EpicGames.Core;
-using Microsoft.Extensions.Logging;
 using UnrealBuildBase;
 
 namespace UnrealBuildTool
@@ -395,6 +396,12 @@ namespace UnrealBuildTool
 		[CommandLine("-PrintLevel")]
 		public int PrintLevel { get; set; } = 1;
 
+		/// <summary>
+		/// Version of the analyzers
+		/// </summary>
+		[CommandLine("-AnalyzerVersion")]
+		public Version? AnalyzerVersion { get; set; } = null;
+
 		/// <summary>
 		/// If all ThirdParty code should be ignored
 		/// </summary>
@@ -436,60 +443,59 @@ namespace UnrealBuildTool
 							continue;
 						}
 
-						bool bCanParse = false;
-
-						string[] Tokens = Line.Split(new string[] { "<#~>" }, StringSplitOptions.None);
-						if (Tokens.Length >= 9)
+						try
 						{
-							//string Trial = Tokens[1];
-							string LineNumberStr = Tokens[2];
-							string FileName = Tokens[3];
-							string WarningCode = Tokens[5];
-							string WarningMessage = Tokens[6];
-							string FalseAlarmStr = Tokens[7];
-							string LevelStr = Tokens[8];
-
-							if (Int32.TryParse(LineNumberStr, out int LineNumber) && Boolean.TryParse(FalseAlarmStr, out bool bFalseAlarm) && Int32.TryParse(LevelStr, out int Level))
+							PVSErrorInfo? ErrorInfo;
+							if (AnalyzerVersion!.CompareTo(new Version("7.30")) < 0)
 							{
-								bCanParse = true;
-
-								// Output the line to the raw output file
-								RawWriter.WriteLine(Line);
-
-								FileReference file;
-								if (!String.IsNullOrWhiteSpace(FileName))
+								if (!TryParseErrorInfo(Line, out ErrorInfo))
 								{
-									file = new FileReference(FileName);
-
-									// Ignore anything in the IgnoredDirectories folders
-									if (IgnoredDirectories.Any() && IgnoredDirectories.Any(x => file.IsUnderDirectory(x)))
-									{
-										continue;
-									}
-
-									if (IgnoreThirdParty && file.FullName.Contains("ThirdParty", StringComparison.OrdinalIgnoreCase))
-									{
-										continue;
-									}
+									throw new FormatException();
 								}
-								else
+							}
+							else
+							{
+								ErrorInfo = JsonConvert.DeserializeObject<PVSErrorInfo>(Line) ?? throw new FormatException();
+							}
+
+							string FileName = ErrorInfo.Positions![0].File!;
+							int LineNumber = ErrorInfo.Positions![0].Lines![0];
+
+							FileReference file;
+							if (!String.IsNullOrWhiteSpace(FileName))
+							{
+								file = new FileReference(FileName);
+								// Ignore anything in the IgnoredDirectories folders
+								if (IgnoredDirectories.Any() && IgnoredDirectories.Any(x => file.IsUnderDirectory(x)))
 								{
-									file = InputFile;
-									FileName = InputFile.FullName;
-									LineNumber = LineIdx + 1;
+									continue;
 								}
 
-								// Output the line to the log
-								if (!bFalseAlarm && Level <= PrintLevel)
+								if (IgnoreThirdParty && file.FullName.Contains("ThirdParty", StringComparison.OrdinalIgnoreCase))
 								{
-									Logger.LogWarning(KnownLogEvents.Compiler, "{Path}({LineNumber}): warning {WarningCode}: {WarningMessage}", LogValue.SourceFile(file, FileName), LineNumber, WarningCode, WarningMessage);
+									continue;
 								}
 							}
-						}
+							else
+							{
+								file = InputFile;
+								FileName = InputFile.FullName;
+								LineNumber = LineIdx + 1;
+							}
 
-						if (!bCanParse)
+							// Output the line to the raw output file
+							RawWriter.WriteLine(Line);
+
+							// Output the line to the log
+							if (ErrorInfo.FalseAlarm != true && ErrorInfo.Level <= PrintLevel)
+							{
+								Logger.LogWarning(KnownLogEvents.Compiler, "{Path}({LineNumber}): warning {WarningCode}: {WarningMessage}", LogValue.SourceFile(file, FileName), LineNumber, ErrorInfo.Code, ErrorInfo.Message);
+							}
+
+						}
+						catch (Exception Ex)
 						{
-							Logger.LogWarning(KnownLogEvents.Compiler, "{Path}({LineNumber}): warning: Unable to parse PVS output line '{Line}' (tokens=|{Tokens}|)", LogValue.SourceFile(InputFile, InputFile.GetFileName()), LineIdx + 1, Line, String.Join("|", Tokens));
+							Logger.LogWarning(KnownLogEvents.Compiler, "{Path}({LineNumber}): warning: Unable to parse PVS output line '{Line}' ({Message})", LogValue.SourceFile(InputFile, InputFile.GetFileName()), LineIdx + 1, Line, Ex.Message);
 						}
 					}
 				}
@@ -497,6 +503,68 @@ namespace UnrealBuildTool
 			Logger.LogInformation("Written {NumItems} {Noun} to {File}.", UniqueItems.Count, (UniqueItems.Count == 1) ? "diagnostic" : "diagnostics", OutputFile.FullName);
 			return Task.FromResult(0);
 		}
+
+		bool TryParseErrorInfo(string Line, out PVSErrorInfo ErrorInfo)
+		{
+			string[] Tokens = Line.Split(new string[] { "<#~>" }, StringSplitOptions.None);
+			if (Tokens.Length >= 9)
+			{
+				string LineNumberStr = Tokens[2];
+				string FileName = Tokens[3];
+				string WarningCode = Tokens[5];
+				string WarningMessage = Tokens[6];
+				string FalseAlarmStr = Tokens[7];
+				string LevelStr = Tokens[8];
+
+				int LineNumber;
+				bool bFalseAlarm;
+				int Level;
+				if (Int32.TryParse(LineNumberStr, out LineNumber) && Boolean.TryParse(FalseAlarmStr, out bFalseAlarm) && Int32.TryParse(LevelStr, out Level))
+				{
+					ErrorInfo = new PVSErrorInfo()
+					{
+						Code = WarningCode,
+						Message = WarningMessage,
+						FalseAlarm = bFalseAlarm,
+						Level = Level,
+						Positions = new PVSPosition[] { new PVSPosition { File = FileName, Lines = new int[] { LineNumber } } },
+
+					};
+
+					return true;
+				}
+			}
+
+			ErrorInfo = new PVSErrorInfo();
+			return false;
+		}
+	}
+
+	class PVSPosition
+	{
+		[JsonProperty(Required = Required.Always)]
+		public string? File;
+
+		[JsonProperty(Required = Required.Always)]
+		public int[]? Lines;
+	}
+
+	class PVSErrorInfo
+	{
+		[JsonProperty(Required = Required.Always)]
+		public string? Code;
+
+		[JsonProperty(Required = Required.Always)]
+		public bool? FalseAlarm;
+
+		[JsonProperty(Required = Required.Always)]
+		public int? Level;
+
+		[JsonProperty(Required = Required.Always)]
+		public string? Message;
+
+		[JsonProperty(Required = Required.Always)]
+		public PVSPosition[]? Positions;
 	}
 
 	class PVSToolChain : ISPCToolChain
@@ -510,6 +578,8 @@ namespace UnrealBuildTool
 		UnrealTargetPlatform Platform;
 		Version AnalyzerVersion;
 
+		string OutputFileExtension => AnalyzerVersion.CompareTo(new Version("7.30")) >= 0 ? ".PVS-Studio.log" : ".pvslog";
+
 		public PVSToolChain(ReadOnlyTargetRules Target, VCToolChain InInnerToolchain, ILogger Logger)
 			: base(Logger)
 		{
@@ -665,10 +735,10 @@ namespace UnrealBuildTool
 		public static bool ShouldCompileAsC(string compilerCommandLine, string sourceFileName)
 		{
 			int CFlagLastPosition = Math.Max(Math.Max(compilerCommandLine.LastIndexOf("/TC "), compilerCommandLine.LastIndexOf("/Tc ")),
-											 Math.Max(compilerCommandLine.LastIndexOf("-TC "), compilerCommandLine.LastIndexOf("-Tc ")));
+												Math.Max(compilerCommandLine.LastIndexOf("-TC "), compilerCommandLine.LastIndexOf("-Tc ")));
 
 			int CppFlagLastPosition = Math.Max(Math.Max(compilerCommandLine.LastIndexOf("/TP "), compilerCommandLine.LastIndexOf("/Tp ")),
-											   Math.Max(compilerCommandLine.LastIndexOf("-TP "), compilerCommandLine.LastIndexOf("-Tp ")));
+												Math.Max(compilerCommandLine.LastIndexOf("-TP "), compilerCommandLine.LastIndexOf("-Tp ")));
 
 			bool compileAsCCode;
 			if (CFlagLastPosition == CppFlagLastPosition)
@@ -805,7 +875,7 @@ namespace UnrealBuildTool
 				{
 					ConfigFileContents.Append("report-disabled-rules=yes\n");
 				}
-				
+
 				if (SourceFileItem.Location.IsUnderDirectory(Unreal.RootDirectory))
 				{
 					ConfigFileContents.AppendFormat("errors-off=V1102\n");
@@ -819,13 +889,18 @@ namespace UnrealBuildTool
 					ConfigFileContents.Append("silent-exit-code-mode=yes\n");
 				}
 
+				if (AnalyzerVersion.CompareTo(new Version("7.30")) >= 0)
+				{
+					ConfigFileContents.Append("new-output-format=yes\n"); ;
+				}
+
 				string BaseFileName = PreprocessedFileItem.Location.GetFileName();
 
 				FileReference ConfigFileLocation = FileReference.Combine(OutputDir, BaseFileName + ".cfg");
 				FileItem ConfigFileItem = Graph.CreateIntermediateTextFile(ConfigFileLocation, ConfigFileContents.ToString());
 
 				// Run the analzyer on the preprocessed source file
-				FileReference OutputFileLocation = FileReference.Combine(OutputDir, BaseFileName + ".pvslog");
+				FileReference OutputFileLocation = FileReference.Combine(OutputDir, BaseFileName + OutputFileExtension);
 				FileItem OutputFileItem = FileItem.GetItemByFileReference(OutputFileLocation);
 
 				Action AnalyzeAction = Graph.CreateAction(ActionType.Compile);
@@ -871,17 +946,19 @@ namespace UnrealBuildTool
 		public override void FinalizeOutput(ReadOnlyTargetRules Target, TargetMakefileBuilder MakefileBuilder)
 		{
 			FileReference OutputFile;
+			string outputFileExtension = OutputFileExtension;
+
 			if (Target.ProjectFile == null)
 			{
-				OutputFile = FileReference.Combine(Unreal.EngineDirectory, "Saved", "PVS-Studio", String.Format("{0}.pvslog", Target.Name));
+				OutputFile = FileReference.Combine(Unreal.EngineDirectory, "Saved", "PVS-Studio", $"{Target.Name}{outputFileExtension}");
 			}
 			else
 			{
-				OutputFile = FileReference.Combine(Target.ProjectFile.Directory, "Saved", "PVS-Studio", String.Format("{0}.pvslog", Target.Name));
+				OutputFile = FileReference.Combine(Target.ProjectFile.Directory, "Saved", "PVS-Studio", $"{Target.Name}{outputFileExtension}");
 			}
 
 			TargetMakefile Makefile = MakefileBuilder.Makefile;
-			ImmutableSortedSet<FileReference> InputFiles = Makefile.OutputItems.Select(x => x.Location).Where(x => x.HasExtension(".pvslog")).ToImmutableSortedSet();
+			ImmutableSortedSet<FileReference> InputFiles = Makefile.OutputItems.Select(x => x.Location).Where(x => x.HasExtension(outputFileExtension)).ToImmutableSortedSet();
 
 			// Collect the sourcefile items off of the Compile action added in CompileCPPFiles so that in SingleFileCompile mode the PVSGather step is also not filtered out
 			ImmutableSortedSet<FileItem> CompileSourceFiles = Makefile.Actions.OfType<VCCompileAction>().Select(x => x.SourceFile!).ToImmutableSortedSet();
@@ -896,7 +973,8 @@ namespace UnrealBuildTool
 			AnalyzeAction.ActionType = ActionType.PostBuildStep;
 			AnalyzeAction.CommandDescription = "Process PVS-Studio Results";
 			AnalyzeAction.CommandPath = Unreal.DotnetPath;
-			AnalyzeAction.CommandArguments = $"\"{Unreal.UnrealBuildToolDllPath}\" -Mode=PVSGather -Input=\"{InputFileListItem.Location}\" -Output=\"{OutputFile}\" -Ignored=\"{IgnoredFileListeItem.Location}\" -PrintLevel={Target.StaticAnalyzerPVSPrintLevel} ";
+			AnalyzeAction.CommandArguments = $"\"{Unreal.UnrealBuildToolDllPath}\" -Mode=PVSGather -Input=\"{InputFileListItem.Location}\" -Output=\"{OutputFile}\" -Ignored=\"{IgnoredFileListeItem.Location}\"" +
+																				$" -PrintLevel={Target.StaticAnalyzerPVSPrintLevel} -AnalyzerVersion={AnalyzerVersion}";
 			AnalyzeAction.WorkingDirectory = Unreal.EngineSourceDirectory;
 			AnalyzeAction.PrerequisiteItems.Add(InputFileListItem);
 			AnalyzeAction.PrerequisiteItems.Add(IgnoredFileListeItem);
