PR #11411: Fixes for Rider project file generator

diff --git a/Engine/Source/Programs/UnrealBuildTool/Platform/Apple/AppleProjectGenerator.cs b/Engine/Source/Programs/UnrealBuildTool/Platform/Apple/AppleProjectGenerator.cs
new file mode 100644
--- /dev/null
+++ b/Engine/Source/Programs/UnrealBuildTool/Platform/Apple/AppleProjectGenerator.cs
@@ -0,0 +1,163 @@
+﻿// Copyright Epic Games, Inc. All Rights Reserved.
+
+using System;
+using System.Collections.Generic;
+using System.Diagnostics;
+using EpicGames.Core;
+using Microsoft.Extensions.Logging;
+
+namespace UnrealBuildTool
+{
+	/// <summary>
+	/// Base class for Apple-based project generators
+	/// </summary>
+	abstract class AppleProjectGenerator : PlatformProjectGenerator
+	{
+		private static readonly XcrunRunner AppleHelper = new XcrunRunner();
+
+		/// <inheritdoc/>
+		protected AppleProjectGenerator(CommandLineArguments Arguments, ILogger Logger)
+			: base(Arguments, Logger)
+		{
+		}
+
+		/// <inheritdoc/>
+		public override bool HasVisualStudioSupport(VSSettings InVSSettings)
+		{
+			return false;
+		}
+
+		/// <inheritdoc/>
+		public override IList<string> GetSystemIncludePaths(UEBuildTarget InTarget)
+		{
+			if (UEBuildPlatform.IsPlatformInGroup(BuildHostPlatform.Current.Platform, UnrealPlatformGroup.Apple))
+			{
+				// Only generate Apple system include paths when host platform is Apple OS
+				return AppleHelper.GetAppleSystemIncludePaths(InTarget.Architectures.SingleArchitecture, InTarget.Platform, Logger);
+			}
+
+			// TODO: Fix case when working with MacOS on Windows host platform
+			return new List<string>(0);
+		}
+
+		private class XcrunRunner
+		{
+			private readonly Dictionary<string, IList<string>> CachedIncludePaths = new Dictionary<string, IList<string>>();
+
+			private string CurrentlyProcessedSDK = String.Empty;
+			private Process? XcrunProcess;
+			private bool IsReadingIncludesSection;
+
+			public IList<string> GetAppleSystemIncludePaths(UnrealArch Architecture, UnrealTargetPlatform Platform, ILogger Logger)
+			{
+				if (!UEBuildPlatform.IsPlatformInGroup(Platform, UnrealPlatformGroup.Apple))
+				{
+					throw new InvalidOperationException("xcrun can be run only for Apple's platforms");
+				}
+
+				string SDKPath = GetSDKPath(Architecture, Platform, Logger);
+				if (!CachedIncludePaths.ContainsKey(SDKPath))
+				{
+					CalculateSystemIncludePaths(SDKPath, Logger);
+				}
+
+				return CachedIncludePaths[SDKPath];
+			}
+
+			private void CalculateSystemIncludePaths(string SDKPath, ILogger Logger)
+			{
+				if (!String.IsNullOrEmpty(CurrentlyProcessedSDK))
+				{
+					throw new InvalidOperationException("Cannot calculate include paths for several platforms at once");
+				}
+
+				CurrentlyProcessedSDK = SDKPath;
+				CachedIncludePaths[SDKPath] = new List<string>();
+				using (XcrunProcess = new Process())
+				{
+					string AppName = "xcrun";
+					string SystemRootArgument = String.IsNullOrEmpty(SDKPath) ? String.Empty : (" -isysroot " + SDKPath);
+					string Arguments = "clang++ -Wp,-v -x c++ - -fsyntax-only" + SystemRootArgument;
+					XcrunProcess.StartInfo.FileName = AppName;
+					XcrunProcess.StartInfo.Arguments = Arguments;
+					XcrunProcess.StartInfo.UseShellExecute = false;
+					XcrunProcess.StartInfo.CreateNoWindow = true;
+
+					// For some weird reason output of this command is written to error channel so we're redirecting both channels
+					XcrunProcess.StartInfo.RedirectStandardOutput = true;
+					XcrunProcess.StartInfo.RedirectStandardError = true;
+					XcrunProcess.OutputDataReceived += OnOutputDataReceived;
+					XcrunProcess.ErrorDataReceived += OnOutputDataReceived;
+					XcrunProcess.Start();
+					XcrunProcess.BeginOutputReadLine();
+					XcrunProcess.BeginErrorReadLine();
+
+					// xcrun is not finished on it's own. It should be killed by OnOutputDataReceived when reading is finished. But we'll add timeout as a safeguard.
+					// While usually it is fast, first launch on macOS might take ~10 seconds so timeout is quite big: [REDACTED_PROFILE_URL]/llvm-project/issues/75179
+					bool HasExited = XcrunProcess.WaitForExit(30_000);
+					if (!HasExited)
+					{
+						Logger.LogWarning("xcrun didn't finish in 30 second. List of system include paths will not be complete");
+						XcrunProcess.Kill();
+					}
+				}
+
+				XcrunProcess = null;
+				IsReadingIncludesSection = false;
+				CurrentlyProcessedSDK = String.Empty;
+			}
+
+			private void OnOutputDataReceived(object Sender, DataReceivedEventArgs Args)
+			{
+				if (Args.Data != null)
+				{
+					if (IsReadingIncludesSection)
+					{
+						if (Args.Data.StartsWith("End of search"))
+						{
+							IsReadingIncludesSection = false;
+							XcrunProcess!.Kill();
+						}
+						else
+						{
+							if (!Args.Data.EndsWith("(framework directory)"))
+							{
+								CachedIncludePaths[CurrentlyProcessedSDK].Add(Args.Data.Trim(' ', '"'));
+							}
+						}
+					}
+
+					if (Args.Data.StartsWith("#include <...>"))
+					{
+						IsReadingIncludesSection = true;
+					}
+				}
+			}
+
+			private string GetSDKPath(UnrealArch Architecture, UnrealTargetPlatform Platform, ILogger Logger)
+			{
+				if (Platform == UnrealTargetPlatform.Mac)
+				{
+					return MacToolChain.Settings.GetSDKPath().FullName;
+				}
+
+				if (Platform == UnrealTargetPlatform.IOS)
+				{
+					return new IOSToolChainSettings(Logger).GetSDKPath(Architecture).FullName;
+				}
+
+				if (Platform == UnrealTargetPlatform.TVOS)
+				{
+					return new TVOSToolChainSettings(Logger).GetSDKPath(Architecture).FullName;
+				}
+
+				if (Platform == UnrealTargetPlatform.VisionOS)
+				{
+					return new VisionOSToolChainSettings(Logger).GetSDKPath(Architecture).FullName;
+				}
+
+				throw new NotImplementedException("Path to SDK has to be specified for each Apple's platform");
+			}
+		}
+	}
+}
diff --git a/Engine/Source/Programs/UnrealBuildTool/Platform/IOS/IOSProjectGenerator.cs b/Engine/Source/Programs/UnrealBuildTool/Platform/IOS/IOSProjectGenerator.cs
--- a/Engine/Source/Programs/UnrealBuildTool/Platform/IOS/IOSProjectGenerator.cs
+++ b/Engine/Source/Programs/UnrealBuildTool/Platform/IOS/IOSProjectGenerator.cs
@@ -9,7 +9,7 @@ namespace UnrealBuildTool
 	/// <summary>
 	/// Base class for platform-specific project generators
 	/// </summary>
-	class IOSProjectGenerator : PlatformProjectGenerator
+	class IOSProjectGenerator : AppleProjectGenerator
 	{
 		/// <summary>
 		/// Constructor
@@ -28,12 +28,5 @@ namespace UnrealBuildTool
 		{
 			yield return UnrealTargetPlatform.IOS;
 		}
-
-		/// <inheritdoc/>
-		public override bool HasVisualStudioSupport(VSSettings InVSSettings)
-		{
-			// iOS is not supported in VisualStudio
-			return false;
-		}
 	}
 }
diff --git a/Engine/Source/Programs/UnrealBuildTool/Platform/Linux/LinuxProjectGenerator.cs b/Engine/Source/Programs/UnrealBuildTool/Platform/Linux/LinuxProjectGenerator.cs
--- a/Engine/Source/Programs/UnrealBuildTool/Platform/Linux/LinuxProjectGenerator.cs
+++ b/Engine/Source/Programs/UnrealBuildTool/Platform/Linux/LinuxProjectGenerator.cs
@@ -1,8 +1,12 @@
 // Copyright Epic Games, Inc. All Rights Reserved.
 
+using System;
 using System.Collections.Generic;
+using System.IO;
+using System.Text.RegularExpressions;
 using EpicGames.Core;
 using Microsoft.Extensions.Logging;
+using UnrealBuildBase;
 
 namespace UnrealBuildTool
 {
@@ -35,5 +39,81 @@ namespace UnrealBuildTool
 		{
 			return false;
 		}
+
+		/// <inheritdoc/>
+		public override IList<string> GetSystemIncludePaths(UEBuildTarget InTarget)
+		{
+			List<string> Result = new List<string>();
+			string EngineDirectory = Unreal.EngineDirectory.ToString();
+			string? UseLibcxxEnvVarOverride = Environment.GetEnvironmentVariable("UE_LINUX_USE_LIBCXX");
+			UnrealArch TargetArchitecture = InTarget.Architectures.SingleArchitecture;
+			if (String.IsNullOrEmpty(UseLibcxxEnvVarOverride) || UseLibcxxEnvVarOverride == "1")
+			{
+				if (TargetArchitecture == UnrealArch.X64 || TargetArchitecture == UnrealArch.Arm64)
+				{
+					// libc++ include directories
+					Result.Add(Path.Combine(EngineDirectory, "Source/ThirdParty/Unix/LibCxx/include/"));
+					Result.Add(Path.Combine(EngineDirectory, "Source/ThirdParty/Unix/LibCxx/include/c++/v1"));
+				}
+			}
+
+			UEBuildPlatformSDK BuildPlatformSdk = UEBuildPlatform.GetSDK(InTarget.Platform)!;
+			string? InternalSdkPath = BuildPlatformSdk.GetInternalSDKPath();
+			if (InternalSdkPath != null)
+			{
+				string PlatformSdkVersionString = BuildPlatformSdk.GetInstalledVersion()!;
+				string Version = GetLinuxToolchainVersionFromFullString(PlatformSdkVersionString);
+				string ClangIncludeDirectory = Path.Combine(InternalSdkPath, "lib/clang/" + Version + "/include/");
+
+				Result.Add(Path.Combine(InternalSdkPath, "include"));
+				Result.Add(Path.Combine(InternalSdkPath, "usr/include"));
+				Result.Add(ClangIncludeDirectory);
+				if (!Directory.Exists(ClangIncludeDirectory))
+				{
+					Logger.LogWarning("Clang include directory doesn't exist on disk. VersionString={VersionString}, ClangDir={ClangDir}",
+						PlatformSdkVersionString, ClangIncludeDirectory);
+				}
+			}
+
+			return Result;
+		}
+
+		/// <summary>
+		/// Get clang toolchain version from full version string
+		/// v17_clang-10.0.1-centos7 -> 10.0.1
+		/// v17_clang-16.0.1-centos7 -> 16
+		/// </summary>
+		/// <param name="FullVersion">Full clang toolchain version string. Example: "v17_clang-10.0.1-centos7"</param>
+		/// <returns>Clang toolchain version. Example: 10.0.1 or 16</returns>
+		/// <remarks>Starting with clang 16.x the directory naming changed to include major version only</remarks>
+		private static string GetLinuxToolchainVersionFromFullString(string FullVersion)
+		{
+			string FullVersionPattern = @"^v[0-9]+_.*-(([0-9]+)\.[0-9]+\.[0-9]+)-.*$";
+			Regex Regex = new Regex(FullVersionPattern);
+			Match Match = Regex.Match(FullVersion);
+			if (!Match.Success)
+			{
+				throw new ArgumentException("Wrong full version string", FullVersion);
+			}
+
+			Group MajorVersionGroup = Match.Groups[2];
+			CaptureCollection MajorVersionCaptures = MajorVersionGroup.Captures;
+			if (MajorVersionCaptures.Count != 1)
+			{
+				throw new ArgumentException("Multiple regex captures in major version string", FullVersion);
+			}
+
+			if (Int32.TryParse(MajorVersionCaptures[0].Value, out int MajorVersion))
+			{
+				if (MajorVersion >= 16)
+				{
+					return MajorVersionCaptures[0].Value;
+				}
+			}
+
+			Group FullNumberVersionGroup = Match.Groups[1];
+			CaptureCollection FullNumberVersionCaptures = FullNumberVersionGroup.Captures;
+			return FullNumberVersionCaptures[0].Value;
+		}
 	}
 }
diff --git a/Engine/Source/Programs/UnrealBuildTool/Platform/Mac/MacProjectGenerator.cs b/Engine/Source/Programs/UnrealBuildTool/Platform/Mac/MacProjectGenerator.cs
--- a/Engine/Source/Programs/UnrealBuildTool/Platform/Mac/MacProjectGenerator.cs
+++ b/Engine/Source/Programs/UnrealBuildTool/Platform/Mac/MacProjectGenerator.cs
@@ -9,7 +9,7 @@ namespace UnrealBuildTool
 	/// <summary>
 	/// Base class for platform-specific project generators
 	/// </summary>
-	class MacProjectGenerator : PlatformProjectGenerator
+	class MacProjectGenerator : AppleProjectGenerator
 	{
 		/// <summary>
 		/// Constructor
@@ -28,12 +28,5 @@ namespace UnrealBuildTool
 		{
 			yield return UnrealTargetPlatform.Mac;
 		}
-
-		/// <inheritdoc/>
-		public override bool HasVisualStudioSupport(VSSettings InVSSettings)
-		{
-			// Mac is not supported in VisualStudio
-			return false;
-		}
 	}
 }
diff --git a/Engine/Source/Programs/UnrealBuildTool/Platform/Windows/WindowsProjectGenerator.cs b/Engine/Source/Programs/UnrealBuildTool/Platform/Windows/WindowsProjectGenerator.cs
--- a/Engine/Source/Programs/UnrealBuildTool/Platform/Windows/WindowsProjectGenerator.cs
+++ b/Engine/Source/Programs/UnrealBuildTool/Platform/Windows/WindowsProjectGenerator.cs
@@ -83,5 +83,17 @@ namespace UnrealBuildTool
 		{
 			return true;
 		}
+
+		/// <inheritdoc/>
+		public override IList<string> GetSystemIncludePaths(UEBuildTarget InTarget)
+		{
+			List<string> Result = new List<string>();
+			foreach (DirectoryReference Path in InTarget.Rules.WindowsPlatform.Environment!.IncludePaths)
+			{
+				Result.Add(Path.FullName);
+			}
+
+			return Result;
+		}
 	}
 }
diff --git a/Engine/Source/Programs/UnrealBuildTool/ProjectFiles/ProjectFileGenerator.cs b/Engine/Source/Programs/UnrealBuildTool/ProjectFiles/ProjectFileGenerator.cs
--- a/Engine/Source/Programs/UnrealBuildTool/ProjectFiles/ProjectFileGenerator.cs
+++ b/Engine/Source/Programs/UnrealBuildTool/ProjectFiles/ProjectFileGenerator.cs
@@ -500,7 +500,7 @@ namespace UnrealBuildTool
 		/// <summary>
 		/// Adds all rules project files to the solution.
 		/// </summary>
-		void AddRulesModules(Rules.RulesFileType RulesFileType, string ProgramSubDirectory, List<ProjectFile> AddedProjectFiles,
+		protected void AddRulesModules(Rules.RulesFileType RulesFileType, string ProgramSubDirectory, List<ProjectFile> AddedProjectFiles,
 			List<FileReference> UnrealProjectFiles, PrimaryProjectFolder RootFolder, PrimaryProjectFolder ProgramsFolder, ILogger Logger)
 		{
 			List<DirectoryReference> GameFolders = new List<DirectoryReference>();
diff --git a/Engine/Source/Programs/UnrealBuildTool/ProjectFiles/Rider/RiderProjectFile.cs b/Engine/Source/Programs/UnrealBuildTool/ProjectFiles/Rider/RiderProjectFile.cs
--- a/Engine/Source/Programs/UnrealBuildTool/ProjectFiles/Rider/RiderProjectFile.cs
+++ b/Engine/Source/Programs/UnrealBuildTool/ProjectFiles/Rider/RiderProjectFile.cs
@@ -2,11 +2,9 @@
 
 using System;
 using System.Collections.Generic;
-using System.Diagnostics;
 using System.IO;
 using System.Linq;
 using System.Text;
-using System.Text.RegularExpressions;
 using EpicGames.Core;
 using Microsoft.Extensions.Logging;
 using UnrealBuildBase;
@@ -29,8 +27,6 @@ namespace UnrealBuildTool
 
 	internal class RiderProjectFile : ProjectFile
 	{
-		private static readonly XcrunRunner AppleHelper = new XcrunRunner();
-
 		private readonly DirectoryReference RootPath;
 		private readonly HashSet<TargetType> TargetTypes;
 		private readonly CommandLineArguments Arguments;
@@ -517,7 +513,7 @@ namespace UnrealBuildTool
 			}
 			Writer.WriteObjectEnd();
 
-			ExportBuildInfo(Writer, Target, PlatformProjectGenerators, bBuildByDefault);
+			ExportBuildInfo(Writer, Target, PlatformProjectGenerators, bBuildByDefault, Logger);
 
 			Writer.WriteArrayStart("EnvironmentIncludePaths");
 			foreach (DirectoryReference Path in GlobalCompileEnvironment.UserIncludePaths)
@@ -538,62 +534,6 @@ namespace UnrealBuildTool
 				}
 			}
 
-			if (UEBuildPlatform.IsPlatformInGroup(Target.Platform, UnrealPlatformGroup.Windows))
-			{
-				foreach (DirectoryReference Path in Target.Rules.WindowsPlatform.Environment!.IncludePaths)
-				{
-					Writer.WriteValue(Path.FullName);
-				}
-			}
-			else if (UEBuildPlatform.IsPlatformInGroup(Target.Platform, UnrealPlatformGroup.Apple) &&
-					 UEBuildPlatform.IsPlatformInGroup(BuildHostPlatform.Current.Platform, UnrealPlatformGroup.Apple))
-			{
-				// Only generate Apple system include paths when host platform is Apple OS
-				// TODO: Fix case when working with MacOS on Windows host platform  
-				foreach (string Path in AppleHelper.GetAppleSystemIncludePaths(GlobalCompileEnvironment.Architecture, Target.Platform, Logger))
-				{
-					Writer.WriteValue(Path);
-				}
-			}
-			else if (UEBuildPlatform.IsPlatformInGroup(Target.Platform, UnrealPlatformGroup.Linux) ||
-					UEBuildPlatform.IsPlatformInGroup(Target.Platform, UnrealPlatformGroup.Unix))
-			{
-				string EngineDirectory = Unreal.EngineDirectory.ToString();
-
-				string? UseLibcxxEnvVarOverride = Environment.GetEnvironmentVariable("UE_LINUX_USE_LIBCXX");
-				// assumes a single architecture
-				UnrealArch TargetArchitecture = Target.Architectures.SingleArchitecture;
-				if (String.IsNullOrEmpty(UseLibcxxEnvVarOverride) || UseLibcxxEnvVarOverride == "1")
-				{
-					if (TargetArchitecture == UnrealArch.X64 ||
-						TargetArchitecture == UnrealArch.Arm64)
-					{
-						// libc++ include directories
-						Writer.WriteValue(Path.Combine(EngineDirectory, "Source/ThirdParty/Unix/LibCxx/include/"));
-						Writer.WriteValue(Path.Combine(EngineDirectory, "Source/ThirdParty/Unix/LibCxx/include/c++/v1"));
-					}
-				}
-
-				UEBuildPlatform BuildPlatform = UEBuildPlatform.GetBuildPlatform(Target.Platform);
-				string PlatformSdkVersionString = UEBuildPlatformSDK.GetSDKForPlatform(BuildPlatform.GetPlatformName())!.GetInstalledVersion()!;
-				string Version = GetLinuxToolchainVersionFromFullString(PlatformSdkVersionString);
-
-				string? InternalSdkPath = UEBuildPlatform.GetSDK(Target.Platform)!.GetInternalSDKPath();
-				if (InternalSdkPath != null)
-				{
-					Writer.WriteValue(Path.Combine(InternalSdkPath, "include"));
-					Writer.WriteValue(Path.Combine(InternalSdkPath, "usr/include"));
-
-					string ClangIncludeDirectory = Path.Combine(InternalSdkPath, "lib/clang/" + Version + "/include/");
-					Writer.WriteValue(ClangIncludeDirectory);
-					if (!Directory.Exists(ClangIncludeDirectory))
-					{
-						Logger.LogWarning("Clang include directory doesn't exist on disk. VersionString={VersionString}, ClangDir={ClangDir}",
-							PlatformSdkVersionString, ClangIncludeDirectory);
-					}
-				}
-			}
-
 			Writer.WriteArrayEnd();
 
 			Writer.WriteArrayStart("EnvironmentDefinitions");
@@ -604,56 +544,66 @@ namespace UnrealBuildTool
 			Writer.WriteArrayEnd();
 		}
 
-		private void ExportBuildInfo(JsonWriter Writer, UEBuildTarget Target, PlatformProjectGeneratorCollection PlatformProjectGenerators, bool bBuildByDefault)
+		private void ExportBuildInfo(JsonWriter Writer, UEBuildTarget Target, PlatformProjectGeneratorCollection PlatformProjectGenerators,
+			bool bBuildByDefault, ILogger Logger)
 		{
 			if (IsStubProject)
 			{
 				return;
 			}
 
-			Writer.WriteObjectStart("BuildInfo");
-			UnrealTargetPlatform HostPlatform = BuildHostPlatform.Current.Platform;
-
-			ProjectTarget ProjectTarget = ProjectTargets.OfType<ProjectTarget>().Single(It => Target.TargetRulesFile == It.TargetFilePath);
-			UnrealTargetPlatform Platform = Target.Platform;
-			UnrealTargetConfiguration Configuration = Target.Configuration;
-
-			string UProjectPath = "";
-			if (IsForeignProject)
-			{
-				UProjectPath = String.Format("\"{0}\"", ProjectTarget.UnrealProjectFilePath!.FullName);
-			}
-
-			Writer.WriteValue("bBuildByDefault", bBuildByDefault);
-
-			if (HostPlatform.IsInGroup(UnrealPlatformGroup.Windows))
+			try
 			{
-				PlatformProjectGenerator? ProjGenerator = PlatformProjectGenerators.GetPlatformProjectGenerator(Platform, true);
-				VCProjectFile.BuildCommandBuilder BuildCommandBuilder = new VCProjectFile.BuildCommandBuilder(
-					new PlatformProjectGenerator.VSSettings(Platform, Configuration, VCProjectFileFormat.Default, null), ProjectTarget, UProjectPath)
+				string BuildScript;
+				string RebuildScript;
+				string CleanScript;
+				string BuildArguments;
+				string RebuildArguments;
+				string CleanArguments;
+				string Output = Target.Binaries[0].OutputFilePath.FullName;
+
+				ProjectTarget ProjectTarget = ProjectTargets.OfType<ProjectTarget>().Single(It => Target.TargetRulesFile == It.TargetFilePath);
+				string UProjectPath = IsForeignProject ? String.Format("\"{0}\"", ProjectTarget.UnrealProjectFilePath!.FullName) : "";
+				UnrealTargetPlatform HostPlatform = BuildHostPlatform.Current.Platform;
+				if (HostPlatform.IsInGroup(UnrealPlatformGroup.Windows))
 				{
-					ProjectGenerator = ProjGenerator,
-					bIsForeignProject = IsForeignProject
-				};
+					PlatformProjectGenerator? ProjGenerator = PlatformProjectGenerators.GetPlatformProjectGenerator(Target.Platform, true);
+					VCProjectFile.BuildCommandBuilder BuildCommandBuilder =
+						new VCProjectFile.BuildCommandBuilder(
+							new PlatformProjectGenerator.VSSettings(Target.Platform, Target.Configuration, VCProjectFileFormat.Default, null),
+							ProjectTarget, UProjectPath)
+						{
+							ProjectGenerator = ProjGenerator,
+							bIsForeignProject = IsForeignProject
+						};
+
+					BuildArguments = RebuildArguments = CleanArguments = BuildCommandBuilder.GetBuildArguments();
+					BuildScript = EscapePath(BuildCommandBuilder.BuildScript.FullName);
+					RebuildScript = EscapePath(BuildCommandBuilder.RebuildScript.FullName);
+					CleanScript = EscapePath(BuildCommandBuilder.CleanScript.FullName);
+				}
+				else
+				{
+					BuildScript = CleanScript = GetBuildScript(HostPlatform);
+					BuildArguments = GetBuildArguments(HostPlatform, ProjectTarget, Target, UProjectPath, false);
+					CleanArguments = GetBuildArguments(HostPlatform, ProjectTarget, Target, UProjectPath, true);
+					RebuildScript = RebuildArguments = "";
+				}
 
-				string BuildArguments = BuildCommandBuilder.GetBuildArguments();
-				WriteCommand(Writer, "BuildCmd", EscapePath(BuildCommandBuilder.BuildScript.FullName), BuildArguments);
-				WriteCommand(Writer, "RebuildCmd", EscapePath(BuildCommandBuilder.RebuildScript.FullName), BuildArguments);
-				WriteCommand(Writer, "CleanCmd", EscapePath(BuildCommandBuilder.CleanScript.FullName), BuildArguments);
+				Writer.WriteObjectStart("BuildInfo");
+				Writer.WriteValue("bBuildByDefault", bBuildByDefault);
+				WriteCommand(Writer, "BuildCmd", BuildScript, BuildArguments);
+				WriteCommand(Writer, "RebuildCmd", RebuildScript, RebuildArguments);
+				WriteCommand(Writer, "CleanCmd", CleanScript, CleanArguments);
+				Writer.WriteValue("Output", Output);
+				Writer.WriteObjectEnd();
 			}
-			else
+			catch (Exception Ex)
 			{
-				string BuildScript = GetBuildScript(HostPlatform);
-				string BuildArguments = GetBuildArguments(HostPlatform, ProjectTarget, Target, UProjectPath, false);
-				string CleanArguments = GetBuildArguments(HostPlatform, ProjectTarget, Target, UProjectPath, true);
-				WriteCommand(Writer, "BuildCmd", BuildScript, BuildArguments);
-				WriteCommand(Writer, "RebuildCmd", "", "");
-				WriteCommand(Writer, "CleanCmd", BuildScript, CleanArguments);
+				Logger.LogWarning(Ex,
+					"Exception while generating build info for Target: {Target}, Platform: {Platform}, Configuration: {Configuration}",
+					Target.TargetName, Target.Platform.ToString(), Target.Configuration.ToString());
 			}
-
-			UEBuildBinary MainBinary = Target.Binaries[0];
-			Writer.WriteValue("Output", MainBinary.OutputFilePath.FullName);
-			Writer.WriteObjectEnd();
 		}
 
 		private string GetBuildScript(UnrealTargetPlatform HostPlatform)
@@ -809,6 +759,7 @@ namespace UnrealBuildTool
 			if (CurrentTarget!.Platform.IsInGroup(UnrealPlatformGroup.Windows))
 			{
 				ToolchainInfo.bEnableAddressSanitizer = CurrentTarget.Rules.WindowsPlatform.bEnableAddressSanitizer;
+				ToolchainInfo.bUpdatedCPPMacro = CurrentTarget.Rules.WindowsPlatform.bUpdatedCPPMacro;
 				WindowsCompiler WindowsPlatformCompiler = CurrentTarget.Rules.WindowsPlatform.Compiler;
 				ToolchainInfo.bStrictConformanceMode = WindowsPlatformCompiler.IsMSVC() && CurrentTarget.Rules.WindowsPlatform.bStrictConformanceMode;
 				ToolchainInfo.bStrictPreprocessorConformanceMode =
@@ -828,158 +779,5 @@ namespace UnrealBuildTool
 
 			return ToolchainInfo;
 		}
-
-		/// <summary>
-		/// Get clang toolchain version from full version string
-		/// v17_clang-10.0.1-centos7 -> 10.0.1
-		/// v17_clang-16.0.1-centos7 -> 16
-		/// </summary>
-		/// <param name="FullVersion">Full clang toolchain version string. Example: "v17_clang-10.0.1-centos7"</param>
-		/// <returns>Clang toolchain version. Example: 10.0.1 or 16</returns>
-		/// <remarks>Starting with clang 16.x the directory naming changed to include major version only</remarks>
-		private static string GetLinuxToolchainVersionFromFullString(string FullVersion)
-		{
-			string FullVersionPattern = @"^v[0-9]+_.*-(([0-9]+)\.[0-9]+\.[0-9]+)-.*$";
-			Regex Regex = new Regex(FullVersionPattern);
-			Match Match = Regex.Match(FullVersion);
-			if (!Match.Success)
-			{
-				throw new ArgumentException("Wrong full version string", FullVersion);
-			}
-
-			Group MajorVersionGroup = Match.Groups[2];
-			CaptureCollection MajorVersionCaptures = MajorVersionGroup.Captures;
-			if (MajorVersionCaptures.Count != 1)
-			{
-				throw new ArgumentException("Multiple regex captures in major version string", FullVersion);
-			}
-
-			if (Int32.TryParse(MajorVersionCaptures[0].Value, out int MajorVersion))
-			{
-				if (MajorVersion >= 16)
-				{
-					return MajorVersionCaptures[0].Value;
-				}
-			}
-
-			Group FullNumberVersionGroup = Match.Groups[1];
-			CaptureCollection FullNumberVersionCaptures = FullNumberVersionGroup.Captures;
-			return FullNumberVersionCaptures[0].Value;
-		}
-
-		private class XcrunRunner
-		{
-			private readonly Dictionary<string, IList<string>> CachedIncludePaths = new Dictionary<string, IList<string>>();
-
-			private string CurrentlyProcessedSDK = String.Empty;
-			private Process? XcrunProcess;
-			private bool IsReadingIncludesSection;
-
-			public IList<string> GetAppleSystemIncludePaths(UnrealArch Architecture, UnrealTargetPlatform Platform, ILogger Logger)
-			{
-				if (!UEBuildPlatform.IsPlatformInGroup(Platform, UnrealPlatformGroup.Apple))
-				{
-					throw new InvalidOperationException("xcrun can be run only for Apple's platforms");
-				}
-
-				string SDKPath = GetSDKPath(Architecture, Platform, Logger);
-				if (!CachedIncludePaths.ContainsKey(SDKPath))
-				{
-					CalculateSystemIncludePaths(SDKPath, Logger);
-				}
-
-				return CachedIncludePaths[SDKPath];
-			}
-
-			private void CalculateSystemIncludePaths(string SDKPath, ILogger Logger)
-			{
-				if (!String.IsNullOrEmpty(CurrentlyProcessedSDK))
-				{
-					throw new InvalidOperationException("Cannot calculate include paths for several platforms at once");
-				}
-
-				CurrentlyProcessedSDK = SDKPath;
-				CachedIncludePaths[SDKPath] = new List<string>();
-				using (XcrunProcess = new Process())
-				{
-					string AppName = "xcrun";
-					string SystemRootArgument = String.IsNullOrEmpty(SDKPath) ? String.Empty : (" -isysroot " + SDKPath);
-					string Arguments = "clang++ -Wp,-v -x c++ - -fsyntax-only" + SystemRootArgument;
-					XcrunProcess.StartInfo.FileName = AppName;
-					XcrunProcess.StartInfo.Arguments = Arguments;
-					XcrunProcess.StartInfo.UseShellExecute = false;
-					XcrunProcess.StartInfo.CreateNoWindow = true;
-
-					// For some weird reason output of this command is written to error channel so we're redirecting both channels
-					XcrunProcess.StartInfo.RedirectStandardOutput = true;
-					XcrunProcess.StartInfo.RedirectStandardError = true;
-					XcrunProcess.OutputDataReceived += OnOutputDataReceived;
-					XcrunProcess.ErrorDataReceived += OnOutputDataReceived;
-					XcrunProcess.Start();
-					XcrunProcess.BeginOutputReadLine();
-					XcrunProcess.BeginErrorReadLine();
-
-					// xcrun is not finished on it's own. It should be killed by OnOutputDataReceived when reading is finished. But we'll add timeout as a safeguard.
-					// While usually it is fast, first launch on macOS might take ~10 seconds so timeout is quite big: [REDACTED_PROFILE_URL]/llvm-project/issues/75179
-					bool HasExited = XcrunProcess.WaitForExit(30_000);
-					if (!HasExited)
-					{
-						Logger.LogWarning("xcrun didn't finish in 30 second. List of system include paths will not be complete");
-						XcrunProcess.Kill();
-					}
-				}
-
-				XcrunProcess = null;
-				IsReadingIncludesSection = false;
-				CurrentlyProcessedSDK = String.Empty;
-			}
-
-			private void OnOutputDataReceived(object Sender, DataReceivedEventArgs Args)
-			{
-				if (Args.Data != null)
-				{
-					if (IsReadingIncludesSection)
-					{
-						if (Args.Data.StartsWith("End of search"))
-						{
-							IsReadingIncludesSection = false;
-							XcrunProcess!.Kill();
-						}
-						else
-						{
-							if (!Args.Data.EndsWith("(framework directory)"))
-							{
-								CachedIncludePaths[CurrentlyProcessedSDK].Add(Args.Data.Trim(' ', '"'));
-							}
-						}
-					}
-
-					if (Args.Data.StartsWith("#include <...>"))
-					{
-						IsReadingIncludesSection = true;
-					}
-				}
-			}
-
-			private string GetSDKPath(UnrealArch Architecture, UnrealTargetPlatform Platform, ILogger Logger)
-			{
-				if (Platform == UnrealTargetPlatform.Mac)
-				{
-					return MacToolChain.Settings.GetSDKPath().FullName;
-				}
-
-				if (Platform == UnrealTargetPlatform.IOS)
-				{
-					return new IOSToolChainSettings(Logger).GetSDKPath(Architecture).FullName;
-				}
-
-				if (Platform == UnrealTargetPlatform.TVOS)
-				{
-					return new TVOSToolChainSettings(Logger).GetSDKPath(Architecture).FullName;
-				}
-
-				throw new NotImplementedException("Path to SDK has to be specified for each Apple's platform");
-			}
-		}
 	}
 }
diff --git a/Engine/Source/Programs/UnrealBuildTool/ProjectFiles/Rider/RiderProjectFileGenerator.cs b/Engine/Source/Programs/UnrealBuildTool/ProjectFiles/Rider/RiderProjectFileGenerator.cs
--- a/Engine/Source/Programs/UnrealBuildTool/ProjectFiles/Rider/RiderProjectFileGenerator.cs
+++ b/Engine/Source/Programs/UnrealBuildTool/ProjectFiles/Rider/RiderProjectFileGenerator.cs
@@ -475,6 +475,11 @@ namespace UnrealBuildTool
 				WriteProjectFiles(PlatformProjectGenerators, Logger);
 			}
 
+			// Generate all necessary .props files for C# projects
+			PrimaryProjectFolder ProgramsFolder = RootFolder.AddSubFolder("Programs");
+			AddRulesModules(Rules.RulesFileType.AutomationModule, "Automation", AutomationProjectFiles, AllGameProjects, RootFolder, ProgramsFolder, Logger);
+			AddRulesModules(Rules.RulesFileType.UbtPlugin, "UnrealBuildTool.Plugins", UbtPluginProjectFiles, AllGameProjects, RootFolder, ProgramsFolder, Logger);
+
 			Logger.LogDebug("Project generation complete ({NumGenerated} generated, {NumImported} imported)", GeneratedProjectFiles.Count,
 				OtherProjectFiles.Count);
 
diff --git a/Engine/Source/Programs/UnrealBuildTool/ProjectFiles/Rider/ToolchainInfo.cs b/Engine/Source/Programs/UnrealBuildTool/ProjectFiles/Rider/ToolchainInfo.cs
--- a/Engine/Source/Programs/UnrealBuildTool/ProjectFiles/Rider/ToolchainInfo.cs
+++ b/Engine/Source/Programs/UnrealBuildTool/ProjectFiles/Rider/ToolchainInfo.cs
@@ -31,6 +31,7 @@ namespace UnrealBuildTool
 		public bool bStrictPreprocessorConformanceMode;
 		public bool bEnableCoroutines;
 		public bool bEnableAddressSanitizer;
+		public bool bUpdatedCPPMacro;
 
 		public IEnumerable<Tuple<string, object?>> GetDiff(ToolchainInfo Other)
 		{
@@ -88,7 +89,8 @@ namespace UnrealBuildTool
 				PrecompiledHeaderAction == Other.PrecompiledHeaderAction && PrecompiledHeaderFile == Other.PrecompiledHeaderFile &&
 				Equals(ForceIncludeFiles, Other.ForceIncludeFiles) && Compiler == Other.Compiler &&
 				bStrictConformanceMode == Other.bStrictConformanceMode && bStrictPreprocessorConformanceMode == Other.bStrictPreprocessorConformanceMode &&
-				bEnableCoroutines == Other.bEnableCoroutines && bEnableAddressSanitizer == Other.bEnableAddressSanitizer;
+				bEnableCoroutines == Other.bEnableCoroutines && bEnableAddressSanitizer == Other.bEnableAddressSanitizer &&
+				bUpdatedCPPMacro == Other.bUpdatedCPPMacro;
 		}
 
 		public override bool Equals(object? Obj)
@@ -132,6 +134,7 @@ namespace UnrealBuildTool
 				HashCode = (HashCode * 397) ^ bStrictPreprocessorConformanceMode.GetHashCode();
 				HashCode = (HashCode * 397) ^ bEnableCoroutines.GetHashCode();
 				HashCode = (HashCode * 397) ^ bEnableAddressSanitizer.GetHashCode();
+				HashCode = (HashCode * 397) ^ bUpdatedCPPMacro.GetHashCode();
 				return HashCode;
 			}
 		}
