UnrealBuildTool: Fix reference to unavailable symbol

#rnx

diff --git a/Engine/Source/Programs/UnrealBuildTool/Platform/Apple/AppleProjectGenerator.cs b/Engine/Source/Programs/UnrealBuildTool/Platform/Apple/AppleProjectGenerator.cs
--- a/Engine/Source/Programs/UnrealBuildTool/Platform/Apple/AppleProjectGenerator.cs
+++ b/Engine/Source/Programs/UnrealBuildTool/Platform/Apple/AppleProjectGenerator.cs
@@ -1,4 +1,4 @@
-﻿// Copyright Epic Games, Inc. All Rights Reserved.
+// Copyright Epic Games, Inc. All Rights Reserved.
 
 using System;
 using System.Collections.Generic;
@@ -141,22 +141,8 @@ namespace UnrealBuildTool
 					return MacToolChain.Settings.GetSDKPath().FullName;
 				}
 
-				if (Platform == UnrealTargetPlatform.IOS)
-				{
-					return new IOSToolChainSettings(Logger).GetSDKPath(Architecture).FullName;
-				}
-
-				if (Platform == UnrealTargetPlatform.TVOS)
-				{
-					return new TVOSToolChainSettings(Logger).GetSDKPath(Architecture).FullName;
-				}
-
-				if (Platform == UnrealTargetPlatform.VisionOS)
-				{
-					return new VisionOSToolChainSettings(Logger).GetSDKPath(Architecture).FullName;
-				}
-
-				throw new NotImplementedException("Path to SDK has to be specified for each Apple's platform");
+				// Resolves to AppleToolChainSettings.GetSDKPath() which is not overridden
+				return new IOSToolChainSettings(Logger).GetSDKPath(Architecture).FullName;
 			}
 		}
 	}
