[UBT]
* Changed so dll linking actions use /NOIMPLIB to prevent writing of .lib/.exp file (instead of deleting them afterwards)
* Changed so lib creation is calling link.exe directly instead of via lib.exe (lib.exe [REDACTED_HANDLE] just spawns link.exe /lib [REDACTED_HANDLE]).

Note, Ideally we should figure out how to do this splitup of two actions with lld-link.exe but there is no way to run lld-link.exe and only create the import lib.. /NOIMPLIB is supported, but not the other way around.

diff --git a/Engine/Source/Programs/UnrealBuildTool/Platform/Windows/VCEnvironment.cs b/Engine/Source/Programs/UnrealBuildTool/Platform/Windows/VCEnvironment.cs
--- a/Engine/Source/Programs/UnrealBuildTool/Platform/Windows/VCEnvironment.cs
+++ b/Engine/Source/Programs/UnrealBuildTool/Platform/Windows/VCEnvironment.cs
@@ -246,7 +246,7 @@ namespace UnrealBuildTool
 				//return FileReference.Combine(CompilerDir, "bin", "compiler", "lld-link.exe");
 				//return FileReference.Combine(CompilerDir, "bin", "xilib.exe");
 			}
-			return FileReference.Combine(GetVCToolPath(ToochainDir, Architecture), "lib.exe");
+			return FileReference.Combine(GetVCToolPath(ToochainDir, Architecture), "link.exe"); // We add /LIB to cmd line so we can use link.exe directly instead of going via lib.exe
 		}
 
 		/// <summary>
diff --git a/Engine/Source/Programs/UnrealBuildTool/Platform/Windows/VCToolChain.cs b/Engine/Source/Programs/UnrealBuildTool/Platform/Windows/VCToolChain.cs
--- a/Engine/Source/Programs/UnrealBuildTool/Platform/Windows/VCToolChain.cs
+++ b/Engine/Source/Programs/UnrealBuildTool/Platform/Windows/VCToolChain.cs
@@ -2475,10 +2475,10 @@ namespace UnrealBuildTool
 			// file is not needed for our builds, but there is no way to prevent MSVC from generating it when
 			// linking targets that have exports.  We don't want this to clobber our LIB file and invalidate the
 			// existing timstamp, so instead we simply emit it with a different name
-			FileReference ImportLibraryFilePath;
+			FileReference? ImportLibraryFilePath = null;
 			if (LinkEnvironment.bIsCrossReferenced && !bBuildImportLibraryOnly)
 			{
-				ImportLibraryFilePath = FileReference.Combine(LinkEnvironment.IntermediateDirectory!, LinkEnvironment.OutputFilePath.GetFileNameWithoutExtension() + ".sup.lib");
+				Arguments.Add("/NOIMPLIB");
 			}
 			else if (Target.bShouldCompileAsDLL)
 			{
@@ -2492,7 +2492,7 @@ namespace UnrealBuildTool
 			FileItem OutputFile;
 			if (bBuildImportLibraryOnly)
 			{
-				OutputFile = FileItem.GetItemByFileReference(ImportLibraryFilePath);
+				OutputFile = FileItem.GetItemByFileReference(ImportLibraryFilePath!);
 			}
 			else
 			{
@@ -2547,7 +2547,7 @@ namespace UnrealBuildTool
 			// Since all DLLs are typically marked as cross referenced now anyway, we can just ignore this file to allow incremental linking to work.
 			if (LinkEnvironment.bHasExports && !LinkEnvironment.bIsBuildingLibrary && !LinkEnvironment.bIsCrossReferenced)
 			{
-				FileReference ExportFilePath = ImportLibraryFilePath.ChangeExtension(".exp");
+				FileReference ExportFilePath = ImportLibraryFilePath!.ChangeExtension(".exp");
 				FileItem ExportFile = FileItem.GetItemByFileReference(ExportFilePath);
 				ProducedItems.Add(ExportFile);
 			}
@@ -2555,11 +2555,11 @@ namespace UnrealBuildTool
 			if (!bIsBuildingLibraryOrImportLibrary)
 			{
 				// There is anything to export
-				if (LinkEnvironment.bHasExports)
+				if (LinkEnvironment.bHasExports && !LinkEnvironment.bIsBuildingLibrary && !LinkEnvironment.bIsCrossReferenced)
 				{
 					// Write the import library to the output directory for nFringe support.
-					FileItem ImportLibraryFile = FileItem.GetItemByFileReference(ImportLibraryFilePath);
-					Arguments.Add($"/IMPLIB:\"{NormalizeCommandLinePath(ImportLibraryFilePath)}\"");
+					FileItem ImportLibraryFile = FileItem.GetItemByFileReference(ImportLibraryFilePath!);
+					Arguments.Add($"/IMPLIB:\"{NormalizeCommandLinePath(ImportLibraryFilePath!)}\"");
 
 					// Like the export file above, don't add the import library as a produced item when it's cross referenced.
 					if (!LinkEnvironment.bIsCrossReferenced)
@@ -2634,12 +2634,13 @@ namespace UnrealBuildTool
 			if (bIsBuildingLibraryOrImportLibrary)
 			{
 				LinkAction.CommandPath = EnvVars.LibraryManagerPath;
+				LinkAction.CommandArguments = $"/LIB @\"{ResponseFileName}\"";
 			}
 			else
 			{
 				LinkAction.CommandPath = EnvVars.LinkerPath;
+				LinkAction.CommandArguments = $"@\"{ResponseFileName}\"";
 			}
-			LinkAction.CommandArguments = $"@\"{ResponseFileName}\"";
 			LinkAction.CommandVersion = EnvVars.ToolChainVersion.ToString();
 			LinkAction.ProducedItems.UnionWith(ProducedItems);
 			LinkAction.PrerequisiteItems.UnionWith(PrerequisiteItems);
@@ -2657,14 +2658,6 @@ namespace UnrealBuildTool
 				LinkAction.DeleteItems.UnionWith(LinkAction.ProducedItems.Where(x => x.Location.HasExtension(".pdb") || x.Location.HasExtension(".full.pdb")));
 			}
 
-			// Delete any .sup.lib files before building, even if they're not tracked
-			if (ImportLibraryFilePath.GetFileName().EndsWith(".sup.lib", StringComparison.OrdinalIgnoreCase))
-			{
-				FileReference ExportFilePath = ImportLibraryFilePath.ChangeExtension(".exp");
-				LinkAction.DeleteItems.Add(FileItem.GetItemByFileReference(ImportLibraryFilePath));
-				LinkAction.DeleteItems.Add(FileItem.GetItemByFileReference(ExportFilePath));
-			}
-
 			// Tell the action that we're building an import library here and it should conditionally be
 			// ignored as a prerequisite for other actions
 			LinkAction.bProducesImportLibrary = bBuildImportLibraryOnly || LinkEnvironment.bIsBuildingDLL;
