UnrealBuildTool: Print correct status when preprocessing for clang platforms

#rnx

diff --git a/Engine/Source/Programs/UnrealBuildTool/ToolChain/ClangToolChain.cs b/Engine/Source/Programs/UnrealBuildTool/ToolChain/ClangToolChain.cs
--- a/Engine/Source/Programs/UnrealBuildTool/ToolChain/ClangToolChain.cs
+++ b/Engine/Source/Programs/UnrealBuildTool/ToolChain/ClangToolChain.cs
@@ -332,6 +332,10 @@ namespace UnrealBuildTool
 		/// </summary>
 		protected bool CompilerVersionLessThan(int Major, int Minor, int Patch) => Info.ClangVersion < new Version(Major, Minor, Patch);
 
+		protected bool IsPreprocessing(CppCompileEnvironment CompileEnvironment) =>
+				CompileEnvironment.PrecompiledHeaderAction != PrecompiledHeaderAction.Create
+			&& CompileEnvironment.bPreprocessOnly;
+
 		protected bool IsAnalyzing(CppCompileEnvironment CompileEnvironment) =>
 				StaticAnalyzer == StaticAnalyzer.Default
 			&& CompileEnvironment.PrecompiledHeaderAction != PrecompiledHeaderAction.Create
@@ -937,7 +941,7 @@ namespace UnrealBuildTool
 					ParentPCHInstance = ParentPCHInstance.ParentPCHInstance;
 				}
 			}
-			else if (CompileEnvironment.bPreprocessOnly)
+			else if (IsPreprocessing(CompileEnvironment))
 			{
 				OutputFile = FileItem.GetItemByFileReference(FileReference.Combine(OutputDir, GetFileNameFromExtension(FileName, ".i")));
 				CompileResult.ObjectFiles.Add(OutputFile);
@@ -1144,7 +1148,7 @@ namespace UnrealBuildTool
 			CompileAction.WorkingDirectory = Unreal.EngineSourceDirectory;
 			CompileAction.CommandPath = Info.Clang;
 			CompileAction.CommandVersion = Info.ClangVersionString;
-			CompileAction.CommandDescription = IsAnalyzing(CompileEnvironment) ? "Analyze" : "Compile";
+			CompileAction.CommandDescription = IsPreprocessing(CompileEnvironment) ? "Preprocess" : IsAnalyzing(CompileEnvironment) ? "Analyze" : "Compile";
 			UnrealArchitectureConfig ArchConfig = UnrealArchitectureConfig.ForPlatform(CompileEnvironment.Platform);
 			if (ArchConfig.Mode != UnrealArchitectureMode.SingleArchitecture)
 			{
