UnrealBuildTool: Fix some warnings

#rnx

diff --git a/Engine/Source/Programs/UnrealBuildTool/Executors/Experimental/UBAExecutor.cs b/Engine/Source/Programs/UnrealBuildTool/Executors/Experimental/UBAExecutor.cs
--- a/Engine/Source/Programs/UnrealBuildTool/Executors/Experimental/UBAExecutor.cs
+++ b/Engine/Source/Programs/UnrealBuildTool/Executors/Experimental/UBAExecutor.cs
@@ -308,7 +308,7 @@ namespace UnrealBuildTool
 			_threadedLogger.LogInformation("  Storage capacity {StoreCapacityGb}Gb", UBAConfig.StoreCapacityGb);
 		}
 
-		private async Task WriteActionOutputFile(IEnumerable<LinkedAction> inputActions)
+		private async Task WriteActionOutputFileAsync(IEnumerable<LinkedAction> inputActions)
 		{
 			if (String.IsNullOrEmpty(UBAConfig.ActionsOutputFile))
 			{
@@ -380,7 +380,7 @@ namespace UnrealBuildTool
 			}
 
 			PrintConfiguration();
-			await WriteActionOutputFile(inputActions);
+			await WriteActionOutputFileAsync(inputActions);
 
 			logger = _threadedLogger;
 
@@ -647,7 +647,9 @@ namespace UnrealBuildTool
 		static bool ForceLocalNoDetour(LinkedAction action)
 		{
 			if (!OperatingSystem.IsMacOS()) // Below code is slow, so early out
+			{
 				return false;
+			}
 			// Don't let Mac run shell commands through Uba as interposing dylibs into
 			// the shell results in dyld errors about no matching architecture.
 			// The shell is used to run various commands during a build like copy/ditto.
diff --git a/Engine/Source/Programs/UnrealBuildTool/Executors/Experimental/UBAHordeSession.cs b/Engine/Source/Programs/UnrealBuildTool/Executors/Experimental/UBAHordeSession.cs
--- a/Engine/Source/Programs/UnrealBuildTool/Executors/Experimental/UBAHordeSession.cs
+++ b/Engine/Source/Programs/UnrealBuildTool/Executors/Experimental/UBAHordeSession.cs
@@ -378,7 +378,7 @@ namespace UnrealBuildTool
 
 		public int NumLogicalCores { get; private set; } = 0;
 
-		public async void RemoveCompleteWorkers()
+		public async void RemoveCompleteWorkersAsync()
 		{
 			for (int idx = 0; idx < _workers.Count; idx++)
 			{
@@ -535,7 +535,7 @@ namespace UnrealBuildTool
 			}
 		}
 
-		public static async Task<UBAHordeSession?> TryCreateHordeSession(UnrealBuildAcceleratorHordeConfig hordeConfig, UBAExecutor executor, bool bStrictErrors, ILogger logger, CancellationToken cancellationToken = default)
+		public static async Task<UBAHordeSession?> TryCreateHordeSessionAsync(UnrealBuildAcceleratorHordeConfig hordeConfig, UBAExecutor executor, bool bStrictErrors, ILogger logger, CancellationToken cancellationToken = default)
 		{
 			if (hordeConfig.bDisableHorde)
 			{
@@ -788,7 +788,7 @@ namespace UnrealBuildTool
 			}
 
 			_cancellationSource = new CancellationTokenSource();
-			_hordeSessionTask = UBAHordeSession.TryCreateHordeSession(HordeConfig, executor, _ubaConfig.bStrict, _logger, _cancellationSource.Token);
+			_hordeSessionTask = UBAHordeSession.TryCreateHordeSessionAsync(HordeConfig, executor, _ubaConfig.bStrict, _logger, _cancellationSource.Token);
 			await _hordeSessionTask;
 		}
 
@@ -817,7 +817,7 @@ namespace UnrealBuildTool
 					return;
 				}
 
-				hordeSession.RemoveCompleteWorkers();
+				hordeSession.RemoveCompleteWorkersAsync();
 
 				if (queue.IsDone)
 				{
