- Deprecate unused bUsesSteam in TargetRules
- Remove bUsesSteam from ShooterGame targets

diff --git a/Engine/Source/Programs/UnrealBuildTool/Configuration/ReadOnlyTargetRules.cs b/Engine/Source/Programs/UnrealBuildTool/Configuration/ReadOnlyTargetRules.cs
--- a/Engine/Source/Programs/UnrealBuildTool/Configuration/ReadOnlyTargetRules.cs
+++ b/Engine/Source/Programs/UnrealBuildTool/Configuration/ReadOnlyTargetRules.cs
@@ -81,6 +81,7 @@ namespace UnrealBuildTool
 
 		public string? OutputFile => Inner.OutputFile;
 
+		[Obsolete("Deprecated in UE5.5 - No longer used in engine.")]
 		public bool bUsesSteam => Inner.bUsesSteam;
 
 		public bool bUsesCEF3 => Inner.bUsesCEF3;
diff --git a/Engine/Source/Programs/UnrealBuildTool/Configuration/TargetRules.cs b/Engine/Source/Programs/UnrealBuildTool/Configuration/TargetRules.cs
--- a/Engine/Source/Programs/UnrealBuildTool/Configuration/TargetRules.cs
+++ b/Engine/Source/Programs/UnrealBuildTool/Configuration/TargetRules.cs
@@ -683,6 +683,7 @@ namespace UnrealBuildTool
 		/// <summary>
 		/// Whether the target uses Steam.
 		/// </summary>
+		[Obsolete("Deprecated in UE5.5 - No longer used in engine.")]
 		public bool bUsesSteam { get; set; }
 
 		/// <summary>
