[UBT]
* Add /NOEXP to crossreferenced link action when building dll (.exp is already written by lib action). Note, this is only enabled for msvc linker (clang linker does not support it)

diff --git a/Engine/Source/Programs/UnrealBuildTool/Platform/Windows/VCToolChain.cs b/Engine/Source/Programs/UnrealBuildTool/Platform/Windows/VCToolChain.cs
--- a/Engine/Source/Programs/UnrealBuildTool/Platform/Windows/VCToolChain.cs
+++ b/Engine/Source/Programs/UnrealBuildTool/Platform/Windows/VCToolChain.cs
@@ -2479,6 +2479,8 @@ namespace UnrealBuildTool
 			if (LinkEnvironment.bIsCrossReferenced && !bBuildImportLibraryOnly)
 			{
 				Arguments.Add("/NOIMPLIB");
+				if (!Target.WindowsPlatform.bAllowClangLinker)
+					Arguments.Add("/NOEXP"); // This compiler flag does not exist on lld-link.exe.. it skips the writing of the .exp file
 			}
 			else if (Target.bShouldCompileAsDLL)
 			{
