Move UBAExecutor out of experimental

#rnx

diff --git a/Engine/Source/Programs/UnrealBuildTool/Executors/Experimental/UBAAgentCoordinator.cs b/Engine/Source/Programs/UnrealBuildTool/Executors/Experimental/UBAAgentCoordinator.cs
deleted file mode 100644
diff --git a/Engine/Source/Programs/UnrealBuildTool/Executors/Experimental/UBAExecutor.cs b/Engine/Source/Programs/UnrealBuildTool/Executors/Experimental/UBAExecutor.cs
deleted file mode 100644
diff --git a/Engine/Source/Programs/UnrealBuildTool/Executors/Experimental/UBAHordeSession.cs b/Engine/Source/Programs/UnrealBuildTool/Executors/Experimental/UBAHordeSession.cs
deleted file mode 100644
diff --git a/Engine/Source/Programs/UnrealBuildTool/Executors/UnrealBuildAccelerator/UBAAgentCoordinator.cs b/Engine/Source/Programs/UnrealBuildTool/Executors/UnrealBuildAccelerator/UBAAgentCoordinator.cs
new file mode 100644
--- /dev/null
+++ b/Engine/Source/Programs/UnrealBuildTool/Executors/UnrealBuildAccelerator/UBAAgentCoordinator.cs
@@ -0,0 +1,22 @@
+// Copyright Epic Games, Inc. All Rights Reserved.
+
+using EpicGames.Core;
+using System;
+using System.Threading.Tasks;
+
+namespace UnrealBuildTool
+{
+
+	interface IUBAAgentCoordinator
+	{
+		DirectoryReference? GetUBARootDir();
+
+		Task InitAsync(UBAExecutor executor);
+
+		void Start(ImmediateActionQueue queue, Func<LinkedAction, bool> canRunRemotely);
+
+		void Stop();
+	
+		Task CloseAsync();
+	}
+}
\ No newline at end of file
diff --git a/Engine/Source/Programs/UnrealBuildTool/Executors/UnrealBuildAccelerator/UBAAgentCoordinatorHorde.cs b/Engine/Source/Programs/UnrealBuildTool/Executors/UnrealBuildAccelerator/UBAAgentCoordinatorHorde.cs
new file mode 100644
--- /dev/null
+++ b/Engine/Source/Programs/UnrealBuildTool/Executors/UnrealBuildAccelerator/UBAAgentCoordinatorHorde.cs
@@ -0,0 +1,777 @@
+// Copyright Epic Games, Inc. All Rights Reserved.
+
+using System;
+using System.Collections.Generic;
+using System.Diagnostics;
+using System.Linq;
+using System.Net.Http;
+using System.Net.Http.Headers;
+using System.Runtime.InteropServices;
+using System.Threading;
+using System.Threading.Tasks;
+using EpicGames.Core;
+using EpicGames.Horde;
+using EpicGames.Horde.Common;
+using EpicGames.Horde.Compute;
+using EpicGames.Horde.Compute.Clients;
+using EpicGames.Horde.Storage;
+using EpicGames.Horde.Storage.Bundles;
+using EpicGames.Horde.Storage.Nodes;
+using EpicGames.OIDC;
+using Microsoft.Extensions.Configuration;
+using Microsoft.Extensions.DependencyInjection;
+using Microsoft.Extensions.Logging;
+using Microsoft.Extensions.Logging.Abstractions;
+using UnrealBuildBase;
+
+namespace UnrealBuildTool
+{
+	class UBAHordeSession : IAsyncDisposable
+	{
+		/// <summary>
+		/// ID for this Horde session
+		/// </summary>
+		readonly Guid _id = Guid.NewGuid();
+
+		const string ResourceLogicalCores = "LogicalCores";
+		readonly ClusterId _clusterId = new("default");
+
+		readonly Uri _hordeUri;
+		readonly string? _pool;
+		readonly bool _allowWine;
+		readonly int _maxCores;
+		readonly bool _strict;
+		readonly ConnectionMode? _connectionMode;
+		readonly Encryption? _encryption;
+		readonly CancellationTokenSource _cancellationTokenSource = new();
+		readonly ILogger _logger;
+
+		readonly UBAExecutor _owner;
+
+		readonly BundleStorageClient _storage = BundleStorageClient.CreateInMemory(NullLogger.Instance);
+		BlobLocator _ubaAgentLocator;
+
+		readonly ServiceProvider _serviceProvider;
+		readonly string _crypto;
+		IComputeClient? _client;
+
+		public struct Worker
+		{
+			public Task BackgroundTask { get; set; }
+			public int NumLogicalCores { get; set; }
+			public Stopwatch StartTime { get; set; }
+			public bool Started { get; set; }
+			public string Ip { get; set; }
+			public ConnectionMetadataPort Port { get; set; }
+			public ConnectionMetadataPort ProxyPort { get; set; }
+		}
+
+		readonly List<Worker> _workers = new();
+
+		public UBAHordeSession(UBAExecutor owner, Uri hordeUri, AuthenticationHeaderValue? authHeader, string? pool, bool allowWine, int maxCores, bool strict, ConnectionMode? connectionMode, Encryption? encryption, ILogger logger)
+		{
+			_owner = owner;
+			_hordeUri = hordeUri;
+			_pool = pool;
+			_allowWine = allowWine;
+			_maxCores = maxCores;
+			_strict = strict;
+			_connectionMode = connectionMode;
+			_encryption = encryption;
+			_logger = logger;
+			_crypto = owner.Crypto;
+
+			void ConfigureHttpClient(HttpClient httpClient)
+			{
+				httpClient.BaseAddress = _hordeUri;
+				httpClient.DefaultRequestHeaders.Authorization = authHeader;
+			}
+
+			ServiceCollection services = new();
+			services.AddHttpClient<HordeHttpClient>(ConfigureHttpClient);
+			_serviceProvider = services.BuildServiceProvider();
+
+			if (connectionMode == ConnectionMode.Relay && String.IsNullOrEmpty(_crypto))
+			{
+				_crypto = UBAExecutor.CreateCrypto();
+				_encryption = Encryption.Ssl;
+			}
+		}
+
+		public async ValueTask DisposeAsync()
+		{
+			if (_workers.Count > 0) // Should handle double-dispose, prevent cancelling twice
+			{
+				_cancellationTokenSource.Cancel();
+
+				for (int idx = _workers.Count - 1; idx >= 0; idx--)
+				{
+					await _workers[idx].BackgroundTask;
+					_workers.RemoveAt(idx);
+				}
+			}
+
+			if (_client != null)
+			{
+				// Set CPU resource need to zero (will also expire on the server if not updated)
+				await UpdateCpuCoreNeedAsync(0);
+				await _client.DisposeAsync();
+				_client = null;
+			}
+
+			await _serviceProvider.DisposeAsync();
+			_cancellationTokenSource.Dispose();
+
+			_storage.Dispose();
+		}
+
+		public async Task InitAsync(bool useSentry, CancellationToken cancellationToken)
+		{
+			if (_client != null)
+			{
+				throw new InvalidOperationException("Session has already been initialized");
+			}
+
+			string? sessionId = null;
+			string? jobId = Environment.GetEnvironmentVariable("UE_HORDE_JOBID");
+			string? batchId = Environment.GetEnvironmentVariable("UE_HORDE_BATCHID");
+			string? stepId = Environment.GetEnvironmentVariable("UE_HORDE_STEPID");
+
+			if (jobId != null && batchId != null && stepId != null)
+			{
+				sessionId = $"{jobId}-{batchId}-{stepId}";
+			}
+
+			_client = new ServerComputeClient(_serviceProvider.GetRequiredService<IHttpClientFactory>(), sessionId, _logger);
+
+			_logger.LogInformation("Creating tool bundle...");
+			DirectoryReference ubaDir;
+			List<string> agentFiles = new();
+			if (OperatingSystem.IsWindows())
+			{
+#pragma warning disable CA1308 // Normalize strings to uppercase
+				ubaDir = DirectoryReference.Combine(Unreal.EngineDirectory, "Binaries", "Win64", "UnrealBuildAccelerator", RuntimeInformation.ProcessArchitecture.ToString().ToLowerInvariant());
+#pragma warning restore CA1308 // Normalize strings to uppercase
+				agentFiles.Add("UbaAgent.exe");
+				if (useSentry)
+				{
+					agentFiles.Add("crashpad_handler.exe");
+					agentFiles.Add("sentry.dll");
+				}
+			}
+			else if (OperatingSystem.IsLinux())
+			{
+				agentFiles.Add("UbaAgent");
+				bool bIsDebug = false;
+				if (bIsDebug)
+				{
+					agentFiles.Add("UbaAgent.debug");
+					agentFiles.Add("UbaAgent.sym");
+					agentFiles.Add("libclang_rt.tsan.so"); // Needs to be copied from autosdk
+					agentFiles.Add("llvm-symbolizer"); // Needs to be copied from autosdk
+				}
+
+				if (RuntimeInformation.ProcessArchitecture == Architecture.X64)
+				{
+					ubaDir = DirectoryReference.Combine(Unreal.EngineDirectory, "Binaries", "Linux", "UnrealBuildAccelerator");
+				}
+				else if (RuntimeInformation.ProcessArchitecture == Architecture.Arm64)
+				{
+					ubaDir = DirectoryReference.Combine(Unreal.EngineDirectory, "Binaries", "LinuxArm64", "UnrealBuildAccelerator");
+				}
+				else
+				{
+					throw new PlatformNotSupportedException();
+				}
+			}
+			else if (OperatingSystem.IsMacOS())
+			{
+				agentFiles.Add("UbaAgent");
+				ubaDir = DirectoryReference.Combine(Unreal.EngineDirectory, "Binaries", "Mac", "UnrealBuildAccelerator");
+			}
+			else
+			{
+				throw new PlatformNotSupportedException();
+			}
+			_ubaAgentLocator = await CreateToolAsync(ubaDir, agentFiles.Select(x => FileReference.Combine(ubaDir, x)), cancellationToken);
+			_logger.LogInformation("Created tool bundle with locator {UbaAgentLocator}", _ubaAgentLocator.ToString());
+		}
+
+		async Task<BlobLocator> CreateToolAsync(DirectoryReference baseDir, IEnumerable<FileReference> files, CancellationToken cancellationToken)
+		{
+			BlobSerializerOptions serializerOptions = new BlobSerializerOptions();
+			serializerOptions.Converters.Add(new InteriorChunkedDataNodeConverter(2)); // Lock to v2 for now. Could change based on protocol version.
+
+			await using IBlobWriter writer = _storage.CreateBlobWriter(serializerOptions: serializerOptions);
+			DirectoryNode sandbox = new();
+			await sandbox.AddFilesAsync(baseDir, files, writer, cancellationToken: cancellationToken);
+			IBlobRef<DirectoryNode> handle = await writer.WriteBlobAsync(sandbox, cancellationToken);
+			await writer.FlushAsync(cancellationToken);
+			return handle.GetLocator();
+		}
+
+		public int NumLogicalCores { get; private set; } = 0;
+
+		public async void RemoveCompleteWorkersAsync()
+		{
+			for (int idx = 0; idx < _workers.Count; idx++)
+			{
+				Worker worker = _workers[idx];
+				if (worker.BackgroundTask.IsCompleted)
+				{
+					await worker.BackgroundTask;
+					NumLogicalCores -= worker.NumLogicalCores;
+					_workers.RemoveAt(idx--);
+				}
+			}
+		}
+
+		public int QueuedUpCores()
+		{
+			int count = 0;
+			foreach (Worker worker in _workers)
+			{
+				if (!worker.Started)
+				{
+					count += worker.NumLogicalCores;
+				}
+			}
+			return count;
+		}
+
+		int _workerId = 0;
+
+		static readonly HashSet<StringView> s_logProperties = new()
+			{
+				"ComputeIp",
+				"CPU",
+				"RAM",
+				"DiskFreeSpace",
+				"PhysicalCores",
+				"LogicalCores"
+			};
+
+		public async Task<bool> AddWorkerAsync(Requirements requirements, UnrealBuildAcceleratorHordeConfig hordeConfig, CancellationToken cancellationToken)
+		{
+			if (_client == null)
+			{
+				throw new InvalidOperationException("Call init first");
+			}
+
+			PrefixLogger workerLogger = new($"[Worker{_workerId}]", _logger);
+
+			const string UbaPortName = "UbaPort";
+			const string UbaProxyPortName = "UbaProxyPort";
+			const int UbaPort = 7001;
+			const int UbaProxyPort = 7002;
+
+			// Request ID that is unique per attempt to acquire the same compute lease/worker
+			// Primarily for tracking worker demand on Horde server as UBAExecutor will repeatedly try adding a new worker
+			string requestId = $"{_id}-worker-{_workerId}";
+			IComputeLease? lease = null;
+			try
+			{
+				Stopwatch stopwatch = Stopwatch.StartNew();
+				ConnectionMetadataRequest cmr = new()
+				{
+					ModePreference = _connectionMode,
+					Encryption = _encryption,
+					Ports = { { UbaPortName, UbaPort }, { UbaProxyPortName, UbaProxyPort } }
+				};
+				lease = await _client.TryAssignWorkerAsync(_clusterId, requirements, requestId, cmr, workerLogger, cancellationToken);
+				if (lease == null)
+				{
+					_logger.LogDebug("Unable to assign a remote worker");
+
+					int missingNumCores = Math.Max(0, _maxCores - NumLogicalCores);
+					await UpdateCpuCoreNeedAsync(missingNumCores, cancellationToken);
+					return false;
+				}
+
+				_workerId++;
+
+				workerLogger.LogDebug("Agent properties:");
+
+				int numLogicalCores = 24; // Assume 24 if something goes wrong here and property is not found
+				string computeIp = String.Empty;
+				foreach (string property in lease.Properties)
+				{
+					int equalsIdx = property.IndexOf('=', StringComparison.OrdinalIgnoreCase);
+					StringView propertyName = new(property, 0, equalsIdx);
+					if (s_logProperties.Contains(propertyName))
+					{
+						_logger.LogDebug("  {Property}", property);
+
+						if (propertyName == ResourceLogicalCores && Int32.TryParse(property.AsSpan(equalsIdx + 1), out int value))
+						{
+							numLogicalCores = value;
+						}
+						else if (propertyName == "ComputeIp")
+						{
+							computeIp = property[(equalsIdx + 1)..];
+						}
+					}
+				}
+
+				// When using relay connection mode, the IP will be relay server's IP
+				string ip = String.IsNullOrEmpty(lease.Ip) ? computeIp : lease.Ip;
+
+				if (!lease.Ports.TryGetValue(UbaPortName, out ConnectionMetadataPort? ubaPort))
+				{
+					ubaPort = new ConnectionMetadataPort(UbaPort, UbaPort);
+				}
+
+				if (!lease.Ports.TryGetValue(UbaProxyPortName, out ConnectionMetadataPort? ubaProxyPort))
+				{
+					ubaProxyPort = new ConnectionMetadataPort(UbaProxyPort, UbaProxyPort);
+				}
+
+				string exeName = OperatingSystem.IsWindows() ? "UbaAgent.exe" : "UbaAgent";
+				BlobLocator locator = _ubaAgentLocator;
+				Worker worker = new()
+				{
+					StartTime = stopwatch,
+					NumLogicalCores = numLogicalCores,
+					Ip = ip,
+					Port = ubaPort,
+					ProxyPort = ubaProxyPort,
+				};
+				worker.BackgroundTask = RunWorkerAsync(worker, lease, locator, exeName, workerLogger, hordeConfig, _cancellationTokenSource.Token);
+				_workers.Add(worker);
+				lease = null; // Will be disposed by RunWorkerAsync
+
+				NumLogicalCores += numLogicalCores;
+				return true;
+			}
+			finally
+			{
+				if (lease != null)
+				{
+					await lease.DisposeAsync();
+				}
+			}
+		}
+
+		async Task UpdateCpuCoreNeedAsync(int targetCoreCount, CancellationToken cancellationToken = default)
+		{
+			if (_client != null && _pool != null)
+			{
+				_logger.LogDebug("Setting CPU core need to {TargetCoreCount}", targetCoreCount);
+				Dictionary<string, int> resourceNeeds = new() { { ResourceLogicalCores, targetCoreCount } };
+				try
+				{
+					await _client.DeclareResourceNeedsAsync(_clusterId, _pool, resourceNeeds, cancellationToken);
+				}
+				catch (Exception e)
+				{
+					_logger.Log(_strict ? LogLevel.Error : LogLevel.Information, KnownLogEvents.Systemic_Horde_Compute, e, "Failed updating resource need to {TargetCoreCount} cores", targetCoreCount);
+				}
+			}
+		}
+
+		public static async Task<UBAHordeSession?> TryCreateHordeSessionAsync(UnrealBuildAcceleratorHordeConfig hordeConfig, UBAExecutor executor, bool bStrictErrors, ILogger logger, CancellationToken cancellationToken = default)
+		{
+			if (hordeConfig.bDisableHorde)
+			{
+				logger.LogInformation("Horde disabled via command line option.");
+				return null;
+			}
+
+			string? server = hordeConfig.HordeServer;
+			string? token = hordeConfig.HordeToken;
+			string? oidcProvider = hordeConfig.HordeOidcProvider;
+
+			if (String.IsNullOrEmpty(server))
+			{
+				server = Environment.GetEnvironmentVariable("UE_HORDE_URL");
+				if (String.IsNullOrEmpty(server))
+				{
+					logger.LogInformation("Horde URL not specified in BuildConfiguration.xml or via UE_HORDE_URL environment variable");
+					return null;
+				}
+				if (String.IsNullOrEmpty(token))
+				{
+					token = Environment.GetEnvironmentVariable("UE_HORDE_TOKEN");
+				}
+			}
+
+			ConnectionMode? connectionMode = Enum.TryParse(hordeConfig.HordeConnectionMode, true, out ConnectionMode cm) ? cm : null;
+			Encryption? encryption = Enum.TryParse(hordeConfig.HordeEncryption, true, out Encryption enc) ? enc : null;
+
+			oidcProvider ??= Environment.GetEnvironmentVariable("UE_HORDE_OIDC_PROVIDER");
+
+			bool hasOidcProvider = !String.IsNullOrEmpty(oidcProvider);
+			logger.LogInformation("Horde URL: {Server}, Pool: {Pool}, Condition: {Condition}, OIDC: {OidcProvider}, Connection: {Connection} HordeEncryption: {Encryption}",
+				server, hordeConfig.HordePool ?? "(none)", hordeConfig.HordeCondition ?? "(none)", hasOidcProvider ? oidcProvider! : "Disabled", connectionMode?.ToString() ?? "(none)", encryption?.ToString() ?? "(none)");
+			try
+			{
+				if (String.IsNullOrEmpty(token) && hasOidcProvider)
+				{
+					token = await GetOidcBearerTokenAsync(null, oidcProvider!, logger, cancellationToken);
+				}
+
+				AuthenticationHeaderValue? authHeader = null;
+				if (!String.IsNullOrEmpty(token))
+				{
+					authHeader = new AuthenticationHeaderValue("Bearer", token);
+				}
+
+				bool allowWine = hordeConfig.bHordeAllowWine && OperatingSystem.IsWindows();
+
+				UBAHordeSession session = new(executor, new Uri(server), authHeader, hordeConfig.HordePool, allowWine, hordeConfig.HordeMaxCores, bStrictErrors, connectionMode, encryption, logger);
+				await session.InitAsync(useSentry: !String.IsNullOrEmpty(hordeConfig.UBASentryUrl), cancellationToken);
+				return session;
+			}
+			catch (TaskCanceledException)
+			{
+				return null;
+			}
+			catch (Exception ex)
+			{
+				logger.Log(bStrictErrors ? LogLevel.Error : LogLevel.Information, ex, "Unable to create Horde session: {Message}", ex.Message);
+				return null;
+			}
+		}
+
+		static async Task<string> GetOidcBearerTokenAsync(DirectoryReference? projectDir, string oidcProvider, ILogger logger, CancellationToken cancellationToken = default)
+		{
+			logger.LogInformation("Performing OIDC token refresh...");
+
+			using ITokenStore tokenStore = TokenStoreFactory.CreateTokenStore();
+			IConfiguration providerConfiguration = ProviderConfigurationFactory.ReadConfiguration(Unreal.EngineDirectory.ToDirectoryInfo(), projectDir?.ToDirectoryInfo());
+			OidcTokenManager oidcTokenManager = OidcTokenManager.CreateTokenManager(providerConfiguration, tokenStore, new List<string>() { oidcProvider });
+
+			OidcTokenInfo result;
+			try
+			{
+				result = await oidcTokenManager.GetAccessToken(oidcProvider, cancellationToken);
+			}
+			catch (NotLoggedInException)
+			{
+				result = await oidcTokenManager.Login(oidcProvider, cancellationToken);
+			}
+
+			if (result.AccessToken == null)
+			{
+				throw new Exception($"Unable to get access token for {oidcProvider}");
+			}
+
+			logger.LogInformation("Received bearer token for {OidcProvider}", oidcProvider);
+			return result.AccessToken;
+		}
+
+		async Task RunWorkerAsync(Worker self, IComputeLease lease, BlobLocator tool, string executable, ILogger logger, UnrealBuildAcceleratorHordeConfig hordeConfig, CancellationToken cancellationToken)
+		{
+			logger.LogDebug("Running worker task..");
+			try
+			{
+				await using (_ = lease)
+				{
+					// Create a message channel on channel id 0. The Horde Agent always listens on this channel for requests.
+					const int PrimaryChannelId = 0;
+					using (AgentMessageChannel channel = lease.Socket.CreateAgentMessageChannel(PrimaryChannelId, 4 * 1024 * 1024))
+					{
+						logger.LogDebug("Waiting for attach...");
+
+						TimeSpan attachTimeout = TimeSpan.FromSeconds(20.0);
+						try
+						{
+							Task attachTask = channel.WaitForAttachAsync(cancellationToken).AsTask();
+							await attachTask.WaitAsync(attachTimeout, cancellationToken);
+						}
+						catch (TimeoutException)
+						{
+							logger.Log(_strict ? LogLevel.Error : LogLevel.Information, KnownLogEvents.Systemic_Horde_Compute, "Waited {Time}s on attach message. Giving up", (int)attachTimeout.TotalSeconds);
+							throw;
+						}
+
+						logger.LogDebug("Uploading files...");
+						await channel.UploadFilesAsync("", tool, _storage.Backend, cancellationToken);
+
+						string hordeHost = _owner.UBAConfig.Host;
+						if (!String.IsNullOrEmpty(hordeConfig.HordeHost))
+						{
+							hordeHost = hordeConfig.HordeHost;
+						}
+
+						bool useListen = !String.IsNullOrEmpty(hordeConfig.HordeHost);
+						List<string> arguments = new();
+
+						if (useListen)
+						{
+							arguments.Add($"-Host={hordeHost}:{_owner.UBAConfig.Port}");
+						}
+						else
+						{
+							arguments.Add($"-Listen={self.Port.AgentPort}");
+						}
+
+						if (!String.IsNullOrEmpty(_crypto))
+						{
+							arguments.Add($"-crypto={_crypto}");
+						}
+
+						arguments.Add("-NoPoll");
+						arguments.Add("-Quiet");
+						if (!String.IsNullOrEmpty(hordeConfig.UBASentryUrl))
+						{
+							arguments.Add($"-Sentry=\"{hordeConfig.UBASentryUrl}\"");
+						}
+						arguments.Add("-ProxyPort=" + self.ProxyPort.AgentPort);
+						if (_owner.UBAConfig.bUseQuic)
+						{
+							arguments.Add("-quic");
+						}
+						//arguments.Add("-NoStore");
+						//arguments.Add("-KillRandom"); // For debugging
+
+						if (System.OperatingSystem.IsMacOS())
+						{
+							// we need to populate the cas with all known xcodes so we can serve all the ones we have installed
+							string xcodeVersion = Utils.RunLocalProcessAndReturnStdOut("/bin/sh", "-c '/usr/bin/defaults read $(xcode-select -p)/../version.plist ProductBuildVersion");
+							arguments.Add($"-populateCasFromXcodeVersion={xcodeVersion}");
+						}
+
+						arguments.Add("-Dir=%UE_HORDE_SHARED_DIR%\\Uba");
+						arguments.Add("-Eventfile=%UE_HORDE_TERMINATION_SIGNAL_FILE%");
+						arguments.Add("-MaxIdle=15");
+						if (_owner.UBAConfig.bLogEnabled)
+						{
+							arguments.Add("-Log");
+						}
+
+						LogLevel logLevel = _owner.UBAConfig.bDetailedLog ? LogLevel.Information : LogLevel.Debug;
+
+						logger.Log(logLevel, "Executing child process: {Executable} {Arguments}", executable, CommandLineArguments.Join(arguments));
+
+						ExecuteProcessFlags execFlags = _allowWine ? ExecuteProcessFlags.UseWine : ExecuteProcessFlags.None;
+						await using AgentManagedProcess process = await channel.ExecuteAsync(executable, arguments, null, null, execFlags, cancellationToken);
+						bool shouldConnect = !useListen;
+						self.Started = true;
+						string? line;
+
+						while ((line = await process.ReadLineAsync(cancellationToken)) != null)
+						{
+							logger.Log(logLevel, "{Line}", line);
+
+							if (shouldConnect && line.Contains("Listening on", StringComparison.OrdinalIgnoreCase)) // This log entry means that the agent is ready for connections.
+							{
+								long totalMs = self.StartTime.ElapsedMilliseconds;
+								logger.LogInformation("Connecting to UbaAgent on {Ip}:{Port} (local agent port {AgentPort}) {Seconds}.{Milliseconds} seconds after assigned",
+									self.Ip, self.Port.Port, self.Port.AgentPort, totalMs / 1000, totalMs % 1000);
+
+								_owner.Server!.AddClient(self.Ip, self.Port.Port, _crypto);
+								shouldConnect = false;
+							}
+						}
+						logger.LogDebug("Shutting down process");
+					}
+
+					logger.LogDebug("Closing channel");
+					await lease.CloseAsync(cancellationToken);
+				}
+			}
+			catch (TimeoutException)
+			{
+			}
+			catch (Exception ex)
+			{
+				if (!cancellationToken.IsCancellationRequested)
+				{
+					logger.Log(_strict ? LogLevel.Error : LogLevel.Information, KnownLogEvents.Systemic_Horde_Compute, ex, "Exception in worker task: {Ex}", ex.ToString());
+
+					// Add additional properties to aid debugging
+					logger.LogInformation(KnownLogEvents.Systemic_Horde_Compute, ex, "UBA agent locator {UBAAgentLocator}", _ubaAgentLocator.ToString());
+				}
+			}
+		}
+	}
+
+	class UBAAgentCoordinatorHorde : IUBAAgentCoordinator, IDisposable
+	{
+		public UBAAgentCoordinatorHorde(ILogger logger, UnrealBuildAcceleratorConfig ubaConfig, CommandLineArguments? additionalArguments = null)
+		{
+			_logger = logger;
+			_ubaConfig = ubaConfig;
+
+			XmlConfig.ApplyTo(HordeConfig);
+			additionalArguments?.ApplyTo(HordeConfig);
+
+			// Sentry is currently unsupported for non-Windows and non-x64
+			if (!OperatingSystem.IsWindows() || RuntimeInformation.ProcessArchitecture != Architecture.X64)
+			{
+				HordeConfig.UBASentryUrl = null;
+			}
+		}
+
+		public DirectoryReference? GetUBARootDir()
+		{
+			DirectoryReference? hordeSharedDir = DirectoryReference.FromString(Environment.GetEnvironmentVariable("UE_HORDE_SHARED_DIR"));
+			if (hordeSharedDir != null)
+			{
+				return DirectoryReference.Combine(hordeSharedDir, "UbaHost");
+			}
+			return null;
+		}
+
+		public async Task InitAsync(UBAExecutor executor)
+		{
+			if (_ubaConfig.bDisableRemote)
+			{
+				return;
+			}
+
+			_cancellationSource = new CancellationTokenSource();
+			_hordeSessionTask = UBAHordeSession.TryCreateHordeSessionAsync(HordeConfig, executor, _ubaConfig.bStrict, _logger, _cancellationSource.Token);
+			await _hordeSessionTask;
+		}
+
+		public void Start(ImmediateActionQueue queue, Func<LinkedAction, bool> canRunRemotely)
+		{
+			int timerPeriod = 5000;
+			bool shownNoAgentsFoundMessage = false;
+
+			if (_hordeSessionTask == null)
+			{
+				return;
+			}
+			_timer = new(async (_) =>
+			{
+				_timer?.Change(Timeout.Infinite, Timeout.Infinite);
+
+				if (_cancellationSource!.IsCancellationRequested)
+				{
+					return;
+				}
+
+				UBAHordeSession? hordeSession = await _hordeSessionTask!;
+
+				if (hordeSession == null)
+				{
+					return;
+				}
+
+				hordeSession.RemoveCompleteWorkersAsync();
+
+				if (queue.IsDone)
+				{
+					return;
+				}
+
+				// We are assuming all active logical cores are already being used.. so queueWeight is essentially work that could be executed but can't because of bandwidth
+				double queueThreshold = _ubaConfig.bForceBuildAllRemote ? 0 : 5;
+
+				try
+				{
+					double queueWeight = queue.EnumerateReadyToCompileActions().Where(x => canRunRemotely(x)).Sum(x => x.Weight);
+
+					queueWeight -= hordeSession.QueuedUpCores();
+					while (true)
+					{
+						int currentLogicalCores = hordeSession.NumLogicalCores;
+
+						if (queueWeight <= queueThreshold || currentLogicalCores >= HordeConfig.HordeMaxCores || _cancellationSource!.IsCancellationRequested)
+						{
+							break;
+						}
+
+						Requirements requirements = new()
+						{
+							Exclusive = true
+						};
+
+						if (!String.IsNullOrEmpty(HordeConfig.HordePool))
+						{
+							requirements.Pool = HordeConfig.HordePool;
+						}
+
+						if (HordeConfig.HordeCondition != null)
+						{
+							requirements.Condition = Condition.Parse(HordeConfig.HordeCondition);
+						}
+
+						if (!await hordeSession.AddWorkerAsync(requirements, HordeConfig, _cancellationSource.Token))
+						{
+							_logger.LogDebug("No additional workers available");
+							break;
+						}
+						int coresAdded = hordeSession.NumLogicalCores - currentLogicalCores;
+						queueWeight -= coresAdded;
+					}
+				}
+				catch (NoComputeAgentsFoundException ex)
+				{
+					if (!shownNoAgentsFoundMessage)
+					{
+						_logger.Log(_ubaConfig.bStrict ? LogLevel.Warning : LogLevel.Information, KnownLogEvents.Systemic_Horde_Compute, ex, "No agents found matching requirements (cluster: {ClusterId}, requirements: {Requirements})", ex.ClusterId, ex.Requirements);
+						shownNoAgentsFoundMessage = true;
+					}
+				}
+				catch (Exception ex)
+				{
+					if (!_cancellationSource!.IsCancellationRequested)
+					{
+						if (_ubaConfig.bStrict)
+						{
+							_logger.Log(LogLevel.Error, KnownLogEvents.Systemic_Horde_Compute, ex, "Unable to get worker: {Ex}", ex.ToString());
+						}
+						else
+						{
+							_logger.Log(LogLevel.Information, KnownLogEvents.Systemic_Horde_Compute, "Unable to get worker: {Ex}", ex.ToString());
+						}
+					}
+				}
+
+				_timer?.Change(timerPeriod, Timeout.Infinite);
+			}, null, HordeConfig.HordeDelay * 1000, timerPeriod);
+		}
+
+		public void Stop()
+		{
+			_cancellationSource?.Cancel();
+		}
+
+		public async Task CloseAsync()
+		{
+			_cancellationSource?.Cancel();
+
+			if (_hordeSessionTask == null)
+			{
+				return;
+			}
+
+			UBAHordeSession? hordeSession = await _hordeSessionTask;
+			if (hordeSession != null)
+			{
+				await hordeSession.DisposeAsync();
+			}
+		}
+
+		public void Dispose()
+		{
+			Stop();
+			CloseAsync().Wait();
+			Dispose(true);
+			GC.SuppressFinalize(this);
+		}
+
+		protected virtual void Dispose(bool disposing)
+		{
+			if (disposing)
+			{
+				_cancellationSource?.Dispose();
+				_cancellationSource = null;
+				_timer?.Dispose();
+				_timer = null;
+			}
+		}
+
+		readonly ILogger _logger;
+		readonly UnrealBuildAcceleratorConfig _ubaConfig;
+		UnrealBuildAcceleratorHordeConfig HordeConfig { get; init; } = new();
+
+		CancellationTokenSource? _cancellationSource;
+		Task<UBAHordeSession?>? _hordeSessionTask;
+		Timer? _timer;
+	}
+}
diff --git a/Engine/Source/Programs/UnrealBuildTool/Executors/UnrealBuildAccelerator/UBAExecutor.cs b/Engine/Source/Programs/UnrealBuildTool/Executors/UnrealBuildAccelerator/UBAExecutor.cs
new file mode 100644
--- /dev/null
+++ b/Engine/Source/Programs/UnrealBuildTool/Executors/UnrealBuildAccelerator/UBAExecutor.cs
@@ -0,0 +1,679 @@
+// Copyright Epic Games, Inc. All Rights Reserved.
+
+using System;
+using System.Collections.Concurrent;
+using System.Collections.Generic;
+using System.Linq;
+using System.Management;
+using System.Runtime.InteropServices;
+using System.Runtime.Versioning;
+using System.Threading.Tasks;
+using EpicGames.Core;
+using EpicGames.UBA;
+using Microsoft.Extensions.Logging;
+using UnrealBuildBase;
+using UnrealBuildTool.Artifacts;
+
+namespace UnrealBuildTool
+{
+	class UBAExecutor : ParallelExecutor
+	{
+		public UnrealBuildAcceleratorConfig UBAConfig { get; init; } = new();
+
+		public string Crypto { get; private set; } = String.Empty;
+		public IServer? Server { get; private set; }
+		ISessionServer? _session;
+		readonly List<IUBAAgentCoordinator> _agentCoordinators = new List<IUBAAgentCoordinator>();
+		DirectoryReference? _rootDirRef;
+		bool _bIsCancelled;
+		bool _bIsRemoteActionsAllowed = true;
+		readonly object _actionsChangedLock = new();
+		bool _bActionsChanged = true;
+		uint _actionsQueuedThatCanRunRemotely = UInt32.MaxValue;
+		readonly ThreadedLogger _threadedLogger;
+
+		// Tracking for LinkedActions that failed remotely that should be retried locally
+		readonly ConcurrentDictionary<LinkedAction, bool> _localRetryActions = new();
+		// Tracking for LinkedActions that failed locally that should be retried without UBA
+		readonly ConcurrentDictionary<LinkedAction, bool> _forcedRetryActions = new();
+
+		protected override void Dispose(bool disposing)
+		{
+			if (disposing)
+			{
+				_session?.Dispose();
+				_session = null;
+				_threadedLogger.Dispose();
+			}
+			base.Dispose(disposing);
+		}
+
+		public override string Name => "Unreal Build Accelerator";
+
+		public static new bool IsAvailable()
+		{
+			return EpicGames.UBA.Utils.IsAvailable();
+		}
+
+		public UBAExecutor(int maxLocalActions, bool bAllCores, bool bCompactOutput, Microsoft.Extensions.Logging.ILogger logger, CommandLineArguments? additionalArguments = null)
+			: base(maxLocalActions, bAllCores, bCompactOutput, logger)
+		{
+			XmlConfig.ApplyTo(this);
+			XmlConfig.ApplyTo(UBAConfig);
+			CommandLine.ParseArguments(Environment.GetCommandLineArgs(), this, logger);
+			additionalArguments?.ApplyTo(this);
+			additionalArguments?.ApplyTo(UBAConfig);
+
+			_threadedLogger = new ThreadedLogger(logger);
+			_agentCoordinators.Add(new UBAAgentCoordinatorHorde(logger, UBAConfig, additionalArguments));
+		}
+
+		private void PrintConfiguration()
+		{
+			_threadedLogger.LogInformation("  Storage capacity {StoreCapacityGb}Gb", UBAConfig.StoreCapacityGb);
+		}
+
+		private async Task WriteActionOutputFileAsync(IEnumerable<LinkedAction> inputActions)
+		{
+			if (String.IsNullOrEmpty(UBAConfig.ActionsOutputFile))
+			{
+				return;
+			}
+
+			if (!UBAConfig.ActionsOutputFile.EndsWith(".yaml", StringComparison.OrdinalIgnoreCase))
+			{
+				_threadedLogger.LogError("UBA actions output file needs to have extension .yaml for UbaCli to understand it");
+			}
+			using System.IO.StreamWriter streamWriter = new System.IO.StreamWriter(UBAConfig.ActionsOutputFile);
+			using System.CodeDom.Compiler.IndentedTextWriter writer = new System.CodeDom.Compiler.IndentedTextWriter(streamWriter, "  ");
+			await writer.WriteAsync("environment: ");
+			await writer.WriteLineAsync(Environment.GetEnvironmentVariable("PATH"));
+			await writer.WriteLineAsync("processes:");
+			writer.Indent++;
+			int index = 0;
+			foreach (LinkedAction action in inputActions)
+			{
+				action.SortIndex = index++;
+				await writer.WriteLineAsync($"- id: {action.SortIndex}");
+				writer.Indent++;
+				await writer.WriteLineAsync($"app: {action.CommandPath}");
+				await writer.WriteLineAsync($"arg: {action.CommandArguments}");
+				await writer.WriteLineAsync($"dir: {action.WorkingDirectory}");
+				await writer.WriteLineAsync($"desc: {action.StatusDescription}");
+				if (action.Weight != 1.0f)
+				{
+					await writer.WriteLineAsync($"weight: {action.Weight}");
+				}
+				if (!action.bCanExecuteInUBA)
+				{
+					await writer.WriteLineAsync("detour: false");
+				}
+				else if (!action.bCanExecuteRemotely)
+				{
+					await writer.WriteLineAsync("remote: false");
+				}
+				if (action.PrerequisiteActions.Any())
+				{
+					await writer.WriteAsync("dep: [");
+					await writer.WriteAsync(String.Join(", ", action.PrerequisiteActions.Select(x => x.SortIndex)));
+					await writer.WriteLineAsync("]");
+				}
+				writer.Indent--;
+				await writer.WriteLineNoTabsAsync(null);
+			}
+		}
+
+		[System.Diagnostics.CodeAnalysis.SuppressMessage("Globalization", "CA1308:Normalize strings to uppercase", Justification = "lowercase crypto string")]
+		public static string CreateCrypto()
+		{
+			byte[] bytes = new byte[16];
+			using System.Security.Cryptography.RandomNumberGenerator random = System.Security.Cryptography.RandomNumberGenerator.Create();
+			random.GetBytes(bytes);
+			return BitConverter.ToString(bytes).Replace("-", "", StringComparison.OrdinalIgnoreCase).ToLowerInvariant(); // "1234567890abcdef1234567890abcdef";
+		}
+
+		public override async Task<bool> ExecuteActionsAsync(IEnumerable<LinkedAction> inputActions, Microsoft.Extensions.Logging.ILogger logger, IActionArtifactCache? actionArtifactCache)
+		{
+			if (!inputActions.Any())
+			{
+				return true;
+			}
+
+			if (inputActions.Count() < NumParallelProcesses && !UBAConfig.bForceBuildAllRemote)
+			{
+				UBAConfig.bDisableRemote = true;
+			}
+
+			PrintConfiguration();
+			await WriteActionOutputFileAsync(inputActions);
+
+			logger = _threadedLogger;
+
+			if (UBAConfig.bUseCrypto)
+			{
+				Crypto = CreateCrypto();
+			}
+
+			if (!String.IsNullOrEmpty(UBAConfig.RootDir))
+			{
+				_rootDirRef = DirectoryReference.FromString(UBAConfig.RootDir);
+			}
+			else
+			{
+				_rootDirRef = DirectoryReference.FromString(Environment.GetEnvironmentVariable("UBA_ROOT") ?? Environment.GetEnvironmentVariable("BOX_ROOT"));
+			}
+
+			List<Task> coordinatorInitTasks = new List<Task>();
+			if (!UBAConfig.bDisableRemote)
+			{
+				foreach (IUBAAgentCoordinator coordinator in _agentCoordinators)
+				{
+					// Do not await here, it will block local tasks from running
+					coordinatorInitTasks.Add(coordinator.InitAsync(this));
+
+					if (_rootDirRef == null)
+					{
+						_rootDirRef = coordinator.GetUBARootDir();
+					}
+				}
+			}
+
+			if (_rootDirRef == null)
+			{
+				if (OperatingSystem.IsWindows())
+				{
+					_rootDirRef = DirectoryReference.Combine(DirectoryReference.GetSpecialFolder(Environment.SpecialFolder.CommonApplicationData)!, "Epic", "UnrealBuildAccelerator");
+				}
+				else
+				{
+					_rootDirRef = DirectoryReference.Combine(DirectoryReference.GetSpecialFolder(Environment.SpecialFolder.UserProfile)!, ".epic", "UnrealBuildAccelerator");
+				}
+			}
+
+			DirectoryReference.CreateDirectory(_rootDirRef);
+
+			FileReference ubaTraceFile;
+			if (!String.IsNullOrEmpty(UBAConfig.TraceFile))
+			{
+				ubaTraceFile = new FileReference(UBAConfig.TraceFile);
+			}
+			else if (Unreal.IsBuildMachine())
+			{
+				ubaTraceFile = FileReference.Combine(Unreal.EngineProgramSavedDirectory, "AutomationTool", "Saved", "Logs", "Trace.uba");
+			}
+			else if (Log.OutputFile != null)
+			{
+				ubaTraceFile = Log.OutputFile.ChangeExtension(".uba");
+			}
+			else
+			{
+				ubaTraceFile = FileReference.Combine(Unreal.EngineProgramSavedDirectory, "UnrealBuildTool", "Trace.uba");
+			}
+
+			DirectoryReference.CreateDirectory(ubaTraceFile.Directory);
+			Log.BackupLogFile(ubaTraceFile);
+
+			IStorageServer? ubaStorage = null;
+			void CancelKeyPress(object? sender, ConsoleCancelEventArgs e)
+			{
+				_bIsCancelled = true;
+				_session?.CancelAll();
+				ubaStorage?.SaveCasTable();
+				foreach (IUBAAgentCoordinator coordinator in _agentCoordinators)
+				{
+					coordinator.CloseAsync().Wait(2000); // Give coordinators some time to close (this makes coordinators like horde return resources faster)
+				}
+			}
+			Console.CancelKeyPress += CancelKeyPress;
+
+			try
+			{
+				if (UBAConfig.bLaunchVisualizer && OperatingSystem.IsWindows())
+				{
+					_ = Task.Run(LaunchVisualizer);
+				}
+
+				using EpicGames.UBA.ILogger ubaLogger = EpicGames.UBA.ILogger.CreateLogger(logger, UBAConfig.bDetailedLog);
+				using (Server = IServer.CreateServer(UBAConfig.MaxWorkers, UBAConfig.SendSize, ubaLogger, UBAConfig.bUseQuic))
+				{
+					using IStorageServer ubaStorageServer = IStorageServer.CreateStorageServer(Server, ubaLogger, new StorageServerCreateInfo(_rootDirRef.FullName, ((ulong)UBAConfig.StoreCapacityGb) * 1000 * 1000 * 1000, !UBAConfig.bStoreRaw, UBAConfig.Zone));
+					using ISessionServerCreateInfo serverCreateInfo = ISessionServerCreateInfo.CreateSessionServerCreateInfo(ubaStorageServer, Server, ubaLogger, new SessionServerCreateInfo(_rootDirRef.FullName, ubaTraceFile.FullName, UBAConfig.bDisableCustomAlloc, false, UBAConfig.bResetCas, UBAConfig.bWriteToDisk, UBAConfig.bDetailedTrace, !UBAConfig.bDisableWaitOnMem, UBAConfig.bAllowKillOnMem));
+					using (_session = ISessionServer.CreateSessionServer(serverCreateInfo))
+					{
+
+						ubaStorage = ubaStorageServer;
+
+						if (!UBAConfig.bDisableRemote)
+						{
+							Server.StartServer(UBAConfig.Host, UBAConfig.Port, Crypto);
+						}
+
+						bool success = ExecuteActionsInternal(inputActions, _session, logger, actionArtifactCache);
+
+						if (!UBAConfig.bDisableRemote)
+						{
+							Server.StopServer();
+						}
+
+						if (UBAConfig.bPrintSummary)
+						{
+							_session.PrintSummary();
+						}
+
+						return success;
+					}
+				}
+			}
+			finally
+			{
+				Console.CancelKeyPress -= CancelKeyPress;
+
+				foreach (IUBAAgentCoordinator coordinator in _agentCoordinators)
+				{
+					await coordinator.CloseAsync();
+				}
+				await _threadedLogger.FinishAsync();
+				_session = null;
+			}
+		}
+
+		[SupportedOSPlatform("windows")]
+		static void LaunchVisualizer()
+		{
+			FileReference visaulizerPath = FileReference.Combine(Unreal.EngineDirectory, "Binaries", "Win64", "UnrealBuildAccelerator", RuntimeInformation.ProcessArchitecture.ToString(), "UbaVisualizer.exe");
+			FileReference tempPath = FileReference.Combine(new DirectoryReference(System.IO.Path.GetTempPath()), visaulizerPath.GetFileName());
+			if (!FileReference.Exists(visaulizerPath))
+			{
+				return;
+			}
+
+			try
+			{
+				// Check if a listening visualizer is already running
+				foreach (System.Diagnostics.Process process in System.Diagnostics.Process.GetProcessesByName(visaulizerPath.GetFileNameWithoutAnyExtensions()))
+				{
+					using ManagementObjectSearcher searcher = new ManagementObjectSearcher($"SELECT CommandLine FROM Win32_Process WHERE ProcessId = {process.Id}");
+					using ManagementObjectCollection objects = searcher.Get();
+					string args = objects.Cast<ManagementBaseObject>().SingleOrDefault()?["CommandLine"]?.ToString() ?? "";
+					if (args.Contains("-listen", StringComparison.OrdinalIgnoreCase))
+					{
+						return;
+					}
+				}
+				if (!FileReference.Exists(tempPath) || tempPath.ToFileInfo().LastWriteTime < visaulizerPath.ToFileInfo().LastWriteTime)
+				{
+					FileReference.Copy(visaulizerPath, tempPath, true);
+				}
+				if (FileReference.Exists(tempPath))
+				{
+					System.Diagnostics.ProcessStartInfo psi = new(BuildHostPlatform.Current.Shell.FullName, $" /C start \"\" \"{tempPath.FullName}\" -listen")
+					{
+						WorkingDirectory = System.IO.Path.GetTempPath(),
+						WindowStyle = System.Diagnostics.ProcessWindowStyle.Hidden,
+						UseShellExecute = true,
+					};
+					System.Diagnostics.Process.Start(psi);
+				}
+			}
+			catch (Exception)
+			{
+			}
+		}
+
+		public override bool VerifyOutputs => UBAConfig.bWriteToDisk;
+
+		/// <summary>
+		/// Executes the provided actions
+		/// </summary>
+		/// <returns>True if all the tasks successfully executed, or false if any of them failed.</returns>
+		bool ExecuteActionsInternal(IEnumerable<LinkedAction> inputActions, ISessionServer session, Microsoft.Extensions.Logging.ILogger logger, IActionArtifactCache? actionArtifactCache)
+		{
+			using ImmediateActionQueue queue = CreateActionQueue(inputActions, actionArtifactCache, logger);
+			int actionLimit = Math.Min(NumParallelProcesses, queue.TotalActions);
+			queue.CreateAutomaticRunner(action => RunActionLocal(queue, action), bUseActionWeights, actionLimit, NumParallelProcesses);
+			ImmediateActionQueueRunner remoteRunner = queue.CreateManualRunner(action => RunActionRemote(queue, action));
+
+			// Setup a notification that alerts uba when an artifact has been read from the cache
+			queue.OnArtifactsRead = (action) =>
+			{
+				HashSet<DirectoryItem> refreshedDirectories = new();
+				foreach (FileItem output in action.ProducedItems)
+				{
+					if (refreshedDirectories.Add(output.Directory))
+					{
+						session.RefreshDirectories(output.Directory.FullName);
+					}
+				}
+			};
+
+			// Start the queue
+			queue.Start();
+
+			// Handle process available from remote
+			session!.RemoteProcessSlotAvailable += (sender, args) =>
+			{
+				if (!queue.TryStartOneAction(remoteRunner) && (_bIsRemoteActionsAllowed && !UBAConfig.bForceBuildAllRemote))
+				{
+					uint count = ActionsLeftThatCanRunRemotely(queue);
+
+					// We didn't find an action to start, let's check how many queued items are left.. if there are less than NumParallelProcesses, then disconnect
+					if (count <= NumParallelProcesses)
+					{
+						_bIsRemoteActionsAllowed = false;
+						foreach (IUBAAgentCoordinator coordinator in _agentCoordinators)
+						{
+							coordinator.Stop();
+						}
+						session!.DisableRemoteExecution();
+					}
+					else
+					{
+						// Tell UBA max number of remote processes left. Providing this information makes it possible for UBA to start disconnecting clients
+						session!.SetMaxRemoteProcessCount(count);
+					}
+				}
+			};
+
+			session!.RemoteProcessReturned += (sender, args) =>
+			{
+				args.Process.Cancel(true);
+				LinkedAction action = (LinkedAction)args.Process.UserData!;
+				//logger.LogInformation("REQUEUING " + action.ProducedItems.FirstOrDefault()!.Name);
+				queue.RequeueAction(action);
+			};
+
+			// Add all actions we can add
+			queue.StartManyActions();
+
+			try
+			{
+				foreach (IUBAAgentCoordinator coordinator in _agentCoordinators)
+				{
+					coordinator.Start(queue, CanRunRemotely);
+				}
+
+				bool res = queue.RunTillDone().Result; // Using inline wait to avoid possible thread switch
+				return res;
+			}
+			finally
+			{
+				foreach (IUBAAgentCoordinator coordinator in _agentCoordinators)
+				{
+					coordinator.Stop();
+				}
+			}
+		}
+
+		/// <summary>
+		/// Determine if an action must be run locally and with no detouring
+		/// </summary>
+		/// <param name="action">The action to check</param>
+		/// <returns>If this action must be local, non-detoured</returns>
+		static bool ForceLocalNoDetour(LinkedAction action)
+		{
+			if (!OperatingSystem.IsMacOS()) // Below code is slow, so early out
+			{
+				return false;
+			}
+			// Don't let Mac run shell commands through Uba as interposing dylibs into
+			// the shell results in dyld errors about no matching architecture.
+			// The shell is used to run various commands during a build like copy/ditto.
+			// So for these actions we need to make sure UBA is not used.
+			// Linking is similarly currently not working in Uba on Mac
+			bool bIsShellAction = action.CommandPath == BuildHostPlatform.Current.Shell;
+			bool bIsLinkAction = action.ActionType == ActionType.Link;
+			return (bIsShellAction || bIsLinkAction);
+		}
+
+		/// <summary>
+		/// Determine if an action is able to be run remotely
+		/// </summary>
+		/// <param name="action">The action to check</param>
+		/// <returns>If this action can be run remotely</returns>
+		bool CanRunRemotely(LinkedAction action) =>
+			action.bCanExecuteInUBA &&
+			action.bCanExecuteRemotely &&
+			(UBAConfig.bLinkRemote || action.ActionType != ActionType.Link) &&
+			!_localRetryActions.ContainsKey(action) &&
+			!_forcedRetryActions.ContainsKey(action) &&
+			!ForceLocalNoDetour(action);
+
+		ProcessStartInfo GetActionStartInfo(LinkedAction action, out FileItem? pchItem)
+		{
+			ProcessStartInfo startInfo = new()
+			{
+				Application = action.CommandPath.FullName,
+				WorkingDirectory = action.WorkingDirectory.FullName,
+				Arguments = action.CommandArguments,
+				Priority = ProcessPriority,
+				OutputStatsThresholdMs = (uint)UBAConfig.OutputStatsThresholdMs,
+				UserData = action,
+				Description = action.StatusDescription,
+				Configuration = action.bIsGCCCompiler ? EpicGames.UBA.ProcessStartInfo.CommonProcessConfigs.CompileClang : EpicGames.UBA.ProcessStartInfo.CommonProcessConfigs.CompileMsvc,
+				LogFile = UBAConfig.bLogEnabled ? action.Inner.ProducedItems.First().Location.GetFileName() : null,
+			};
+
+			bool usingLtcg = true; // ltcg linking checks what pch was used and it seems like it needs to be identical in other ways than timestamp
+			pchItem = action.Inner.ProducedItems.FirstOrDefault(item => item.Name.EndsWith(".pch", StringComparison.OrdinalIgnoreCase));
+			if (pchItem != null)
+			{
+				startInfo.Priority = System.Diagnostics.ProcessPriorityClass.AboveNormal;
+				if (!usingLtcg && action.ArtifactMode.HasFlag(ArtifactMode.PropagateInputs))
+				{
+					startInfo.TrackInputs = true;
+				}
+				else
+				{
+					pchItem = null;
+				}
+			}
+
+			return startInfo;
+		}
+
+		Func<Task>? RunActionLocal(ImmediateActionQueue queue, LinkedAction action)
+		{
+			if (UBAConfig.bForceBuildAllRemote && CanRunRemotely(action))
+			{
+				return null;
+			}
+
+			return () =>
+			{
+				bool enableDetour = !ForceLocalNoDetour(action) && action.bCanExecuteInUBA && !_forcedRetryActions.ContainsKey(action);
+
+				ProcessStartInfo startInfo = GetActionStartInfo(action, out FileItem? pchItem);
+				using (IProcess process = _session!.RunProcess(startInfo, false, null, enableDetour))
+				{
+					if (process.ExitCode != 0 && UBAConfig.bForcedRetry)
+					{
+						_threadedLogger.LogWarning("{Description} {StatusDescription}: Exited with error code {ExitCode}. This action will retry without UBA", action.CommandDescription, action.StatusDescription, process.ExitCode);
+						_forcedRetryActions.AddOrUpdate(action, false, (k, v) => false);
+						queue.RequeueAction(action);
+						return Task.CompletedTask;
+					}
+
+					if (!enableDetour && process.ExitCode == 0)
+					{
+						_session!.RegisterNewFiles(action.ProducedItems.Where(x => FileReference.Exists(x.Location)).Select(x => x.FullName).ToArray());
+					}
+
+					TimeSpan processorTime = process.TotalProcessorTime;
+					TimeSpan executionTime = process.TotalWallTime;
+					List<string> logLines = process.LogLines;
+					logLines.RemoveAll((line) => line.StartsWith("   Creating library ", StringComparison.OrdinalIgnoreCase) && line.EndsWith(".exp", StringComparison.OrdinalIgnoreCase) || line.EndsWith("file(s) copied.", StringComparison.OrdinalIgnoreCase));
+
+					string? additionalDescription = !enableDetour ? "(UBA disabled)" : null;
+					ActionFinished(queue, new ExecuteResults(logLines, process.ExitCode, executionTime, processorTime, additionalDescription), action, pchItem, process);
+				}
+				return Task.CompletedTask;
+			};
+		}
+
+		Func<Task>? RunActionRemote(ImmediateActionQueue queue, LinkedAction action)
+		{
+			if (!CanRunRemotely(action))
+			{
+				return null;
+			}
+
+			return () =>
+			{
+				uint knownInputsCount = 0;
+				byte[]? knownInputs = null;
+				if (UBAConfig.bUseKnownInputs)
+				{
+					int sizeOfChar = System.OperatingSystem.IsWindows() ? 2 : 1;
+
+					int byteCount = 0;
+					foreach (FileItem item in action.PrerequisiteItems)
+					{
+						byteCount += (item.FullName.Length + 1) * sizeOfChar;
+						++knownInputsCount;
+					}
+
+					knownInputs = new byte[byteCount + sizeOfChar];
+
+					int byteOffset = 0;
+					foreach (FileItem item in action.PrerequisiteItems)
+					{
+						string str = item.FullName;
+						int strBytes = str.Length * sizeOfChar;
+						if (sizeOfChar == 1) // Unmanaged size uses ascii
+						{
+							System.Buffer.BlockCopy(System.Text.Encoding.ASCII.GetBytes(str.ToCharArray()), 0, knownInputs, byteOffset, strBytes);
+						}
+						else
+						{
+							System.Buffer.BlockCopy(str.ToCharArray(), 0, knownInputs, byteOffset, strBytes);
+						}
+
+						byteOffset += strBytes + sizeOfChar;
+					}
+				}
+
+				ProcessStartInfo startInfo = GetActionStartInfo(action, out FileItem? pchItem);
+				_session!.RunProcessRemote(startInfo, (s, e) =>
+				{
+					if (e.ExitCode != 0 && !e.LogLines.Any())
+					{
+						RemoteActionFailedNoOutput(queue, action, e.ExitCode, e.ExecutingHost ?? "Unknown");
+						return;
+					}
+					else if ((uint)e.ExitCode == 0xC0000005)
+					{
+						RemoteActionFailedCrash(queue, action, e.ExitCode, e.ExecutingHost ?? "Unknown", "Access violation");
+						return;
+					}
+					else if ((uint)e.ExitCode == 0xC0000409)
+					{
+						RemoteActionFailedCrash(queue, action, e.ExitCode, e.ExecutingHost ?? "Unknown", "Stack buffer overflow");
+						return;
+					}
+					else if (e.ExitCode != 0 && e.LogLines.Any(x => x.Contains(" C1001: ", StringComparison.Ordinal)))
+					{
+						RemoteActionFailedCrash(queue, action, e.ExitCode, e.ExecutingHost ?? "Unknown", "C1001");
+						return;
+					}
+
+					string additionalDescription = $"[RemoteExecutor: {e.ExecutingHost}]";
+					TimeSpan processorTime = e.TotalProcessorTime;
+					TimeSpan executionTime = e.TotalWallTime;
+					List<string> logLines = e.LogLines;
+					logLines.RemoveAll((line) => line.StartsWith("   Creating library ", StringComparison.OrdinalIgnoreCase) && line.EndsWith(".exp", StringComparison.OrdinalIgnoreCase));
+					ActionFinished(queue, new ExecuteResults(logLines, e.ExitCode, executionTime, processorTime, additionalDescription), action, pchItem, s as IProcess);
+				}, action.Weight, knownInputs, knownInputsCount);
+				return Task.CompletedTask;
+			};
+		}
+
+		uint ActionsLeftThatCanRunRemotely(ImmediateActionQueue queue)
+		{
+			lock (_actionsChangedLock)
+			{
+				if (_bActionsChanged)
+				{
+					_actionsQueuedThatCanRunRemotely = queue.GetQueuedActionsCount(CanRunRemotely);
+					_bActionsChanged = false;
+				}
+				return _actionsQueuedThatCanRunRemotely;
+			}
+		}
+
+		protected void ActionFinished(ImmediateActionQueue queue, ExecuteResults results, LinkedAction action, FileItem? pchItem = null, IProcess? process = null)
+		{
+			if (_bIsCancelled)
+			{
+				return;
+			}
+
+			bool success = results.ExitCode == 0;
+			if (pchItem != null && process != null && success)
+			{
+				_session!.SetCustomCasKeyFromTrackedInputs(pchItem.FullName, action.WorkingDirectory.FullName, process);
+			}
+
+			queue.OnActionCompleted(action, success, results);
+
+			lock (_actionsChangedLock)
+			{
+				_bActionsChanged = true;
+			}
+		}
+
+		void RemoteActionFailedNoOutput(ImmediateActionQueue queue, LinkedAction action, int exitCode, string executingHost)
+		{
+			if (_bIsCancelled)
+			{
+				return;
+			}
+
+			_threadedLogger.LogWarning("{Description} {StatusDescription} [RemoteExecutor: {ExecutingHost}]: Exited with error code {ExitCode} with no output. This action will retry locally", action.CommandDescription, action.StatusDescription, executingHost, exitCode);
+			_localRetryActions.AddOrUpdate(action, false, (k, v) => false);
+			queue.RequeueAction(action);
+
+			lock (_actionsChangedLock)
+			{
+				_bActionsChanged = true;
+			}
+		}
+
+		void RemoteActionFailedCrash(ImmediateActionQueue queue, LinkedAction action, int exitCode, string executingHost, string error)
+		{
+			if (_bIsCancelled)
+			{
+				return;
+			}
+
+			_threadedLogger.LogWarning("{Description} {StatusDescription} [RemoteExecutor: {ExecutingHost}]: Exited with error code {ExitCode} ({Error}). This action will retry locally", action.CommandDescription, action.StatusDescription, executingHost, exitCode, error);
+			_localRetryActions.AddOrUpdate(action, false, (k, v) => false);
+			queue.RequeueAction(action);
+
+			lock (_actionsChangedLock)
+			{
+				_bActionsChanged = true;
+			}
+		}
+	}
+
+	/// <summary>
+	/// UBAExecutor, but force local only compiles regardless of other settings
+	/// </summary>
+	class UBALocalExecutor : UBAExecutor
+	{
+		public override string Name => "Unreal Build Accelerator local";
+
+		public static new bool IsAvailable()
+		{
+			return EpicGames.UBA.Utils.IsAvailable();
+		}
+
+		public UBALocalExecutor(int maxLocalActions, bool bAllCores, bool bCompactOutput, Microsoft.Extensions.Logging.ILogger logger, CommandLineArguments? additionalArguments = null)
+			: base(maxLocalActions, bAllCores, bCompactOutput, logger, additionalArguments)
+		{
+			UBAConfig.bDisableRemote = true;
+			UBAConfig.bForceBuildAllRemote = false;
+			UBAConfig.Zone = "local";
+		}
+	}
+}
diff --git a/Engine/Source/Programs/UnrealBuildTool/Executors/UnrealBuildAccelerator/UnrealBuildAcceleratorConfig.cs b/Engine/Source/Programs/UnrealBuildTool/Executors/UnrealBuildAccelerator/UnrealBuildAcceleratorConfig.cs
new file mode 100644
--- /dev/null
+++ b/Engine/Source/Programs/UnrealBuildTool/Executors/UnrealBuildAccelerator/UnrealBuildAcceleratorConfig.cs
@@ -0,0 +1,242 @@
+// Copyright Epic Games, Inc. All Rights Reserved.
+
+using System;
+using EpicGames.Core;
+
+namespace UnrealBuildTool
+{
+	/// <summary>
+	/// Configuration for Unreal Build Accelerator
+	/// </summary>
+	[System.Diagnostics.CodeAnalysis.SuppressMessage("Style", "IDE1006:Naming Styles", Justification = "UnrealBuildTool naming style")]
+	class UnrealBuildAcceleratorConfig
+	{
+		/// <summary>
+		/// When set to true, UBA will not use any remote help
+		/// </summary>
+		[XmlConfigFile(Category = "UnrealBuildAccelerator")]
+		[CommandLine("-BoxDisableRemote")]
+		[CommandLine("-UBADisableRemote")]
+		public bool bDisableRemote { get; set; } = false;
+
+		/// <summary>
+		/// When set to true, UBA will force all actions that can be built remotely to be built remotely. This will hang if there are no remote agents available
+		/// </summary>
+		[XmlConfigFile(Category = "UnrealBuildAccelerator")]
+		[CommandLine("-BoxForceRemote")]
+		[CommandLine("-UBAForceRemote")]
+		public bool bForceBuildAllRemote { get; set; } = false;
+
+		/// <summary>
+		/// When set to true, actions that fail locally with UBA will be retried without UBA.
+		/// </summary>
+		[XmlConfigFile(Category = "UnrealBuildAccelerator")]
+		[CommandLine("-BoxForcedRetry")]
+		[CommandLine("-UBAForcedRetry")]
+		public bool bForcedRetry { get; set; } = false;
+
+		/// <summary>
+		/// When set to true, all errors and warnings from UBA will be output at the appropriate severity level to the log (rather than being output as 'information' and attempting to continue regardless).
+		/// </summary>
+		[XmlConfigFile(Category = "UnrealBuildAccelerator")]
+		[CommandLine("-BoxStrict")]
+		[CommandLine("-UBAStrict")]
+		public bool bStrict { get; set; } = false;
+
+		/// <summary>
+		/// If UBA should store cas compressed or raw
+		/// </summary>
+		[XmlConfigFile(Category = "UnrealBuildAccelerator")]
+		[CommandLine("-BoxStoreRaw")]
+		[CommandLine("-UBAStoreRaw")]
+		public bool bStoreRaw { get; set; } = false;
+
+		/// <summary>
+		/// If UBA should distribute linking to remote workers. This needs bandwidth but can be an optimization
+		/// </summary>
+		[XmlConfigFile(Category = "UnrealBuildAccelerator")]
+		[CommandLine("-BoxLinkRemote")]
+		[CommandLine("-UBALinkRemote")]
+		public bool bLinkRemote { get; set; } = false;
+
+		/// <summary>
+		/// The amount of gigabytes UBA is allowed to use to store workset and cached data. It is a good idea to have this >10gb
+		/// </summary>
+		[XmlConfigFile(Category = "UnrealBuildAccelerator")]
+		[CommandLine("-BoxStoreCapacityGb")]
+		[CommandLine("-UBAStoreCapacityGb")]
+		public int StoreCapacityGb { get; set; } = 40;
+
+		/// <summary>
+		/// Max number of worker threads that can handle messages from remotes. 
+		/// </summary>
+		[XmlConfigFile(Category = "UnrealBuildAccelerator")]
+		[CommandLine("-BoxMaxWorkers")]
+		[CommandLine("-UBAMaxWorkers")]
+		public int MaxWorkers { get; set; } = 192;
+
+		/// <summary>
+		/// Max size of each message sent from server to client
+		/// </summary>
+		[XmlConfigFile(Category = "UnrealBuildAccelerator")]
+		[CommandLine("-BoxSendSize")]
+		[CommandLine("-UBASendSize")]
+		public int SendSize { get; set; } = 256 * 1024;
+
+		/// <summary>
+		/// Which ip UBA server should listen to for connections
+		/// </summary>
+		[XmlConfigFile(Category = "UnrealBuildAccelerator")]
+		[CommandLine("-BoxHost")]
+		[CommandLine("-UBAHost")]
+		public string Host { get; set; } = String.Empty;
+
+		/// <summary>
+		/// Which port UBA server should listen to for connections.
+		/// </summary>
+		[XmlConfigFile(Category = "UnrealBuildAccelerator")]
+		[CommandLine("-BoxPort")]
+		[CommandLine("-UBAPort")]
+		public int Port { get; set; } = 1345;
+
+		/// <summary>
+		/// Which directory to store files for UBA.
+		/// </summary>
+		[XmlConfigFile(Category = "UnrealBuildAccelerator")]
+		[CommandLine("-BoxRootDir")]
+		[CommandLine("-UBARootDir")]
+		public string? RootDir { get; set; } = null;
+
+		/// <summary>
+		/// Use Quic protocol instead of Tcp (experimental)
+		/// </summary>
+		[XmlConfigFile(Category = "UnrealBuildAccelerator")]
+		[CommandLine("-BoxQuic", Value = "true")]
+		[CommandLine("-UBAQuic", Value = "true")]
+		public bool bUseQuic { get; set; } = false;
+
+		/// <summary>
+		/// Enable logging of UBA processes
+		/// </summary>
+		[XmlConfigFile(Category = "UnrealBuildAccelerator")]
+		[CommandLine("-BoxLog", Value = "true")]
+		[CommandLine("-UBALog", Value = "true")]
+		public bool bLogEnabled { get; set; } = false;
+
+		/// <summary>
+		/// Prints summary of UBA stats at end of build
+		/// </summary>
+		[XmlConfigFile(Category = "UnrealBuildAccelerator")]
+		[CommandLine("-BoxPrintSummary", Value = "true")]
+		[CommandLine("-UBAPrintSummary", Value = "true")]
+		public bool bPrintSummary { get; set; } = false;
+
+		/// <summary>
+		/// Launch visualizer application which shows build progress
+		/// </summary>
+		[XmlConfigFile(Category = "UnrealBuildAccelerator")]
+		[CommandLine("-BoxVisualizer", Value = "true")]
+		[CommandLine("-UBAVisualizer", Value = "true")]
+		public bool bLaunchVisualizer { get; set; } = false;
+
+		/// <summary>
+		/// Resets the cas cache
+		/// </summary>
+		[XmlConfigFile(Category = "UnrealBuildAccelerator")]
+		[CommandLine("-BoxResetCas", Value = "true")]
+		[CommandLine("-UBAResetCas", Value = "true")]
+		public bool bResetCas { get; set; } = false;
+
+		/// <summary>
+		/// Provide custom path for trace output file
+		/// </summary>
+		[XmlConfigFile(Category = "UnrealBuildAccelerator")]
+		[CommandLine("-BoxTraceOutputFile")]
+		[CommandLine("-UBATraceOutputFile")]
+		public string TraceFile { get; set; } = String.Empty;
+
+		/// <summary>
+		/// Add verbose details to the UBA trace
+		/// </summary>
+		[XmlConfigFile(Category = "UnrealBuildAccelerator")]
+		[CommandLine("-BoxDetailedTrace", Value = "true")]
+		[CommandLine("-UBADetailedTrace", Value = "true")]
+		public bool bDetailedTrace { get; set; }
+
+		/// <summary>
+		/// Disable UBA waiting on available memory before spawning new processes
+		/// </summary>
+		[XmlConfigFile(Category = "UnrealBuildAccelerator")]
+		[CommandLine("-UBADisableWaitOnMem", Value = "true")]
+		public bool bDisableWaitOnMem { get; set; }
+
+		/// <summary>
+		/// Let UBA kill running processes when close to out of memory
+		/// </summary>
+		[XmlConfigFile(Category = "UnrealBuildAccelerator")]
+		[CommandLine("-BoxAllowKillOnMem", Value = "true")]
+		[CommandLine("-UBAAllowKillOnMem", Value = "true")]
+		public bool bAllowKillOnMem { get; set; }
+
+		/// <summary>
+		/// Threshold for when executor should output logging for the process. Defaults to never
+		/// </summary>
+		[XmlConfigFile(Category = "UnrealBuildAccelerator")]
+		[CommandLine("-BoxOutputStatsThresholdMs")]
+		[CommandLine("-UBAOutputStatsThresholdMs")]
+		public int OutputStatsThresholdMs { get; set; } = Int32.MaxValue;
+
+		/// <summary>
+		/// Skip writing intermediate and output files to disk. Useful for validation builds where we don't need the output
+		/// </summary>
+		[XmlConfigFile(Category = "UnrealBuildAccelerator")]
+		[CommandLine("-BoxNoWrite", Value = "false")]
+		[CommandLine("-UBANoWrite", Value = "false")]
+		public bool bWriteToDisk { get; set; } = true;
+
+		/// <summary>
+		/// Set to true to disable mimalloc and detouring of memory allocations.
+		/// </summary>
+		[XmlConfigFile(Category = "UnrealBuildAccelerator")]
+		[CommandLine("-BoxNoCustoMalloc", Value = "true")]
+		[CommandLine("-UBANoCustoMalloc", Value = "true")]
+		public bool bDisableCustomAlloc { get; set; } = false;
+
+		/// <summary>
+		/// The zone to use for UBA.
+		/// </summary>
+		[XmlConfigFile(Category = "UnrealBuildAccelerator")]
+		[CommandLine("-BoxZone=")]
+		[CommandLine("-UBAZone=")]
+		public string Zone { get; set; } = String.Empty;
+
+		/// <summary>
+		/// Set to true to enable encryption when transfering files over the network.
+		/// </summary>
+		[XmlConfigFile(Category = "UnrealBuildAccelerator")]
+		[CommandLine("-BoxCrypto", Value = "true")]
+		[CommandLine("-UBACrypto", Value = "true")]
+		public bool bUseCrypto { get; set; } = false;
+
+		/// <summary>
+		/// Set to true to provide known inputs to processes that are run remote. This is an experimental feature to speed up build times when ping is higher
+		/// </summary>
+		[XmlConfigFile(Category = "UnrealBuildAccelerator")]
+		[CommandLine("-UBAUseKnownInputs", Value = "true")]
+		public bool bUseKnownInputs { get; set; } = false;
+
+		/// <summary>
+		/// Write yaml file with all actions that are queued for build. This can be used to replay using "UbaCli.exe local file.yaml"
+		/// </summary>
+		[XmlConfigFile(Category = "UnrealBuildAccelerator")]
+		[CommandLine("-UBAActionsOutputFile")]
+		public string ActionsOutputFile { get; set; } = String.Empty;
+
+		/// <summary>
+		/// Set to true to see more info about what is happening inside uba and also log output from agents
+		/// </summary>
+		[XmlConfigFile(Category = "UnrealBuildAccelerator")]
+		[CommandLine("-UBADetailedLog", Value = "true")]
+		public bool bDetailedLog { get; set; } = false;
+	}
+}
\ No newline at end of file
diff --git a/Engine/Source/Programs/UnrealBuildTool/Executors/UnrealBuildAccelerator/UnrealBuildAcceleratorHordeConfig.cs b/Engine/Source/Programs/UnrealBuildTool/Executors/UnrealBuildAccelerator/UnrealBuildAcceleratorHordeConfig.cs
new file mode 100644
--- /dev/null
+++ b/Engine/Source/Programs/UnrealBuildTool/Executors/UnrealBuildAccelerator/UnrealBuildAcceleratorHordeConfig.cs
@@ -0,0 +1,174 @@
+// Copyright Epic Games, Inc. All Rights Reserved.
+
+using System;
+using EpicGames.Core;
+using EpicGames.Horde.Compute;
+
+namespace UnrealBuildTool
+{
+	/// <summary>
+	/// Configuration for Unreal Build Accelerator Horde session
+	/// </summary>
+	[System.Diagnostics.CodeAnalysis.SuppressMessage("Style", "IDE1006:Naming Styles", Justification = "UnrealBuildTool naming style")]
+	class UnrealBuildAcceleratorHordeConfig
+	{
+		/// <summary>
+		/// Uri of the Horde server
+		/// </summary>
+		[XmlConfigFile(Category = "Horde", Name = "Server")]
+		[CommandLine("-BoxHorde=")]
+		[CommandLine("-UBAHorde=")]
+		public string? HordeServer { get; set; }
+
+		/// <summary>
+		/// Uri of the Horde server
+		/// </summary>
+		[XmlConfigFile(Category = "Horde", Name = "Token")]
+		[CommandLine("-BoxHordeToken=")]
+		[CommandLine("-UBAHordeToken=")]
+		public string? HordeToken { get; set; }
+
+		/// <summary>
+		/// OIDC id for the login to use
+		/// </summary>
+		[XmlConfigFile(Category = "Horde", Name = "OidcProvider")]
+		[CommandLine("-BoxHordeOidc=")]
+		[CommandLine("-UBAHordeOidc=")]
+		public string? HordeOidcProvider { get; set; }
+
+		/// <summary>
+		/// Pool for the Horde agent to assign, calculated based on platform and overrides
+		/// </summary>
+		public string? HordePool
+		{
+			get
+			{
+				string? Pool = DefaultHordePool;
+				if (OverrideHordePool != null)
+				{
+					Pool = OverrideHordePool;
+				}
+				else if (OperatingSystem.IsWindows() && WindowsHordePool != null)
+				{
+					Pool = WindowsHordePool;
+				}
+				else if (OperatingSystem.IsMacOS() && MacHordePool != null)
+				{
+					Pool = MacHordePool;
+				}
+				else if (OperatingSystem.IsLinux() && LinuxHordePool != null)
+				{
+					Pool = LinuxHordePool;
+				}
+
+				//Console.WriteLine("CHOSEN UBA POOL: {0}", Pool);
+				return Pool;
+			}
+		}
+
+		/// <summary>
+		/// Pool for the Horde agent to assign if no override current platform doesn't have it set
+		/// </summary>
+		[XmlConfigFile(Category = "Horde", Name = "Pool")]
+		public string? DefaultHordePool { get; set; }
+
+		/// <summary>
+		/// Pool for the Horde agent to assign, only used for commandline override
+		/// </summary>
+		[CommandLine("-BoxHordePool=")]
+		[CommandLine("-UBAHordePool=")]
+		public string? OverrideHordePool { get; set; }
+
+		/// <summary>
+		/// Pool for the Horde agent to assign when on Linux
+		/// </summary>
+		[XmlConfigFile(Category = "Horde", Name = "LinuxPool")]
+		public string? LinuxHordePool { get; set; }
+
+		/// <summary>
+		/// Pool for the Horde agent to assign when on Mac
+		/// </summary>
+		[XmlConfigFile(Category = "Horde", Name = "MacPool")]
+		public string? MacHordePool { get; set; }
+
+		/// <summary>
+		/// Pool for the Horde agent to assign when on Windows
+		/// </summary>
+		[XmlConfigFile(Category = "Horde", Name = "WindowsPool")]
+		public string? WindowsHordePool { get; set; }
+
+		/// <summary>
+		/// Requirements for the Horde agent to assign
+		/// </summary>
+		[XmlConfigFile(Category = "Horde", Name = "Requirements")]
+		[CommandLine("-BoxHordeRequirements=")]
+		[CommandLine("-UBAHordeRequirements=")]
+		public string? HordeCondition { get; set; }
+
+		/// <summary>
+		/// Which ip UBA server should give to agents. This will invert so host listens and agents connect
+		/// </summary>
+		[XmlConfigFile(Category = "Horde", Name = "LocalHost")]
+		[CommandLine("-BoxHordeHost")]
+		[CommandLine("-UBAHordeHost")]
+		public string HordeHost { get; set; } = String.Empty;
+
+		/// <summary>
+		/// Max cores allowed to be used by build session
+		/// </summary>
+		[XmlConfigFile(Category = "UnrealBuildAccelerator")]
+		[CommandLine("-BoxHordeMaxCores")]
+		[CommandLine("-UBAHordeMaxCores")]
+		public int HordeMaxCores { get; set; } = 576;
+
+		/// <summary>
+		/// How long UBT should wait to ask for help. Useful in build configs where machine can delay remote work and still get same wall time results (pch dependencies etc)
+		/// </summary>
+		[XmlConfigFile(Category = "UnrealBuildAccelerator")]
+		[CommandLine("-BoxHordeDelay")]
+		[CommandLine("-UBAHordeDelay")]
+		public int HordeDelay { get; set; } = 0;
+
+		/// <summary>
+		/// Allow use of Wine. Only applicable to Horde agents running Linux. Can still be ignored if Wine executable is not set on agent.
+		/// </summary>
+		[XmlConfigFile(Category = "UnrealBuildAccelerator")]
+		[CommandLine("-BoxHordeAllowWine", Value = "true")]
+		[CommandLine("-UBAHordeAllowWine", Value = "true")]
+		public bool bHordeAllowWine { get; set; } = true;
+
+		/// <summary>
+		/// Connection mode for agent/compute communication
+		/// <see cref="ConnectionMode" /> for valid modes.
+		/// </summary>
+		[XmlConfigFile(Category = "Horde", Name = "ConnectionMode")]
+		[CommandLine("-BoxHordeConnectionMode=")]
+		[CommandLine("-UBAHordeConnectionMode=")]
+		public string? HordeConnectionMode { get; set; }
+
+		/// <summary>
+		/// Encryption to use for agent/compute communication. Note that UBA agent uses its own encryption.
+		/// <see cref="Encryption" /> for valid modes.
+		/// </summary>
+		[XmlConfigFile(Category = "Horde", Name = "Encryption")]
+		[CommandLine("-BoxHordeEncryption=")]
+		[CommandLine("-UBAHordeEncryption=")]
+		public string? HordeEncryption { get; set; }
+
+		/// <summary>
+		/// Sentry URL to send box data to. Optional.
+		/// </summary>
+		[XmlConfigFile(Category = "Horde")]
+		[CommandLine("-BoxSentryUrl=")]
+		[CommandLine("-UBASentryUrl=")]
+		public string? UBASentryUrl { get; set; }
+
+		/// <summary>
+		/// Disable horde all together
+		/// </summary>
+		[XmlConfigFile(Category = "UnrealBuildAccelerator")]
+		[CommandLine("-BoxDisableHorde")]
+		[CommandLine("-UBADisableHorde")]
+		public bool bDisableHorde { get; set; } = false;
+	}
+}
\ No newline at end of file
