[UBT]
* Added braces around single line conditional

diff --git a/Engine/Source/Programs/UnrealBuildTool/Platform/Windows/VCToolChain.cs b/Engine/Source/Programs/UnrealBuildTool/Platform/Windows/VCToolChain.cs
--- a/Engine/Source/Programs/UnrealBuildTool/Platform/Windows/VCToolChain.cs
+++ b/Engine/Source/Programs/UnrealBuildTool/Platform/Windows/VCToolChain.cs
@@ -2480,7 +2480,9 @@ namespace UnrealBuildTool
 			{
 				Arguments.Add("/NOIMPLIB");
 				if (!Target.WindowsPlatform.bAllowClangLinker)
+				{
 					Arguments.Add("/NOEXP"); // This compiler flag does not exist on lld-link.exe.. it skips the writing of the .exp file
+				}
 			}
 			else if (Target.bShouldCompileAsDLL)
 			{
