Making .h files compilable in VS.

diff --git a/Engine/Source/Programs/UnrealBuildTool/ProjectFiles/VisualStudio/VCProject.cs b/Engine/Source/Programs/UnrealBuildTool/ProjectFiles/VisualStudio/VCProject.cs
--- a/Engine/Source/Programs/UnrealBuildTool/ProjectFiles/VisualStudio/VCProject.cs
+++ b/Engine/Source/Programs/UnrealBuildTool/ProjectFiles/VisualStudio/VCProject.cs
@@ -1248,12 +1248,12 @@ namespace UnrealBuildTool
 				Parallel.ForEach(LocalAliasedFiles, LocalAliasedFile =>
 				{
 					// get the filetype as represented to Visual Studio
-					string VCFileType = GetVCFileType(LocalAliasedFile.FileSystemPath);
+					VCFileType FileType = GetVCFileType(LocalAliasedFile.FileSystemPath);
 					DirectoryReference FileSystemPathDir = new DirectoryReference(LocalAliasedFile.FileSystemPath);
 
 					// if the filetype is an include and its path is filtered out, skip it entirely (should we do this for any type of
 					// file? Possibly, but not today due to potential fallout)
-					if (VCFileType == "ClInclude" && IncludePathIsFilteredOut(FileSystemPathDir))
+					if (CanVCFileTypeBeIncluded(FileType) && IncludePathIsFilteredOut(FileSystemPathDir))
 					{
 						return;
 					}
@@ -1264,7 +1264,7 @@ namespace UnrealBuildTool
 						return;
 					}
 
-					if (VCFileType != "ClCompile")
+					if (!CanVCFileTypeBeCompiled(FileType))
 					{
 						return;
 					}
@@ -1564,12 +1564,12 @@ namespace UnrealBuildTool
 					}
 
 					// get the filetype as represented to Visual Studio
-					string VCFileType = GetVCFileType(AliasedFile.FileSystemPath);
+					VCFileType FileType = GetVCFileType(AliasedFile.FileSystemPath);
 					DirectoryReference FileSystemPathDir = new DirectoryReference(AliasedFile.FileSystemPath);
 
 					// if the filetype is an include and its path is filtered out, skip it entirely (should we do this for any type of
 					// file? Possibly, but not today due to potential fallout)
-					if (VCFileType == "ClInclude" && IncludePathIsFilteredOut(FileSystemPathDir))
+					if (CanVCFileTypeBeIncluded(FileType) && IncludePathIsFilteredOut(FileSystemPathDir))
 					{
 						continue;
 					}
@@ -1580,7 +1580,9 @@ namespace UnrealBuildTool
 						continue;
 					}
 
-					if (VCFileType != "ClCompile")
+					string VCFileType = GetVCFileTypeString(FileType);
+
+					if (!CanVCFileTypeBeCompiled(FileType))
 					{
 						VCProjectFileContent.AppendLine("    <{0} Include=\"{1}\"/>", VCFileType, EscapeFileName(AliasedFile.FileSystemPath));
 					}
@@ -2056,35 +2058,81 @@ namespace UnrealBuildTool
 			return FiltersFileIsNeeded;
 		}
 
+		private enum VCFileType
+		{
+			None,
+			CCode,
+			Header,
+			Inline,
+			Resource,
+			Manifest,
+		};
+
 		/// <summary>
-		/// Returns the VCFileType element name based on the file path.
+		/// Returns the VCFileType string based on the VCFileType enum.
+		/// </summary>
+		/// <param name="FileType"></param>
+		/// <returns>Name of the element in MSBuild project file for this file type</returns>
+		private string GetVCFileTypeString(VCFileType FileType)
+		{
+			switch (FileType)
+			{
+				default:
+				case VCFileType.None:
+					return "None";
+				case VCFileType.CCode:
+					return "ClCompile";
+				case VCFileType.Header:
+					return "ClCompile";
+				case VCFileType.Inline:
+					return "ClInclude";
+				case VCFileType.Resource:
+					return "ResourceCompile";
+				case VCFileType.Manifest:
+					return "Manifest";
+			}
+		}
+
+		private bool CanVCFileTypeBeIncluded(VCFileType FileType)
+		{
+			return FileType == VCFileType.Header || FileType == VCFileType.Inline;
+		}
+
+		private bool CanVCFileTypeBeCompiled(VCFileType FileType)
+		{
+			return FileType == VCFileType.CCode || FileType == VCFileType.Header;
+		}
+
+		/// <summary>
+		/// Returns the VCFileType enum value based on the file path.
 		/// </summary>
 		/// <param name="Path">The path of the file to return type for.</param>
-		/// <returns>Name of the element in MSBuild project file for this file.</returns>
-		private string GetVCFileType(string Path)
+		/// <returns>VCFileType enum value for this file.</returns>
+		private VCFileType GetVCFileType(string Path)
 		{
 			// What type of file is this?
-			if (Path.EndsWith(".h", StringComparison.InvariantCultureIgnoreCase) ||
-				Path.EndsWith(".inl", StringComparison.InvariantCultureIgnoreCase))
+			if (Path.EndsWith(".h", StringComparison.InvariantCultureIgnoreCase))
 			{
-				return "ClInclude";
+				return VCFileType.Header;
 			}
-			else if (Path.EndsWith(".cpp", StringComparison.InvariantCultureIgnoreCase))
+			if (Path.EndsWith(".inl", StringComparison.InvariantCultureIgnoreCase))
 			{
-				return "ClCompile";
+				return VCFileType.Inline;
 			}
-			else if (Path.EndsWith(".rc", StringComparison.InvariantCultureIgnoreCase))
+			if (Path.EndsWith(".cpp", StringComparison.InvariantCultureIgnoreCase))
 			{
-				return "ResourceCompile";
+				return VCFileType.CCode;
 			}
-			else if (Path.EndsWith(".manifest", StringComparison.InvariantCultureIgnoreCase))
+			if (Path.EndsWith(".rc", StringComparison.InvariantCultureIgnoreCase))
 			{
-				return "Manifest";
+				return VCFileType.Resource;
 			}
-			else
+			if (Path.EndsWith(".manifest", StringComparison.InvariantCultureIgnoreCase))
 			{
-				return "None";
+				return VCFileType.Manifest;
 			}
+
+			return VCFileType.None;
 		}
 
 		// Helper class to generate NMake build commands and arguments
