UHT: plugin/Source is also a valid base directory for include paths

#rnx

diff --git a/Engine/Source/Programs/UnrealBuildTool/System/ExternalExecution.cs b/Engine/Source/Programs/UnrealBuildTool/System/ExternalExecution.cs
--- a/Engine/Source/Programs/UnrealBuildTool/System/ExternalExecution.cs
+++ b/Engine/Source/Programs/UnrealBuildTool/System/ExternalExecution.cs
@@ -360,6 +360,16 @@ namespace UnrealBuildTool
 
 		static DirectoryReference FindIncludeBase(UEBuildModuleCPP Module, ILogger Logger)
 		{
+			// Plugin Source directory may be an available include path for modules under a UPlugin (see UEBuildTarget.FindOrCreateModuleByName)
+			if (Module.Rules.Plugin != null)
+			{
+				DirectoryReference PluginSourceDirectoryName = DirectoryReference.Combine(Module.Rules.Plugin.Directory, "Source");
+				if (Module.Rules.File.IsUnderDirectory(PluginSourceDirectoryName))
+				{
+					return PluginSourceDirectoryName;
+				}
+			}
+
 			// Project Source directory is always an available include path for modules under a UProject (see UEBuildTarget.FindOrCreateModuleByName)
 			if (Module.Rules.Target.ProjectFile != null)
 			{
