Add event matcher to temporarily silence a xcode warning

#rnx

diff --git a/Engine/Source/Programs/UnrealBuildTool/Matchers/CompileEventMatcher.cs b/Engine/Source/Programs/UnrealBuildTool/Matchers/CompileEventMatcher.cs
--- a/Engine/Source/Programs/UnrealBuildTool/Matchers/CompileEventMatcher.cs
+++ b/Engine/Source/Programs/UnrealBuildTool/Matchers/CompileEventMatcher.cs
@@ -5,6 +5,7 @@ using System.Diagnostics.CodeAnalysis;
 using System.Linq;
 using System.Text.RegularExpressions;
 using EpicGames.Core;
+using Microsoft.CodeAnalysis.Differencing;
 using Microsoft.Extensions.Logging;
 
 namespace UnrealBuildTool.Matchers
@@ -185,7 +186,7 @@ namespace UnrealBuildTool.Matchers
 			else if (input.IsMatch(s_scriptCompilePattern))
 			{
 				LogEventBuilder builder = new LogEventBuilder(input);
-				return builder.ToMatch(LogEventPriority.High, LogLevel.Error, KnownLogEvents.Compiler_Summary);
+				return builder.ToMatch(LogEventPriority.High, LogLevel.Error, KnownLogEvents.Systemic_XCode);
 			}
 			return null;
 		}
@@ -203,6 +204,12 @@ namespace UnrealBuildTool.Matchers
 				return true;
 			}
 
+			if (builder.Current.Contains("was built for newer macOS version"))
+			{
+				outEvent = builder.ToMatch(LogEventPriority.None, LogLevel.Information, KnownLogEvents.Compiler);
+				return true;
+			}
+
 			outEvent = null;
 			return false;
 		}
