Headers that use UE_DEPRECATED_HEADER shouldn't be included in header units.

diff --git a/Engine/Source/Programs/UnrealBuildTool/System/SourceFileMetadataCache.cs b/Engine/Source/Programs/UnrealBuildTool/System/SourceFileMetadataCache.cs
--- a/Engine/Source/Programs/UnrealBuildTool/System/SourceFileMetadataCache.cs
+++ b/Engine/Source/Programs/UnrealBuildTool/System/SourceFileMetadataCache.cs
@@ -392,6 +392,10 @@ namespace UnrealBuildTool
 						++InsideDeprecationScope;
 					}
 				}
+				else if (LineSpan.StartsWith("UE_DEPRECATED_HEADER"))
+				{
+					HeaderFileInfo.UnitType = HeaderUnitType.Skip;
+				}
 				else
 				{
 					int HeaderUnitIndex = Line.IndexOf("HEADER_UNIT_");
