Adding Systemic matcher to out of heap memory errors

#rnx

diff --git a/Engine/Source/Programs/UnrealBuildTool/Matchers/CompileEventMatcher.cs b/Engine/Source/Programs/UnrealBuildTool/Matchers/CompileEventMatcher.cs
--- a/Engine/Source/Programs/UnrealBuildTool/Matchers/CompileEventMatcher.cs
+++ b/Engine/Source/Programs/UnrealBuildTool/Matchers/CompileEventMatcher.cs
@@ -262,6 +262,14 @@ namespace UnrealBuildTool.Matchers
 						return true;
 					}
 				}
+				else if (codeGroup.Value.StartsWith("C", StringComparison.Ordinal))
+				{
+					if (codeGroup.Value.Equals("C1060", StringComparison.Ordinal))
+					{
+						outEvent = builder.ToMatch(LogEventPriority.High, LogLevel.Error, KnownLogEvents.Systemic_MSBuild);
+						return true;
+					}
+				}
 			}
 
 			builder.AnnotateSourceFile(match.Groups["file"], sourceFileBaseDir);
