Fixing mac system matcher

#rnx

diff --git a/Engine/Source/Programs/UnrealBuildTool/Matchers/CompileEventMatcher.cs b/Engine/Source/Programs/UnrealBuildTool/Matchers/CompileEventMatcher.cs
--- a/Engine/Source/Programs/UnrealBuildTool/Matchers/CompileEventMatcher.cs
+++ b/Engine/Source/Programs/UnrealBuildTool/Matchers/CompileEventMatcher.cs
@@ -186,7 +186,7 @@ namespace UnrealBuildTool.Matchers
 			else if (input.IsMatch(s_scriptCompilePattern))
 			{
 				LogEventBuilder builder = new LogEventBuilder(input);
-				return builder.ToMatch(LogEventPriority.High, LogLevel.Error, KnownLogEvents.Systemic_XCode);
+				return builder.ToMatch(LogEventPriority.High, LogLevel.Error, KnownLogEvents.Compiler_Summary);
 			}
 			return null;
 		}
@@ -206,7 +206,7 @@ namespace UnrealBuildTool.Matchers
 
 			if (builder.Current.Contains("was built for newer macOS version"))
 			{
-				outEvent = builder.ToMatch(LogEventPriority.None, LogLevel.Information, KnownLogEvents.Compiler);
+				outEvent = builder.ToMatch(LogEventPriority.Highest, LogLevel.Information, KnownLogEvents.Systemic_XCode);
 				return true;
 			}
 
