UnrealBuildTool: Include sdk config json files in the platform folders in the project

#rnx

diff --git a/Engine/Source/Programs/UnrealBuildTool/UnrealBuildTool.csproj b/Engine/Source/Programs/UnrealBuildTool/UnrealBuildTool.csproj
--- a/Engine/Source/Programs/UnrealBuildTool/UnrealBuildTool.csproj
+++ b/Engine/Source/Programs/UnrealBuildTool/UnrealBuildTool.csproj
@@ -100,6 +100,16 @@
        messier directory structure for these few files -->
       <Link>$([System.Text.RegularExpressions.Regex]::Replace(%(FullPath), '^.+?[\\/]UnrealBuildTool[\\/]', ''))</Link>
     </Compile>
+
+    <!-- SDK config files -->
+    <None Include="../../../Platforms/*/Config/*_SDK.json">
+      <!-- Since RecursiveDir is bugged on VS Mac, we use the FullPath, and use Regex to pull out portions of the path to get Platform name -->
+      <Link>Platform/$([System.Text.RegularExpressions.Regex]::Replace(%(FullPath), '^(.*[\\/]Platforms[\\/])(.+?)[\\/](.*)$', '$2'))/%(Filename)%(Extension)</Link>
+    </None>
+    <None Include="../../../Config/*/*_SDK.json">
+      <!-- Since RecursiveDir is bugged on VS Mac, we use the FullPath, and use Regex to pull out portions of the path to get Platform name -->
+      <Link>Platform/$([System.Text.RegularExpressions.Regex]::Replace(%(FullPath), '^(.*[\\/]Config[\\/])(.+?)[\\/](.*)$', '$2'))/%(Filename)%(Extension)</Link>
+    </None>
   </ItemGroup>
 
   <ItemGroup>
