UBT: Establish Horde connection using regular EpicGames.Horde methods raher than using custom OIDC token discovery. Removes the need to configure an OIDC provider in an engine config file when using auth.

OIDCProvider command line option/XmlConfig setting is now unused and marked obsolete, but is kept around for now to avoid triggering warnings for Megastation users.

diff --git a/Engine/Source/Programs/UnrealBuildTool/Executors/UnrealBuildAccelerator/UBAAgentCoordinatorHorde.cs b/Engine/Source/Programs/UnrealBuildTool/Executors/UnrealBuildAccelerator/UBAAgentCoordinatorHorde.cs
--- a/Engine/Source/Programs/UnrealBuildTool/Executors/UnrealBuildAccelerator/UBAAgentCoordinatorHorde.cs
+++ b/Engine/Source/Programs/UnrealBuildTool/Executors/UnrealBuildAccelerator/UBAAgentCoordinatorHorde.cs
@@ -36,7 +36,6 @@ namespace UnrealBuildTool
 		const string ResourceLogicalCores = "LogicalCores";
 		readonly ClusterId _clusterId = new("default");
 
-		readonly Uri _hordeUri;
 		readonly string? _pool;
 		readonly bool _allowWine;
 		readonly int _maxCores;
@@ -68,10 +67,9 @@ namespace UnrealBuildTool
 
 		readonly List<Worker> _workers = new();
 
-		public UBAHordeSession(UBAExecutor owner, Uri hordeUri, AuthenticationHeaderValue? authHeader, string? pool, bool allowWine, int maxCores, bool strict, ConnectionMode? connectionMode, Encryption? encryption, ILogger logger)
+		public UBAHordeSession(UBAExecutor owner, Uri? serverUrl, string? accessToken, string? pool, bool allowWine, int maxCores, bool strict, ConnectionMode? connectionMode, Encryption? encryption, ILogger logger)
 		{
 			_owner = owner;
-			_hordeUri = hordeUri;
 			_pool = pool;
 			_allowWine = allowWine;
 			_maxCores = maxCores;
@@ -81,14 +79,12 @@ namespace UnrealBuildTool
 			_logger = logger;
 			_crypto = owner.Crypto;
 
-			void ConfigureHttpClient(HttpClient httpClient)
-			{
-				httpClient.BaseAddress = _hordeUri;
-				httpClient.DefaultRequestHeaders.Authorization = authHeader;
-			}
-
 			ServiceCollection services = new();
-			services.AddHttpClient<HordeHttpClient>(ConfigureHttpClient);
+			services.AddHorde(options => 
+			{
+				options.ServerUrl = serverUrl;
+				options.AccessToken = accessToken;
+			});
 			_serviceProvider = services.BuildServiceProvider();
 
 			if (connectionMode == ConnectionMode.Relay && String.IsNullOrEmpty(_crypto))
@@ -377,48 +373,19 @@ namespace UnrealBuildTool
 				return null;
 			}
 
-			string? server = hordeConfig.HordeServer;
+			Uri? server = (hordeConfig.HordeServer == null)? null : new Uri(hordeConfig.HordeServer);
 			string? token = hordeConfig.HordeToken;
-			string? oidcProvider = hordeConfig.HordeOidcProvider;
-
-			if (String.IsNullOrEmpty(server))
-			{
-				server = Environment.GetEnvironmentVariable("UE_HORDE_URL");
-				if (String.IsNullOrEmpty(server))
-				{
-					logger.LogInformation("Horde URL not specified in BuildConfiguration.xml or via UE_HORDE_URL environment variable");
-					return null;
-				}
-				if (String.IsNullOrEmpty(token))
-				{
-					token = Environment.GetEnvironmentVariable("UE_HORDE_TOKEN");
-				}
-			}
 
 			ConnectionMode? connectionMode = Enum.TryParse(hordeConfig.HordeConnectionMode, true, out ConnectionMode cm) ? cm : null;
 			Encryption? encryption = Enum.TryParse(hordeConfig.HordeEncryption, true, out Encryption enc) ? enc : null;
 
-			oidcProvider ??= Environment.GetEnvironmentVariable("UE_HORDE_OIDC_PROVIDER");
-
-			bool hasOidcProvider = !String.IsNullOrEmpty(oidcProvider);
-			logger.LogInformation("Horde URL: {Server}, Pool: {Pool}, Condition: {Condition}, OIDC: {OidcProvider}, Connection: {Connection} HordeEncryption: {Encryption}",
-				server, hordeConfig.HordePool ?? "(none)", hordeConfig.HordeCondition ?? "(none)", hasOidcProvider ? oidcProvider! : "Disabled", connectionMode?.ToString() ?? "(none)", encryption?.ToString() ?? "(none)");
+			logger.LogInformation("Horde URL: {Server}, Pool: {Pool}, Condition: {Condition}, Connection: {Connection} HordeEncryption: {Encryption}",
+				server, hordeConfig.HordePool ?? "(none)", hordeConfig.HordeCondition ?? "(none)", connectionMode?.ToString() ?? "(none)", encryption?.ToString() ?? "(none)");
 			try
 			{
-				if (String.IsNullOrEmpty(token) && hasOidcProvider)
-				{
-					token = await GetOidcBearerTokenAsync(null, oidcProvider!, logger, cancellationToken);
-				}
-
-				AuthenticationHeaderValue? authHeader = null;
-				if (!String.IsNullOrEmpty(token))
-				{
-					authHeader = new AuthenticationHeaderValue("Bearer", token);
-				}
-
 				bool allowWine = hordeConfig.bHordeAllowWine && OperatingSystem.IsWindows();
 
-				UBAHordeSession session = new(executor, new Uri(server), authHeader, hordeConfig.HordePool, allowWine, hordeConfig.HordeMaxCores, bStrictErrors, connectionMode, encryption, logger);
+				UBAHordeSession session = new(executor, server, token, hordeConfig.HordePool, allowWine, hordeConfig.HordeMaxCores, bStrictErrors, connectionMode, encryption, logger);
 				await session.InitAsync(useSentry: !String.IsNullOrEmpty(hordeConfig.UBASentryUrl), cancellationToken);
 				return session;
 			}
diff --git a/Engine/Source/Programs/UnrealBuildTool/Executors/UnrealBuildAccelerator/UnrealBuildAcceleratorHordeConfig.cs b/Engine/Source/Programs/UnrealBuildTool/Executors/UnrealBuildAccelerator/UnrealBuildAcceleratorHordeConfig.cs
--- a/Engine/Source/Programs/UnrealBuildTool/Executors/UnrealBuildAccelerator/UnrealBuildAcceleratorHordeConfig.cs
+++ b/Engine/Source/Programs/UnrealBuildTool/Executors/UnrealBuildAccelerator/UnrealBuildAcceleratorHordeConfig.cs
@@ -21,7 +21,7 @@ namespace UnrealBuildTool
 		public string? HordeServer { get; set; }
 
 		/// <summary>
-		/// Uri of the Horde server
+		/// Auth token for the Horde server
 		/// </summary>
 		[XmlConfigFile(Category = "Horde", Name = "Token")]
 		[CommandLine("-BoxHordeToken=")]
@@ -31,6 +31,7 @@ namespace UnrealBuildTool
 		/// <summary>
 		/// OIDC id for the login to use
 		/// </summary>
+		[Obsolete("The OidcProvider option is no longer used")]
 		[XmlConfigFile(Category = "Horde", Name = "OidcProvider")]
 		[CommandLine("-BoxHordeOidc=")]
 		[CommandLine("-UBAHordeOidc=")]
