Remove some obsolete command line args

#rnx

diff --git a/Engine/Source/Programs/UnrealBuildTool/Configuration/BuildConfiguration.cs b/Engine/Source/Programs/UnrealBuildTool/Configuration/BuildConfiguration.cs
--- a/Engine/Source/Programs/UnrealBuildTool/Configuration/BuildConfiguration.cs
+++ b/Engine/Source/Programs/UnrealBuildTool/Configuration/BuildConfiguration.cs
@@ -8,6 +8,7 @@ namespace UnrealBuildTool
 	/// <summary>
 	/// Global settings for building. Should not contain any target-specific settings.
 	/// </summary>
+	[System.Diagnostics.CodeAnalysis.SuppressMessage("Style", "IDE1006:Naming Styles", Justification = "UnrealBuildTool naming style")]
 	internal sealed class BuildConfiguration
 	{
 		/// <summary>
@@ -40,22 +41,13 @@ namespace UnrealBuildTool
 		public bool bAllowHybridExecutor { get; set; } = false;
 
 		/// <summary>
-		/// Priority order for remote executors (XGE, SNDBS, FASTBuild, UBAExecutor)
+		/// Priority order for remote executors (XGE, SNDBS, FASTBuild, UBA)
 		/// </summary>
 		[XmlConfigFile]
 		public string[] RemoteExecutorPriority { get; set; } = new string[] { "XGE", "SNDBS", "FASTBuild", "UBA" };
 
 		/// <summary>
-		/// This property is being renamed and should not be used
-		/// </summary>
-		[XmlConfigFile]
-		[CommandLine("-Box", Value = "true")]
-		[CommandLine("-NoBox", Value = "false")]
-		[Obsolete("Replace with bAllowUBAExecutor")]
-		public bool bAllowBoxExecutor { get => bAllowUBAExecutor; set => bAllowUBAExecutor = value; }
-
-		/// <summary>
-		/// Whether the experimental UnrealBuildAccelerator executor will be used.
+		/// Whether the UnrealBuildAccelerator executor will be used.
 		/// </summary>
 		[XmlConfigFile]
 		[CommandLine("-UBA", Value = "true")]
@@ -63,16 +55,7 @@ namespace UnrealBuildTool
 		public bool bAllowUBAExecutor { get; set; } = false;
 
 		/// <summary>
-		/// This property is being renamed and should not be used
-		/// </summary>
-		[XmlConfigFile]
-		[CommandLine("-BoxLocal", Value = "true")]
-		[CommandLine("-NoBoxLocal", Value = "false")]
-		[Obsolete("Replace with bAllowUBALocalExecutor")]
-		public bool bAllowBoxLocalExecutor { get => bAllowUBALocalExecutor; set => bAllowUBALocalExecutor = value; }
-
-		/// <summary>
-		/// Whether the experimental UnrealBuildAccelerator (local only) executor will be used.
+		/// Whether the UnrealBuildAccelerator (local only) executor will be used.
 		/// </summary>
 		[XmlConfigFile]
 		[CommandLine("-UBALocal", Value = "true")]
@@ -92,14 +75,14 @@ namespace UnrealBuildTool
 		public bool bAllowXGE { get; set; } = true;
 
 		/// <summary>
-		/// Whether FASTBuild may be used if availabe, default is true.
+		/// Whether FASTBuild may be used if available, default is true.
 		/// </summary>
 		[XmlConfigFile]
 		[CommandLine("-NoFASTBuild", Value = "false")]
 		public bool bAllowFASTBuild { get; set; } = true;
 
 		/// <summary>
-		/// Whether SN-DBS may be used if availabe, default is true.
+		/// Whether SN-DBS may be used if available, default is true.
 		/// </summary>
 		[XmlConfigFile]
 		[CommandLine("-NoSNDBS", Value = "false")]
@@ -213,7 +196,7 @@ namespace UnrealBuildTool
 		public int MaxRootPathLength { get; set; } = 50;
 
 		/// <summary>
-		/// Maximum length of a path relative to the root directory. Used on Windows to ensure paths are portable between machines. Defaults to off.
+		/// Maximum length of a path relative to the root directory. Used on Windows to ensure paths are portable between machines.
 		/// </summary>
 		[XmlConfigFile(Category = "WindowsPlatform")]
 		public int MaxNestedPathLength { get; set; } = 200;
diff --git a/Engine/Source/Programs/UnrealBuildTool/Executors/UnrealBuildAccelerator/UnrealBuildAcceleratorConfig.cs b/Engine/Source/Programs/UnrealBuildTool/Executors/UnrealBuildAccelerator/UnrealBuildAcceleratorConfig.cs
--- a/Engine/Source/Programs/UnrealBuildTool/Executors/UnrealBuildAccelerator/UnrealBuildAcceleratorConfig.cs
+++ b/Engine/Source/Programs/UnrealBuildTool/Executors/UnrealBuildAccelerator/UnrealBuildAcceleratorConfig.cs
@@ -15,7 +15,6 @@ namespace UnrealBuildTool
 		/// When set to true, UBA will not use any remote help
 		/// </summary>
 		[XmlConfigFile(Category = "UnrealBuildAccelerator")]
-		[CommandLine("-BoxDisableRemote")]
 		[CommandLine("-UBADisableRemote")]
 		public bool bDisableRemote { get; set; } = false;
 
@@ -23,7 +22,6 @@ namespace UnrealBuildTool
 		/// When set to true, UBA will force all actions that can be built remotely to be built remotely. This will hang if there are no remote agents available
 		/// </summary>
 		[XmlConfigFile(Category = "UnrealBuildAccelerator")]
-		[CommandLine("-BoxForceRemote")]
 		[CommandLine("-UBAForceRemote")]
 		public bool bForceBuildAllRemote { get; set; } = false;
 
@@ -31,7 +29,6 @@ namespace UnrealBuildTool
 		/// When set to true, actions that fail locally with UBA will be retried without UBA.
 		/// </summary>
 		[XmlConfigFile(Category = "UnrealBuildAccelerator")]
-		[CommandLine("-BoxForcedRetry")]
 		[CommandLine("-UBAForcedRetry")]
 		public bool bForcedRetry { get; set; } = false;
 
@@ -39,7 +36,6 @@ namespace UnrealBuildTool
 		/// When set to true, all errors and warnings from UBA will be output at the appropriate severity level to the log (rather than being output as 'information' and attempting to continue regardless).
 		/// </summary>
 		[XmlConfigFile(Category = "UnrealBuildAccelerator")]
-		[CommandLine("-BoxStrict")]
 		[CommandLine("-UBAStrict")]
 		public bool bStrict { get; set; } = false;
 
@@ -47,7 +43,6 @@ namespace UnrealBuildTool
 		/// If UBA should store cas compressed or raw
 		/// </summary>
 		[XmlConfigFile(Category = "UnrealBuildAccelerator")]
-		[CommandLine("-BoxStoreRaw")]
 		[CommandLine("-UBAStoreRaw")]
 		public bool bStoreRaw { get; set; } = false;
 
@@ -55,7 +50,6 @@ namespace UnrealBuildTool
 		/// If UBA should distribute linking to remote workers. This needs bandwidth but can be an optimization
 		/// </summary>
 		[XmlConfigFile(Category = "UnrealBuildAccelerator")]
-		[CommandLine("-BoxLinkRemote")]
 		[CommandLine("-UBALinkRemote")]
 		public bool bLinkRemote { get; set; } = false;
 
@@ -63,7 +57,6 @@ namespace UnrealBuildTool
 		/// The amount of gigabytes UBA is allowed to use to store workset and cached data. It is a good idea to have this >10gb
 		/// </summary>
 		[XmlConfigFile(Category = "UnrealBuildAccelerator")]
-		[CommandLine("-BoxStoreCapacityGb")]
 		[CommandLine("-UBAStoreCapacityGb")]
 		public int StoreCapacityGb { get; set; } = 40;
 
@@ -71,7 +64,6 @@ namespace UnrealBuildTool
 		/// Max number of worker threads that can handle messages from remotes. 
 		/// </summary>
 		[XmlConfigFile(Category = "UnrealBuildAccelerator")]
-		[CommandLine("-BoxMaxWorkers")]
 		[CommandLine("-UBAMaxWorkers")]
 		public int MaxWorkers { get; set; } = 192;
 
@@ -79,7 +71,6 @@ namespace UnrealBuildTool
 		/// Max size of each message sent from server to client
 		/// </summary>
 		[XmlConfigFile(Category = "UnrealBuildAccelerator")]
-		[CommandLine("-BoxSendSize")]
 		[CommandLine("-UBASendSize")]
 		public int SendSize { get; set; } = 256 * 1024;
 
@@ -87,7 +78,6 @@ namespace UnrealBuildTool
 		/// Which ip UBA server should listen to for connections
 		/// </summary>
 		[XmlConfigFile(Category = "UnrealBuildAccelerator")]
-		[CommandLine("-BoxHost")]
 		[CommandLine("-UBAHost")]
 		public string Host { get; set; } = String.Empty;
 
@@ -95,7 +85,6 @@ namespace UnrealBuildTool
 		/// Which port UBA server should listen to for connections.
 		/// </summary>
 		[XmlConfigFile(Category = "UnrealBuildAccelerator")]
-		[CommandLine("-BoxPort")]
 		[CommandLine("-UBAPort")]
 		public int Port { get; set; } = 1345;
 
@@ -103,7 +92,6 @@ namespace UnrealBuildTool
 		/// Which directory to store files for UBA.
 		/// </summary>
 		[XmlConfigFile(Category = "UnrealBuildAccelerator")]
-		[CommandLine("-BoxRootDir")]
 		[CommandLine("-UBARootDir")]
 		public string? RootDir { get; set; } = null;
 
@@ -111,7 +99,6 @@ namespace UnrealBuildTool
 		/// Use Quic protocol instead of Tcp (experimental)
 		/// </summary>
 		[XmlConfigFile(Category = "UnrealBuildAccelerator")]
-		[CommandLine("-BoxQuic", Value = "true")]
 		[CommandLine("-UBAQuic", Value = "true")]
 		public bool bUseQuic { get; set; } = false;
 
@@ -119,7 +106,6 @@ namespace UnrealBuildTool
 		/// Enable logging of UBA processes
 		/// </summary>
 		[XmlConfigFile(Category = "UnrealBuildAccelerator")]
-		[CommandLine("-BoxLog", Value = "true")]
 		[CommandLine("-UBALog", Value = "true")]
 		public bool bLogEnabled { get; set; } = false;
 
@@ -127,7 +113,6 @@ namespace UnrealBuildTool
 		/// Prints summary of UBA stats at end of build
 		/// </summary>
 		[XmlConfigFile(Category = "UnrealBuildAccelerator")]
-		[CommandLine("-BoxPrintSummary", Value = "true")]
 		[CommandLine("-UBAPrintSummary", Value = "true")]
 		public bool bPrintSummary { get; set; } = false;
 
@@ -135,7 +120,6 @@ namespace UnrealBuildTool
 		/// Launch visualizer application which shows build progress
 		/// </summary>
 		[XmlConfigFile(Category = "UnrealBuildAccelerator")]
-		[CommandLine("-BoxVisualizer", Value = "true")]
 		[CommandLine("-UBAVisualizer", Value = "true")]
 		public bool bLaunchVisualizer { get; set; } = false;
 
@@ -143,7 +127,6 @@ namespace UnrealBuildTool
 		/// Resets the cas cache
 		/// </summary>
 		[XmlConfigFile(Category = "UnrealBuildAccelerator")]
-		[CommandLine("-BoxResetCas", Value = "true")]
 		[CommandLine("-UBAResetCas", Value = "true")]
 		public bool bResetCas { get; set; } = false;
 
@@ -151,7 +134,6 @@ namespace UnrealBuildTool
 		/// Provide custom path for trace output file
 		/// </summary>
 		[XmlConfigFile(Category = "UnrealBuildAccelerator")]
-		[CommandLine("-BoxTraceOutputFile")]
 		[CommandLine("-UBATraceOutputFile")]
 		public string TraceFile { get; set; } = String.Empty;
 
@@ -159,7 +141,6 @@ namespace UnrealBuildTool
 		/// Add verbose details to the UBA trace
 		/// </summary>
 		[XmlConfigFile(Category = "UnrealBuildAccelerator")]
-		[CommandLine("-BoxDetailedTrace", Value = "true")]
 		[CommandLine("-UBADetailedTrace", Value = "true")]
 		public bool bDetailedTrace { get; set; }
 
@@ -174,7 +155,6 @@ namespace UnrealBuildTool
 		/// Let UBA kill running processes when close to out of memory
 		/// </summary>
 		[XmlConfigFile(Category = "UnrealBuildAccelerator")]
-		[CommandLine("-BoxAllowKillOnMem", Value = "true")]
 		[CommandLine("-UBAAllowKillOnMem", Value = "true")]
 		public bool bAllowKillOnMem { get; set; }
 
@@ -182,7 +162,6 @@ namespace UnrealBuildTool
 		/// Threshold for when executor should output logging for the process. Defaults to never
 		/// </summary>
 		[XmlConfigFile(Category = "UnrealBuildAccelerator")]
-		[CommandLine("-BoxOutputStatsThresholdMs")]
 		[CommandLine("-UBAOutputStatsThresholdMs")]
 		public int OutputStatsThresholdMs { get; set; } = Int32.MaxValue;
 
@@ -190,7 +169,6 @@ namespace UnrealBuildTool
 		/// Skip writing intermediate and output files to disk. Useful for validation builds where we don't need the output
 		/// </summary>
 		[XmlConfigFile(Category = "UnrealBuildAccelerator")]
-		[CommandLine("-BoxNoWrite", Value = "false")]
 		[CommandLine("-UBANoWrite", Value = "false")]
 		public bool bWriteToDisk { get; set; } = true;
 
@@ -198,7 +176,6 @@ namespace UnrealBuildTool
 		/// Set to true to disable mimalloc and detouring of memory allocations.
 		/// </summary>
 		[XmlConfigFile(Category = "UnrealBuildAccelerator")]
-		[CommandLine("-BoxNoCustoMalloc", Value = "true")]
 		[CommandLine("-UBANoCustoMalloc", Value = "true")]
 		public bool bDisableCustomAlloc { get; set; } = false;
 
@@ -206,7 +183,6 @@ namespace UnrealBuildTool
 		/// The zone to use for UBA.
 		/// </summary>
 		[XmlConfigFile(Category = "UnrealBuildAccelerator")]
-		[CommandLine("-BoxZone=")]
 		[CommandLine("-UBAZone=")]
 		public string Zone { get; set; } = String.Empty;
 
@@ -214,7 +190,6 @@ namespace UnrealBuildTool
 		/// Set to true to enable encryption when transfering files over the network.
 		/// </summary>
 		[XmlConfigFile(Category = "UnrealBuildAccelerator")]
-		[CommandLine("-BoxCrypto", Value = "true")]
 		[CommandLine("-UBACrypto", Value = "true")]
 		public bool bUseCrypto { get; set; } = false;
 
diff --git a/Engine/Source/Programs/UnrealBuildTool/Executors/UnrealBuildAccelerator/UnrealBuildAcceleratorHordeConfig.cs b/Engine/Source/Programs/UnrealBuildTool/Executors/UnrealBuildAccelerator/UnrealBuildAcceleratorHordeConfig.cs
--- a/Engine/Source/Programs/UnrealBuildTool/Executors/UnrealBuildAccelerator/UnrealBuildAcceleratorHordeConfig.cs
+++ b/Engine/Source/Programs/UnrealBuildTool/Executors/UnrealBuildAccelerator/UnrealBuildAcceleratorHordeConfig.cs
@@ -16,7 +16,6 @@ namespace UnrealBuildTool
 		/// Uri of the Horde server
 		/// </summary>
 		[XmlConfigFile(Category = "Horde", Name = "Server")]
-		[CommandLine("-BoxHorde=")]
 		[CommandLine("-UBAHorde=")]
 		public string? HordeServer { get; set; }
 
@@ -24,7 +23,6 @@ namespace UnrealBuildTool
 		/// Auth token for the Horde server
 		/// </summary>
 		[XmlConfigFile(Category = "Horde", Name = "Token")]
-		[CommandLine("-BoxHordeToken=")]
 		[CommandLine("-UBAHordeToken=")]
 		public string? HordeToken { get; set; }
 
@@ -33,7 +31,6 @@ namespace UnrealBuildTool
 		/// </summary>
 		[Obsolete("The OidcProvider option is no longer used")]
 		[XmlConfigFile(Category = "Horde", Name = "OidcProvider")]
-		[CommandLine("-BoxHordeOidc=")]
 		[CommandLine("-UBAHordeOidc=")]
 		public string? HordeOidcProvider { get; set; }
 
@@ -76,7 +73,6 @@ namespace UnrealBuildTool
 		/// <summary>
 		/// Pool for the Horde agent to assign, only used for commandline override
 		/// </summary>
-		[CommandLine("-BoxHordePool=")]
 		[CommandLine("-UBAHordePool=")]
 		public string? OverrideHordePool { get; set; }
 
@@ -102,7 +98,6 @@ namespace UnrealBuildTool
 		/// Requirements for the Horde agent to assign
 		/// </summary>
 		[XmlConfigFile(Category = "Horde", Name = "Requirements")]
-		[CommandLine("-BoxHordeRequirements=")]
 		[CommandLine("-UBAHordeRequirements=")]
 		public string? HordeCondition { get; set; }
 
@@ -110,7 +105,6 @@ namespace UnrealBuildTool
 		/// Which ip UBA server should give to agents. This will invert so host listens and agents connect
 		/// </summary>
 		[XmlConfigFile(Category = "Horde", Name = "LocalHost")]
-		[CommandLine("-BoxHordeHost")]
 		[CommandLine("-UBAHordeHost")]
 		public string HordeHost { get; set; } = String.Empty;
 
@@ -118,7 +112,6 @@ namespace UnrealBuildTool
 		/// Max cores allowed to be used by build session
 		/// </summary>
 		[XmlConfigFile(Category = "UnrealBuildAccelerator")]
-		[CommandLine("-BoxHordeMaxCores")]
 		[CommandLine("-UBAHordeMaxCores")]
 		public int HordeMaxCores { get; set; } = 576;
 
@@ -126,7 +119,6 @@ namespace UnrealBuildTool
 		/// How long UBT should wait to ask for help. Useful in build configs where machine can delay remote work and still get same wall time results (pch dependencies etc)
 		/// </summary>
 		[XmlConfigFile(Category = "UnrealBuildAccelerator")]
-		[CommandLine("-BoxHordeDelay")]
 		[CommandLine("-UBAHordeDelay")]
 		public int HordeDelay { get; set; } = 0;
 
@@ -134,7 +126,6 @@ namespace UnrealBuildTool
 		/// Allow use of Wine. Only applicable to Horde agents running Linux. Can still be ignored if Wine executable is not set on agent.
 		/// </summary>
 		[XmlConfigFile(Category = "UnrealBuildAccelerator")]
-		[CommandLine("-BoxHordeAllowWine", Value = "true")]
 		[CommandLine("-UBAHordeAllowWine", Value = "true")]
 		public bool bHordeAllowWine { get; set; } = true;
 
@@ -143,7 +134,6 @@ namespace UnrealBuildTool
 		/// <see cref="ConnectionMode" /> for valid modes.
 		/// </summary>
 		[XmlConfigFile(Category = "Horde", Name = "ConnectionMode")]
-		[CommandLine("-BoxHordeConnectionMode=")]
 		[CommandLine("-UBAHordeConnectionMode=")]
 		public string? HordeConnectionMode { get; set; }
 
@@ -152,7 +142,6 @@ namespace UnrealBuildTool
 		/// <see cref="Encryption" /> for valid modes.
 		/// </summary>
 		[XmlConfigFile(Category = "Horde", Name = "Encryption")]
-		[CommandLine("-BoxHordeEncryption=")]
 		[CommandLine("-UBAHordeEncryption=")]
 		public string? HordeEncryption { get; set; }
 
@@ -160,7 +149,6 @@ namespace UnrealBuildTool
 		/// Sentry URL to send box data to. Optional.
 		/// </summary>
 		[XmlConfigFile(Category = "Horde")]
-		[CommandLine("-BoxSentryUrl=")]
 		[CommandLine("-UBASentryUrl=")]
 		public string? UBASentryUrl { get; set; }
 
@@ -168,7 +156,6 @@ namespace UnrealBuildTool
 		/// Disable horde all together
 		/// </summary>
 		[XmlConfigFile(Category = "UnrealBuildAccelerator")]
-		[CommandLine("-BoxDisableHorde")]
 		[CommandLine("-UBADisableHorde")]
 		public bool bDisableHorde { get; set; } = false;
 	}
diff --git a/Engine/Source/Programs/UnrealBuildTool/System/Action.cs b/Engine/Source/Programs/UnrealBuildTool/System/Action.cs
--- a/Engine/Source/Programs/UnrealBuildTool/System/Action.cs
+++ b/Engine/Source/Programs/UnrealBuildTool/System/Action.cs
@@ -462,10 +462,10 @@ namespace UnrealBuildTool
 				Action.bCanExecuteRemotelyWithXGE = bCanExecuteRemotelyWithXGE;
 			}
 
-			bool bCanExecuteInBox;
-			if (Object.TryGetBoolField("bCanExecuteInUBA", out bCanExecuteInBox))
+			bool bCanExecuteInUBA;
+			if (Object.TryGetBoolField("bCanExecuteInUBA", out bCanExecuteInUBA))
 			{
-				Action.bCanExecuteInUBA = bCanExecuteInBox;
+				Action.bCanExecuteInUBA = bCanExecuteInUBA;
 			}
 
 			bool bIsGCCCompiler;
