Minor code cleanup

#rnx

diff --git a/Engine/Source/Programs/UnrealBuildTool/Executors/UnrealBuildAccelerator/UBAAgentCoordinatorHorde.cs b/Engine/Source/Programs/UnrealBuildTool/Executors/UnrealBuildAccelerator/UBAAgentCoordinatorHorde.cs
--- a/Engine/Source/Programs/UnrealBuildTool/Executors/UnrealBuildAccelerator/UBAAgentCoordinatorHorde.cs
+++ b/Engine/Source/Programs/UnrealBuildTool/Executors/UnrealBuildAccelerator/UBAAgentCoordinatorHorde.cs
@@ -195,7 +195,7 @@ namespace UnrealBuildTool
 
 		async Task<BlobLocator> CreateToolAsync(DirectoryReference baseDir, IEnumerable<FileReference> files, CancellationToken cancellationToken)
 		{
-			BlobSerializerOptions serializerOptions = new BlobSerializerOptions();
+			BlobSerializerOptions serializerOptions = new();
 			serializerOptions.Converters.Add(new InteriorChunkedDataNodeConverter(2)); // Lock to v2 for now. Could change based on protocol version.
 
 			await using IBlobWriter writer = _storage.CreateBlobWriter(serializerOptions: serializerOptions);
diff --git a/Engine/Source/Programs/UnrealBuildTool/Executors/UnrealBuildAccelerator/UBAExecutor.cs b/Engine/Source/Programs/UnrealBuildTool/Executors/UnrealBuildAccelerator/UBAExecutor.cs
--- a/Engine/Source/Programs/UnrealBuildTool/Executors/UnrealBuildAccelerator/UBAExecutor.cs
+++ b/Engine/Source/Programs/UnrealBuildTool/Executors/UnrealBuildAccelerator/UBAExecutor.cs
@@ -23,7 +23,7 @@ namespace UnrealBuildTool
 		public string Crypto { get; private set; } = String.Empty;
 		public IServer? Server { get; private set; }
 		ISessionServer? _session;
-		readonly List<IUBAAgentCoordinator> _agentCoordinators = new List<IUBAAgentCoordinator>();
+		readonly List<IUBAAgentCoordinator> _agentCoordinators = new();
 		DirectoryReference? _rootDirRef;
 		bool _bIsCancelled;
 		bool _bIsRemoteActionsAllowed = true;
@@ -84,8 +84,8 @@ namespace UnrealBuildTool
 			{
 				_threadedLogger.LogError("UBA actions output file needs to have extension .yaml for UbaCli to understand it");
 			}
-			using System.IO.StreamWriter streamWriter = new System.IO.StreamWriter(UBAConfig.ActionsOutputFile);
-			using System.CodeDom.Compiler.IndentedTextWriter writer = new System.CodeDom.Compiler.IndentedTextWriter(streamWriter, "  ");
+			using System.IO.StreamWriter streamWriter = new(UBAConfig.ActionsOutputFile);
+			using System.CodeDom.Compiler.IndentedTextWriter writer = new(streamWriter, "  ");
 			await writer.WriteAsync("environment: ");
 			await writer.WriteLineAsync(Environment.GetEnvironmentVariable("PATH"));
 			await writer.WriteLineAsync("processes:");
@@ -163,7 +163,7 @@ namespace UnrealBuildTool
 				_rootDirRef = DirectoryReference.FromString(Environment.GetEnvironmentVariable("UBA_ROOT") ?? Environment.GetEnvironmentVariable("BOX_ROOT"));
 			}
 
-			List<Task> coordinatorInitTasks = new List<Task>();
+			List<Task> coordinatorInitTasks = new();
 			if (!UBAConfig.bDisableRemote)
 			{
 				foreach (IUBAAgentCoordinator coordinator in _agentCoordinators)
@@ -292,7 +292,7 @@ namespace UnrealBuildTool
 				// Check if a listening visualizer is already running
 				foreach (System.Diagnostics.Process process in System.Diagnostics.Process.GetProcessesByName(visaulizerPath.GetFileNameWithoutAnyExtensions()))
 				{
-					using ManagementObjectSearcher searcher = new ManagementObjectSearcher($"SELECT CommandLine FROM Win32_Process WHERE ProcessId = {process.Id}");
+					using ManagementObjectSearcher searcher = new($"SELECT CommandLine FROM Win32_Process WHERE ProcessId = {process.Id}");
 					using ManagementObjectCollection objects = searcher.Get();
 					string args = objects.Cast<ManagementBaseObject>().SingleOrDefault()?["CommandLine"]?.ToString() ?? "";
 					if (args.Contains("-listen", StringComparison.OrdinalIgnoreCase))
diff --git a/Engine/Source/Programs/UnrealBuildTool/Executors/UnrealBuildAccelerator/UnrealBuildAcceleratorConfig.cs b/Engine/Source/Programs/UnrealBuildTool/Executors/UnrealBuildAccelerator/UnrealBuildAcceleratorConfig.cs
--- a/Engine/Source/Programs/UnrealBuildTool/Executors/UnrealBuildAccelerator/UnrealBuildAcceleratorConfig.cs
+++ b/Engine/Source/Programs/UnrealBuildTool/Executors/UnrealBuildAccelerator/UnrealBuildAcceleratorConfig.cs
@@ -176,7 +176,7 @@ namespace UnrealBuildTool
 		/// Set to true to disable mimalloc and detouring of memory allocations.
 		/// </summary>
 		[XmlConfigFile(Category = "UnrealBuildAccelerator")]
-		[CommandLine("-UBANoCustoMalloc", Value = "true")]
+		[CommandLine("-UBANoCustomMalloc", Value = "true")]
 		public bool bDisableCustomAlloc { get; set; } = false;
 
 		/// <summary>
@@ -187,7 +187,7 @@ namespace UnrealBuildTool
 		public string Zone { get; set; } = String.Empty;
 
 		/// <summary>
-		/// Set to true to enable encryption when transfering files over the network.
+		/// Set to true to enable encryption when transferring files over the network.
 		/// </summary>
 		[XmlConfigFile(Category = "UnrealBuildAccelerator")]
 		[CommandLine("-UBACrypto", Value = "true")]
