Modifying how we grab pgd and pgc files on Windows to account for the new merged pgd pipeline.

diff --git a/Engine/Source/Programs/UnrealBuildTool/Platform/Windows/UEBuildWindows.cs b/Engine/Source/Programs/UnrealBuildTool/Platform/Windows/UEBuildWindows.cs
--- a/Engine/Source/Programs/UnrealBuildTool/Platform/Windows/UEBuildWindows.cs
+++ b/Engine/Source/Programs/UnrealBuildTool/Platform/Windows/UEBuildWindows.cs
@@ -232,6 +232,13 @@ namespace UnrealBuildTool
 		[CommandLine("-PGOFastGen")]
 		public bool bUseFastGenProfile = false;
 
+		/// <summary>
+		/// If specified along with -PGOOptimize, will use the specified per-merged pgd file instead of the usual pgd file with loose pgc files.
+		/// </summary>
+		[XmlConfigFile(Category = "WindowsPlatform")]
+		[CommandLine("-PGOMergedPGD")]
+		public string? PreMergedPgdFilename = null;
+
 		/// <summary>
 		/// If specified along with -PGOProfile, prevent the usage of extra counters. Please note that by default /FASTGENPROFILE doesnt use extra counters
 		/// </summary>
@@ -793,6 +800,8 @@ namespace UnrealBuildTool
 
 		public bool bUseFastGenProfile => Inner.bUseFastGenProfile;
 
+		public string? PreMergedPgdFilename => Inner.PreMergedPgdFilename;
+
 		public bool bPGONoExtraCounters => Inner.bPGONoExtraCounters;
 
 		public bool bSampleBasedPGO => Inner.bSampleBasedPGO;
@@ -1699,6 +1708,7 @@ namespace UnrealBuildTool
 
 				LinkEnvironment.PGODirectory = CompileEnvironment.PGODirectory;
 				LinkEnvironment.PGOFilenamePrefix = CompileEnvironment.PGOFilenamePrefix;
+				LinkEnvironment.PGOMergedFilenamePrefix = Target.Platform == UnrealTargetPlatform.Win64 ? Target.WindowsPlatform.PreMergedPgdFilename : null;
 			}
 
 			CompileEnvironment.Definitions.Add("WINDOWS_MAX_NUM_TLS_SLOTS=" + Target.WindowsPlatform.MaxNumTlsSlots.ToString());
diff --git a/Engine/Source/Programs/UnrealBuildTool/Platform/Windows/VCToolChain.cs b/Engine/Source/Programs/UnrealBuildTool/Platform/Windows/VCToolChain.cs
--- a/Engine/Source/Programs/UnrealBuildTool/Platform/Windows/VCToolChain.cs
+++ b/Engine/Source/Programs/UnrealBuildTool/Platform/Windows/VCToolChain.cs
@@ -2711,27 +2711,51 @@ namespace UnrealBuildTool
 
 				// prefer a PGD file that matches the output file
 				string PGDFile = Path.Combine(LinkEnvironment.PGODirectory!, LinkEnvironment.PGOFilenamePrefix + ".pgd");
-				if (!File.Exists(PGDFile))
-				{
-					string[] PGDFiles = Directory.GetFiles(LinkEnvironment.PGODirectory!, "*.pgd");
-					if (PGDFiles.Length > 1)
+				string[] PGCFiles = {};
+				
+				bool bUsingMergedPGD = false;
+				
+				// check if we are using a pre-merged pgd file, if so use it instead
+				if (LinkEnvironment.PGOMergedFilenamePrefix != null)
+				{
+					string MergedPGDFile = Path.Combine(LinkEnvironment.PGODirectory!, LinkEnvironment.PGOMergedFilenamePrefix + ".pgd");
+					if (File.Exists(MergedPGDFile))
 					{
-						throw new BuildException("More than one .pgd file found in \"{0}\" and \"{1}\" not found ", LinkEnvironment.PGODirectory, PGDFile);
+						// Only use the merged pgd file if it actually exists, otherwise keep the default behavior
+						PGDFile = MergedPGDFile;
+						bUsingMergedPGD = true;
 					}
-					else if (PGDFiles.Length == 0)
+					else
 					{
-						Logger.LogWarning("No .pgd files found in \"{PgoDir}\".", LinkEnvironment.PGODirectory);
-						return false;
+						Logger.LogWarning("The specified merged .pgd file \"{MergedPgdFile}\" was not found.", LinkEnvironment.PGOMergedFilenamePrefix);
 					}
-
-					PGDFile = PGDFiles.First();
 				}
 
-				string[] PGCFiles = Directory.GetFiles(LinkEnvironment.PGODirectory!, "*.pgc");
-				if (PGCFiles.Length == 0)
+				if (!bUsingMergedPGD)
 				{
-					Logger.LogWarning("No .pgc files found in \"{PgoDir}\".", LinkEnvironment.PGODirectory);
-					return false;
+					if (!File.Exists(PGDFile))
+					{
+						string[] PGDFiles = Directory.GetFiles(LinkEnvironment.PGODirectory!, "*.pgd");
+						if (PGDFiles.Length > 1)
+						{
+							throw new BuildException("More than one .pgd file found in \"{0}\" and \"{1}\" not found ", LinkEnvironment.PGODirectory,
+								PGDFile);
+						}
+						else if (PGDFiles.Length == 0)
+						{
+							Logger.LogWarning("No .pgd files found in \"{PgoDir}\".", LinkEnvironment.PGODirectory);
+							return false;
+						}
+
+						PGDFile = PGDFiles.First();
+					}
+
+					PGCFiles = Directory.GetFiles(LinkEnvironment.PGODirectory!, "*.pgc");
+					if (PGCFiles.Length == 0)
+					{
+						Logger.LogWarning("No .pgc files found in \"{PgoDir}\".", LinkEnvironment.PGODirectory);
+						return false;
+					}
 				}
 
 				// Make sure the destination directory exists!
diff --git a/Engine/Source/Programs/UnrealBuildTool/System/LinkEnvironment.cs b/Engine/Source/Programs/UnrealBuildTool/System/LinkEnvironment.cs
--- a/Engine/Source/Programs/UnrealBuildTool/System/LinkEnvironment.cs
+++ b/Engine/Source/Programs/UnrealBuildTool/System/LinkEnvironment.cs
@@ -254,6 +254,11 @@ namespace UnrealBuildTool
 		/// </summary>
 		public string? PGOFilenamePrefix;
 
+		/// <summary>
+		/// Platform specific filename where pre-merged PGO profiling data is saved.
+		/// </summary>
+		public string? PGOMergedFilenamePrefix;
+
 		/// <summary>
 		/// Whether to request the linker create a map file as part of the build
 		/// </summary>
