[IWYU]
* Added special handling for TSharedPtr/TSharedRef to prevent forward declaration of these types
* Fixed so PlatformSurvey.h is always used and not XxPlatformSurvey.h

diff --git a/Engine/Source/Programs/UnrealBuildTool/Modes/IWYUMode.cs b/Engine/Source/Programs/UnrealBuildTool/Modes/IWYUMode.cs
--- a/Engine/Source/Programs/UnrealBuildTool/Modes/IWYUMode.cs
+++ b/Engine/Source/Programs/UnrealBuildTool/Modes/IWYUMode.cs
@@ -825,6 +825,7 @@ namespace UnrealBuildTool
 					GetInfo("VerseInteropUtils.h", "Restricted/NotForLicensees/Plugins/Solaris/Source/VerseNative/Public/VerseInteropTypes.h"),
 					GetInfo("Containers/ContainersFwd.h", "Source/Runtime/Core/Public/Containers/ContainersFwd.h"),
 					GetInfo("Misc/OptionalFwd.h", "Source/Runtime/Core/Public/Misc/OptionalFwd.h"),
+					GetInfo("Templates/SharedPointerFwd.h", "Source/Runtime/Core/Public/Templates/SharedPointerFwd.h"),
 				}
 			);
 
@@ -835,6 +836,8 @@ namespace UnrealBuildTool
 				{ "TArray", SpecialIncludes["Containers/ContainersFwd.h"] },
 				{ "TArrayView", SpecialIncludes["Containers/ContainersFwd.h"] },
 				{ "TOptional", SpecialIncludes["Misc/OptionalFwd.h"] },
+				{ "TSharedPtr", SpecialIncludes["Templates/SharedPointerFwd.h"] },
+				{ "TSharedRef", SpecialIncludes["Templates/SharedPointerFwd.h"] },
 			};
 
 			// Add all .generated.h files as entries in the lookup and explicitly add the includes they have which will never be removed
