Fixing the LinkEnvironment copy constructor to resolve an issue with the merged pgd filename not being passed to the VC toolchain.

#tests Ran a PF with the change, issue of the missing value no longer occurs.

diff --git a/Engine/Source/Programs/UnrealBuildTool/System/LinkEnvironment.cs b/Engine/Source/Programs/UnrealBuildTool/System/LinkEnvironment.cs
--- a/Engine/Source/Programs/UnrealBuildTool/System/LinkEnvironment.cs
+++ b/Engine/Source/Programs/UnrealBuildTool/System/LinkEnvironment.cs
@@ -424,6 +424,7 @@ namespace UnrealBuildTool
 			bPGOProfile = Other.bPGOProfile;
 			PGODirectory = Other.PGODirectory;
 			PGOFilenamePrefix = Other.PGOFilenamePrefix;
+			PGOMergedFilenamePrefix = Other.PGOMergedFilenamePrefix;
 			bCreateMapFile = Other.bCreateMapFile;
 			bUsePDBFiles = Other.bUsePDBFiles;
 			bUseFastPDBLinking = Other.bUseFastPDBLinking;
