Disabling PGOMGR warning for missing pgc files when using a pre-merged pgd file.

diff --git a/Engine/Source/Programs/UnrealBuildTool/Platform/Windows/VCToolChain.cs b/Engine/Source/Programs/UnrealBuildTool/Platform/Windows/VCToolChain.cs
--- a/Engine/Source/Programs/UnrealBuildTool/Platform/Windows/VCToolChain.cs
+++ b/Engine/Source/Programs/UnrealBuildTool/Platform/Windows/VCToolChain.cs
@@ -2724,6 +2724,7 @@ namespace UnrealBuildTool
 						// Only use the merged pgd file if it actually exists, otherwise keep the default behavior
 						PGDFile = MergedPGDFile;
 						bUsingMergedPGD = true;
+						Environment.SetEnvironmentVariable("PGOMGR", "/nowarn:188", EnvironmentVariableTarget.Process);
 					}
 					else
 					{
