UnrealBuildTool: Add support for studio telemetry. Requires connection info to be defined in a BaseEngine.ini config, as well as opting into those providers via a BuildConfiguration.xml setting

diff --git a/Engine/Source/Programs/Shared/EpicGames.Core/Telemetry/DataRouterTelemetryService.cs b/Engine/Source/Programs/Shared/EpicGames.Core/Telemetry/DataRouterTelemetryService.cs
--- a/Engine/Source/Programs/Shared/EpicGames.Core/Telemetry/DataRouterTelemetryService.cs
+++ b/Engine/Source/Programs/Shared/EpicGames.Core/Telemetry/DataRouterTelemetryService.cs
@@ -22,32 +22,32 @@ namespace EpicGames.Core.Telemetry
 		/// <summary>
 		/// The app's id
 		/// </summary>
-		protected string AppId { get; } = String.Empty;
+		protected string AppId { get; set; } = String.Empty;
 
 		/// <summary>
 		/// The app's version
 		/// </summary>
-		protected string AppVersion { get; } = String.Empty;
+		protected string AppVersion { get; set; } = String.Empty;
 
 		/// <summary>
 		/// The app's environment
 		/// </summary>
-		protected string AppEnvironment { get; } = String.Empty;
+		protected string AppEnvironment { get; set; } = String.Empty;
 
 		/// <summary>
 		/// The upload type
 		/// </summary>
-		protected string UploadType { get; } = String.Empty;
+		protected string UploadType { get; set; } = String.Empty;
 
 		/// <summary>
 		///  The user's id
 		/// </summary>
-		protected string UserId { get; } = String.Empty;
+		protected string UserId { get; set; } = String.Empty;
 
 		/// <summary>
 		/// The session's id
 		/// </summary>
-		protected string SessionId { get; } = String.Empty;
+		protected string SessionId { get; set; } = String.Empty;
 
 		/// <summary>
 		/// The TimeStamp of when the event was added to queue UNTIL we serialize so we use the name TimeStamp in code, and serialize as DateOffset
@@ -68,6 +68,15 @@ namespace EpicGames.Core.Telemetry
 			get;
 		}
 
+		/// <summary>
+		/// Constructor
+		/// </summary>
+		/// <param name="inTimeStamp"></param>
+		protected DataRouterEvent(DateTime? inTimeStamp = null)
+		{
+			TimeStamp = inTimeStamp ?? DateTime.UtcNow;
+		}
+
 		/// <summary>
 		/// Constructor
 		/// </summary>
@@ -214,6 +223,67 @@ namespace EpicGames.Core.Telemetry
 		}
 	}
 
+	/// <summary>
+	/// DateTime converter for data router events to convert to Unix timestamp
+	/// </summary>
+	public class UnixTimestampDataRouterEventConverter : JsonConverter<DateTime>
+	{
+		/// <summary>
+		/// Reader
+		/// </summary>
+		/// <param name="reader"></param>
+		/// <param name="typeToConvert"></param>
+		/// <param name="options"></param>
+		/// <returns></returns>
+		/// <exception cref="NotImplementedException"></exception>
+		public override DateTime Read(ref Utf8JsonReader reader, Type typeToConvert, JsonSerializerOptions options)
+		{
+			throw new NotImplementedException();
+		}
+
+		/// <summary>
+		/// Writer
+		/// </summary>
+		/// <param name="writer"></param>
+		/// <param name="value"></param>
+		/// <param name="options"></param>
+		public override void Write(Utf8JsonWriter writer, DateTime value, JsonSerializerOptions options)
+		{
+			TimeSpan diff = value.ToUniversalTime() - DateTime.UnixEpoch;
+			writer.WriteNumberValue(diff.TotalSeconds);
+		}
+	}
+
+	/// <summary>
+	/// TimeSpan converter for data router events to convert to seconds
+	/// </summary>
+	public class DurationDataRouterEventConverter : JsonConverter<TimeSpan>
+	{
+		/// <summary>
+		/// Reader
+		/// </summary>
+		/// <param name="reader"></param>
+		/// <param name="typeToConvert"></param>
+		/// <param name="options"></param>
+		/// <returns></returns>
+		/// <exception cref="NotImplementedException"></exception>
+		public override TimeSpan Read(ref Utf8JsonReader reader, Type typeToConvert, JsonSerializerOptions options)
+		{
+			throw new NotImplementedException();
+		}
+
+		/// <summary>
+		/// Writer
+		/// </summary>
+		/// <param name="writer"></param>
+		/// <param name="value"></param>
+		/// <param name="options"></param>
+		public override void Write(Utf8JsonWriter writer, TimeSpan value, JsonSerializerOptions options)
+		{
+			writer.WriteNumberValue(value.TotalSeconds);
+		}
+	}
+
 	/// <summary>
 	/// Implementation for TelemetryService using DataRouter.
 	/// FlushEvents() should be called before application exits otherwise some events may remain unsent
diff --git a/Engine/Source/Programs/UnrealBuildTool/Configuration/TargetDescriptor.cs b/Engine/Source/Programs/UnrealBuildTool/Configuration/TargetDescriptor.cs
--- a/Engine/Source/Programs/UnrealBuildTool/Configuration/TargetDescriptor.cs
+++ b/Engine/Source/Programs/UnrealBuildTool/Configuration/TargetDescriptor.cs
@@ -607,6 +607,10 @@ namespace UnrealBuildTool
 					}
 				}
 			}
+
+			// Register any found descriptors for telemetry, using the first as the primary target
+			TelemetryService.Get().SetPrimaryTargetDetails(TargetDescriptors.FirstOrDefault());
+			TargetDescriptors.ForEach(x => TelemetryService.Get().AddEndpointsFromConfig(x.ProjectFile?.Directory));
 		}
 
 		/// <summary>
diff --git a/Engine/Source/Programs/UnrealBuildTool/Executors/Experimental/FASTBuild.cs b/Engine/Source/Programs/UnrealBuildTool/Executors/Experimental/FASTBuild.cs
--- a/Engine/Source/Programs/UnrealBuildTool/Executors/Experimental/FASTBuild.cs
+++ b/Engine/Source/Programs/UnrealBuildTool/Executors/Experimental/FASTBuild.cs
@@ -309,6 +309,11 @@ namespace UnrealBuildTool
 			return false;
 		}
 
+		private TelemetryExecutorEvent? telemetryEvent;
+
+		/// <inheritdoc/>
+		public override TelemetryExecutorEvent? GetTelemetryEvent() => telemetryEvent ?? LocalExecutor.GetTelemetryEvent();
+
 		//////////////////////////////////////////
 		// Action Helpers
 
@@ -476,13 +481,15 @@ namespace UnrealBuildTool
 			IEnumerable<LinkedAction> CompileActions = Actions.Where(Action => Action.ActionType == ActionType.Compile && Action.bCanExecuteRemotely && Action.bCanExecuteRemotelyWithSNDBS);
 			if (CompileActions.Any() && DetectBuildType(CompileActions, Logger))
 			{
+				DateTime startTimeUTC = DateTime.UtcNow;
 				string FASTBuildFilePath = Path.Combine(Unreal.EngineDirectory.FullName, "Intermediate", "Build", "fbuild.bff");
 				if (!CreateBffFile(Actions, FASTBuildFilePath, Logger))
 				{
 					return false;
 				}
 
-				return ExecuteBffFile(FASTBuildFilePath, Logger);
+				bool result = ExecuteBffFile(FASTBuildFilePath, Logger);
+				telemetryEvent = new TelemetryExecutorEvent(Name, startTimeUTC, result, Actions.Count(), -1, -1, DateTime.UtcNow);
 			}
 
 			return await LocalExecutor.ExecuteActionsAsync(Actions, Logger, actionArtifactCache);
diff --git a/Engine/Source/Programs/UnrealBuildTool/Executors/ImmediateActionQueue.cs b/Engine/Source/Programs/UnrealBuildTool/Executors/ImmediateActionQueue.cs
--- a/Engine/Source/Programs/UnrealBuildTool/Executors/ImmediateActionQueue.cs
+++ b/Engine/Source/Programs/UnrealBuildTool/Executors/ImmediateActionQueue.cs
@@ -507,6 +507,19 @@ namespace UnrealBuildTool
 			return Success;
 		}
 
+		/// <summary>
+		/// Get the number of actions that were run, and how many succeeded or failed
+		/// </summary>
+		/// <param name="totalActions">Out parameter, the total number of actions</param>
+		/// <param name="succeededActions">Out parameter, the number of successful actions</param>
+		/// <param name="failedActions">Out parameter, the number of failed actions</param>
+		public void GetActionResultCounts(out int totalActions, out int succeededActions, out int failedActions)
+		{
+			totalActions = Actions.Length;
+			succeededActions = Actions.Where(x => x.Results?.ExitCode == 0).Count();
+			failedActions = Actions.Where(x => x.Results != null && x.Results.ExitCode != 0).Count();
+		}
+
 		/// <summary>
 		/// Return an enumeration of ready compile tasks.  This is not executed under a lock and 
 		/// does not modify the state of any actions.
diff --git a/Engine/Source/Programs/UnrealBuildTool/Executors/ParallelExecutor.cs b/Engine/Source/Programs/UnrealBuildTool/Executors/ParallelExecutor.cs
--- a/Engine/Source/Programs/UnrealBuildTool/Executors/ParallelExecutor.cs
+++ b/Engine/Source/Programs/UnrealBuildTool/Executors/ParallelExecutor.cs
@@ -150,6 +150,14 @@ namespace UnrealBuildTool
 			return true;
 		}
 
+		/// <summary>
+		/// Telemetry event for this executor
+		/// </summary>
+		protected TelemetryExecutorEvent? telemetryEvent;
+
+		/// <inheritdoc/>
+		public override TelemetryExecutorEvent? GetTelemetryEvent() => telemetryEvent;
+
 		/// <summary>
 		/// Create an action queue
 		/// </summary>
@@ -179,6 +187,9 @@ namespace UnrealBuildTool
 				return true;
 			}
 
+			DateTime startTimeUTC = DateTime.UtcNow;
+			bool result;
+
 			// The "useAutomaticQueue" should always be true unless manual queue is being tested
 			bool useAutomaticQueue = true;
 			if (useAutomaticQueue)
@@ -188,7 +199,10 @@ namespace UnrealBuildTool
 				queue.CreateAutomaticRunner(action => RunAction(queue, action), bUseActionWeights, actionLimit, NumParallelProcesses);
 				queue.Start();
 				queue.StartManyActions();
-				return await queue.RunTillDone();
+				result = await queue.RunTillDone();
+
+				queue.GetActionResultCounts(out int totalActions, out int succeededActions, out int failedActions);
+				telemetryEvent = new TelemetryExecutorEvent(Name, startTimeUTC, result, totalActions, succeededActions, failedActions, DateTime.UtcNow);
 			}
 			else
 			{
@@ -198,8 +212,13 @@ namespace UnrealBuildTool
 				queue.Start();
 				using Timer timer = new((_) => queue.StartManyActions(runner), null, 0, 500);
 				queue.StartManyActions();
-				return await queue.RunTillDone();
+				result = await queue.RunTillDone();
+
+				queue.GetActionResultCounts(out int totalActions, out int succeededActions, out int failedActions);
+				telemetryEvent = new TelemetryExecutorEvent(Name, startTimeUTC, result, totalActions, succeededActions, failedActions, DateTime.UtcNow);
 			}
+
+			return result;
 		}
 
 		private static Func<Task>? RunAction(ImmediateActionQueue queue, LinkedAction action)
diff --git a/Engine/Source/Programs/UnrealBuildTool/Executors/SNDBS.cs b/Engine/Source/Programs/UnrealBuildTool/Executors/SNDBS.cs
--- a/Engine/Source/Programs/UnrealBuildTool/Executors/SNDBS.cs
+++ b/Engine/Source/Programs/UnrealBuildTool/Executors/SNDBS.cs
@@ -216,6 +216,11 @@ namespace UnrealBuildTool
 			return true;
 		}
 
+		private TelemetryExecutorEvent? telemetryEvent;
+
+		/// <inheritdoc/>
+		public override TelemetryExecutorEvent? GetTelemetryEvent() => telemetryEvent;
+
 		/// <inheritdoc/>
 		public override Task<bool> ExecuteActionsAsync(IEnumerable<LinkedAction> ActionsToExecute, ILogger Logger, IActionArtifactCache? actionArtifactCache)
 		{
@@ -353,6 +358,7 @@ namespace UnrealBuildTool
 				try
 				{
 					// Start the process, redirecting stdout/stderr if requested.
+					DateTime startTimeUTC = DateTime.UtcNow;
 					Process LocalProcess = new Process();
 					LocalProcess.StartInfo = SnDbsStartInfo;
 					bool bShouldRedirectOuput = EventHandlerWrapper != null;
@@ -377,7 +383,10 @@ namespace UnrealBuildTool
 
 					// Wait until the process is finished and return whether it all the tasks successfully executed.
 					LocalProcess.WaitForExit();
-					return LocalProcess.ExitCode == 0;
+					bool result = LocalProcess.ExitCode == 0;
+
+					telemetryEvent = new TelemetryExecutorEvent(Name, startTimeUTC, result, NumActions, -1, -1, DateTime.UtcNow);
+					return result;
 				}
 				catch (Exception Ex)
 				{
diff --git a/Engine/Source/Programs/UnrealBuildTool/Executors/UnrealBuildAccelerator/UBAExecutor.cs b/Engine/Source/Programs/UnrealBuildTool/Executors/UnrealBuildAccelerator/UBAExecutor.cs
--- a/Engine/Source/Programs/UnrealBuildTool/Executors/UnrealBuildAccelerator/UBAExecutor.cs
+++ b/Engine/Source/Programs/UnrealBuildTool/Executors/UnrealBuildAccelerator/UBAExecutor.cs
@@ -1,5 +1,8 @@
 // Copyright Epic Games, Inc. All Rights Reserved.
 
+using EpicGames.Core;
+using EpicGames.UBA;
+using Microsoft.Extensions.Logging;
 using System;
 using System.Collections.Concurrent;
 using System.Collections.Generic;
@@ -7,10 +10,8 @@ using System.Linq;
 using System.Management;
 using System.Runtime.InteropServices;
 using System.Runtime.Versioning;
+using System.Threading;
 using System.Threading.Tasks;
-using EpicGames.Core;
-using EpicGames.UBA;
-using Microsoft.Extensions.Logging;
 using UnrealBuildBase;
 using UnrealBuildTool.Artifacts;
 
@@ -37,6 +38,11 @@ namespace UnrealBuildTool
 		// Tracking for LinkedActions that failed locally that should be retried without UBA
 		readonly ConcurrentDictionary<LinkedAction, bool> _forcedRetryActions = new();
 
+		// Tracking for all actions processed locally
+		int _localProcessedActions = 0;
+		// Tracking for all actions processed remotely
+		int _remoteProcessedActions = 0;
+
 		protected override void Dispose(bool disposing)
 		{
 			if (disposing)
@@ -328,6 +334,7 @@ namespace UnrealBuildTool
 		/// <returns>True if all the tasks successfully executed, or false if any of them failed.</returns>
 		bool ExecuteActionsInternal(IEnumerable<LinkedAction> inputActions, ISessionServer session, Microsoft.Extensions.Logging.ILogger logger, IActionArtifactCache? actionArtifactCache)
 		{
+			DateTime startTimeUTC = DateTime.UtcNow;
 			using ImmediateActionQueue queue = CreateActionQueue(inputActions, actionArtifactCache, logger);
 			int actionLimit = Math.Min(NumParallelProcesses, queue.TotalActions);
 			queue.CreateAutomaticRunner(action => RunActionLocal(queue, action), bUseActionWeights, actionLimit, NumParallelProcesses);
@@ -393,6 +400,11 @@ namespace UnrealBuildTool
 				}
 
 				bool res = queue.RunTillDone().Result; // Using inline wait to avoid possible thread switch
+
+				queue.GetActionResultCounts(out int totalActions, out int succeededActions, out int failedActions);
+				telemetryEvent = new TelemetryExecutorUBAEvent(Name, startTimeUTC, res, totalActions, succeededActions, failedActions,
+					_localProcessedActions, _remoteProcessedActions, _localRetryActions.Count, _forcedRetryActions.Count, DateTime.UtcNow);
+
 				return res;
 			}
 			finally
@@ -485,6 +497,7 @@ namespace UnrealBuildTool
 				ProcessStartInfo startInfo = GetActionStartInfo(action, out FileItem? pchItem);
 				using (IProcess process = _session!.RunProcess(startInfo, false, null, enableDetour))
 				{
+					Interlocked.Add(ref _localProcessedActions, 1);
 					if (process.ExitCode != 0 && UBAConfig.bForcedRetry)
 					{
 						_threadedLogger.LogWarning("{Description} {StatusDescription}: Exited with error code {ExitCode}. This action will retry without UBA", action.CommandDescription, action.StatusDescription, process.ExitCode);
@@ -555,6 +568,7 @@ namespace UnrealBuildTool
 				ProcessStartInfo startInfo = GetActionStartInfo(action, out FileItem? pchItem);
 				_session!.RunProcessRemote(startInfo, (s, e) =>
 				{
+					Interlocked.Add(ref _remoteProcessedActions, 1);
 					if (e.ExitCode != 0 && !e.LogLines.Any())
 					{
 						RemoteActionFailedNoOutput(queue, action, e.ExitCode, e.ExecutingHost ?? "Unknown");
diff --git a/Engine/Source/Programs/UnrealBuildTool/Executors/XGE.cs b/Engine/Source/Programs/UnrealBuildTool/Executors/XGE.cs
--- a/Engine/Source/Programs/UnrealBuildTool/Executors/XGE.cs
+++ b/Engine/Source/Programs/UnrealBuildTool/Executors/XGE.cs
@@ -369,6 +369,11 @@ namespace UnrealBuildTool
 			return true;
 		}
 
+		private TelemetryExecutorEvent? telemetryEvent;
+
+		/// <inheritdoc/>
+		public override TelemetryExecutorEvent? GetTelemetryEvent() => telemetryEvent;
+
 		// precompile the Regex needed to parse the XGE output (the ones we want are of the form "File (Duration at +time)"
 		//private static Regex XGEDurationRegex = new Regex(@"(?<Filename>.*) *\((?<Duration>[0-9:\.]+) at [0-9\+:\.]+\)", RegexOptions.ExplicitCapture);
 
@@ -771,8 +776,14 @@ namespace UnrealBuildTool
 					}
 				};
 
+				DateTime startTimeUTC = DateTime.UtcNow;
+
 				// Run through the standard XGE executor
-				return ExecuteTaskFile(TaskFilePath, EventHandlerWrapper, NumActions, Logger);
+				bool result = ExecuteTaskFile(TaskFilePath, EventHandlerWrapper, NumActions, Logger);
+
+				telemetryEvent = new TelemetryExecutorEvent(Name, startTimeUTC, result, NumActions, -1, -1, DateTime.UtcNow);
+
+				return result;
 			}
 		}
 	}
diff --git a/Engine/Source/Programs/UnrealBuildTool/System/ActionExecutor.cs b/Engine/Source/Programs/UnrealBuildTool/System/ActionExecutor.cs
--- a/Engine/Source/Programs/UnrealBuildTool/System/ActionExecutor.cs
+++ b/Engine/Source/Programs/UnrealBuildTool/System/ActionExecutor.cs
@@ -64,6 +64,24 @@ namespace UnrealBuildTool
 		/// <returns>True if the build succeeded, false otherwise</returns>
 		public abstract Task<bool> ExecuteActionsAsync(IEnumerable<LinkedAction> ActionsToExecute, ILogger Logger, IActionArtifactCache? actionArtifactCache = null);
 
+		/// <summary>
+		/// When completed, get a populated telemetry event for this executor
+		/// </summary>
+		/// <returns></returns>
+		public abstract TelemetryExecutorEvent? GetTelemetryEvent();
+
+		/// <summary>
+		/// Post a telemetry event if enabled
+		/// </summary>
+		public void PostTelemetryEvent()
+		{
+			TelemetryExecutorEvent? telemetryEvent = GetTelemetryEvent();
+			if (telemetryEvent != null)
+			{
+				TelemetryService.Get().RecordEvent(telemetryEvent);
+			}
+		}
+
 		/// <summary>
 		/// Will verify that produced items exists on disk
 		/// </summary>
diff --git a/Engine/Source/Programs/UnrealBuildTool/System/ActionGraph.cs b/Engine/Source/Programs/UnrealBuildTool/System/ActionGraph.cs
--- a/Engine/Source/Programs/UnrealBuildTool/System/ActionGraph.cs
+++ b/Engine/Source/Programs/UnrealBuildTool/System/ActionGraph.cs
@@ -443,6 +443,7 @@ namespace UnrealBuildTool
 				// Execute the build
 				Stopwatch Timer = Stopwatch.StartNew();
 				bool Result = await Executor.ExecuteActionsAsync(ActionsToExecute, Logger, actionArtifactCache);
+				Executor.PostTelemetryEvent();
 
 				Logger.LogInformation("Total time in {ExecutorName} executor: {TotalSeconds:0.00} seconds", Executor.Name, Timer.Elapsed.TotalSeconds);
 
diff --git a/Engine/Source/Programs/UnrealBuildTool/System/ActionGraphBuilder.cs b/Engine/Source/Programs/UnrealBuildTool/System/ActionGraphBuilder.cs
--- a/Engine/Source/Programs/UnrealBuildTool/System/ActionGraphBuilder.cs
+++ b/Engine/Source/Programs/UnrealBuildTool/System/ActionGraphBuilder.cs
@@ -314,7 +314,7 @@ namespace UnrealBuildTool
 			Action NewAction = Graph.CreateAction(Type);
 			NewAction.CommandPath = Unreal.DotnetPath;
 			NewAction.WorkingDirectory = Unreal.EngineSourceDirectory;
-			NewAction.CommandArguments = $"\"{Unreal.UnrealBuildToolDllPath}\" -Mode={Attribute.Name} {Arguments}";
+			NewAction.CommandArguments = $"\"{Unreal.UnrealBuildToolDllPath}\" -Session=\"{UnrealBuildTool.SessionIdentifier}\" -Mode={Attribute.Name} {Arguments}";
 			NewAction.CommandDescription = Attribute.Name;
 			NewAction.bCanExecuteRemotely = false;
 			NewAction.bCanExecuteRemotelyWithSNDBS = false;
diff --git a/Engine/Source/Programs/UnrealBuildTool/System/Telemetry.cs b/Engine/Source/Programs/UnrealBuildTool/System/Telemetry.cs
new file mode 100644
--- /dev/null
+++ b/Engine/Source/Programs/UnrealBuildTool/System/Telemetry.cs
@@ -0,0 +1,402 @@
+// Copyright Epic Games, Inc. All Rights Reserved.
+
+using EpicGames.Core;
+using EpicGames.Core.Telemetry;
+using System;
+using System.Collections.Concurrent;
+using System.Collections.Generic;
+using System.Linq;
+using System.Net.Http;
+using System.Text.Json.Serialization;
+using UnrealBuildBase;
+
+namespace UnrealBuildTool
+{
+	internal abstract class TelemetryEvent : DataRouterEvent
+	{
+		/// <inheritdoc/>
+		public override string EventName => "Core.UBT";
+
+		public abstract int SchemaVersion { get; }
+
+		public string Mode { get; protected set; } = String.Empty;
+
+		// If multiple targets are built, these will contain the first target found
+		public string Target { get; protected set; } = String.Empty;
+		public string Configuration { get; protected set; } = String.Empty;
+		public string Platform { get; protected set; } = String.Empty;
+		public string Architecture { get; protected set; } = String.Empty;
+
+		public string Build_BranchName { get; protected set; } = String.Empty;
+		public bool Config_IsBuildMachine { get; protected set; } = false;
+
+		public string Horde_BatchID { get; protected set; } = String.Empty;
+		public string Horde_JobID { get; protected set; } = String.Empty;
+		public string Horde_JobURL { get; protected set; } = String.Empty;
+		public string Horde_StepID { get; protected set; } = String.Empty;
+		public string Horde_StepName { get; protected set; } = String.Empty;
+		public string Horde_StepURL { get; protected set; } = String.Empty;
+		public string Horde_TemplateID { get; protected set; } = String.Empty;
+		public string Horde_TemplateName { get; protected set; } = String.Empty;
+
+		protected TelemetryEvent(DateTime? timestamp = null) : base(timestamp)
+		{
+		}
+
+		public void SetMetadata(string appId, TelemetryMetadata metadata)
+		{
+			AppId = appId;
+			AppVersion = metadata.AppVersion;
+			AppEnvironment = metadata.AppEnvironment;
+			UploadType = metadata.UploadType;
+			UserId = metadata.UserId;
+			SessionId = metadata.SessionId;
+			Mode = metadata.Mode;
+			Target = metadata.Target;
+			Configuration = metadata.Configuration;
+			Platform = metadata.Platform;
+			Architecture = metadata.Architecture;
+
+			Build_BranchName = metadata.Build_BranchName;
+			Config_IsBuildMachine = metadata.Config_IsBuildMachine;
+
+			Horde_TemplateID = metadata.Horde_TemplateID;
+			Horde_TemplateName = metadata.Horde_TemplateName;
+			Horde_JobURL = metadata.Horde_JobURL;
+			Horde_JobID = metadata.Horde_JobID;
+			Horde_StepName = metadata.Horde_StepName;
+			Horde_StepID = metadata.Horde_StepID;
+			Horde_StepURL = metadata.Horde_StepURL;
+			Horde_BatchID = metadata.Horde_BatchID;
+		}
+
+		public TelemetryEvent CloneWithMetadata(string appId, TelemetryMetadata metadata)
+		{
+			TelemetryEvent clone = (TelemetryEvent)MemberwiseClone();
+			clone.SetMetadata(appId, metadata);
+			return clone;
+		}
+	}
+
+	internal class TelemetryCompletedEvent : TelemetryEvent
+	{
+		/// <inheritdoc/>
+		public override string EventName => "Core.UBT.Completed";
+
+		public override int SchemaVersion => 3;
+
+		public string[] Arguments { get; }
+
+		[JsonConverter(typeof(UnixTimestampDataRouterEventConverter))]
+		public DateTime SessionStartUTC { get; }
+		[JsonConverter(typeof(DurationDataRouterEventConverter))]
+		public TimeSpan Duration { get; }
+		public CompilationResult Result { get; }
+
+		public TelemetryCompletedEvent(string[] arguments, DateTime sessionStartUTC, CompilationResult result, DateTime timestamp)
+			: base(timestamp)
+		{
+			Arguments = arguments;
+			SessionStartUTC = sessionStartUTC;
+			Duration = timestamp - SessionStartUTC;
+			Result = result;
+		}
+	}
+
+	internal class TelemetryExecutorEvent : TelemetryEvent
+	{
+		/// <inheritdoc/>
+		public override string EventName => "Core.UBT.Executor";
+
+		public override int SchemaVersion => 3;
+
+		public string Executor { get; }
+		[JsonConverter(typeof(UnixTimestampDataRouterEventConverter))]
+		public DateTime StartUTC { get; }
+		[JsonConverter(typeof(DurationDataRouterEventConverter))]
+		public TimeSpan Duration { get; }
+		public bool Result { get; }
+		public int TotalActions { get; }
+		public int SucceededActions { get; }
+		public int FailedActions { get; }
+		public int SkippedActions { get; }
+
+		public double SuccessRate => SucceededActions / (double)TotalActions;
+		public double FailureRate => FailedActions / (double)TotalActions;
+		public double SkipRate => SkippedActions / (double)TotalActions;
+
+		public TelemetryExecutorEvent(string executor, DateTime startUTC, bool result, int totalActions, int succeededActions, int failedActions, DateTime timestamp)
+			: base(timestamp)
+		{
+			Executor = executor;
+			StartUTC = startUTC;
+			Duration = timestamp - StartUTC;
+			Result = result;
+			TotalActions = totalActions;
+			SucceededActions = succeededActions;
+			FailedActions = failedActions;
+			SkippedActions = TotalActions - SucceededActions - FailedActions;
+		}
+	}
+
+	internal class TelemetryExecutorUBAEvent : TelemetryExecutorEvent
+	{
+		public int LocalActions { get; }
+		public int RemoteActions { get; }
+
+		public int RetriedLocalActions { get; }
+		public int RetriedDisabledActions { get; }
+
+		public double RetriedLocalRate => RetriedLocalActions / (double)TotalActions;
+		public double RetriedDisabledRate => RetriedDisabledActions / (double)TotalActions;
+
+		public double LocalUsage => LocalActions / (double)(TotalActions + RetriedLocalActions + RetriedDisabledActions - SkippedActions);
+		public double RemoteUsage => RemoteActions / (double)(TotalActions + RetriedLocalActions + RetriedDisabledActions - SkippedActions);
+
+		public TelemetryExecutorUBAEvent(string executor, DateTime startUTC, bool result, int totalActions, int succeededActions, int failedActions, int localActions, int remoteActions, int retriedLocalActions, int retriedDisabledActions, DateTime timestamp)
+			: base(executor, startUTC, result, totalActions, succeededActions, failedActions, timestamp)
+		{
+			LocalActions = localActions;
+			RemoteActions = remoteActions;
+			RetriedLocalActions = retriedLocalActions;
+			RetriedDisabledActions = retriedDisabledActions;
+		}
+	}
+
+	internal class TelemetryEndpoint
+	{
+		readonly Lazy<DataRouterTelemetryService> _service;
+		readonly Lazy<HttpClient> _httpClient;
+		readonly Uri _baseAddress;
+		readonly string _appId;
+		readonly Lazy<TelemetryMetadata> _metadata;
+
+		public TelemetryEndpoint(Lazy<HttpClient> httpClient, Uri baseAddress, string appId, Lazy<TelemetryMetadata> metadata)
+		{
+			_httpClient = httpClient;
+			_baseAddress = baseAddress;
+			_appId = appId;
+			_metadata = metadata;
+			_service = new(() =>
+			{
+				DataRouterTelemetryService service = new(_httpClient.Value, _baseAddress);
+#if DEBUG
+				service.IsDryRun = true;
+#endif
+				return service;
+			});
+		}
+
+		public void FlushEvents() => _service.Value.FlushEvents();
+
+		public void RecordEvent(TelemetryEvent eventData) => _service.Value.RecordEvent(eventData.CloneWithMetadata(_appId, _metadata.Value));
+	}
+
+	internal readonly struct TelemetryMetadata
+	{
+		public string AppVersion { get; }
+		public string AppEnvironment { get; }
+		public string UploadType { get; }
+		public string UserId { get; }
+		public string SessionId { get; }
+		public string Mode { get; }
+		public string Target { get; }
+		public string Configuration { get; }
+		public string Platform { get; }
+		public string Architecture { get; }
+
+		public string Build_BranchName { get; }
+		public bool Config_IsBuildMachine { get; }
+
+		public string Horde_BatchID { get; }
+		public string Horde_JobID { get; }
+		public string Horde_JobURL { get; }
+		public string Horde_StepID { get; }
+		public string Horde_StepName { get; }
+		public string Horde_StepURL { get; }
+		public string Horde_TemplateID { get; }
+		public string Horde_TemplateName { get; }
+
+		public TelemetryMetadata(string appVersion, string appEnvironment, string uploadType, string userID, string sessionId, string mode, TargetDescriptor? descriptor,
+			string? buildBranchName, bool configIsBuildMachine,
+			string? hordeBatchID, string? hordeJobID, string? hordeURL, string? hordeStepID, string? hordeStepName, string? hordeTemplateID, string? hordeTemplateName)
+		{
+			AppVersion = appVersion;
+			AppEnvironment = appEnvironment;
+			UploadType = uploadType;
+			UserId = userID;
+			SessionId = sessionId;
+			Mode = mode;
+			Target = descriptor?.Name ?? "Unknown";
+			Configuration = descriptor?.Configuration.ToString() ?? "Unknown";
+			Platform = descriptor?.Platform.ToString() ?? "Unknown";
+			Architecture = descriptor?.Architectures.ToString() ?? "Unknown";
+
+			Build_BranchName = buildBranchName ?? String.Empty;
+			Config_IsBuildMachine = configIsBuildMachine;
+
+			Horde_BatchID = hordeBatchID ?? String.Empty;
+			Horde_JobID = hordeJobID ?? String.Empty;
+			Horde_JobURL = !String.IsNullOrEmpty(hordeURL) && !String.IsNullOrEmpty(Horde_JobID) ? $"{hordeURL}job/{Horde_JobID}" : String.Empty;
+			Horde_StepID = hordeStepID ?? String.Empty;
+			Horde_StepName = hordeStepName ?? String.Empty;
+			Horde_StepURL = !String.IsNullOrEmpty(Horde_JobURL) && !String.IsNullOrEmpty(Horde_StepID) ? $"{Horde_JobURL}?step={Horde_StepID}" : String.Empty;
+			Horde_TemplateID = hordeTemplateID ?? String.Empty;
+			Horde_TemplateName = hordeTemplateName ?? String.Empty;
+
+			if (Mode.EndsWith("Mode", StringComparison.Ordinal))
+			{
+				Mode = Mode[..^4];
+			}
+
+			if (descriptor?.bRebuild == true && Mode == "Build")
+			{
+				Mode = "Rebuild";
+			}
+		}
+	}
+
+	/// <summary>
+	/// Telemetry service for UnrealBuildTool
+	/// </summary>
+	internal class TelemetryService : ITelemetryService<TelemetryEvent>
+	{
+		const string ConfigCategoryPrefix = "StudioTelemetry.";
+
+		// Lazy to defer resolution until an event is first sent
+		readonly Lazy<TelemetryMetadata> _metadata = new(() =>
+		{
+			BuildVersion.TryRead(BuildVersion.GetDefaultFileName(), out BuildVersion? version);
+#pragma warning disable CA1308 // Normalize strings to uppercase
+			return new TelemetryMetadata(
+				version != null ? $"{version.MajorVersion}.{version.MinorVersion}.{version.PatchVersion}-{version.Changelist}+{version.BranchName}" : String.Empty,
+				"UnrealBuildTool",
+				"eteventstream",
+				Environment.UserName,
+				UnrealBuildTool.SessionIdentifier,
+				UnrealBuildTool.BuildMode,
+				Get()._descriptor,
+				version?.BranchName?.ToLowerInvariant() ?? String.Empty,
+				Unreal.IsBuildMachine(),
+				Environment.GetEnvironmentVariable("UE_HORDE_BATCHID"),
+				Environment.GetEnvironmentVariable("UE_HORDE_JOBID"),
+				Environment.GetEnvironmentVariable("UE_HORDE_URL"),
+				Environment.GetEnvironmentVariable("UE_HORDE_STEPID"),
+				Environment.GetEnvironmentVariable("UE_HORDE_STEPNAME"),
+				Environment.GetEnvironmentVariable("UE_HORDE_TEMPLATEID"),
+				Environment.GetEnvironmentVariable("UE_HORDE_TEMPLATENAME")
+			);
+#pragma warning restore CA1308 // Normalize strings to uppercase
+		});
+
+		readonly Lazy<HttpClient> _httpClient;
+		readonly ConcurrentDictionary<Tuple<Uri, string>, TelemetryEndpoint> _endpoints;
+		readonly HashSet<string> _telemetryProviders = new();
+		bool _disposed = false;
+
+		TargetDescriptor? _descriptor = null;
+
+		#region Singleton
+		static readonly Lazy<TelemetryService> s_instance = new(() => new TelemetryService());
+		public static TelemetryService Get() => s_instance.Value;
+
+		private TelemetryService()
+		{
+			_httpClient = new(() => HttpClientDefaults.GetClient());
+			_endpoints = new();
+		}
+		#endregion
+
+		public void SetPrimaryTargetDetails(TargetDescriptor? descriptor) => _descriptor ??= descriptor;
+
+		public void AddTelemetryConfigProviders(IEnumerable<string> providers) => _telemetryProviders.UnionWith(providers);
+
+		public void AddEndpointsFromConfig(DirectoryReference? projectDir = null)
+		{
+			ConfigHierarchy engineIni = ConfigCache.ReadHierarchy(ConfigHierarchyType.Engine, projectDir, BuildHostPlatform.Current.Platform);
+			foreach (string category in _telemetryProviders)
+			{
+				AddEndpointFromConfigInternal(engineIni, category);
+			}
+		}
+
+		public void AddAllEndpointsFromConfig(DirectoryReference? projectDir = null)
+		{
+			ConfigHierarchy engineIni = ConfigCache.ReadHierarchy(ConfigHierarchyType.Engine, projectDir, BuildHostPlatform.Current.Platform);
+			foreach (string section in engineIni.SectionNames.Where(x => x.StartsWith(ConfigCategoryPrefix, StringComparison.Ordinal)))
+			{
+				AddEndpointFromConfigInternal(engineIni, section);
+			}
+		}
+
+		private bool AddEndpointFromConfigInternal(ConfigHierarchy engineIni, string? provider)
+		{
+			if (String.IsNullOrEmpty(provider))
+			{
+				return false;
+			}
+
+			ConfigHierarchySection section = engineIni.FindSection(provider);
+
+			if (section.TryGetValue("APIKeyET", out string? apiKeyET) && section.TryGetValue("APIServerET", out string? apiServerET) && section.TryGetValue("APIEndpointET", out string? apiEndpointET))
+			{
+				Uri baseAddress = new($"{apiServerET}{apiEndpointET}");
+				Tuple<Uri, string> key = new(baseAddress, apiKeyET);
+				_endpoints.GetOrAdd(key, _ => new TelemetryEndpoint(_httpClient, baseAddress, apiKeyET, _metadata));
+				return true;
+			}
+
+			return false;
+		}
+
+		#region ITelemetryService
+		/// <inheritdoc/>
+		public void RecordEvent(TelemetryEvent eventData)
+		{
+			foreach (KeyValuePair<Tuple<Uri, string>, TelemetryEndpoint> endpoint in _endpoints)
+			{
+				endpoint.Value.RecordEvent(eventData);
+			}
+		}
+
+		/// <inheritdoc/>
+		public void FlushEvents()
+		{
+			foreach (KeyValuePair<Tuple<Uri, string>, TelemetryEndpoint> endpoint in _endpoints)
+			{
+				endpoint.Value.FlushEvents();
+			}
+		}
+		#endregion
+
+		#region IDisposable
+		/// <summary>
+		/// Dispose
+		/// </summary>
+		/// <param name="disposing"></param>
+		protected virtual void Dispose(bool disposing)
+		{
+			if (!_disposed)
+			{
+				if (disposing)
+				{
+					// dispose managed state (managed objects)
+				}
+
+				// free unmanaged resources (unmanaged objects) and override finalizer
+				// set large fields to null
+
+				FlushEvents();
+				_disposed = true;
+			}
+		}
+
+		/// <inheritdoc/>
+		public void Dispose()
+		{
+			Dispose(disposing: true);
+			GC.SuppressFinalize(this);
+		}
+		#endregion
+	}
+}
\ No newline at end of file
diff --git a/Engine/Source/Programs/UnrealBuildTool/UnrealBuildTool.cs b/Engine/Source/Programs/UnrealBuildTool/UnrealBuildTool.cs
--- a/Engine/Source/Programs/UnrealBuildTool/UnrealBuildTool.cs
+++ b/Engine/Source/Programs/UnrealBuildTool/UnrealBuildTool.cs
@@ -24,6 +24,26 @@ namespace UnrealBuildTool
 		/// </summary>
 		public static DateTime StartTimeUtc { get; } = DateTime.UtcNow;
 
+		/// <summary>
+		/// Whether this is a recursive run of of the application
+		/// </summary>
+		public static bool IsRecursive = false;
+
+		/// <summary>
+		/// Unique id to track this session
+		/// </summary>
+		public static string SessionIdentifier = Guid.NewGuid().ToString("B");
+
+		/// <summary>
+		///  The mode of this instance
+		/// </summary>
+		public static string BuildMode = "BuildMode";
+
+		/// <summary>
+		/// The result of running the application
+		/// </summary>
+		private static CompilationResult ApplicationResult = CompilationResult.Unknown;
+
 		/// <summary>
 		/// The environment at boot time.
 		/// </summary>
@@ -320,6 +340,24 @@ namespace UnrealBuildTool
 			[XmlConfigFile(Category = "BuildConfiguration")]
 			public bool bDeleteTempDirectory = false;
 
+			/// <summary>
+			/// Providers to load opt-in telemetry connection information from ini. If unset, or the provider categories do not contain connection info, no telemetry will be sent.
+			/// </summary>
+			[XmlConfigFile(Category = "Telemetry", Name = "Providers")]
+			public string[] TelemetryProviders = Array.Empty<string>();
+
+			/// <summary>
+			/// Additional command line providers to load opt-in telemetry connection information from ini.
+			/// </summary>
+			[CommandLine(Prefix = "-TelemetryProvider", Description = "List of ini providers for telemetry", ListSeparator = '+')]
+			public List<string> CmdTelemetryProviders = new();
+
+			/// <summary>
+			/// Session identifier for this run of UBT, if unset defaults to a random Guid
+			/// </summary>
+			[CommandLine(Prefix = "-Session", Description = "Session identifier for this run of UBT, if unset defaults to a random Guid")]
+			public string? TelemetrySession = null;
+
 			/// <summary>
 			/// Initialize the options with the given command line arguments
 			/// </summary>
@@ -525,6 +563,7 @@ namespace UnrealBuildTool
 						return 1;
 					}
 				}
+				BuildMode = ModeType.Name;
 
 				// Get the options for which systems have to be initialized for this mode
 				ToolModeOptions ModeOptions = ModeType.GetCustomAttribute<ToolModeAttribute>()!.Options;
@@ -569,6 +608,16 @@ namespace UnrealBuildTool
 					Log.AddFileWriter("LogTraceListener", Options.LogFileName);
 				}
 
+				// Initialize the telemetry service
+				if (!String.IsNullOrEmpty(Options.TelemetrySession))
+				{
+					IsRecursive = true;
+					SessionIdentifier = Options.TelemetrySession;
+				}
+
+				TelemetryService.Get().AddTelemetryConfigProviders(Options.TelemetryProviders.Concat(Options.CmdTelemetryProviders));
+				TelemetryService.Get().AddEndpointsFromConfig();
+
 				// Create a UbtRun file
 				try
 				{
@@ -684,18 +733,21 @@ namespace UnrealBuildTool
 					}
 				}
 
+				ApplicationResult = (CompilationResult)Result;
 				return Result;
 			}
 			catch (CompilationResultException Ex)
 			{
 				// Used to return a propagate a specific exit code after an error has occurred.
 				Ex.LogException(Logger);
+				ApplicationResult = Ex.Result;
 				return (int)Ex.Result;
 			}
 			catch (BuildLogEventException Ex)
 			{
 				// BuildExceptions should have nicely formatted messages.
 				Ex.LogException(Logger);
+				ApplicationResult = CompilationResult.OtherCompilationError;
 				return (int)CompilationResult.OtherCompilationError;
 			}
 			catch (JsonException Ex)
@@ -704,6 +756,7 @@ namespace UnrealBuildTool
 				LogValue FileValue = LogValue.SourceFile(source, source.GetFileName());
 				Logger.LogError(KnownLogEvents.Compiler, "{File}({Line}): error:{Message}", FileValue, Ex.LineNumber ?? 0, ExceptionUtils.FormatException(Ex));
 				Logger.LogDebug(KnownLogEvents.Compiler, "{File}({Line}): error:{Message}", FileValue, Ex.LineNumber ?? 0, ExceptionUtils.FormatExceptionDetails(Ex));
+				ApplicationResult = CompilationResult.OtherCompilationError;
 				return (int)CompilationResult.OtherCompilationError;
 			}
 			catch (BuildException Ex)
@@ -717,6 +770,7 @@ namespace UnrealBuildTool
 				// Unhandled exception.
 				Logger.LogError(Ex, "Unhandled exception: {Ex}", ExceptionUtils.FormatException(Ex));
 				Logger.LogDebug(Ex, "Unhandled exception: {Ex}", ExceptionUtils.FormatExceptionDetails(Ex));
+				ApplicationResult = CompilationResult.OtherCompilationError;
 				return (int)CompilationResult.OtherCompilationError;
 			}
 			finally
@@ -737,6 +791,14 @@ namespace UnrealBuildTool
 				// Useful when investigating why UBT takes time.
 				//DirectoryItem.WriteDebugFileWithAllEnumeratedFiles(@"c:\temp\AllFiles.txt");
 
+				if (!IsRecursive)
+				{
+					TelemetryService.Get().RecordEvent(new TelemetryCompletedEvent(ArgumentsArray, StartTimeUtc, ApplicationResult, DateTime.UtcNow));
+				}
+
+				// Flush any remaining telemetry events
+				TelemetryService.Get().FlushEvents();
+
 				Utils.LogWriteFileIfChangedActivity(Logger);
 
 				// Print out all the performance info
