Added bBuildConsoleAppOnly flag which overrides bBuildAdditionalConsoleApp. This removes the redundant creation of the non '-Cmd' versions of certain applications.

diff --git a/Engine/Source/Programs/UnrealBuildTool/Configuration/ReadOnlyTargetRules.cs b/Engine/Source/Programs/UnrealBuildTool/Configuration/ReadOnlyTargetRules.cs
--- a/Engine/Source/Programs/UnrealBuildTool/Configuration/ReadOnlyTargetRules.cs
+++ b/Engine/Source/Programs/UnrealBuildTool/Configuration/ReadOnlyTargetRules.cs
@@ -310,6 +310,8 @@ namespace UnrealBuildTool
 
 		public bool bBuildAdditionalConsoleApp => Inner.bBuildAdditionalConsoleApp;
 
+		public bool bBuildConsoleAppOnly => Inner.bBuildConsoleAppOnly;
+
 		public bool bDisableSymbolCache => Inner.bDisableSymbolCache;
 
 		public bool bUseUnityBuild => Inner.bUseUnityBuild;
diff --git a/Engine/Source/Programs/UnrealBuildTool/Configuration/TargetRules.cs b/Engine/Source/Programs/UnrealBuildTool/Configuration/TargetRules.cs
--- a/Engine/Source/Programs/UnrealBuildTool/Configuration/TargetRules.cs
+++ b/Engine/Source/Programs/UnrealBuildTool/Configuration/TargetRules.cs
@@ -1561,6 +1561,10 @@ namespace UnrealBuildTool
 		}
 		private bool? bBuildAdditionalConsoleAppOverride;
 
+		/// <summary>
+		/// If true, only creates an additional console application. Overrides bBuildAdditionalConsoleApp.		
+		/// </summary>
+		public bool bBuildConsoleAppOnly { get; set; } = false;
 		/// <summary>
 		/// True if debug symbols that are cached for some platforms should not be created.
 		/// </summary>
diff --git a/Engine/Source/Programs/UnrealBuildTool/Configuration/UEBuildBinary.cs b/Engine/Source/Programs/UnrealBuildTool/Configuration/UEBuildBinary.cs
--- a/Engine/Source/Programs/UnrealBuildTool/Configuration/UEBuildBinary.cs
+++ b/Engine/Source/Programs/UnrealBuildTool/Configuration/UEBuildBinary.cs
@@ -97,6 +97,12 @@ namespace UnrealBuildTool
 		/// </summary>
 		public bool bBuildAdditionalConsoleApp = false;
 
+		/// <summary>
+		/// If true, replaces the executable with a console application. Hack for Windows, where it's not possible to conditionally inherit a parent's console Window depending on how
+		/// the application is invoked
+		/// </summary>
+		public bool bBuildConsoleAppOnly = false;
+
 		/// <summary>
 		/// 
 		/// </summary>
@@ -130,6 +136,7 @@ namespace UnrealBuildTool
 		/// <param name="IntermediateDirectory"></param>
 		/// <param name="bAllowExports"></param>
 		/// <param name="bBuildAdditionalConsoleApp"></param>
+		/// <param name="bBuildConsoleAppOnly"></param>
 		/// <param name="PrimaryModule"></param>
 		/// <param name="bUsePrecompiled"></param>
 		public UEBuildBinary(
@@ -138,6 +145,7 @@ namespace UnrealBuildTool
 				DirectoryReference IntermediateDirectory,
 				bool bAllowExports,
 				bool bBuildAdditionalConsoleApp,
+				bool bBuildConsoleAppOnly,
 				UEBuildModuleCPP PrimaryModule,
 				bool bUsePrecompiled
 			)
@@ -148,6 +156,7 @@ namespace UnrealBuildTool
 			this.IntermediateDirectory = IntermediateDirectory;
 			this.bAllowExports = bAllowExports;
 			this.bBuildAdditionalConsoleApp = bBuildAdditionalConsoleApp;
+			this.bBuildConsoleAppOnly = bBuildConsoleAppOnly;			
 			this.PrimaryModule = PrimaryModule;
 			this.bUsePrecompiled = bUsePrecompiled;
 
@@ -273,36 +282,46 @@ namespace UnrealBuildTool
 			// Create the import library if needed
 			OutputFiles.AddRange(ToolChain.LinkImportLibrary(BinaryLinkEnvironment, Graph));
 
-			// Link the binary.
-			FileItem[] Executables = ToolChain.LinkAllFiles(BinaryLinkEnvironment, false, Graph);
-			OutputFiles.AddRange(Executables);
-
-			// Save all the output items for this binary. This is used for hot-reload, and excludes any items added in PostBuild (such as additional files copied into the app).
-			if (Target.LinkType == TargetLinkType.Modular)
-			{
-				Graph.SetOutputItemsForModule(PrimaryModule.Name, OutputFiles.ToArray());
-			}
-
-			// Produce additional console app if requested
-			if (bBuildAdditionalConsoleApp)
-			{
+			
+			// Override the build to be a console app (i.e. only build a console app)
+			if (bBuildConsoleAppOnly)
+			{				
+				BinaryLinkEnvironment.bIsBuildingConsoleApplication = true;
+				BinaryLinkEnvironment.bCodeCoverage = CompileEnvironment.bCodeCoverage;
+				BinaryLinkEnvironment.WindowsEntryPointOverride = "WinMainCRTStartup";       // For WinMain() instead of "main()" for Launch module
+				BinaryLinkEnvironment.OutputFilePaths = BinaryLinkEnvironment.OutputFilePaths.Select(Path => GetAdditionalConsoleAppPath(Path)).ToList();
+			} 
+			else if (bBuildAdditionalConsoleApp)
+			{				
 				// Produce additional binary but link it as a console app
-				LinkEnvironment ConsoleAppLinkEvironment = new LinkEnvironment(BinaryLinkEnvironment);
-				ConsoleAppLinkEvironment.bIsBuildingConsoleApplication = true;
-				ConsoleAppLinkEvironment.bCodeCoverage = CompileEnvironment.bCodeCoverage;
-				ConsoleAppLinkEvironment.WindowsEntryPointOverride = "WinMainCRTStartup";       // For WinMain() instead of "main()" for Launch module
-				ConsoleAppLinkEvironment.OutputFilePaths = ConsoleAppLinkEvironment.OutputFilePaths.Select(Path => GetAdditionalConsoleAppPath(Path)).ToList();
+				LinkEnvironment ConsoleAppLinkEnvironment = new LinkEnvironment(BinaryLinkEnvironment);
+				ConsoleAppLinkEnvironment.bIsBuildingConsoleApplication = true;
+				ConsoleAppLinkEnvironment.bCodeCoverage = CompileEnvironment.bCodeCoverage;
+				ConsoleAppLinkEnvironment.WindowsEntryPointOverride = "WinMainCRTStartup";       // For WinMain() instead of "main()" for Launch module
+				ConsoleAppLinkEnvironment.OutputFilePaths = ConsoleAppLinkEnvironment.OutputFilePaths.Select(Path => GetAdditionalConsoleAppPath(Path)).ToList();
 
 				// Link the console app executable
-				FileItem[] ConsoleAppOutputFiles = ToolChain.LinkAllFiles(ConsoleAppLinkEvironment, false, Graph);
+				FileItem[] ConsoleAppOutputFiles = ToolChain.LinkAllFiles(ConsoleAppLinkEnvironment, false, Graph);
 				OutputFiles.AddRange(ConsoleAppOutputFiles);
 
 				foreach (FileItem Executable in ConsoleAppOutputFiles)
 				{
-					OutputFiles.AddRange(ToolChain.PostBuild(Target, Executable, ConsoleAppLinkEvironment, Graph));
+					OutputFiles.AddRange(ToolChain.PostBuild(Target, Executable, ConsoleAppLinkEnvironment, Graph));
 				}
 			}
 
+	
+
+			// Link the binary.
+			FileItem[] Executables = ToolChain.LinkAllFiles(BinaryLinkEnvironment, false, Graph);
+			OutputFiles.AddRange(Executables);
+			
+			// Save all the output items for this binary. This is used for hot-reload, and excludes any items added in PostBuild (such as additional files copied into the app).
+			if (Target.LinkType == TargetLinkType.Modular)
+			{
+				Graph.SetOutputItemsForModule(PrimaryModule.Name, OutputFiles.ToArray());
+			}
+
 			foreach (FileItem Executable in Executables)
 			{
 				OutputFiles.AddRange(ToolChain.PostBuild(Target, Executable, BinaryLinkEnvironment, Graph));
@@ -540,13 +559,23 @@ namespace UnrealBuildTool
 
 				// Add the primary build products
 				string[] DebugExtensions = UEBuildPlatform.GetBuildPlatform(Target.Platform).GetDebugInfoExtensions(Target, Type);
-				foreach (FileReference OutputFilePath in OutputFilePaths)
+
+				if (Type == UEBuildBinaryType.Executable && bBuildConsoleAppOnly)
 				{
-					AddBuildProductAndDebugFiles(OutputFilePath, OutputType, DebugExtensions, BuildProducts, ToolChain, bCreateDebugInfo);
+					foreach (FileReference OutputFilePath in OutputFilePaths)
+					{
+						AddBuildProductAndDebugFiles(GetAdditionalConsoleAppPath(OutputFilePath), OutputType, DebugExtensions, BuildProducts, ToolChain, bCreateDebugInfo);
+					}
+				} else 
+				{
+					foreach (FileReference OutputFilePath in OutputFilePaths)
+					{
+						AddBuildProductAndDebugFiles(OutputFilePath, OutputType, DebugExtensions, BuildProducts, ToolChain, bCreateDebugInfo);
+					}
 				}
 
 				// Add the console app, if there is one
-				if (Type == UEBuildBinaryType.Executable && bBuildAdditionalConsoleApp)
+				if (Type == UEBuildBinaryType.Executable && !bBuildConsoleAppOnly && bBuildAdditionalConsoleApp)
 				{
 					foreach (FileReference OutputFilePath in OutputFilePaths)
 					{
diff --git a/Engine/Source/Programs/UnrealBuildTool/Configuration/UEBuildTarget.cs b/Engine/Source/Programs/UnrealBuildTool/Configuration/UEBuildTarget.cs
--- a/Engine/Source/Programs/UnrealBuildTool/Configuration/UEBuildTarget.cs
+++ b/Engine/Source/Programs/UnrealBuildTool/Configuration/UEBuildTarget.cs
@@ -4167,6 +4167,7 @@ namespace UnrealBuildTool
 				IntermediateDirectory: Module.IntermediateDirectory,
 				bAllowExports: true,
 				bBuildAdditionalConsoleApp: false,
+				bBuildConsoleAppOnly: false,
 				PrimaryModule: Module,
 				bUsePrecompiled: Module.Rules.bUsePrecompiled
 			);
@@ -4805,6 +4806,7 @@ namespace UnrealBuildTool
 				IntermediateDirectory: IntermediateDirectory,
 				bAllowExports: Rules.bHasExports,
 				bBuildAdditionalConsoleApp: Rules.bBuildAdditionalConsoleApp,
+				bBuildConsoleAppOnly: Rules.bBuildConsoleAppOnly,
 				PrimaryModule: LaunchModule,
 				bUsePrecompiled: LaunchModule.Rules.bUsePrecompiled && OutputPaths[0].IsUnderDirectory(Unreal.EngineDirectory)
 			);
@@ -4962,7 +4964,7 @@ namespace UnrealBuildTool
 			if (RelativeBaseDir != null)
 			{
 				GlobalCompileEnvironment.Definitions.Add(String.Format("UE_RELATIVE_BASE_DIR=\"{0}/\"", RelativeBaseDir));
-				if (Rules.bBuildAdditionalConsoleApp)
+				if (Rules.bBuildAdditionalConsoleApp || Rules.bBuildConsoleAppOnly)
 				{
 					string? CmdletRelativeBaseDir = GetRelativeBaseDir(GetExecutableDir(), Platform);
 					GlobalCompileEnvironment.Definitions.Add(String.Format("UE_CMDLET_RELATIVE_BASE_DIR=\"{0}/\"", CmdletRelativeBaseDir));
