UnrealBuildTool: Don't send telemetry or show dryrun output in debug

#rnx

diff --git a/Engine/Source/Programs/UnrealBuildTool/System/Telemetry.cs b/Engine/Source/Programs/UnrealBuildTool/System/Telemetry.cs
--- a/Engine/Source/Programs/UnrealBuildTool/System/Telemetry.cs
+++ b/Engine/Source/Programs/UnrealBuildTool/System/Telemetry.cs
@@ -181,7 +181,7 @@ namespace UnrealBuildTool
 			{
 				DataRouterTelemetryService service = new(_httpClient.Value, _baseAddress);
 #if DEBUG
-				service.IsDryRun = true;
+				service.IsDryRun = false;
 #endif
 				return service;
 			});
@@ -210,6 +210,7 @@ namespace UnrealBuildTool
 
 		public string Horde_BatchID { get; }
 		public string Horde_JobID { get; }
+		public string Horde_URL { get; }
 		public string Horde_JobURL { get; }
 		public string Horde_StepID { get; }
 		public string Horde_StepName { get; }
@@ -237,7 +238,8 @@ namespace UnrealBuildTool
 
 			Horde_BatchID = hordeBatchID ?? String.Empty;
 			Horde_JobID = hordeJobID ?? String.Empty;
-			Horde_JobURL = !String.IsNullOrEmpty(hordeURL) && !String.IsNullOrEmpty(Horde_JobID) ? $"{hordeURL}job/{Horde_JobID}" : String.Empty;
+			Horde_URL = hordeURL ?? String.Empty;
+			Horde_JobURL = !String.IsNullOrEmpty(Horde_URL) && !String.IsNullOrEmpty(Horde_JobID) ? $"{Horde_URL}job/{Horde_JobID}" : String.Empty;
 			Horde_StepID = hordeStepID ?? String.Empty;
 			Horde_StepName = hordeStepName ?? String.Empty;
 			Horde_StepURL = !String.IsNullOrEmpty(Horde_JobURL) && !String.IsNullOrEmpty(Horde_StepID) ? $"{Horde_JobURL}?step={Horde_StepID}" : String.Empty;
@@ -261,6 +263,10 @@ namespace UnrealBuildTool
 	/// </summary>
 	internal class TelemetryService : ITelemetryService<TelemetryEvent>
 	{
+#if DEBUG
+		readonly bool _bAllowRecordingEventsDebug = false;
+#endif
+
 		const string ConfigCategoryPrefix = "StudioTelemetry.";
 
 		// Lazy to defer resolution until an event is first sent
@@ -353,6 +359,16 @@ namespace UnrealBuildTool
 		/// <inheritdoc/>
 		public void RecordEvent(TelemetryEvent eventData)
 		{
+#if DEBUG
+			if (_bAllowRecordingEventsDebug)
+			{
+				foreach (KeyValuePair<Tuple<Uri, string>, TelemetryEndpoint> endpoint in _endpoints)
+				{
+					endpoint.Value.RecordEvent(eventData);
+				}
+				return;
+			}
+#endif
 			foreach (KeyValuePair<Tuple<Uri, string>, TelemetryEndpoint> endpoint in _endpoints)
 			{
 				endpoint.Value.RecordEvent(eventData);
