UnrealBuildTool: If running via Horde, always add the current Horde server as a telemetry endpoint

#rnx

diff --git a/Engine/Source/Programs/UnrealBuildTool/System/Telemetry.cs b/Engine/Source/Programs/UnrealBuildTool/System/Telemetry.cs
--- a/Engine/Source/Programs/UnrealBuildTool/System/Telemetry.cs
+++ b/Engine/Source/Programs/UnrealBuildTool/System/Telemetry.cs
@@ -324,6 +324,12 @@ namespace UnrealBuildTool
 			{
 				AddEndpointFromConfigInternal(engineIni, category);
 			}
+
+			// If we are running via horde, always the current horde url as a provider
+			if (!String.IsNullOrEmpty(_metadata.Value.Horde_URL))
+			{
+				AddEndpointFromConfigInternal(engineIni, "StudioTelemetry.Provider.Horde", _metadata.Value.Horde_URL);
+			}
 		}
 
 		public void AddAllEndpointsFromConfig(DirectoryReference? projectDir = null)
@@ -335,7 +341,7 @@ namespace UnrealBuildTool
 			}
 		}
 
-		private bool AddEndpointFromConfigInternal(ConfigHierarchy engineIni, string? provider)
+		private bool AddEndpointFromConfigInternal(ConfigHierarchy engineIni, string provider, string? apiServerET = null)
 		{
 			if (String.IsNullOrEmpty(provider))
 			{
@@ -344,7 +350,12 @@ namespace UnrealBuildTool
 
 			ConfigHierarchySection section = engineIni.FindSection(provider);
 
-			if (section.TryGetValue("APIKeyET", out string? apiKeyET) && section.TryGetValue("APIServerET", out string? apiServerET) && section.TryGetValue("APIEndpointET", out string? apiEndpointET))
+			if (String.IsNullOrEmpty(apiServerET))
+			{
+				section.TryGetValue("APIServerET", out apiServerET);
+			}
+
+			if (!String.IsNullOrEmpty(apiServerET) && section.TryGetValue("APIKeyET", out string? apiKeyET) && section.TryGetValue("APIEndpointET", out string? apiEndpointET))
 			{
 				Uri baseAddress = new($"{apiServerET}{apiEndpointET}");
 				Tuple<Uri, string> key = new(baseAddress, apiKeyET);
