UnrealBuildTool: Add option to reduce optimizations for huge functions for MSVC platforms

#rnx

diff --git a/Engine/Source/Programs/UnrealBuildTool/Platform/Windows/UEBuildWindows.cs b/Engine/Source/Programs/UnrealBuildTool/Platform/Windows/UEBuildWindows.cs
--- a/Engine/Source/Programs/UnrealBuildTool/Platform/Windows/UEBuildWindows.cs
+++ b/Engine/Source/Programs/UnrealBuildTool/Platform/Windows/UEBuildWindows.cs
@@ -622,6 +622,21 @@ namespace UnrealBuildTool
 		/// </summary>
 		public bool bOptimizeGlobalData = true;
 
+		/// <summary>
+		/// Whether to reduce optimizations for huge functions over an instruction threshold to improve compile time
+		/// https://devblogs.microsoft.com/cppblog/msvc-backend-updates-in-visual-studio-2019-versions-16-3-and-16-4/
+		/// </summary>
+		[XmlConfigFile(Category = "WindowsPlatform")]
+		[RequiresUniqueBuildEnvironment]
+		public bool bReducedOptimizeHugeFunctions = false;
+
+		/// <summary>
+		/// The instruction threshold to use when reducing optimizations for huge functions, default 20000.
+		/// </summary>
+		[XmlConfigFile(Category = "WindowsPlatform")]
+		[RequiresUniqueBuildEnvironment]
+		public int bReducedOptimizeHugeFunctionsThreshold = 20000;
+
 		/// <summary>
 		/// (Experimental) Appends the -ftime-trace argument to the command line for Clang to output a JSON file containing a timeline for the compile. 
 		/// See http://aras-p.info/blog/2019/01/16/time-trace-timeline-flame-chart-profiler-for-Clang/ for more info.
@@ -898,6 +913,10 @@ namespace UnrealBuildTool
 
 		public bool bOptimizeGlobalData => Inner.bOptimizeGlobalData;
 
+		public bool bReducedOptimizeHugeFunctions => Inner.bReducedOptimizeHugeFunctions;
+
+		public int bReducedOptimizeHugeFunctionsThreshold => Inner.bReducedOptimizeHugeFunctionsThreshold;
+
 		public bool bClangTimeTrace => Inner.bClangTimeTrace;
 
 		public bool bCompilerTrace => Inner.bCompilerTrace;
diff --git a/Engine/Source/Programs/UnrealBuildTool/Platform/Windows/VCToolChain.cs b/Engine/Source/Programs/UnrealBuildTool/Platform/Windows/VCToolChain.cs
--- a/Engine/Source/Programs/UnrealBuildTool/Platform/Windows/VCToolChain.cs
+++ b/Engine/Source/Programs/UnrealBuildTool/Platform/Windows/VCToolChain.cs
@@ -491,6 +491,13 @@ namespace UnrealBuildTool
 				Arguments.Add("/Gw");
 			}
 
+            // Reduce optimizations for huge functions, may improve compile time a the expense of speed for functions over the threshold
+            if (Target.WindowsPlatform.Compiler.IsMSVC() && Target.WindowsPlatform.bReducedOptimizeHugeFunctions)
+            {
+                Arguments.Add("/d2ReducedOptimizeHugeFunctions");
+                Arguments.Add($"/d2ReducedOptimizeThreshold:{Target.WindowsPlatform.bReducedOptimizeHugeFunctionsThreshold}");
+            }
+
 			// Separate functions for linker.
 			Arguments.Add("/Gy");
 
@@ -1353,10 +1360,17 @@ namespace UnrealBuildTool
 				//Arguments.Add("/PDBCompress"); // Do not turn this on, it makes link times almost 2x slower. This is _only_ to save local disk space. Will _not_ make actual file smaller for network transfer
 			}
 
-			//
-			//	Shipping & LTCG
-			//
-			if (LinkEnvironment.bAllowLTCG)
+            // Reduce optimizations for huge functions, may improve compile time a the expense of speed for functions over the threshold
+            if (Target.WindowsPlatform.Compiler.IsMSVC() && Target.WindowsPlatform.bReducedOptimizeHugeFunctions)
+            {
+                Arguments.Add("/d2:\"-ReducedOptimizeHugeFunctions\"");
+                Arguments.Add($"/d2:\"-ReducedOptimizeThreshold:{Target.WindowsPlatform.bReducedOptimizeHugeFunctionsThreshold}\"");
+            }
+
+            //
+            //	Shipping & LTCG
+            //
+            if (LinkEnvironment.bAllowLTCG)
 			{
 				// Use link-time code generation.
 				Arguments.Add("/LTCG");
