UnrealBuildTool: Whoops an int isn't a bool, don't prefix with b

#rnx

diff --git a/Engine/Source/Programs/UnrealBuildTool/Platform/Windows/UEBuildWindows.cs b/Engine/Source/Programs/UnrealBuildTool/Platform/Windows/UEBuildWindows.cs
--- a/Engine/Source/Programs/UnrealBuildTool/Platform/Windows/UEBuildWindows.cs
+++ b/Engine/Source/Programs/UnrealBuildTool/Platform/Windows/UEBuildWindows.cs
@@ -635,7 +635,7 @@ namespace UnrealBuildTool
 		/// </summary>
 		[XmlConfigFile(Category = "WindowsPlatform")]
 		[RequiresUniqueBuildEnvironment]
-		public int bReducedOptimizeHugeFunctionsThreshold = 20000;
+		public int ReducedOptimizeHugeFunctionsThreshold = 20000;
 
 		/// <summary>
 		/// (Experimental) Appends the -ftime-trace argument to the command line for Clang to output a JSON file containing a timeline for the compile. 
@@ -915,7 +915,7 @@ namespace UnrealBuildTool
 
 		public bool bReducedOptimizeHugeFunctions => Inner.bReducedOptimizeHugeFunctions;
 
-		public int bReducedOptimizeHugeFunctionsThreshold => Inner.bReducedOptimizeHugeFunctionsThreshold;
+		public int ReducedOptimizeHugeFunctionsThreshold => Inner.ReducedOptimizeHugeFunctionsThreshold;
 
 		public bool bClangTimeTrace => Inner.bClangTimeTrace;
 
diff --git a/Engine/Source/Programs/UnrealBuildTool/Platform/Windows/VCToolChain.cs b/Engine/Source/Programs/UnrealBuildTool/Platform/Windows/VCToolChain.cs
--- a/Engine/Source/Programs/UnrealBuildTool/Platform/Windows/VCToolChain.cs
+++ b/Engine/Source/Programs/UnrealBuildTool/Platform/Windows/VCToolChain.cs
@@ -495,7 +495,7 @@ namespace UnrealBuildTool
             if (Target.WindowsPlatform.Compiler.IsMSVC() && Target.WindowsPlatform.bReducedOptimizeHugeFunctions)
             {
                 Arguments.Add("/d2ReducedOptimizeHugeFunctions");
-                Arguments.Add($"/d2ReducedOptimizeThreshold:{Target.WindowsPlatform.bReducedOptimizeHugeFunctionsThreshold}");
+                Arguments.Add($"/d2ReducedOptimizeThreshold:{Target.WindowsPlatform.ReducedOptimizeHugeFunctionsThreshold}");
             }
 
 			// Separate functions for linker.
@@ -1364,7 +1364,7 @@ namespace UnrealBuildTool
             if (Target.WindowsPlatform.Compiler.IsMSVC() && Target.WindowsPlatform.bReducedOptimizeHugeFunctions)
             {
                 Arguments.Add("/d2:\"-ReducedOptimizeHugeFunctions\"");
-                Arguments.Add($"/d2:\"-ReducedOptimizeThreshold:{Target.WindowsPlatform.bReducedOptimizeHugeFunctionsThreshold}\"");
+                Arguments.Add($"/d2:\"-ReducedOptimizeThreshold:{Target.WindowsPlatform.ReducedOptimizeHugeFunctionsThreshold}\"");
             }
 
             //
