Fix iOS deploy error

diff --git a/Engine/Source/Programs/UnrealBuildTool/ProjectFiles/Xcode/XcodeProject.cs b/Engine/Source/Programs/UnrealBuildTool/ProjectFiles/Xcode/XcodeProject.cs
--- a/Engine/Source/Programs/UnrealBuildTool/ProjectFiles/Xcode/XcodeProject.cs
+++ b/Engine/Source/Programs/UnrealBuildTool/ProjectFiles/Xcode/XcodeProject.cs
@@ -1349,9 +1349,9 @@ namespace UnrealBuildTool.XcodeProjectXcconfig
 				"DEST_EXE=\\\"${CONFIGURATION_BUILD_DIR}/${EXECUTABLE_PATH}\\\"",
 				"DEST_EXE_DIR=`dirname \\\"${DEST_EXE}\\\"`",
 				"",
-				"echo Copying executable and any standalone dylibs into ${DEST_EXE_DIR}",
+				"echo Copying executable and any standalone dylibs into ${DEST_EXE_DIR} but do not overwrite unless src is newer",
 				"mkdir -p \\\"${DEST_EXE_DIR}\\\"",
-				"ditto \\\"${SRC_EXE}\\\" \\\"${DEST_EXE}\\\"",
+				"rsync -au \\\"${SRC_EXE}\\\" \\\"${DEST_EXE}\\\"",
 			});
 
 			IEnumerable<ModuleRules.RuntimeDependency>? Dylibs;
