[UBT]
* Added so if uba agent exits with non-zero return code we add it to the log

diff --git a/Engine/Source/Programs/UnrealBuildTool/Executors/UnrealBuildAccelerator/UBAAgentCoordinatorHorde.cs b/Engine/Source/Programs/UnrealBuildTool/Executors/UnrealBuildAccelerator/UBAAgentCoordinatorHorde.cs
--- a/Engine/Source/Programs/UnrealBuildTool/Executors/UnrealBuildAccelerator/UBAAgentCoordinatorHorde.cs
+++ b/Engine/Source/Programs/UnrealBuildTool/Executors/UnrealBuildAccelerator/UBAAgentCoordinatorHorde.cs
@@ -532,6 +532,9 @@ namespace UnrealBuildTool
 							}
 						}
 						logger.LogDebug("Shutting down process");
+						int ExitCode = await process.WaitForExitAsync(cancellationToken);
+						if (ExitCode != 0)
+							logger.LogInformation($"UbaAgent exited with exit code: {ExitCode}");
 					}
 
 					logger.LogDebug("Closing channel");
