Shared C#: Few more analysis fixes

#rnx

diff --git a/Engine/Source/Programs/Shared/EpicGames.Tracing/UnrealInsights/Reader.cs b/Engine/Source/Programs/Shared/EpicGames.Tracing/UnrealInsights/Reader.cs
--- a/Engine/Source/Programs/Shared/EpicGames.Tracing/UnrealInsights/Reader.cs
+++ b/Engine/Source/Programs/Shared/EpicGames.Tracing/UnrealInsights/Reader.cs
@@ -111,11 +111,10 @@ namespace EpicGames.Tracing.UnrealInsights
 			while (reader.BaseStream.Position != reader.BaseStream.Length)
 			{
 				ushort uid;
-				ushort size = 0;
 				if (threadId == TransportPacket.ThreadIdEvents || threadId == TransportPacket.ThreadIdImportants)
 				{
 					uid = reader.ReadUInt16();
-					size = reader.ReadUInt16(); // TODO: validate and respect the size
+					reader.ReadUInt16(); // TODO: validate and respect the size
 				}
 				else
 				{
diff --git a/Engine/Source/Programs/Shared/EpicGames.UBA/ILogger.cs b/Engine/Source/Programs/Shared/EpicGames.UBA/ILogger.cs
--- a/Engine/Source/Programs/Shared/EpicGames.UBA/ILogger.cs
+++ b/Engine/Source/Programs/Shared/EpicGames.UBA/ILogger.cs
@@ -59,11 +59,10 @@ namespace EpicGames.UBA
 		/// Create a ILogger object
 		/// </summary>
 		/// <param name="logger">The Microsoft.Extensions.Logging.ILogger to wrap</param>
-		/// <param name="showDetail">Will log out detailed log entries from native uba to logger</param>
 		/// <returns>The ILogger</returns>
-		public static ILogger CreateLogger(Microsoft.Extensions.Logging.ILogger logger, bool showDetail)
+		public static ILogger CreateLogger(Microsoft.Extensions.Logging.ILogger logger)
 		{
-			return new LoggerImpl(logger, showDetail);
+			return new LoggerImpl(logger);
 		}
 	}
 }
\ No newline at end of file
diff --git a/Engine/Source/Programs/Shared/EpicGames.UBA/Impl/LoggerImpl.cs b/Engine/Source/Programs/Shared/EpicGames.UBA/Impl/LoggerImpl.cs
--- a/Engine/Source/Programs/Shared/EpicGames.UBA/Impl/LoggerImpl.cs
+++ b/Engine/Source/Programs/Shared/EpicGames.UBA/Impl/LoggerImpl.cs
@@ -29,7 +29,7 @@ namespace EpicGames.UBA
 		static extern void DestroyCallbackLogWriter(IntPtr logger);
 		#endregion
 
-		public LoggerImpl(Microsoft.Extensions.Logging.ILogger logger, bool showDetail)
+		public LoggerImpl(Microsoft.Extensions.Logging.ILogger logger)
 		{
 			_logger = logger;
 			_beginScopeCallbackDelegate = BeginScope;
diff --git a/Engine/Source/Programs/Shared/EpicGames.UHT/Parsers/UhtHeaderFileParser.cs b/Engine/Source/Programs/Shared/EpicGames.UHT/Parsers/UhtHeaderFileParser.cs
--- a/Engine/Source/Programs/Shared/EpicGames.UHT/Parsers/UhtHeaderFileParser.cs
+++ b/Engine/Source/Programs/Shared/EpicGames.UHT/Parsers/UhtHeaderFileParser.cs
@@ -39,6 +39,7 @@ namespace EpicGames.UHT.Parsers
 	/// Compiler directives
 	/// </summary>
 	[Flags]
+	[SuppressMessage("Usage", "CA2217:Do not mark enums with FlagsAttribute")]
 	public enum UhtCompilerDirective
 	{
 		/// <summary>
diff --git a/Engine/Source/Programs/Shared/EpicGames.UHT/Types/UhtType.cs b/Engine/Source/Programs/Shared/EpicGames.UHT/Types/UhtType.cs
--- a/Engine/Source/Programs/Shared/EpicGames.UHT/Types/UhtType.cs
+++ b/Engine/Source/Programs/Shared/EpicGames.UHT/Types/UhtType.cs
@@ -2,6 +2,7 @@
 
 using System;
 using System.Collections.Generic;
+using System.Diagnostics.CodeAnalysis;
 using System.Text;
 using System.Text.Json.Serialization;
 using System.Threading;
@@ -528,6 +529,7 @@ namespace EpicGames.UHT.Types
 	/// Not all types support the given options.
 	/// </summary>
 	[Flags]
+	[SuppressMessage("Usage", "CA2217:Do not mark enums with FlagsAttribute")]
 	public enum UhtDefineScope
 	{
 
diff --git a/Engine/Source/Programs/UnrealBuildTool/Executors/UnrealBuildAccelerator/UBAExecutor.cs b/Engine/Source/Programs/UnrealBuildTool/Executors/UnrealBuildAccelerator/UBAExecutor.cs
--- a/Engine/Source/Programs/UnrealBuildTool/Executors/UnrealBuildAccelerator/UBAExecutor.cs
+++ b/Engine/Source/Programs/UnrealBuildTool/Executors/UnrealBuildAccelerator/UBAExecutor.cs
@@ -239,7 +239,7 @@ namespace UnrealBuildTool
 					_ = Task.Run(LaunchVisualizer);
 				}
 
-				using EpicGames.UBA.ILogger ubaLogger = EpicGames.UBA.ILogger.CreateLogger(logger, UBAConfig.bDetailedLog);
+				using EpicGames.UBA.ILogger ubaLogger = EpicGames.UBA.ILogger.CreateLogger(logger);
 				using (Server = IServer.CreateServer(UBAConfig.MaxWorkers, UBAConfig.SendSize, ubaLogger, UBAConfig.bUseQuic))
 				{
 					using IStorageServer ubaStorageServer = IStorageServer.CreateStorageServer(Server, ubaLogger, new StorageServerCreateInfo(_rootDirRef.FullName, ((ulong)UBAConfig.StoreCapacityGb) * 1000 * 1000 * 1000, !UBAConfig.bStoreRaw, UBAConfig.Zone));
