UnrealBuildTool: Fix inconsistent configuration in WithAutomationTests

#rnx

diff --git a/Engine/Source/Programs/UnrealBuildTool/Configuration/ReadOnlyTargetRules.cs b/Engine/Source/Programs/UnrealBuildTool/Configuration/ReadOnlyTargetRules.cs
--- a/Engine/Source/Programs/UnrealBuildTool/Configuration/ReadOnlyTargetRules.cs
+++ b/Engine/Source/Programs/UnrealBuildTool/Configuration/ReadOnlyTargetRules.cs
@@ -662,7 +662,7 @@ namespace UnrealBuildTool
 			{
 				if (!_bWithAutomationTestsPrivate.HasValue)
 				{
-					bool bCompileDevTests = Configuration != UnrealTargetConfiguration.Shipping;
+					bool bCompileDevTests = (Configuration != UnrealTargetConfiguration.Test && Configuration != UnrealTargetConfiguration.Shipping);
 					bool bCompilePerfTests = bCompileDevTests;
 
 					if (bForceCompileDevelopmentAutomationTests)
