UnrealBuildTool: Resolve IDE0018, IDE0019, IDE0034, IDE0038, IDE0041, IDE0054, IDE0057, IDE0074, IDE0090

#rnx

diff --git a/Engine/Platforms/VisionOS/Source/Programs/UnrealBuildTool/UEBuildVisionOS.cs b/Engine/Platforms/VisionOS/Source/Programs/UnrealBuildTool/UEBuildVisionOS.cs
--- a/Engine/Platforms/VisionOS/Source/Programs/UnrealBuildTool/UEBuildVisionOS.cs
+++ b/Engine/Platforms/VisionOS/Source/Programs/UnrealBuildTool/UEBuildVisionOS.cs
@@ -126,7 +126,7 @@ namespace UnrealBuildTool
 		/// </summary>
 		public override void RegisterBuildPlatforms(ILogger Logger)
 		{
-			ApplePlatformSDK SDK = new ApplePlatformSDK(Logger);
+			ApplePlatformSDK SDK = new(Logger);
 
 			// Register this build platform for IOS
 			UEBuildPlatform.RegisterBuildPlatform(new VisionOSPlatform(SDK, Logger), Logger);
diff --git a/Engine/Platforms/VisionOS/Source/Programs/UnrealBuildTool/UEDeployVisionOS.cs b/Engine/Platforms/VisionOS/Source/Programs/UnrealBuildTool/UEDeployVisionOS.cs
--- a/Engine/Platforms/VisionOS/Source/Programs/UnrealBuildTool/UEDeployVisionOS.cs
+++ b/Engine/Platforms/VisionOS/Source/Programs/UnrealBuildTool/UEDeployVisionOS.cs
@@ -46,36 +46,30 @@ namespace UnrealBuildTool
 			ConfigHierarchy Ini = ConfigCache.ReadHierarchy(ConfigHierarchyType.Engine, DirRef, UnrealTargetPlatform.IOS);
 
 			// bundle display name
-			string BundleDisplayName;
-			Ini.GetString("/Script/IOSRuntimeSettings.IOSRuntimeSettings", "BundleDisplayName", out BundleDisplayName);
+			Ini.GetString("/Script/IOSRuntimeSettings.IOSRuntimeSettings", "BundleDisplayName", out string BundleDisplayName);
 
 			// bundle identifier
-			string BundleIdentifier;
-			Ini.GetString("/Script/IOSRuntimeSettings.IOSRuntimeSettings", "BundleIdentifier", out BundleIdentifier);
+			Ini.GetString("/Script/IOSRuntimeSettings.IOSRuntimeSettings", "BundleIdentifier", out string BundleIdentifier);
 			if (!String.IsNullOrEmpty(BundleID))
 			{
 				BundleIdentifier = BundleID;
 			}
 
 			// bundle name
-			string BundleName;
-			Ini.GetString("/Script/IOSRuntimeSettings.IOSRuntimeSettings", "BundleName", out BundleName);
+			Ini.GetString("/Script/IOSRuntimeSettings.IOSRuntimeSettings", "BundleName", out string BundleName);
 
 			// short version string
-			string BundleShortVersion;
-			Ini.GetString("/Script/IOSRuntimeSettings.IOSRuntimeSettings", "VersionInfo", out BundleShortVersion);
+			Ini.GetString("/Script/IOSRuntimeSettings.IOSRuntimeSettings", "VersionInfo", out string BundleShortVersion);
 
 			// required capabilities
 			string RequiredCaps = "\t\t<string>arm64</string>\n";
 
 			// minimum iOS version
-			string MinVersionSetting = "";
-			Ini.GetString("/Script/IOSRuntimeSettings.IOSRuntimeSettings", "MinimumiOSVersion", out MinVersionSetting);
+			Ini.GetString("/Script/IOSRuntimeSettings.IOSRuntimeSettings", "MinimumiOSVersion", out string MinVersionSetting);
 			string MinVersion = GetMinimumOSVersion(MinVersionSetting, Logger);
 
 			// extra plist data
-			string ExtraData = "";
-			Ini.GetString("/Script/IOSRuntimeSettings.IOSRuntimeSettings", "AdditionalPlistData", out ExtraData);
+			Ini.GetString("/Script/IOSRuntimeSettings.IOSRuntimeSettings", "AdditionalPlistData", out string ExtraData);
 
 			// create the final display name, including converting all entities for XML use
 			string FinalDisplayName = BundleDisplayName.Replace("[PROJECT_NAME]", ProjectName).Replace("_", "");
@@ -86,7 +80,7 @@ namespace UnrealBuildTool
 			FinalDisplayName = FinalDisplayName.Replace(">", "&gt;");
 
 			// generate the plist file
-			StringBuilder Text = new StringBuilder();
+			StringBuilder Text = new();
 			Text.AppendLine("<?xml version=\"1.0\" encoding=\"UTF-8\"?>");
 			Text.AppendLine("<!DOCTYPE plist PUBLIC \"-//Apple//DTD PLIST 1.0//EN\" \"http://www.apple.com/DTDs/PropertyList-1.0.dtd\">");
 			Text.AppendLine("<plist version=\"1.0\">");
diff --git a/Engine/Source/Programs/UnrealBuildTool/Configuration/UEBuildPlatform.cs b/Engine/Source/Programs/UnrealBuildTool/Configuration/UEBuildPlatform.cs
--- a/Engine/Source/Programs/UnrealBuildTool/Configuration/UEBuildPlatform.cs
+++ b/Engine/Source/Programs/UnrealBuildTool/Configuration/UEBuildPlatform.cs
@@ -415,10 +415,7 @@ namespace UnrealBuildTool
 		/// </summary>
 		public IReadOnlySet<string> GetExcludedFolderNames()
 		{
-			if (CachedExcludedFolderNames == null)
-			{
-				CachedExcludedFolderNames = new HashSet<string>(GetPlatformFolderNames().Except(GetIncludedFolderNames()), DirectoryReference.Comparer);
-			}
+			CachedExcludedFolderNames ??= new HashSet<string>(GetPlatformFolderNames().Except(GetIncludedFolderNames()), DirectoryReference.Comparer);
 			return CachedExcludedFolderNames;
 		}
 
diff --git a/Engine/Source/Programs/UnrealBuildTool/Configuration/UEBuildTarget.cs b/Engine/Source/Programs/UnrealBuildTool/Configuration/UEBuildTarget.cs
--- a/Engine/Source/Programs/UnrealBuildTool/Configuration/UEBuildTarget.cs
+++ b/Engine/Source/Programs/UnrealBuildTool/Configuration/UEBuildTarget.cs
@@ -38,10 +38,7 @@ namespace UnrealBuildTool
 		// here this means initializing "StringRegistry" directly to "new UniqueStringRegistry()" may not be executed before FindOrAddByName() has been called as part of initializing a static member variable of another part of the partial struct
 		private static UniqueStringRegistry GetUniqueStringRegistry()
 		{
-			if (StringRegistry == null)
-			{
-				StringRegistry = new UniqueStringRegistry();
-			}
+			StringRegistry ??= new UniqueStringRegistry();
 			return StringRegistry;
 		}
 
@@ -124,7 +121,7 @@ namespace UnrealBuildTool
 		/// <returns></returns>
 		public override bool Equals(object? B)
 		{
-			if (Object.ReferenceEquals(B, null))
+			if (B is null)
 			{
 				return false;
 			}
@@ -369,10 +366,7 @@ namespace UnrealBuildTool
 		// #jira UE-88908 (see above)
 		private static UniqueStringRegistry GetUniqueStringRegistry()
 		{
-			if (StringRegistry == null)
-			{
-				StringRegistry = new UniqueStringRegistry();
-			}
+			StringRegistry ??= new UniqueStringRegistry();
 			return StringRegistry;
 		}
 
@@ -443,7 +437,7 @@ namespace UnrealBuildTool
 		/// <returns></returns>
 		public override bool Equals(object? B)
 		{
-			if (Object.ReferenceEquals(B, null))
+			if (B is null)
 			{
 				return false;
 			}
@@ -594,10 +588,7 @@ namespace UnrealBuildTool
 		// here this means initializing "StringRegistry" directly to "new UniqueStringRegistry()" may not be executed before FindOrAddByName() has been called as part of initializing a static member variable of another part of the partial struct
 		private static UniqueStringRegistry GetUniqueStringRegistry()
 		{
-			if (StringRegistry == null)
-			{
-				StringRegistry = new UniqueStringRegistry();
-			}
+			StringRegistry ??= new UniqueStringRegistry();
 			return StringRegistry;
 		}
 
@@ -674,7 +665,7 @@ namespace UnrealBuildTool
 		/// <returns></returns>
 		public override bool Equals(object? B)
 		{
-			if (Object.ReferenceEquals(B, null))
+			if (B is null)
 			{
 				return false;
 			}
diff --git a/Engine/Source/Programs/UnrealBuildTool/Executors/ImmediateActionQueue.cs b/Engine/Source/Programs/UnrealBuildTool/Executors/ImmediateActionQueue.cs
--- a/Engine/Source/Programs/UnrealBuildTool/Executors/ImmediateActionQueue.cs
+++ b/Engine/Source/Programs/UnrealBuildTool/Executors/ImmediateActionQueue.cs
@@ -887,10 +887,7 @@ namespace UnrealBuildTool
 					lock (_actionsToLog)
 					{
 						_actionsToLog.Add(action.SortIndex);
-						if (_actionsToLogTask == null)
-						{
-							_actionsToLogTask = Task.Run(LogActions);
-						}
+						_actionsToLogTask ??= Task.Run(LogActions);
 					}
 				}
 
diff --git a/Engine/Source/Programs/UnrealBuildTool/Modes/BuildMode.cs b/Engine/Source/Programs/UnrealBuildTool/Modes/BuildMode.cs
--- a/Engine/Source/Programs/UnrealBuildTool/Modes/BuildMode.cs
+++ b/Engine/Source/Programs/UnrealBuildTool/Modes/BuildMode.cs
@@ -127,10 +127,7 @@ namespace UnrealBuildTool
 			}
 
 			// Fixup the log path if it wasn't overridden by a config file
-			if (BaseLogFileName == null)
-			{
-				BaseLogFileName = FileReference.Combine(Unreal.EngineProgramSavedDirectory, "UnrealBuildTool", "Log.txt").FullName;
-			}
+			BaseLogFileName ??= FileReference.Combine(Unreal.EngineProgramSavedDirectory, "UnrealBuildTool", "Log.txt").FullName;
 
 			// Create the log file, and flush the startup listener to it
 			if (!Arguments.HasOption("-NoLog") && !Log.HasFileWriter())
diff --git a/Engine/Source/Programs/UnrealBuildTool/Modes/IWYUMode.cs b/Engine/Source/Programs/UnrealBuildTool/Modes/IWYUMode.cs
--- a/Engine/Source/Programs/UnrealBuildTool/Modes/IWYUMode.cs
+++ b/Engine/Source/Programs/UnrealBuildTool/Modes/IWYUMode.cs
@@ -337,7 +337,7 @@ namespace UnrealBuildTool
 			string NewModulePath = ModulePath.Replace('\\', '/');
 			if (!NewModulePath.EndsWith('/'))
 			{
-				NewModulePath = NewModulePath + '/';
+				NewModulePath += '/';
 			}
 			return NewModulePath;
 		}
@@ -357,10 +357,7 @@ namespace UnrealBuildTool
 			Logger.LogInformation($"====================================================");
 
 			// Fixup the log path if it wasn't overridden by a config file
-			if (IWYUBaseLogFileName == null)
-			{
-				IWYUBaseLogFileName = FileReference.Combine(Unreal.EngineProgramSavedDirectory, "UnrealBuildTool", "IWYULog.txt").FullName;
-			}
+			IWYUBaseLogFileName ??= FileReference.Combine(Unreal.EngineProgramSavedDirectory, "UnrealBuildTool", "IWYULog.txt").FullName;
 
 			// Create the log file, and flush the startup listener to it
 			if (!Arguments.HasOption("-NoLog") && !Log.HasFileWriter())
@@ -388,7 +385,7 @@ namespace UnrealBuildTool
 
 			if (TargetDescriptors.Count != 1)
 			{
-				Logger.LogError($"IWYUMode can only handle command lines that produce one target (Cmdline: {Arguments.ToString()})");
+				Logger.LogError($"IWYUMode can only handle command lines that produce one target (Cmdline: {Arguments})");
 				return 0;
 			}
 
@@ -1487,10 +1484,7 @@ namespace UnrealBuildTool
 						// Might be forward declaration.. 
 						if (ForwardDeclarationsToAdd.Remove(Line)) // Skip adding the ones that already exists
 						{
-							if (FirstForwardDeclareLine == null)
-							{
-								FirstForwardDeclareLine = Line;
-							}
+							FirstForwardDeclareLine ??= Line;
 						}
 
 						if (Line.Contains("IWYU pragma: "))
@@ -2202,7 +2196,7 @@ namespace UnrealBuildTool
 				}
 
 				float Percent = 100.0f - (((float)OptimizedSize) / (OptimizedSize + Saved) * 100.0f);
-				Logger.LogInformation($"{Path.GetFileName(File.File)}   {PrettySize(OptimizedSize + Saved)} -> {PrettySize(OptimizedSize)}  (Saved {PrettySize(Saved)} or {Percent.ToString("0.0")}%)");
+				Logger.LogInformation($"{Path.GetFileName(File.File)}   {PrettySize(OptimizedSize + Saved)} -> {PrettySize(OptimizedSize)}  (Saved {PrettySize(Saved)} or {Percent:0.0}%)");
 			}
 			Logger.LogInformation("");
 
diff --git a/Engine/Source/Programs/UnrealBuildTool/Modes/PipInstallMode.cs b/Engine/Source/Programs/UnrealBuildTool/Modes/PipInstallMode.cs
--- a/Engine/Source/Programs/UnrealBuildTool/Modes/PipInstallMode.cs
+++ b/Engine/Source/Programs/UnrealBuildTool/Modes/PipInstallMode.cs
@@ -120,10 +120,7 @@ namespace UnrealBuildTool.Modes
 				return 0;
 			}
 
-			if (Action == null)
-			{
-				Action = PipAction.Install;
-			}
+			Action ??= PipAction.Install;
 
 			DirectoryReference ProjectDir = DirectoryReference.FromFile(TargetDescriptor.ProjectFile);
 			DirectoryReference InstallDir = DirectoryReference.Combine(ProjectDir, "Intermediate", "PipInstall");
diff --git a/Engine/Source/Programs/UnrealBuildTool/Modes/ProfileUnitySizesMode.cs b/Engine/Source/Programs/UnrealBuildTool/Modes/ProfileUnitySizesMode.cs
--- a/Engine/Source/Programs/UnrealBuildTool/Modes/ProfileUnitySizesMode.cs
+++ b/Engine/Source/Programs/UnrealBuildTool/Modes/ProfileUnitySizesMode.cs
@@ -286,10 +286,7 @@ namespace UnrealBuildTool
 			}
 
 			TimingData? BestTimingData = AllTimingData.MinBy(TimingData => TimingData.ExecutorTiming);
-			if (BestTimingData == null)
-			{
-				BestTimingData = AllTimingData[0];
-			}
+			BestTimingData ??= AllTimingData[0];
 			return BestTimingData;
 		}
 
diff --git a/Engine/Source/Programs/UnrealBuildTool/Platform/Android/UEDeployAndroid.cs b/Engine/Source/Programs/UnrealBuildTool/Platform/Android/UEDeployAndroid.cs
--- a/Engine/Source/Programs/UnrealBuildTool/Platform/Android/UEDeployAndroid.cs
+++ b/Engine/Source/Programs/UnrealBuildTool/Platform/Android/UEDeployAndroid.cs
@@ -425,10 +425,7 @@ namespace UnrealBuildTool
 		protected bool ReadPackageDataInsideApkFromIni(ConfigHierarchy? Ini)
 		{
 			// make a new one if one wasn't passed in
-			if (Ini == null)
-			{
-				Ini = GetConfigCacheIni(ConfigHierarchyType.Engine);
-			}
+			Ini ??= GetConfigCacheIni(ConfigHierarchyType.Engine);
 
 			// we check this a lot, so make it easy 
 			bool bIniPackageDataInsideApk;
@@ -445,10 +442,7 @@ namespace UnrealBuildTool
 			}
 
 			// make a new one if one wasn't passed in
-			if (Ini == null)
-			{
-				Ini = GetConfigCacheIni(ConfigHierarchyType.Engine);
-			}
+			Ini ??= GetConfigCacheIni(ConfigHierarchyType.Engine);
 
 			// we check this a lot, so make it easy 
 			bool bUseExternalFilesDir;
@@ -478,10 +472,7 @@ namespace UnrealBuildTool
 		public bool IsPackagingForMetaQuest(ConfigHierarchy? Ini = null)
 		{
 			// make a new one if one wasn't passed in
-			if (Ini == null)
-			{
-				Ini = GetConfigCacheIni(ConfigHierarchyType.Engine);
-			}
+			Ini ??= GetConfigCacheIni(ConfigHierarchyType.Engine);
 
 			List<string> TargetOculusDevices = GetTargetOculusMobileDevices(Ini); // Backcompat for deprecated oculus device target setting
 			bool bTargetOculusDevices = (TargetOculusDevices != null && TargetOculusDevices.Count() > 0); // Backcompat for deprecated oculus device target setting
@@ -494,10 +485,7 @@ namespace UnrealBuildTool
 		public bool DisableVerifyOBBOnStartUp(ConfigHierarchy? Ini = null)
 		{
 			// make a new one if one wasn't passed in
-			if (Ini == null)
-			{
-				Ini = GetConfigCacheIni(ConfigHierarchyType.Engine);
-			}
+			Ini ??= GetConfigCacheIni(ConfigHierarchyType.Engine);
 
 			// we check this a lot, so make it easy 
 			bool bDisableVerifyOBBOnStartUp;
@@ -609,7 +597,7 @@ namespace UnrealBuildTool
 
 					// preserve timestamp and clear read-only flags
 					FileInfo DestFileInfo = new FileInfo(DestFilename);
-					DestFileInfo.Attributes = DestFileInfo.Attributes & ~FileAttributes.ReadOnly;
+					DestFileInfo.Attributes &= ~FileAttributes.ReadOnly;
 					File.SetLastWriteTimeUtc(DestFilename, File.GetLastWriteTimeUtc(Filename));
 				}
 			}
@@ -640,7 +628,7 @@ namespace UnrealBuildTool
 				{
 					// remove any read only flags
 					FileInfo FileInfo = new FileInfo(Filename);
-					FileInfo.Attributes = FileInfo.Attributes & ~FileAttributes.ReadOnly;
+					FileInfo.Attributes &= ~FileAttributes.ReadOnly;
 					FileInfo.Delete();
 				}
 				catch (Exception)
@@ -836,7 +824,7 @@ namespace UnrealBuildTool
 
 					// preserve timestamp and clear read-only flags
 					FileInfo DestFileInfo = new FileInfo(DestFilename);
-					DestFileInfo.Attributes = DestFileInfo.Attributes & ~FileAttributes.ReadOnly;
+					DestFileInfo.Attributes &= ~FileAttributes.ReadOnly;
 					File.SetLastWriteTimeUtc(DestFilename, File.GetLastWriteTimeUtc(Filename));
 
 					Logger.LogInformation("Copied file {DestFilename}.", DestFilename);
@@ -1027,7 +1015,7 @@ namespace UnrealBuildTool
 						string EmbeddedAppendDisplayName = "";
 						if (Ini.GetString("/Script/AndroidRuntimeSettings.AndroidRuntimeSettings", "EmbeddedAppendDisplayName", out EmbeddedAppendDisplayName))
 						{
-							VersionDisplayName = VersionDisplayName + EmbeddedAppendDisplayName;
+							VersionDisplayName += EmbeddedAppendDisplayName;
 						}
 					}
 				}
@@ -1460,7 +1448,7 @@ namespace UnrealBuildTool
 					SafeDeleteFile(DestFilename);
 					File.Copy(SourceFilename, DestFilename);
 					FileInfo DestFileInfo = new FileInfo(DestFilename);
-					DestFileInfo.Attributes = DestFileInfo.Attributes & ~FileAttributes.ReadOnly;
+					DestFileInfo.Attributes &= ~FileAttributes.ReadOnly;
 					File.SetLastWriteTimeUtc(DestFilename, File.GetLastWriteTimeUtc(SourceFilename));
 				}
 				else
@@ -3268,7 +3256,7 @@ namespace UnrealBuildTool
 					{
 						// remove any read only flags if invalid so it can be replaced
 						FileInfo DestFileInfo = new FileInfo(Filename);
-						DestFileInfo.Attributes = DestFileInfo.Attributes & ~FileAttributes.ReadOnly;
+						DestFileInfo.Attributes &= ~FileAttributes.ReadOnly;
 
 						//preserve the rest of the file, just fix up this line
 						string NewLine = Line.Replace("\"app_id\">" + XmlAppId + "</string>", "\"app_id\">" + ReplacementId + "</string>");
@@ -5121,7 +5109,7 @@ popd
 
 					// remove any read only flags
 					FileInfo DestFileInfo = new FileInfo(FinalSOName);
-					DestFileInfo.Attributes = DestFileInfo.Attributes & ~FileAttributes.ReadOnly;
+					DestFileInfo.Attributes &= ~FileAttributes.ReadOnly;
 					File.SetLastWriteTimeUtc(FinalSOName, File.GetLastWriteTimeUtc(SourceSOName));
 				}
 
@@ -5397,7 +5385,7 @@ popd
 
 							// preserve timestamp and clear read-only flags
 							FileInfo DestFileInfo = new FileInfo(DestFilename);
-							DestFileInfo.Attributes = DestFileInfo.Attributes & ~FileAttributes.ReadOnly;
+							DestFileInfo.Attributes &= ~FileAttributes.ReadOnly;
 							File.SetLastWriteTimeUtc(DestFilename, File.GetLastWriteTimeUtc(Filename));
 
 							Logger.LogInformation("Copied file {DestFilename}.", DestFilename);
diff --git a/Engine/Source/Programs/UnrealBuildTool/Platform/IOS/IOSCodeSigning.cs b/Engine/Source/Programs/UnrealBuildTool/Platform/IOS/IOSCodeSigning.cs
--- a/Engine/Source/Programs/UnrealBuildTool/Platform/IOS/IOSCodeSigning.cs
+++ b/Engine/Source/Programs/UnrealBuildTool/Platform/IOS/IOSCodeSigning.cs
@@ -1380,12 +1380,12 @@ namespace UnrealBuildTool
 							if (File.Exists(TargetFile))
 							{
 								DestFileInfo = new FileInfo(TargetFile);
-								DestFileInfo.Attributes = DestFileInfo.Attributes & ~FileAttributes.ReadOnly;
+								DestFileInfo.Attributes &= ~FileAttributes.ReadOnly;
 							}
 							_logger?.LogInformation("  Copying {0} -> {1}", Provision, TargetFile);
 							File.Copy(Provision, TargetFile, true);
 							DestFileInfo = new FileInfo(TargetFile);
-							DestFileInfo.Attributes = DestFileInfo.Attributes & ~FileAttributes.ReadOnly;
+							DestFileInfo.Attributes &= ~FileAttributes.ReadOnly;
 							if (!File.Exists(TargetFile))
 							{
 								_logger?.LogInformation("ERROR: Failed to copy {0} -> {1}", Provision, TargetFile);
diff --git a/Engine/Source/Programs/UnrealBuildTool/Platform/IOS/IOSToolChain.cs b/Engine/Source/Programs/UnrealBuildTool/Platform/IOS/IOSToolChain.cs
--- a/Engine/Source/Programs/UnrealBuildTool/Platform/IOS/IOSToolChain.cs
+++ b/Engine/Source/Programs/UnrealBuildTool/Platform/IOS/IOSToolChain.cs
@@ -894,7 +894,7 @@ namespace UnrealBuildTool
 					Dir = file.Replace(Path.Combine(EngineDir, "Build", "TVOS"), IntermediateDir);
 					File.Copy(file, Dir, true);
 					FileInfo DestFileInfo = new FileInfo(Dir);
-					DestFileInfo.Attributes = DestFileInfo.Attributes & ~FileAttributes.ReadOnly;
+					DestFileInfo.Attributes &= ~FileAttributes.ReadOnly;
 				}
 				// copy the icons from the game directory if it has any
 				string[][] Images = {
@@ -926,7 +926,7 @@ namespace UnrealBuildTool
 
 						File.Copy(Image, Path.Combine(Dir, Images[Index][1], Images[Index][0]), true);
 						FileInfo DestFileInfo = new FileInfo(Path.Combine(Dir, Images[Index][1], Images[Index][0]));
-						DestFileInfo.Attributes = DestFileInfo.Attributes & ~FileAttributes.ReadOnly;
+						DestFileInfo.Attributes &= ~FileAttributes.ReadOnly;
 					}
 				}
 			}
@@ -953,7 +953,7 @@ namespace UnrealBuildTool
 					Dir = file.Replace(Path.Combine(EngineDir, "Build", "IOS"), IntermediateDir);
 					File.Copy(file, Dir, true);
 					FileInfo DestFileInfo = new FileInfo(Dir);
-					DestFileInfo.Attributes = DestFileInfo.Attributes & ~FileAttributes.ReadOnly;
+					DestFileInfo.Attributes &= ~FileAttributes.ReadOnly;
 				}
 				// copy the icons from the game directory if it has any
 				string[][] Images = {
@@ -984,7 +984,7 @@ namespace UnrealBuildTool
 
 						File.Copy(Image, Path.Combine(Dir, Images[Index][0]), true);
 						FileInfo DestFileInfo = new FileInfo(Path.Combine(Dir, Images[Index][0]));
-						DestFileInfo.Attributes = DestFileInfo.Attributes & ~FileAttributes.ReadOnly;
+						DestFileInfo.Attributes &= ~FileAttributes.ReadOnly;
 					}
 				}
 
diff --git a/Engine/Source/Programs/UnrealBuildTool/Platform/IOS/UEDeployIOS.cs b/Engine/Source/Programs/UnrealBuildTool/Platform/IOS/UEDeployIOS.cs
--- a/Engine/Source/Programs/UnrealBuildTool/Platform/IOS/UEDeployIOS.cs
+++ b/Engine/Source/Programs/UnrealBuildTool/Platform/IOS/UEDeployIOS.cs
@@ -1097,11 +1097,11 @@ namespace UnrealBuildTool
 						if (LocalFileExists)
 						{
 							DestFileInfo = new FileInfo(LocalProvisionFile);
-							DestFileInfo.Attributes = DestFileInfo.Attributes & ~FileAttributes.ReadOnly;
+							DestFileInfo.Attributes &= ~FileAttributes.ReadOnly;
 						}
 						File.Copy(Provision, LocalProvisionFile, true);
 						DestFileInfo = new FileInfo(LocalProvisionFile);
-						DestFileInfo.Attributes = DestFileInfo.Attributes & ~FileAttributes.ReadOnly;
+						DestFileInfo.Attributes &= ~FileAttributes.ReadOnly;
 					}
 				}
 			}
@@ -1211,11 +1211,11 @@ namespace UnrealBuildTool
 				if (File.Exists(Environment.GetEnvironmentVariable("HOME") + "/Library/MobileDevice/Provisioning Profiles/" + ProjectProvision))
 				{
 					DestFileInfo = new FileInfo(Environment.GetEnvironmentVariable("HOME") + "/Library/MobileDevice/Provisioning Profiles/" + ProjectProvision);
-					DestFileInfo.Attributes = DestFileInfo.Attributes & ~FileAttributes.ReadOnly;
+					DestFileInfo.Attributes &= ~FileAttributes.ReadOnly;
 				}
 				File.Copy(ProvisionWithPrefix, Environment.GetEnvironmentVariable("HOME") + "/Library/MobileDevice/Provisioning Profiles/" + ProjectProvision, true);
 				DestFileInfo = new FileInfo(Environment.GetEnvironmentVariable("HOME") + "/Library/MobileDevice/Provisioning Profiles/" + ProjectProvision);
-				DestFileInfo.Attributes = DestFileInfo.Attributes & ~FileAttributes.ReadOnly;
+				DestFileInfo.Attributes &= ~FileAttributes.ReadOnly;
 			}
 			if (!File.Exists(ProvisionWithPrefix) || Unreal.IsBuildMachine())
 			{
@@ -1252,11 +1252,11 @@ namespace UnrealBuildTool
 				if (File.Exists(Environment.GetEnvironmentVariable("HOME") + "/Library/MobileDevice/Provisioning Profiles/" + InProjectName + "_Distro.mobileprovision"))
 				{
 					DestFileInfo = new FileInfo(Environment.GetEnvironmentVariable("HOME") + "/Library/MobileDevice/Provisioning Profiles/" + InProjectName + "_Distro.mobileprovision");
-					DestFileInfo.Attributes = DestFileInfo.Attributes & ~FileAttributes.ReadOnly;
+					DestFileInfo.Attributes &= ~FileAttributes.ReadOnly;
 				}
 				File.Copy(ProvisionWithPrefix, Environment.GetEnvironmentVariable("HOME") + "/Library/MobileDevice/Provisioning Profiles/" + InProjectName + "_Distro.mobileprovision", true);
 				DestFileInfo = new FileInfo(Environment.GetEnvironmentVariable("HOME") + "/Library/MobileDevice/Provisioning Profiles/" + InProjectName + "_Distro.mobileprovision");
-				DestFileInfo.Attributes = DestFileInfo.Attributes & ~FileAttributes.ReadOnly;
+				DestFileInfo.Attributes &= ~FileAttributes.ReadOnly;
 			}
 
 			GeneratePList(ProjectFile, Config, InProjectDirectory, bIsUnrealGame, GameExeName, false, InProjectName, InEngineDir, AppDirectory, UPLScripts, BundleID, bBuildAsFramework);
@@ -1265,7 +1265,7 @@ namespace UnrealBuildTool
 			if (File.Exists(AppDirectory + "/" + GameName))
 			{
 				FileInfo GameFileInfo = new FileInfo(AppDirectory + "/" + GameName);
-				GameFileInfo.Attributes = GameFileInfo.Attributes & ~FileAttributes.ReadOnly;
+				GameFileInfo.Attributes &= ~FileAttributes.ReadOnly;
 			}
 
 			// copy the GameName binary
diff --git a/Engine/Source/Programs/UnrealBuildTool/Platform/Mac/ApplePlatformSDK.Versions.cs b/Engine/Source/Programs/UnrealBuildTool/Platform/Mac/ApplePlatformSDK.Versions.cs
--- a/Engine/Source/Programs/UnrealBuildTool/Platform/Mac/ApplePlatformSDK.Versions.cs
+++ b/Engine/Source/Programs/UnrealBuildTool/Platform/Mac/ApplePlatformSDK.Versions.cs
@@ -26,10 +26,7 @@ namespace UnrealBuildTool
 				DeploymentTarget = GetVersionFromConfig("EditorDeploymentTarget");
 			}
 
-			if (DeploymentTarget == null)
-			{
-				DeploymentTarget = GetRequiredVersionFromConfig("DeploymentTarget");
-			}
+			DeploymentTarget ??= GetRequiredVersionFromConfig("DeploymentTarget");
 
 			return DeploymentTarget;
 		}
@@ -49,10 +46,7 @@ namespace UnrealBuildTool
 				DeploymentTarget = GetVersionFromConfig("EditorBuildTarget");
 			}
 
-			if (DeploymentTarget == null)
-			{
-				DeploymentTarget = GetRequiredVersionFromConfig("BuildTarget");
-			}
+			DeploymentTarget ??= GetRequiredVersionFromConfig("BuildTarget");
 
 			return DeploymentTarget;
 		}
diff --git a/Engine/Source/Programs/UnrealBuildTool/Platform/Mac/MacToolChain.cs b/Engine/Source/Programs/UnrealBuildTool/Platform/Mac/MacToolChain.cs
--- a/Engine/Source/Programs/UnrealBuildTool/Platform/Mac/MacToolChain.cs
+++ b/Engine/Source/Programs/UnrealBuildTool/Platform/Mac/MacToolChain.cs
@@ -537,10 +537,7 @@ namespace UnrealBuildTool
 				AppendMacLine(FinalizeAppBundleScript, "cd \"{0}\"", BinariesPath.Replace("$", "\\$"));
 
 				string BundleVersion = LinkEnvironment.BundleVersion!;
-				if (BundleVersion == null)
-				{
-					BundleVersion = LoadEngineDisplayVersion();
-				}
+				BundleVersion ??= LoadEngineDisplayVersion();
 
 				string ExeName = Path.GetFileName(OutputFiles[0].AbsolutePath);
 				bool bIsLauncherProduct = ExeName.StartsWith("EpicGamesLauncher") || ExeName.StartsWith("EpicGamesBootstrapLauncher");
@@ -858,10 +855,7 @@ namespace UnrealBuildTool
 			{
 				// Add the output file to the command-line.
 				string? InstallName = LinkEnvironment.InstallName;
-				if (InstallName == null)
-				{
-					InstallName = String.Format("{0}/{1}", DylibsPath, Path.GetFileName(OutputFile.AbsolutePath).Replace($".dylib_{LinkEnvironment.Architecture}", ".dylib"));
-				}
+				InstallName ??= String.Format("{0}/{1}", DylibsPath, Path.GetFileName(OutputFile.AbsolutePath).Replace($".dylib_{LinkEnvironment.Architecture}", ".dylib"));
 				LinkCommand += String.Format(" -install_name \"{0}\"", InstallName);
 			}
 
diff --git a/Engine/Source/Programs/UnrealBuildTool/Preprocessor/Preprocessor.cs b/Engine/Source/Programs/UnrealBuildTool/Preprocessor/Preprocessor.cs
--- a/Engine/Source/Programs/UnrealBuildTool/Preprocessor/Preprocessor.cs
+++ b/Engine/Source/Programs/UnrealBuildTool/Preprocessor/Preprocessor.cs
@@ -1284,7 +1284,7 @@ namespace UnrealBuildTool
 			// Add an empty variable argument if one was not specified
 			if (macro.HasVariableArgumentList && arguments.Count == macro.Parameters!.Count - 1)
 			{
-				arguments.Add(new List<Token> { new Token(TokenType.Placemarker, TokenFlags.None) });
+				arguments.Add(new List<Token> { new(TokenType.Placemarker, TokenFlags.None) });
 			}
 
 			// Validate the argument list
diff --git a/Engine/Source/Programs/UnrealBuildTool/ProjectFiles/KDevelop/KDevelopGenerator.cs b/Engine/Source/Programs/UnrealBuildTool/ProjectFiles/KDevelop/KDevelopGenerator.cs
--- a/Engine/Source/Programs/UnrealBuildTool/ProjectFiles/KDevelop/KDevelopGenerator.cs
+++ b/Engine/Source/Programs/UnrealBuildTool/ProjectFiles/KDevelop/KDevelopGenerator.cs
@@ -108,7 +108,7 @@ namespace UnrealBuildTool
 			else if (Type == 3)
 			{
 				ToolType = "Clean";
-				ConfName = ConfName + " -clean";
+				ConfName += " -clean";
 			}
 
 			FileContent.Append(String.Format("[CustomBuildSystem][BuildConfig{0}][Tool{1}]\n", BuildConfigIndex, ToolType));
diff --git a/Engine/Source/Programs/UnrealBuildTool/ProjectFiles/ProjectFileGenerator.cs b/Engine/Source/Programs/UnrealBuildTool/ProjectFiles/ProjectFileGenerator.cs
--- a/Engine/Source/Programs/UnrealBuildTool/ProjectFiles/ProjectFileGenerator.cs
+++ b/Engine/Source/Programs/UnrealBuildTool/ProjectFiles/ProjectFileGenerator.cs
@@ -3007,10 +3007,7 @@ namespace UnrealBuildTool
 								string TargetTypeSuffix = /*EngineTarget.TargetRules!.Type == TargetType.Game ? "" :*/ EngineTarget.TargetRules!.Type.ToString();
 								ProjectFile = FindOrAddProject(GetProjectLocation($"{ProjectName}{TargetTypeSuffix}"), ContentOnlyGameProject.Directory, IncludeInGeneratedProjects: true, bAlreadyExisted: out _);
 							}
-							else if (ProjectFile == null)
-							{
-								ProjectFile = FindOrAddProject(GetProjectLocation(ProjectName), ContentOnlyGameProject.Directory, IncludeInGeneratedProjects: true, bAlreadyExisted: out _);
-							}
+							else ProjectFile ??= FindOrAddProject(GetProjectLocation(ProjectName), ContentOnlyGameProject.Directory, IncludeInGeneratedProjects: true, bAlreadyExisted: out _);
 							ProjectFile.IsForeignProject = true;
 							ProjectFile.IsGeneratedProject = true;
 							ProjectFile.IsStubProject = false;
diff --git a/Engine/Source/Programs/UnrealBuildTool/ProjectFiles/Rider/ToolchainInfo.cs b/Engine/Source/Programs/UnrealBuildTool/ProjectFiles/Rider/ToolchainInfo.cs
--- a/Engine/Source/Programs/UnrealBuildTool/ProjectFiles/Rider/ToolchainInfo.cs
+++ b/Engine/Source/Programs/UnrealBuildTool/ProjectFiles/Rider/ToolchainInfo.cs
@@ -69,7 +69,7 @@ namespace UnrealBuildTool
 
 		public bool Equals(ToolchainInfo? Other)
 		{
-			if (ReferenceEquals(null, Other))
+			if (Other is null)
 			{
 				return false;
 			}
@@ -95,7 +95,7 @@ namespace UnrealBuildTool
 
 		public override bool Equals(object? Obj)
 		{
-			if (ReferenceEquals(null, Obj))
+			if (Obj is null)
 			{
 				return false;
 			}
diff --git a/Engine/Source/Programs/UnrealBuildTool/ProjectFiles/Xcode/XcodeProject.File.cs b/Engine/Source/Programs/UnrealBuildTool/ProjectFiles/Xcode/XcodeProject.File.cs
--- a/Engine/Source/Programs/UnrealBuildTool/ProjectFiles/Xcode/XcodeProject.File.cs
+++ b/Engine/Source/Programs/UnrealBuildTool/ProjectFiles/Xcode/XcodeProject.File.cs
@@ -282,7 +282,7 @@ namespace UnrealBuildTool.XcodeProjectXcconfig
 				}
 				else
 				{
-					GroupName = GroupName ?? (File.Reference.IsUnderDirectory(Unreal.EngineDirectory) ? "EngineReferences" : "ExternalReferences");
+					GroupName ??= (File.Reference.IsUnderDirectory(Unreal.EngineDirectory) ? "EngineReferences" : "ExternalReferences");
 					if (bIsFolder)
 					{
 						AddFolderReference(File.FileRefGuid, File.Reference.MakeRelativeTo(ProjectDirectory), GroupName);
diff --git a/Engine/Source/Programs/UnrealBuildTool/ProjectFiles/Xcode/XcodeProject.cs b/Engine/Source/Programs/UnrealBuildTool/ProjectFiles/Xcode/XcodeProject.cs
--- a/Engine/Source/Programs/UnrealBuildTool/ProjectFiles/Xcode/XcodeProject.cs
+++ b/Engine/Source/Programs/UnrealBuildTool/ProjectFiles/Xcode/XcodeProject.cs
@@ -622,10 +622,7 @@ namespace UnrealBuildTool.XcodeProjectXcconfig
 			}
 
 			// if the SubPath (or optional AlProjectsubPath) wasn't found, then fall back to engine location
-			if (FinalPath == null)
-			{
-				FinalPath = Path.Combine(Unreal.EngineDirectory.FullName, SubPath);
-			}
+			FinalPath ??= Path.Combine(Unreal.EngineDirectory.FullName, SubPath);
 			if (bMakeRelative)
 			{
 				FinalPath = new FileReference(FinalPath).MakeRelativeTo(XcodeProjectFileLocation.ParentDirectory!);
@@ -818,10 +815,7 @@ namespace UnrealBuildTool.XcodeProjectXcconfig
 		/// <param name="WrittenNodes"></param>
 		public static void WriteNodeAndReferences(StringBuilder Content, XcodeProjectNode Node, ILogger Logger, HashSet<XcodeProjectNode>? WrittenNodes = null)
 		{
-			if (WrittenNodes == null)
-			{
-				WrittenNodes = new();
-			}
+			WrittenNodes ??= new();
 
 			// write the node into the xcode project file
 			Node.Write(Content);
diff --git a/Engine/Source/Programs/UnrealBuildTool/System/ActionGraph.cs b/Engine/Source/Programs/UnrealBuildTool/System/ActionGraph.cs
--- a/Engine/Source/Programs/UnrealBuildTool/System/ActionGraph.cs
+++ b/Engine/Source/Programs/UnrealBuildTool/System/ActionGraph.cs
@@ -600,7 +600,7 @@ namespace UnrealBuildTool
 					if (!ActionIsNonCyclical.ContainsKey(Action))
 					{
 						string CycleDescription = "";
-						CycleDescription += $"Action #{ActionToIndex[Action].ToString()}: {Action.CommandPath}\n";
+						CycleDescription += $"Action #{ActionToIndex[Action]}: {Action.CommandPath}\n";
 						CycleDescription += $"\twith arguments: {Action.CommandArguments}\n";
 						foreach (FileItem PrerequisiteItem in Action.PrerequisiteItems)
 						{
@@ -623,11 +623,11 @@ namespace UnrealBuildTool
 										CyclicPrerequisiteAction.ProducedItems.ToList();
 									if (CyclicProducedItems.Count == 1)
 									{
-										CycleDescription += $"\t\t{ActionToIndex[CyclicPrerequisiteAction].ToString()} (produces: {CyclicProducedItems[0].AbsolutePath})\n";
+										CycleDescription += $"\t\t{ActionToIndex[CyclicPrerequisiteAction]} (produces: {CyclicProducedItems[0].AbsolutePath})\n";
 									}
 									else
 									{
-										CycleDescription += $"\t\t{ActionToIndex[CyclicPrerequisiteAction].ToString()}\n";
+										CycleDescription += $"\t\t{ActionToIndex[CyclicPrerequisiteAction]}\n";
 										foreach (FileItem CyclicProducedItem in CyclicProducedItems)
 										{
 											CycleDescription += $"\t\t\tproduces:   {CyclicProducedItem.AbsolutePath}\n";
diff --git a/Engine/Source/Programs/UnrealBuildTool/System/CommandLine.cs b/Engine/Source/Programs/UnrealBuildTool/System/CommandLine.cs
--- a/Engine/Source/Programs/UnrealBuildTool/System/CommandLine.cs
+++ b/Engine/Source/Programs/UnrealBuildTool/System/CommandLine.cs
@@ -110,7 +110,7 @@ namespace UnrealBuildTool
 						{
 							if (NonNullableType(FieldInfo.FieldType) != typeof(bool) && Attribute.Value == null && !Prefix.EndsWith("=") && !Prefix.EndsWith(":"))
 							{
-								Prefix = Prefix + "=";
+								Prefix += "=";
 							}
 						}
 						PrefixToParameter.Add(Prefix, new Parameter(Prefix, FieldInfo, Attribute));
diff --git a/Engine/Source/Programs/UnrealBuildTool/System/ConfigCache.cs b/Engine/Source/Programs/UnrealBuildTool/System/ConfigCache.cs
--- a/Engine/Source/Programs/UnrealBuildTool/System/ConfigCache.cs
+++ b/Engine/Source/Programs/UnrealBuildTool/System/ConfigCache.cs
@@ -339,10 +339,7 @@ namespace UnrealBuildTool
 				}
 			}
 
-			if (CustomConfig == null)
-			{
-				CustomConfig = String.Empty;
-			}
+			CustomConfig ??= String.Empty;
 
 			// Get the key to use for the cache. It cannot be null, so we use the engine directory if a project directory is not given.
 			ConfigHierarchyKey Key = new ConfigHierarchyKey(Type, ProjectDir, Platform, CustomConfig, OverrideStrings, HotfixDir);
diff --git a/Engine/Source/Programs/UnrealBuildTool/System/ConfigHierarchy.cs b/Engine/Source/Programs/UnrealBuildTool/System/ConfigHierarchy.cs
--- a/Engine/Source/Programs/UnrealBuildTool/System/ConfigHierarchy.cs
+++ b/Engine/Source/Programs/UnrealBuildTool/System/ConfigHierarchy.cs
@@ -454,7 +454,7 @@ namespace UnrealBuildTool
 			string? Text;
 			if (!TryGetValue(SectionName, KeyName, out Text))
 			{
-				Value = default(T);
+				Value = default;
 				return false;
 			}
 			return Enum.TryParse<T>(Text, out Value);
@@ -486,7 +486,7 @@ namespace UnrealBuildTool
 				return ConfigValueParser.TryParseGeneric( Line, out Value);
 			}
 
-			Value = default(T);
+			Value = default;
 			return false;
 		}
 
diff --git a/Engine/Source/Programs/UnrealBuildTool/System/ConfigValueTracker.cs b/Engine/Source/Programs/UnrealBuildTool/System/ConfigValueTracker.cs
--- a/Engine/Source/Programs/UnrealBuildTool/System/ConfigValueTracker.cs
+++ b/Engine/Source/Programs/UnrealBuildTool/System/ConfigValueTracker.cs
@@ -99,7 +99,7 @@ namespace UnrealBuildTool
 		/// <returns>True if the keys are equal, false otherwise</returns>
 		public bool Equals(ConfigDependencyKey? Other)
 		{
-			return !ReferenceEquals(Other, null) && Type == Other.Type && ProjectDir == Other.ProjectDir && Platform == Other.Platform && SectionName == Other.SectionName && KeyName == Other.KeyName;
+			return Other is not null && Type == Other.Type && ProjectDir == Other.ProjectDir && Platform == Other.Platform && SectionName == Other.SectionName && KeyName == Other.KeyName;
 		}
 
 		/// <summary>
diff --git a/Engine/Source/Programs/UnrealBuildTool/System/ExternalExecution.cs b/Engine/Source/Programs/UnrealBuildTool/System/ExternalExecution.cs
--- a/Engine/Source/Programs/UnrealBuildTool/System/ExternalExecution.cs
+++ b/Engine/Source/Programs/UnrealBuildTool/System/ExternalExecution.cs
@@ -433,10 +433,7 @@ namespace UnrealBuildTool
 					Info.GeneratedCPPFilenameBase = Path.Combine(GeneratedCodeDirectoryUHT.FullName, Info.ModuleName) + ".gen";
 					if (!Module.Rules.bUsePrecompiled)
 					{
-						if (Module.GeneratedCppDirectories == null)
-						{
-							Module.GeneratedCppDirectories = new List<string>();
-						}
+						Module.GeneratedCppDirectories ??= new List<string>();
 						Module.GeneratedCppDirectories.Add(GeneratedCodeDirectoryUHT.FullName);
 					}
 
diff --git a/Engine/Source/Programs/UnrealBuildTool/System/RulesAssembly.cs b/Engine/Source/Programs/UnrealBuildTool/System/RulesAssembly.cs
--- a/Engine/Source/Programs/UnrealBuildTool/System/RulesAssembly.cs
+++ b/Engine/Source/Programs/UnrealBuildTool/System/RulesAssembly.cs
@@ -182,10 +182,7 @@ namespace UnrealBuildTool
 				RulesAssembly? CurrentParent = Parent;
 				while (CurrentParent != null && CurrentParent.CompiledAssembly != null)
 				{
-					if (ReferencedAssembies == null)
-					{
-						ReferencedAssembies = new List<string>();
-					}
+					ReferencedAssembies ??= new List<string>();
 
 					ReferencedAssembies.Add(CurrentParent.CompiledAssembly.Location);
 					CurrentParent = CurrentParent.Parent;
@@ -845,10 +842,7 @@ namespace UnrealBuildTool
 				TargetName = TargetDescriptor.GetTestedName(TargetName);
 			}
 
-			if (Architectures == null)
-			{
-				Architectures = UnrealArchitectureConfig.ForPlatform(Platform).ActiveArchitectures(ProjectFile, TargetName);
-			}
+			Architectures ??= UnrealArchitectureConfig.ForPlatform(Platform).ActiveArchitectures(ProjectFile, TargetName);
 
 			bool bFoundTargetName = TargetNameToTargetFile.ContainsKey(TargetName);
 			if (bFoundTargetName == false)
diff --git a/Engine/Source/Programs/UnrealBuildTool/System/UnrealPluginLanguage.cs b/Engine/Source/Programs/UnrealBuildTool/System/UnrealPluginLanguage.cs
--- a/Engine/Source/Programs/UnrealBuildTool/System/UnrealPluginLanguage.cs
+++ b/Engine/Source/Programs/UnrealBuildTool/System/UnrealPluginLanguage.cs
@@ -911,10 +911,7 @@ namespace UnrealBuildTool
 
 		private ConfigCacheIni_UPL GetConfigCacheIni_UPL(string baseIniName)
 		{
-			if (ConfigCache == null)
-			{
-				ConfigCache = new Dictionary<string, ConfigCacheIni_UPL>();
-			}
+			ConfigCache ??= new Dictionary<string, ConfigCacheIni_UPL>();
 			ConfigCacheIni_UPL? config = null;
 			if (!ConfigCache.TryGetValue(baseIniName, out config))
 			{
@@ -991,7 +988,7 @@ namespace UnrealBuildTool
 
 					// remove any read only flags and keep timestamp
 					FileInfo DestFileInfo = new FileInfo(DestFilename);
-					DestFileInfo.Attributes = DestFileInfo.Attributes & ~FileAttributes.ReadOnly;
+					DestFileInfo.Attributes &= ~FileAttributes.ReadOnly;
 					File.SetLastWriteTimeUtc(DestFilename, File.GetLastWriteTimeUtc(Filename));
 				}
 			}
@@ -1786,7 +1783,7 @@ namespace UnrealBuildTool
 
 										// remove any read only flags and keep timestamp
 										FileInfo DestFileInfo = new FileInfo(Dst);
-										DestFileInfo.Attributes = DestFileInfo.Attributes & ~FileAttributes.ReadOnly;
+										DestFileInfo.Attributes &= ~FileAttributes.ReadOnly;
 										File.SetLastWriteTimeUtc(Dst, File.GetLastWriteTimeUtc(Src));
 									}
 								}
diff --git a/Engine/Source/Programs/UnrealBuildTool/System/Utils.cs b/Engine/Source/Programs/UnrealBuildTool/System/Utils.cs
--- a/Engine/Source/Programs/UnrealBuildTool/System/Utils.cs
+++ b/Engine/Source/Programs/UnrealBuildTool/System/Utils.cs
@@ -628,18 +628,12 @@ namespace UnrealBuildTool
 					if (C != UseDirectorySeparatorChar)
 					{
 						C = UseDirectorySeparatorChar;
-						if (CleanPath == null)
-						{
-							CleanPath = new StringBuilder(FilePath.Substring(0, Index), FilePath.Length);
-						}
+						CleanPath ??= new StringBuilder(FilePath.Substring(0, Index), FilePath.Length);
 					}
 
 					if (bCanCheckDoubleSeparators && C == PrevC)
 					{
-						if (CleanPath == null)
-						{
-							CleanPath = new StringBuilder(FilePath.Substring(0, Index), FilePath.Length);
-						}
+						CleanPath ??= new StringBuilder(FilePath.Substring(0, Index), FilePath.Length);
 						continue;
 					}
 				}
