UnrealBuildTool: Fix exposing static variables for various Unreal types, use get only properties instead

#rnx

diff --git a/Engine/Source/Programs/UnrealBuildTool/Configuration/UEBuildTarget.cs b/Engine/Source/Programs/UnrealBuildTool/Configuration/UEBuildTarget.cs
--- a/Engine/Source/Programs/UnrealBuildTool/Configuration/UEBuildTarget.cs
+++ b/Engine/Source/Programs/UnrealBuildTool/Configuration/UEBuildTarget.cs
@@ -245,42 +245,42 @@ namespace UnrealBuildTool
 		/// <summary>
 		/// 64-bit Windows
 		/// </summary>
-		public static UnrealTargetPlatform Win64 = FindOrAddByName("Win64");
+		public static UnrealTargetPlatform Win64 { get; } = FindOrAddByName("Win64");
 
 		/// <summary>
 		/// Mac
 		/// </summary>
-		public static UnrealTargetPlatform Mac = FindOrAddByName("Mac");
+		public static UnrealTargetPlatform Mac { get; } = FindOrAddByName("Mac");
 
 		/// <summary>
 		/// iOS
 		/// </summary>
-		public static UnrealTargetPlatform IOS = FindOrAddByName("IOS");
+		public static UnrealTargetPlatform IOS { get; } = FindOrAddByName("IOS");
 
 		/// <summary>
 		/// Android
 		/// </summary>
-		public static UnrealTargetPlatform Android = FindOrAddByName("Android");
+		public static UnrealTargetPlatform Android { get; } = FindOrAddByName("Android");
 
 		/// <summary>
 		/// Linux
 		/// </summary>
-		public static UnrealTargetPlatform Linux = FindOrAddByName("Linux");
+		public static UnrealTargetPlatform Linux { get; } = FindOrAddByName("Linux");
 
 		/// <summary>
 		/// LinuxArm64
 		/// </summary>
-		public static UnrealTargetPlatform LinuxArm64 = FindOrAddByName("LinuxArm64");
+		public static UnrealTargetPlatform LinuxArm64 { get; } = FindOrAddByName("LinuxArm64");
 
 		/// <summary>
 		/// TVOS
 		/// </summary>
-		public static UnrealTargetPlatform TVOS = FindOrAddByName("TVOS");
+		public static UnrealTargetPlatform TVOS { get; } = FindOrAddByName("TVOS");
 
 		/// <summary>
 		/// VisionOS
 		/// </summary>
-		public static UnrealTargetPlatform VisionOS = FindOrAddByName("VisionOS");
+		public static UnrealTargetPlatform VisionOS { get; } = FindOrAddByName("VisionOS");
 	}
 
 	internal class UnrealTargetPlatformTypeConverter : TypeConverter
@@ -513,57 +513,57 @@ namespace UnrealBuildTool
 		/// <summary>
 		/// this group is just to lump Win32 and Win64 into Windows directories, removing the special Windows logic in MakeListOfUnsupportedPlatforms
 		/// </summary>
-		public static UnrealPlatformGroup Windows = FindOrAddByName("Windows");
+		public static UnrealPlatformGroup Windows { get; } = FindOrAddByName("Windows");
 
 		/// <summary>
 		/// Microsoft platforms
 		/// </summary>
-		public static UnrealPlatformGroup Microsoft = FindOrAddByName("Microsoft");
+		public static UnrealPlatformGroup Microsoft { get; } = FindOrAddByName("Microsoft");
 
 		/// <summary>
 		/// Apple platforms
 		/// </summary>
-		public static UnrealPlatformGroup Apple = FindOrAddByName("Apple");
+		public static UnrealPlatformGroup Apple { get; } = FindOrAddByName("Apple");
 
 		/// <summary>
 		/// making IOS a group allows TVOS to compile IOS code
 		/// </summary>
-		public static UnrealPlatformGroup IOS = FindOrAddByName("IOS");
+		public static UnrealPlatformGroup IOS { get; } = FindOrAddByName("IOS");
 
 		/// <summary>
 		/// Unix platforms
 		/// </summary>
-		public static UnrealPlatformGroup Unix = FindOrAddByName("Unix");
+		public static UnrealPlatformGroup Unix { get; } = FindOrAddByName("Unix");
 
 		/// <summary>
 		/// Linux platforms
 		/// </summary>
-		public static UnrealPlatformGroup Linux = FindOrAddByName("Linux");
+		public static UnrealPlatformGroup Linux { get; } = FindOrAddByName("Linux");
 
 		/// <summary>
 		/// Android platforms
 		/// </summary>
-		public static UnrealPlatformGroup Android = FindOrAddByName("Android");
+		public static UnrealPlatformGroup Android { get; } = FindOrAddByName("Android");
 
 		/// <summary>
 		/// Desktop group - used by UnrealPlatformClass.Desktop
 		/// </summary>
-		public static UnrealPlatformGroup Desktop = FindOrAddByName("Desktop");
+		public static UnrealPlatformGroup Desktop { get; } = FindOrAddByName("Desktop");
 
 		/// <summary>
 		/// SDLPlatform is for platforms that use SDL for windows, cursors, etc
 		/// </summary>
-		public static UnrealPlatformGroup SDLPlatform = FindOrAddByName("SDLPlatform");
+		public static UnrealPlatformGroup SDLPlatform { get; } = FindOrAddByName("SDLPlatform");
 
 		/// <summary>
 		/// 30Hz is for platforms that typically run at 30Hz (Android, Switch) vs. the baseline 60Hz 
 		/// </summary>
-		public static UnrealPlatformGroup ThirtyHz = FindOrAddByName("30Hz");
+		public static UnrealPlatformGroup ThirtyHz { get; } = FindOrAddByName("30Hz");
 
 		/// <summary>
 		/// POSIX-compliant platforms
 		/// </summary>
-		public static UnrealPlatformGroup PosixOS = FindOrAddByName("PosixOS");
+		public static UnrealPlatformGroup PosixOS { get; } = FindOrAddByName("PosixOS");
 
 	}
 
@@ -824,22 +824,22 @@ namespace UnrealBuildTool
 		/// <summary>
 		/// 64-bit x86/intel
 		/// </summary>
-		public static UnrealArch X64 = FindOrAddByName("x64", bIsX64: true);
+		public static UnrealArch X64 { get; } = FindOrAddByName("x64", bIsX64: true);
 
 		/// <summary>
 		/// 64-bit arm
 		/// </summary>
-		public static UnrealArch Arm64 = FindOrAddByName("arm64", bIsX64: false);
+		public static UnrealArch Arm64 { get; } = FindOrAddByName("arm64", bIsX64: false);
 
 		/// <summary>
 		/// Used in place when needing to handle deprecated architectures that a licensee may still have. Do not use in normal logic
 		/// </summary>
-		public static UnrealArch Deprecated = FindOrAddByName("deprecated", bIsX64: false);
+		public static UnrealArch Deprecated { get; } = FindOrAddByName("deprecated", bIsX64: false);
 
 		/// <summary>
 		/// Maps to the currently running architecture on the host platform (lazy because this needs other classes to be intialized)
 		/// </summary>
-		public static Lazy<UnrealArch> Host = new(() => UnrealArchitectureConfig.ForPlatform(BuildHostPlatform.Current.Platform).GetHostArchitecture());
+		public static Lazy<UnrealArch> Host { get; } = new(() => UnrealArchitectureConfig.ForPlatform(BuildHostPlatform.Current.Platform).GetHostArchitecture());
 	}
 
 	internal class UnrealArchPlatformTypeConverter : TypeConverter
diff --git a/Engine/Source/Programs/UnrealBuildTool/Platform/IOS/UEBuildIOS.cs b/Engine/Source/Programs/UnrealBuildTool/Platform/IOS/UEBuildIOS.cs
--- a/Engine/Source/Programs/UnrealBuildTool/Platform/IOS/UEBuildIOS.cs
+++ b/Engine/Source/Programs/UnrealBuildTool/Platform/IOS/UEBuildIOS.cs
@@ -19,12 +19,12 @@ namespace UnrealBuildTool
 		/// <summary>
 		/// IOS Simulator
 		/// </summary>
-		public static UnrealArch IOSSimulator = FindOrAddByName("iossimulator", bIsX64: false);
+		public static UnrealArch IOSSimulator { get; } = FindOrAddByName("iossimulator", bIsX64: false);
 
 		/// <summary>
 		/// TVOS Simulator
 		/// </summary>
-		public static UnrealArch TVOSSimulator = FindOrAddByName("tvossimulator", bIsX64: false);
+		public static UnrealArch TVOSSimulator { get; } = FindOrAddByName("tvossimulator", bIsX64: false);
 
 		private static IReadOnlyDictionary<UnrealArch, string> AppleToolchainArchitectures = new Dictionary<UnrealArch, string>()
 		{
diff --git a/Engine/Source/Programs/UnrealBuildTool/Platform/Windows/UEBuildWindows.cs b/Engine/Source/Programs/UnrealBuildTool/Platform/Windows/UEBuildWindows.cs
--- a/Engine/Source/Programs/UnrealBuildTool/Platform/Windows/UEBuildWindows.cs
+++ b/Engine/Source/Programs/UnrealBuildTool/Platform/Windows/UEBuildWindows.cs
@@ -19,7 +19,7 @@ namespace UnrealBuildTool
 		/// <summary>
 		/// Version of Arm64 that can interop with X64 (Emulation Compatible)
 		/// </summary>
-		public static UnrealArch Arm64ec = FindOrAddByName("arm64ec", bIsX64: false);
+		public static UnrealArch Arm64ec { get; } = FindOrAddByName("arm64ec", bIsX64: false);
 
 		private struct WindowsArchInfo
 		{
