UnrealBuildTool: Fix CS8603: Possible null reference return

#rnx

diff --git a/Engine/Source/Programs/UnrealBuildTool/Configuration/UEBuildModuleCPP.cs b/Engine/Source/Programs/UnrealBuildTool/Configuration/UEBuildModuleCPP.cs
--- a/Engine/Source/Programs/UnrealBuildTool/Configuration/UEBuildModuleCPP.cs
+++ b/Engine/Source/Programs/UnrealBuildTool/Configuration/UEBuildModuleCPP.cs
@@ -894,7 +894,7 @@ namespace UnrealBuildTool
 
 				if (VisitedIncludes.Add(TranformedHeaderInclude))
 				{
-					Func<DirectoryReference, FileItem> SearchForFileItem = (DirectoryReference dir) =>
+					Func<DirectoryReference, FileItem?> SearchForFileItem = (DirectoryReference dir) =>
 					{
 						FileReference FileRef = FileReference.Combine(dir, TranformedHeaderInclude);
 						HashSet<FileReference>? Files = GetIncludeFiles(FileRef.Directory);
