UnrealBuildTool: Resolve CA1813

#rnx

diff --git a/Engine/Source/Programs/UnrealBuildTool/System/ConfigFileAttribute.cs b/Engine/Source/Programs/UnrealBuildTool/System/ConfigFileAttribute.cs
--- a/Engine/Source/Programs/UnrealBuildTool/System/ConfigFileAttribute.cs
+++ b/Engine/Source/Programs/UnrealBuildTool/System/ConfigFileAttribute.cs
@@ -8,7 +8,7 @@ namespace UnrealBuildTool
 	/// Attribute indicating a value which should be populated from a UE .ini config file
 	/// </summary>
 	[AttributeUsage(AttributeTargets.Field | AttributeTargets.Property)]
-	public class ConfigFileAttribute : Attribute
+	public sealed class ConfigFileAttribute : Attribute
 	{
 		/// <summary>
 		/// Name of the config hierarchy to read from
diff --git a/Engine/Source/Programs/UnrealBuildTool/System/ModuleGroupsAttribute.cs b/Engine/Source/Programs/UnrealBuildTool/System/ModuleGroupsAttribute.cs
--- a/Engine/Source/Programs/UnrealBuildTool/System/ModuleGroupsAttribute.cs
+++ b/Engine/Source/Programs/UnrealBuildTool/System/ModuleGroupsAttribute.cs
@@ -11,7 +11,7 @@ namespace UnrealBuildTool
 	/// Attribute which can be applied to a ModuleRules-dervied class to indicate which module groups it belongs to
 	/// </summary>
 	[AttributeUsage(AttributeTargets.Class, AllowMultiple = true)]
-	public class ModuleGroupsAttribute : Attribute
+	public sealed class ModuleGroupsAttribute : Attribute
 	{
 		/// <summary>
 		/// Array of module group names
diff --git a/Engine/Source/Programs/UnrealBuildTool/System/ObsoleteOverrideAttribute.cs b/Engine/Source/Programs/UnrealBuildTool/System/ObsoleteOverrideAttribute.cs
--- a/Engine/Source/Programs/UnrealBuildTool/System/ObsoleteOverrideAttribute.cs
+++ b/Engine/Source/Programs/UnrealBuildTool/System/ObsoleteOverrideAttribute.cs
@@ -8,7 +8,7 @@ namespace UnrealBuildTool
 	/// Attribute used to denote that a method should no longer be overriden. Used by RulesCompiler.
 	/// </summary>
 	[AttributeUsage(AttributeTargets.Method)]
-	public class ObsoleteOverrideAttribute : Attribute
+	public sealed class ObsoleteOverrideAttribute : Attribute
 	{
 		/// <summary>
 		/// Message to display to the user if the method is overridden.
diff --git a/Engine/Source/Programs/UnrealBuildTool/System/SupportedConfigurationsAttribute.cs b/Engine/Source/Programs/UnrealBuildTool/System/SupportedConfigurationsAttribute.cs
--- a/Engine/Source/Programs/UnrealBuildTool/System/SupportedConfigurationsAttribute.cs
+++ b/Engine/Source/Programs/UnrealBuildTool/System/SupportedConfigurationsAttribute.cs
@@ -8,7 +8,7 @@ namespace UnrealBuildTool
 	/// Attribute which can be applied to a TargetRules-dervied class to indicate which configurations it supports
 	/// </summary>
 	[AttributeUsage(AttributeTargets.Class, AllowMultiple = true)]
-	public class SupportedConfigurationsAttribute : Attribute
+	public sealed class SupportedConfigurationsAttribute : Attribute
 	{
 		/// <summary>
 		/// Array of supported platforms
diff --git a/Engine/Source/Programs/UnrealBuildTool/System/SupportedPlatformGroupsAttribute.cs b/Engine/Source/Programs/UnrealBuildTool/System/SupportedPlatformGroupsAttribute.cs
--- a/Engine/Source/Programs/UnrealBuildTool/System/SupportedPlatformGroupsAttribute.cs
+++ b/Engine/Source/Programs/UnrealBuildTool/System/SupportedPlatformGroupsAttribute.cs
@@ -10,7 +10,7 @@ namespace UnrealBuildTool
 	/// Attribute which can be applied to a TargetRules-dervied class to indicate which platforms it supports
 	/// </summary>
 	[AttributeUsage(AttributeTargets.Class, AllowMultiple = true)]
-	public class SupportedPlatformGroupsAttribute : SupportedPlatformsAttribute
+	public sealed class SupportedPlatformGroupsAttribute : SupportedPlatformsAttribute
 	{
 		/// <summary>
 		/// Initialize the attribute with a list of platform groups
diff --git a/Engine/Source/Programs/UnrealBuildTool/System/SupportedPlatformsAttribute.cs b/Engine/Source/Programs/UnrealBuildTool/System/SupportedPlatformsAttribute.cs
--- a/Engine/Source/Programs/UnrealBuildTool/System/SupportedPlatformsAttribute.cs
+++ b/Engine/Source/Programs/UnrealBuildTool/System/SupportedPlatformsAttribute.cs
@@ -8,6 +8,7 @@ namespace UnrealBuildTool
 	/// Attribute which can be applied to a TargetRules-dervied class to indicate which platforms it supports
 	/// </summary>
 	[AttributeUsage(AttributeTargets.Class, AllowMultiple = true)]
+	[System.Diagnostics.CodeAnalysis.SuppressMessage("Performance", "CA1813:Avoid unsealed attributes")]
 	public class SupportedPlatformsAttribute : Attribute
 	{
 		/// <summary>
diff --git a/Engine/Source/Programs/UnrealBuildTool/System/SupportedTargetTypesAttribute.cs b/Engine/Source/Programs/UnrealBuildTool/System/SupportedTargetTypesAttribute.cs
--- a/Engine/Source/Programs/UnrealBuildTool/System/SupportedTargetTypesAttribute.cs
+++ b/Engine/Source/Programs/UnrealBuildTool/System/SupportedTargetTypesAttribute.cs
@@ -8,7 +8,7 @@ namespace UnrealBuildTool
 	/// Attribute which can be applied to a TargetRules-dervied class to indicate which target types it supports
 	/// </summary>
 	[AttributeUsage(AttributeTargets.Class, AllowMultiple = true)]
-	public class SupportedTargetTypesAttribute : Attribute
+	public sealed class SupportedTargetTypesAttribute : Attribute
 	{
 		/// <summary>
 		/// Array of supported target types
diff --git a/Engine/Source/Programs/UnrealBuildTool/System/ToolMode.cs b/Engine/Source/Programs/UnrealBuildTool/System/ToolMode.cs
--- a/Engine/Source/Programs/UnrealBuildTool/System/ToolMode.cs
+++ b/Engine/Source/Programs/UnrealBuildTool/System/ToolMode.cs
@@ -62,7 +62,7 @@ namespace UnrealBuildTool
 	/// <summary>
 	/// Attribute used to specify options for a UBT mode.
 	/// </summary>
-	class ToolModeAttribute : Attribute
+	sealed class ToolModeAttribute : Attribute
 	{
 		/// <summary>
 		/// Name of this mode
