UnrealBuildTool: Format documents, remove and sort usings

#rnx

diff --git a/Engine/Platforms/VisionOS/Source/Programs/UnrealBuildTool/UEBuildVisionOS.cs b/Engine/Platforms/VisionOS/Source/Programs/UnrealBuildTool/UEBuildVisionOS.cs
--- a/Engine/Platforms/VisionOS/Source/Programs/UnrealBuildTool/UEBuildVisionOS.cs
+++ b/Engine/Platforms/VisionOS/Source/Programs/UnrealBuildTool/UEBuildVisionOS.cs
@@ -1,6 +1,5 @@
 // Copyright Epic Games, Inc. All Rights Reserved.
 
-using System;
 using EpicGames.Core;
 using Microsoft.Extensions.Logging;
 
@@ -63,19 +62,19 @@ namespace UnrealBuildTool
 
 			if (!Target.bBuildRequiresCookedData)
 			{
-			    if (Target.Platform == UnrealTargetPlatform.Mac || Target.Platform == UnrealTargetPlatform.Win64)
-			    {
-				    if ((ModuleName == "Engine" && Target.bBuildDeveloperTools) ||
-					    (ModuleName == "TargetPlatform" && Target.bForceBuildTargetPlatforms))
-				    {
+				if (Target.Platform == UnrealTargetPlatform.Mac || Target.Platform == UnrealTargetPlatform.Win64)
+				{
+					if ((ModuleName == "Engine" && Target.bBuildDeveloperTools) ||
+						(ModuleName == "TargetPlatform" && Target.bForceBuildTargetPlatforms))
+					{
 						Rules.DynamicallyLoadedModuleNames.Add("VisionOSTargetPlatformSettings");
 						if (bIsPlatformAvailableForTargetWithSDK)
 						{
 							Rules.DynamicallyLoadedModuleNames.Add("VisionOSTargetPlatform");
 							Rules.DynamicallyLoadedModuleNames.Add("VisionOSTargetPlatformControls");
 						}
-				    }
-			    }
+					}
+				}
 			}
 
 			if (ModuleName == "UnrealEd" && bIsPlatformAvailableForTargetWithSDK)
diff --git a/Engine/Source/Programs/UnrealBuildTool.Tests/ArtifactsTest.cs b/Engine/Source/Programs/UnrealBuildTool.Tests/ArtifactsTest.cs
--- a/Engine/Source/Programs/UnrealBuildTool.Tests/ArtifactsTest.cs
+++ b/Engine/Source/Programs/UnrealBuildTool.Tests/ArtifactsTest.cs
@@ -7,8 +7,6 @@ using Microsoft.Extensions.Logging.Abstractions;
 using Microsoft.VisualStudio.TestTools.UnitTesting;
 using UnrealBuildTool.Artifacts;
 
-#nullable enable
-
 namespace UnrealBuildToolTests
 {
 
diff --git a/Engine/Source/Programs/UnrealBuildTool.Tests/LogEventMatcherTests.cs b/Engine/Source/Programs/UnrealBuildTool.Tests/LogEventMatcherTests.cs
--- a/Engine/Source/Programs/UnrealBuildTool.Tests/LogEventMatcherTests.cs
+++ b/Engine/Source/Programs/UnrealBuildTool.Tests/LogEventMatcherTests.cs
@@ -8,8 +8,6 @@ using EpicGames.Core;
 using Microsoft.Extensions.Logging;
 using Microsoft.VisualStudio.TestTools.UnitTesting;
 
-#nullable enable
-
 namespace UnrealBuildToolTests
 {
 	[TestClass]
diff --git a/Engine/Source/Programs/UnrealBuildTool/Artifacts/HordeStorageArtifactCache.cs b/Engine/Source/Programs/UnrealBuildTool/Artifacts/HordeStorageArtifactCache.cs
--- a/Engine/Source/Programs/UnrealBuildTool/Artifacts/HordeStorageArtifactCache.cs
+++ b/Engine/Source/Programs/UnrealBuildTool/Artifacts/HordeStorageArtifactCache.cs
@@ -9,7 +9,6 @@ using System.Threading.Tasks;
 using EpicGames.Core;
 using EpicGames.Horde.Storage;
 using EpicGames.Horde.Storage.Bundles;
-using EpicGames.Horde.Storage.Clients;
 using EpicGames.Horde.Storage.Nodes;
 using Microsoft.Extensions.Logging;
 using Microsoft.Extensions.Logging.Abstractions;
diff --git a/Engine/Source/Programs/UnrealBuildTool/Configuration/UEBuildFramework.cs b/Engine/Source/Programs/UnrealBuildTool/Configuration/UEBuildFramework.cs
--- a/Engine/Source/Programs/UnrealBuildTool/Configuration/UEBuildFramework.cs
+++ b/Engine/Source/Programs/UnrealBuildTool/Configuration/UEBuildFramework.cs
@@ -1,6 +1,5 @@
 // Copyright Epic Games, Inc. All Rights Reserved.
 
-using System;
 using System.Collections.Generic;
 using System.Xml;
 using EpicGames.Core;
diff --git a/Engine/Source/Programs/UnrealBuildTool/Configuration/UEBuildPlugin.cs b/Engine/Source/Programs/UnrealBuildTool/Configuration/UEBuildPlugin.cs
--- a/Engine/Source/Programs/UnrealBuildTool/Configuration/UEBuildPlugin.cs
+++ b/Engine/Source/Programs/UnrealBuildTool/Configuration/UEBuildPlugin.cs
@@ -4,7 +4,6 @@ using System.Collections.Generic;
 using System.Linq;
 using EpicGames.Core;
 using Microsoft.Extensions.Logging;
-using UnrealBuildBase;
 
 namespace UnrealBuildTool
 {
diff --git a/Engine/Source/Programs/UnrealBuildTool/Executors/UnrealBuildAccelerator/UBAAgentCoordinatorHorde.cs b/Engine/Source/Programs/UnrealBuildTool/Executors/UnrealBuildAccelerator/UBAAgentCoordinatorHorde.cs
--- a/Engine/Source/Programs/UnrealBuildTool/Executors/UnrealBuildAccelerator/UBAAgentCoordinatorHorde.cs
+++ b/Engine/Source/Programs/UnrealBuildTool/Executors/UnrealBuildAccelerator/UBAAgentCoordinatorHorde.cs
@@ -5,7 +5,6 @@ using System.Collections.Generic;
 using System.Diagnostics;
 using System.Linq;
 using System.Net.Http;
-using System.Net.Http.Headers;
 using System.Runtime.InteropServices;
 using System.Threading;
 using System.Threading.Tasks;
@@ -80,7 +79,7 @@ namespace UnrealBuildTool
 			_crypto = owner.Crypto;
 
 			ServiceCollection services = new();
-			services.AddHorde(options => 
+			services.AddHorde(options =>
 			{
 				options.ServerUrl = serverUrl;
 				options.AccessToken = accessToken;
@@ -373,7 +372,7 @@ namespace UnrealBuildTool
 				return null;
 			}
 
-			Uri? server = (hordeConfig.HordeServer == null)? null : new Uri(hordeConfig.HordeServer);
+			Uri? server = (hordeConfig.HordeServer == null) ? null : new Uri(hordeConfig.HordeServer);
 			string? token = hordeConfig.HordeToken;
 
 			ConnectionMode? connectionMode = Enum.TryParse(hordeConfig.HordeConnectionMode, true, out ConnectionMode cm) ? cm : null;
diff --git a/Engine/Source/Programs/UnrealBuildTool/Matchers/CompileEventMatcher.cs b/Engine/Source/Programs/UnrealBuildTool/Matchers/CompileEventMatcher.cs
--- a/Engine/Source/Programs/UnrealBuildTool/Matchers/CompileEventMatcher.cs
+++ b/Engine/Source/Programs/UnrealBuildTool/Matchers/CompileEventMatcher.cs
@@ -5,7 +5,6 @@ using System.Diagnostics.CodeAnalysis;
 using System.Linq;
 using System.Text.RegularExpressions;
 using EpicGames.Core;
-using Microsoft.CodeAnalysis.Differencing;
 using Microsoft.Extensions.Logging;
 
 namespace UnrealBuildTool.Matchers
@@ -67,7 +66,7 @@ namespace UnrealBuildTool.Matchers
 		static readonly Regex s_scriptCompilePattern = new Regex(@"^\s*[A-Za-z0-9_\.]+ ERROR:.* [A-Za-z_]+ failed to compile\.");
 		static readonly Regex s_cscSummaryPattern = new Regex(@"^\s+\d+ (?:Warning|Error)\(s\)");
 		static readonly Regex s_cscOutputPattern = new Regex(@"^  [^ ]+ -> ");
-	
+
 		static readonly string[] s_invalidExtensions =
 		{
 			".obj",
diff --git a/Engine/Source/Programs/UnrealBuildTool/Modes/PipInstallMode.cs b/Engine/Source/Programs/UnrealBuildTool/Modes/PipInstallMode.cs
--- a/Engine/Source/Programs/UnrealBuildTool/Modes/PipInstallMode.cs
+++ b/Engine/Source/Programs/UnrealBuildTool/Modes/PipInstallMode.cs
@@ -3,14 +3,13 @@
 using System;
 using System.Collections.Generic;
 using System.Diagnostics;
+using System.IO;
 using System.Linq;
-using System.Text;
+using System.Text.RegularExpressions;
 using System.Threading.Tasks;
-using Microsoft.Extensions.Logging;
 using EpicGames.Core;
+using Microsoft.Extensions.Logging;
 using UnrealBuildBase;
-using System.IO;
-using System.Text.RegularExpressions;
 
 namespace UnrealBuildTool.Modes
 {
@@ -20,7 +19,7 @@ namespace UnrealBuildTool.Modes
 	[ToolMode("PipInstall", ToolModeOptions.XmlConfig | ToolModeOptions.BuildPlatformsHostOnly | ToolModeOptions.SingleInstance | ToolModeOptions.StartPrefetchingEngine | ToolModeOptions.ShowExecutionTime)]
 	class PipInstallMode : ToolMode
 	{
-		private enum ActionBits: byte
+		private enum ActionBits : byte
 		{
 			NoOp = 0,
 			GenReqs = 1,
@@ -30,7 +29,7 @@ namespace UnrealBuildTool.Modes
 			ViewLicenses = 16,
 		}
 
-		public enum PipAction: byte
+		public enum PipAction : byte
 		{
 			OnlySetupParse = ActionBits.SetupPip | ActionBits.ParseReqs,
 			OnlyInstall = ActionBits.InstallReqs,
@@ -114,7 +113,7 @@ namespace UnrealBuildTool.Modes
 			}
 
 			UEBuildTarget Target = UEBuildTarget.Create(TargetDescriptor, BuildConfiguration, Logger);
-			if ( Target.TargetType != TargetType.Editor )
+			if (Target.TargetType != TargetType.Editor)
 			{
 				Logger.LogWarning("PipInstall unsupported for non-editor target: {TargetName} (Skipping)", TargetDescriptor.Name);
 				return 0;
@@ -144,7 +143,7 @@ namespace UnrealBuildTool.Modes
 				{
 					return 1;
 				}
-			}	
+			}
 
 			if ((Action & (PipAction)ActionBits.SetupPip) != 0)
 			{
@@ -230,7 +229,7 @@ namespace UnrealBuildTool.Modes
 			{
 				// Only delete virtual environment if version mismatch
 				FileReference VenvConfig = FileReference.Combine(InstallDir, "pyvenv.cfg");
-				if ( FileReference.Exists(VenvConfig) )
+				if (FileReference.Exists(VenvConfig))
 				{
 					CleanVenvDir();
 				}
@@ -245,7 +244,7 @@ namespace UnrealBuildTool.Modes
 				return;
 			}
 
-			using (IBaseCmdProgressLogger SimpleLogger =new SimpleCmdLogger(Logger))
+			using (IBaseCmdProgressLogger SimpleLogger = new SimpleCmdLogger(Logger))
 			{
 				const string PyInterpVerCheckCmd = "import sys; exit(0) if f'{sys.version_info[0]}.{sys.version_info[1]}.{sys.version_info[2]}' == sys.argv[1] else exit(1)";
 				if (RunPythonCmd(EnginePythonInterp, $"-c \"{PyInterpVerCheckCmd}\" \"{PythonVenvVer}\"", SimpleLogger) != 0)
@@ -272,7 +271,7 @@ namespace UnrealBuildTool.Modes
 			}
 
 			List<PluginInfo> CheckPlugins = new List<PluginInfo>();
-			if ( bAllPlugins )
+			if (bAllPlugins)
 			{
 				CheckPlugins.AddAll(Plugins.ReadEnginePlugins(Unreal.EngineDirectory).ToArray());
 				CheckPlugins.AddAll(Plugins.ReadProjectPlugins(Target.ProjectDirectory).ToArray());
@@ -451,7 +450,7 @@ namespace UnrealBuildTool.Modes
 			foreach (JsonObject PlatformReqs in RequirementsJson)
 			{
 				PlatformReqs.TryGetStringField("Platform", out string? PlatformField);
-				if (!CompatiblePlatform(PlatformField,Platform))
+				if (!CompatiblePlatform(PlatformField, Platform))
 				{
 					continue;
 				}
@@ -477,7 +476,7 @@ namespace UnrealBuildTool.Modes
 
 			string ConfigInfo = FileReference.ReadAllText(VenvConfig);
 			Match m = Regex.Match(ConfigInfo, @"version\s*=\s*(\d+\.\d+\.\d+)", RegexOptions.IgnoreCase);
-			if ( !m.Success )
+			if (!m.Success)
 			{
 				Logger.LogWarning("PipInstall: Unable to match venv version config: {VenvFile}", ConfigInfo);
 				return null;
@@ -495,7 +494,7 @@ namespace UnrealBuildTool.Modes
 			}
 
 			// HACK: On windows these script files are set read-only and can't be deleted
-			foreach (FileReference File in DirectoryReference.EnumerateFiles(DirectoryReference.Combine(InstallDir,"Scripts")))
+			foreach (FileReference File in DirectoryReference.EnumerateFiles(DirectoryReference.Combine(InstallDir, "Scripts")))
 			{
 				FileReference.SetAttributes(File, FileAttributes.Normal);
 			}
@@ -528,7 +527,7 @@ namespace UnrealBuildTool.Modes
 			{
 				Args += "--index-url " + ForceIndexUrl;
 			}
-			else if ( ExtraUrls != null )
+			else if (ExtraUrls != null)
 			{
 				foreach (string Url in ExtraUrls)
 				{
@@ -795,9 +794,9 @@ namespace UnrealBuildTool.Modes
 			Logger.LogError("{ErrorData}", ErrorLine.Data);
 		}
 
-		public void Dispose() {}
+		public void Dispose() { }
 
-		public void FinishProgress() {}
+		public void FinishProgress() { }
 	}
 
 	/// <summary>
@@ -939,7 +938,7 @@ namespace UnrealBuildTool.Modes
 
 		// Start strings to use 
 		private static readonly string[] MatchStrs = { "Requirement", "Collecting", "Installing" };
-		private readonly Dictionary<string,string> LogReplaceStrs = new();
+		private readonly Dictionary<string, string> LogReplaceStrs = new();
 
 		public PipProgressLogger(ILogger InLogger, string message, int GuessSteps)
 		{
@@ -967,7 +966,7 @@ namespace UnrealBuildTool.Modes
 
 		string ReplaceMatchStr(string CheckStr)
 		{
-			foreach(KeyValuePair<string,string> ChkPair in LogReplaceStrs)
+			foreach (KeyValuePair<string, string> ChkPair in LogReplaceStrs)
 			{
 				if (CheckStr.Contains(ChkPair.Key))
 				{
diff --git a/Engine/Source/Programs/UnrealBuildTool/Platform/Android/AndroidExports.cs b/Engine/Source/Programs/UnrealBuildTool/Platform/Android/AndroidExports.cs
--- a/Engine/Source/Programs/UnrealBuildTool/Platform/Android/AndroidExports.cs
+++ b/Engine/Source/Programs/UnrealBuildTool/Platform/Android/AndroidExports.cs
@@ -1,6 +1,5 @@
 // Copyright Epic Games, Inc. All Rights Reserved.
 
-using System;
 using System.Collections.Generic;
 using EpicGames.Core;
 using Microsoft.Extensions.Logging;
diff --git a/Engine/Source/Programs/UnrealBuildTool/Platform/Android/AndroidPlatformSDK.Versions.cs b/Engine/Source/Programs/UnrealBuildTool/Platform/Android/AndroidPlatformSDK.Versions.cs
--- a/Engine/Source/Programs/UnrealBuildTool/Platform/Android/AndroidPlatformSDK.Versions.cs
+++ b/Engine/Source/Programs/UnrealBuildTool/Platform/Android/AndroidPlatformSDK.Versions.cs
@@ -1,7 +1,5 @@
 // Copyright Epic Games, Inc. All Rights Reserved.
 
-using EpicGames.Core;
-
 namespace UnrealBuildTool
 {
 	////////////////////////////////////////////////////////////////////////////////////////////
diff --git a/Engine/Source/Programs/UnrealBuildTool/Platform/IOS/UEBuildIOS.cs b/Engine/Source/Programs/UnrealBuildTool/Platform/IOS/UEBuildIOS.cs
--- a/Engine/Source/Programs/UnrealBuildTool/Platform/IOS/UEBuildIOS.cs
+++ b/Engine/Source/Programs/UnrealBuildTool/Platform/IOS/UEBuildIOS.cs
@@ -6,7 +6,6 @@ using System.Diagnostics;
 using System.IO;
 using System.Linq;
 using System.Text;
-using System.Text.RegularExpressions;
 using EpicGames.Core;
 using Microsoft.Extensions.Logging;
 using UnrealBuildBase;
diff --git a/Engine/Source/Programs/UnrealBuildTool/Platform/Mac/ApplePlatformSDK.cs b/Engine/Source/Programs/UnrealBuildTool/Platform/Mac/ApplePlatformSDK.cs
--- a/Engine/Source/Programs/UnrealBuildTool/Platform/Mac/ApplePlatformSDK.cs
+++ b/Engine/Source/Programs/UnrealBuildTool/Platform/Mac/ApplePlatformSDK.cs
@@ -1,6 +1,5 @@
 // Copyright Epic Games, Inc. All Rights Reserved.
 
-using System;
 using EpicGames.Core;
 using Microsoft.Extensions.Logging;
 using UnrealBuildBase;
diff --git a/Engine/Source/Programs/UnrealBuildTool/Platform/Mac/MacToolChain.cs b/Engine/Source/Programs/UnrealBuildTool/Platform/Mac/MacToolChain.cs
--- a/Engine/Source/Programs/UnrealBuildTool/Platform/Mac/MacToolChain.cs
+++ b/Engine/Source/Programs/UnrealBuildTool/Platform/Mac/MacToolChain.cs
@@ -2,7 +2,6 @@
 
 using System;
 using System.Collections.Generic;
-using System.Globalization;
 using System.IO;
 using System.Linq;
 using System.Text.RegularExpressions;
@@ -35,7 +34,7 @@ namespace UnrealBuildTool
 		/// </summary>
 		/// <param name="bVerbose">Whether to output verbose logging</param>
 		/// <param name="Logger">Logger for output</param>
-		public MacToolChainSettings(bool bVerbose, ILogger Logger) 
+		public MacToolChainSettings(bool bVerbose, ILogger Logger)
 			: base("MacOSX", null, "macos", bVerbose, Logger)
 		{
 		}
@@ -380,7 +379,7 @@ namespace UnrealBuildTool
 					FinalExeDir = ExeAbsolutePath + ".app/Contents/MacOS";
 				}
 				string RelativePath = Utils.MakePathRelativeTo(LibraryDir, FinalExeDir).Replace("\\", "/");
-				
+
 
 				if (bCanUseMultipleRPATHs)
 				{
diff --git a/Engine/Source/Programs/UnrealBuildTool/Platform/TVOS/UEBuildTVOS.cs b/Engine/Source/Programs/UnrealBuildTool/Platform/TVOS/UEBuildTVOS.cs
--- a/Engine/Source/Programs/UnrealBuildTool/Platform/TVOS/UEBuildTVOS.cs
+++ b/Engine/Source/Programs/UnrealBuildTool/Platform/TVOS/UEBuildTVOS.cs
@@ -1,6 +1,5 @@
 // Copyright Epic Games, Inc. All Rights Reserved.
 
-using System;
 using EpicGames.Core;
 using Microsoft.Extensions.Logging;
 
diff --git a/Engine/Source/Programs/UnrealBuildTool/Platform/Windows/MicrosoftPlatformSDK.Versions.cs b/Engine/Source/Programs/UnrealBuildTool/Platform/Windows/MicrosoftPlatformSDK.Versions.cs
--- a/Engine/Source/Programs/UnrealBuildTool/Platform/Windows/MicrosoftPlatformSDK.Versions.cs
+++ b/Engine/Source/Programs/UnrealBuildTool/Platform/Windows/MicrosoftPlatformSDK.Versions.cs
@@ -4,7 +4,6 @@ using System;
 using System.Diagnostics;
 using System.Linq;
 using EpicGames.Core;
-using Microsoft.Build.Framework;
 
 namespace UnrealBuildTool
 {
diff --git a/Engine/Source/Programs/UnrealBuildTool/Platform/Windows/PVSToolChain.cs b/Engine/Source/Programs/UnrealBuildTool/Platform/Windows/PVSToolChain.cs
--- a/Engine/Source/Programs/UnrealBuildTool/Platform/Windows/PVSToolChain.cs
+++ b/Engine/Source/Programs/UnrealBuildTool/Platform/Windows/PVSToolChain.cs
@@ -1,8 +1,5 @@
 // Copyright Epic Games, Inc. All Rights Reserved.
 
-using EpicGames.Core;
-using Microsoft.Extensions.Logging;
-using Newtonsoft.Json;
 using System;
 using System.Collections.Generic;
 using System.Collections.Immutable;
@@ -13,6 +10,9 @@ using System.Text;
 using System.Text.RegularExpressions;
 using System.Threading.Tasks;
 using System.Xml.Serialization;
+using EpicGames.Core;
+using Microsoft.Extensions.Logging;
+using Newtonsoft.Json;
 using UnrealBuildBase;
 
 namespace UnrealBuildTool
@@ -491,7 +491,6 @@ namespace UnrealBuildTool
 							{
 								Logger.LogWarning(KnownLogEvents.Compiler, "{Path}({LineNumber}): warning {WarningCode}: {WarningMessage}", LogValue.SourceFile(file, FileName), LineNumber, ErrorInfo.Code, ErrorInfo.Message);
 							}
-
 						}
 						catch (Exception Ex)
 						{
@@ -891,7 +890,7 @@ namespace UnrealBuildTool
 
 				if (AnalyzerVersion.CompareTo(new Version("7.30")) >= 0)
 				{
-					ConfigFileContents.Append("new-output-format=yes\n"); ;
+					ConfigFileContents.Append("new-output-format=yes\n");
 				}
 
 				string BaseFileName = PreprocessedFileItem.Location.GetFileName();
diff --git a/Engine/Source/Programs/UnrealBuildTool/ProjectFiles/VisualStudio/VCProject.cs b/Engine/Source/Programs/UnrealBuildTool/ProjectFiles/VisualStudio/VCProject.cs
--- a/Engine/Source/Programs/UnrealBuildTool/ProjectFiles/VisualStudio/VCProject.cs
+++ b/Engine/Source/Programs/UnrealBuildTool/ProjectFiles/VisualStudio/VCProject.cs
@@ -6,7 +6,6 @@ using System.Collections.Generic;
 using System.Diagnostics.CodeAnalysis;
 using System.IO;
 using System.Linq;
-using System.Runtime.InteropServices;
 using System.Security;
 using System.Text;
 using System.Text.RegularExpressions;
@@ -14,10 +13,9 @@ using System.Threading.Tasks;
 using System.Xml;
 using System.Xml.Linq;
 using EpicGames.Core;
+using Microsoft.CodeAnalysis;
 using Microsoft.Extensions.Logging;
 using UnrealBuildBase;
-using Newtonsoft.Json.Linq;
-using Microsoft.CodeAnalysis;
 using static UnrealBuildTool.PlatformProjectGenerator;
 
 namespace UnrealBuildTool
@@ -1368,7 +1366,7 @@ namespace UnrealBuildTool
 					DefaultRules != null ? (" " + GetConformanceCompileArguments(DefaultRules)) : String.Empty,
 					CommonAdditionalOptions.Length > 0 ? (" " + CommonAdditionalOptions) : String.Empty);
 			};
-			
+
 			string DefaultAdditionalOptions = GetAdditionalOptionsString(DefaultRules);
 
 			// Write common IntelliSense info
@@ -1416,7 +1414,7 @@ namespace UnrealBuildTool
 				WriteConfiguration(ProjectName, Combination, VCProjectFileContent, PlatformProjectGenerators, bGenerateUserFileContent ? VCUserFileContent : null, bGenerateUserFileContent ? VCUserFileSettings : null);
 			}
 
-			{ 
+			{
 				// Collapse common values
 				{
 					StringBuilder CommonProjectFileContent = new StringBuilder();
@@ -1895,7 +1893,7 @@ namespace UnrealBuildTool
 						.Elements()
 						.Where(Element => UserFileSettings.PropertiesToPatch.Contains(Element.Name.LocalName))
 						.ToArray();
-					
+
 					// Check if all properties are already has the correct value and order, and skip patching if so.
 					if (CurrentPropertyGroupElementsForPatch.Length == NewPropertyGroupElementsForPatch.Length &&
 						!CurrentPropertyGroupElementsForPatch.Where((CurrentProperty, i) => CurrentProperty.Name != NewPropertyGroupElementsForPatch[i].Name || CurrentProperty.Value != NewPropertyGroupElementsForPatch[i].Value).Any())
@@ -2709,5 +2707,4 @@ namespace UnrealBuildTool
 		/// Cache of parsed info about this project
 		protected readonly Dictionary<UnrealTargetConfiguration, CsProjectInfo> CachedProjectInfo = new Dictionary<UnrealTargetConfiguration, CsProjectInfo>();
 	}
-
 }
diff --git a/Engine/Source/Programs/UnrealBuildTool/ProjectFiles/VisualStudio/VCProjectFileGenerator.cs b/Engine/Source/Programs/UnrealBuildTool/ProjectFiles/VisualStudio/VCProjectFileGenerator.cs
--- a/Engine/Source/Programs/UnrealBuildTool/ProjectFiles/VisualStudio/VCProjectFileGenerator.cs
+++ b/Engine/Source/Programs/UnrealBuildTool/ProjectFiles/VisualStudio/VCProjectFileGenerator.cs
@@ -11,7 +11,6 @@ using System.Xml.Linq;
 using EpicGames.Core;
 using Microsoft.Extensions.Logging;
 using UnrealBuildBase;
-using Newtonsoft.Json.Linq;
 
 namespace UnrealBuildTool
 {
@@ -633,14 +632,15 @@ namespace UnrealBuildTool
 						new XElement(NS + "Target",
 							new XAttribute("Name", "CleanUpStaleDlls"),
 							new XAttribute("AfterTargets", "Build"),
-							AutomationProjectFiles.SelectMany(AutomationProject => {
-									string BaseFilename = FileReference.Combine(AutomationToolBinariesDir, AutomationProject.ProjectFilePath.GetFileNameWithoutExtension()).FullName;
-									return new List<XElement>() {
-										new XElement(NS + "Delete",	new XAttribute("Files", BaseFilename + ".dll")),
-										new XElement(NS + "Delete",	new XAttribute("Files", BaseFilename + ".dll.config")),
-										new XElement(NS + "Delete",	new XAttribute("Files", BaseFilename + ".pdb"))
+							AutomationProjectFiles.SelectMany(AutomationProject =>
+							{
+								string BaseFilename = FileReference.Combine(AutomationToolBinariesDir, AutomationProject.ProjectFilePath.GetFileNameWithoutExtension()).FullName;
+								return new List<XElement>() {
+										new XElement(NS + "Delete", new XAttribute("Files", BaseFilename + ".dll")),
+										new XElement(NS + "Delete", new XAttribute("Files", BaseFilename + ".dll.config")),
+										new XElement(NS + "Delete", new XAttribute("Files", BaseFilename + ".pdb"))
 									};
-								}
+							}
 							)
 						)
 					)
diff --git a/Engine/Source/Programs/UnrealBuildTool/ProjectFiles/VisualStudioWorkspace/VSWorkspaceProjectFile.cs b/Engine/Source/Programs/UnrealBuildTool/ProjectFiles/VisualStudioWorkspace/VSWorkspaceProjectFile.cs
--- a/Engine/Source/Programs/UnrealBuildTool/ProjectFiles/VisualStudioWorkspace/VSWorkspaceProjectFile.cs
+++ b/Engine/Source/Programs/UnrealBuildTool/ProjectFiles/VisualStudioWorkspace/VSWorkspaceProjectFile.cs
@@ -1,7 +1,6 @@
 // Copyright Epic Games, Inc. All Rights Reserved.
 
 using System;
-using System.Collections.Concurrent;
 using System.Collections.Generic;
 using System.IO;
 using System.Linq;
@@ -77,8 +76,8 @@ namespace UnrealBuildTool
 						}
 
 						// Skip Editor for all platforms except for current platform
-						if (ProjectTarget.TargetRules.Type == TargetType.Editor && 
-							(BuildHostPlatform.Current.Platform != Platform || 
+						if (ProjectTarget.TargetRules.Type == TargetType.Editor &&
+							(BuildHostPlatform.Current.Platform != Platform ||
 							(Configuration == UnrealTargetConfiguration.Test || Configuration == UnrealTargetConfiguration.Shipping)))
 						{
 							continue;
diff --git a/Engine/Source/Programs/UnrealBuildTool/ProjectFiles/Xcode/XcodeProject.cs b/Engine/Source/Programs/UnrealBuildTool/ProjectFiles/Xcode/XcodeProject.cs
--- a/Engine/Source/Programs/UnrealBuildTool/ProjectFiles/Xcode/XcodeProject.cs
+++ b/Engine/Source/Programs/UnrealBuildTool/ProjectFiles/Xcode/XcodeProject.cs
@@ -5,7 +5,6 @@ using System.Collections.Generic;
 using System.IO;
 using System.Linq;
 using System.Text;
-using System.Xml.Linq;
 using EpicGames.Core;
 using Microsoft.CodeAnalysis;
 using Microsoft.Extensions.Logging;
@@ -599,7 +598,7 @@ namespace UnrealBuildTool.XcodeProjectXcconfig
 			return null;
 		}
 
-		public string ProjectOrEnginePath(string SubPath, bool bMakeRelative, string? AltProjectSubPath=null)
+		public string ProjectOrEnginePath(string SubPath, bool bMakeRelative, string? AltProjectSubPath = null)
 		{
 			string? FinalPath = null;
 			if (ProductDirectory != Unreal.EngineDirectory)
@@ -961,9 +960,9 @@ namespace UnrealBuildTool.XcodeProjectXcconfig
 		public void AddFramework(DirectoryReference Framework, string FileRefGuid)
 		{
 			XcodeSourceFile FrameworkSource = new XcodeSourceFile(new FileReference(Framework.FullName), null, FileRefGuid);
-			FileCollection.ProcessFile(FrameworkSource, true, false, "Frameworks", ""); ;
+			FileCollection.ProcessFile(FrameworkSource, true, false, "Frameworks", "");
+			;
 			FileItems.Add(FrameworkSource);
-
 		}
 	}
 
@@ -1055,7 +1054,7 @@ namespace UnrealBuildTool.XcodeProjectXcconfig
 					Content.WriteLine(4, $"\"VALID_ARCHS[sdk=macos*]\" = \"{SupportedMacArchitectures}\";");
 				}
 			}
-				
+
 			Content.WriteLine(3, "};");
 			Content.WriteLine(3, $"name = \"{Info.DisplayName}\";");
 			Content.WriteLine(2, "};");
@@ -1430,7 +1429,8 @@ namespace UnrealBuildTool.XcodeProjectXcconfig
 						$"  rm -rf \\\"${{CONFIGURATION_BUILD_DIR}}/${{CONTENTS_FOLDER_PATH}}{SyncDestSubdir}\\\"",
 						"  exit -0",
 						"fi",
-					}); ;
+					});
+					;
 				}
 
 				// when we bring stated data into the .app, we have to skip some temp stuff that went into it
@@ -1510,7 +1510,7 @@ namespace UnrealBuildTool.XcodeProjectXcconfig
 			string AssetsPath = UnrealData.ProjectOrEnginePath(AssetsSubPath, false, AssetsAltSubPath);
 			ResourcesBuildPhase.AddResource(new FileReference(AssetsPath));
 			StoryboardPath = UnrealData.FindFile(StoryboardPaths, Platform, false);
-			
+
 			if (StoryboardPath != null)
 			{
 				ResourcesBuildPhase.AddResource(new FileReference(StoryboardPath));
@@ -1673,7 +1673,7 @@ namespace UnrealBuildTool.XcodeProjectXcconfig
 		private string GetPlistSigningName(string ProvisionSetting, ILogger Logger)
 		{
 			string? SigningName;
-			lock(PlistFileMap)
+			lock (PlistFileMap)
 			{
 				if (!PlistFileMap.TryGetValue(ProvisionSetting, out SigningName))
 				{
@@ -1814,7 +1814,7 @@ namespace UnrealBuildTool.XcodeProjectXcconfig
 					{
 						SupportedPlatforms += " iphonesimulator";
 					}
-				
+
 					DeploymentTargetKey = "IPHONEOS_DEPLOYMENT_TARGET";
 					SupportedDevices = UnrealData.IOSProjectSettings!.RuntimeDevices;
 					DeploymentTarget = UnrealData.IOSProjectSettings.RuntimeVersion;
@@ -1858,7 +1858,6 @@ namespace UnrealBuildTool.XcodeProjectXcconfig
 					{
 						ExtraConfigLines.Add($"INFOPLIST_KEY_UIApplicationSceneManifest_Generation = YES");
 					}
-
 				}
 				else
 				{
diff --git a/Engine/Source/Programs/UnrealBuildTool/System/ModuleGroupsAttribute.cs b/Engine/Source/Programs/UnrealBuildTool/System/ModuleGroupsAttribute.cs
--- a/Engine/Source/Programs/UnrealBuildTool/System/ModuleGroupsAttribute.cs
+++ b/Engine/Source/Programs/UnrealBuildTool/System/ModuleGroupsAttribute.cs
@@ -1,11 +1,10 @@
 // Copyright Epic Games, Inc. All Rights Reserved.
 
 using System;
-using System.Collections.Generic;
 
 namespace UnrealBuildTool
 {
-	#nullable enable
+#nullable enable
 
 	/// <summary>
 	/// Attribute which can be applied to a ModuleRules-dervied class to indicate which module groups it belongs to
diff --git a/Engine/Source/Programs/UnrealBuildTool/ToolChain/AppleToolChain.cs b/Engine/Source/Programs/UnrealBuildTool/ToolChain/AppleToolChain.cs
--- a/Engine/Source/Programs/UnrealBuildTool/ToolChain/AppleToolChain.cs
+++ b/Engine/Source/Programs/UnrealBuildTool/ToolChain/AppleToolChain.cs
@@ -6,7 +6,6 @@ using System.Diagnostics;
 using System.IO;
 using System.Linq;
 using System.Security.Cryptography.X509Certificates;
-using System.Runtime.InteropServices;
 using System.Text;
 using System.Text.RegularExpressions;
 using System.Threading.Tasks;
@@ -584,7 +583,6 @@ namespace UnrealBuildTool
 					Arguments.Add(GetUserIncludePathArgument(OutputInteropHeader.GetDirectoryItem().Location));
 					CompileAction.PrerequisiteItems.Add(OutputInteropHeader);
 				}
-
 			}
 			Arguments.Add("-DUE_USE_SWIFT_UI_MAIN=" + (bUseSwiftUIMain ? "1" : "0"));
 			return Output;
diff --git a/Engine/Source/Programs/UnrealBuildTool/ToolChain/ClangWarnings.cs b/Engine/Source/Programs/UnrealBuildTool/ToolChain/ClangWarnings.cs
--- a/Engine/Source/Programs/UnrealBuildTool/ToolChain/ClangWarnings.cs
+++ b/Engine/Source/Programs/UnrealBuildTool/ToolChain/ClangWarnings.cs
@@ -1,7 +1,5 @@
 ﻿// Copyright Epic Games, Inc. All Rights Reserved.
 
-using EpicGames.UHT.Utils;
-using Microsoft.Extensions.Logging;
 using System.Collections.Generic;
 using EpicGames.Core;
 
diff --git a/Engine/Source/Programs/UnrealBuildTool/UnrealBuildTool.cs b/Engine/Source/Programs/UnrealBuildTool/UnrealBuildTool.cs
--- a/Engine/Source/Programs/UnrealBuildTool/UnrealBuildTool.cs
+++ b/Engine/Source/Programs/UnrealBuildTool/UnrealBuildTool.cs
@@ -8,7 +8,6 @@ using System.Linq;
 using System.Reflection;
 using System.Text;
 using System.Text.Json;
-using System.Text.RegularExpressions;
 using EpicGames.Core;
 using Microsoft.Extensions.Logging;
 using OpenTracing.Util;
