UnrealBuildTool: Resolve CA1019

#rnx

diff --git a/Engine/Source/Programs/UnrealBuildTool/System/ConfigFileAttribute.cs b/Engine/Source/Programs/UnrealBuildTool/System/ConfigFileAttribute.cs
--- a/Engine/Source/Programs/UnrealBuildTool/System/ConfigFileAttribute.cs
+++ b/Engine/Source/Programs/UnrealBuildTool/System/ConfigFileAttribute.cs
@@ -13,29 +13,29 @@ namespace UnrealBuildTool
 		/// <summary>
 		/// Name of the config hierarchy to read from
 		/// </summary>
-		public ConfigHierarchyType ConfigType;
+		public ConfigHierarchyType ConfigType { get; }
 
 		/// <summary>
 		/// Section containing the setting
 		/// </summary>
-		public string SectionName;
+		public string SectionName { get; }
 
 		/// <summary>
 		/// Key name to search for
 		/// </summary>
-		public string? KeyName;
+		public string? KeyName { get; }
 
 		/// <summary>
 		/// Constructor
 		/// </summary>
-		/// <param name="ConfigType">Type of the config hierarchy to read from</param>
-		/// <param name="SectionName">Section containing the setting</param>
-		/// <param name="KeyName">Key name to search for. Optional; uses the name of the field if not set.</param>
-		public ConfigFileAttribute(ConfigHierarchyType ConfigType, string SectionName, string? KeyName = null)
+		/// <param name="configType">Type of the config hierarchy to read from</param>
+		/// <param name="sectionName">Section containing the setting</param>
+		/// <param name="keyName">Key name to search for. Optional; uses the name of the field if not set.</param>
+		public ConfigFileAttribute(ConfigHierarchyType configType, string sectionName, string? keyName = null)
 		{
-			this.ConfigType = ConfigType;
-			this.SectionName = SectionName;
-			this.KeyName = KeyName;
+			ConfigType = configType;
+			SectionName = sectionName;
+			KeyName = keyName;
 		}
 	}
 }
diff --git a/Engine/Source/Programs/UnrealBuildTool/System/ModuleGroupsAttribute.cs b/Engine/Source/Programs/UnrealBuildTool/System/ModuleGroupsAttribute.cs
--- a/Engine/Source/Programs/UnrealBuildTool/System/ModuleGroupsAttribute.cs
+++ b/Engine/Source/Programs/UnrealBuildTool/System/ModuleGroupsAttribute.cs
@@ -4,8 +4,6 @@ using System;
 
 namespace UnrealBuildTool
 {
-#nullable enable
-
 	/// <summary>
 	/// Attribute which can be applied to a ModuleRules-dervied class to indicate which module groups it belongs to
 	/// </summary>
@@ -15,15 +13,15 @@ namespace UnrealBuildTool
 		/// <summary>
 		/// Array of module group names
 		/// </summary>
-		public readonly string[] ModuleGroups;
+		public string[] ModuleGroups { get; }
 
 		/// <summary>
 		/// Initialize the attribute with a list of module groups
 		/// </summary>
-		/// <param name="ModuleGroups">Variable-length array of module group arguments</param>
-		public ModuleGroupsAttribute(params string[] ModuleGroups)
+		/// <param name="moduleGroups">Variable-length array of module group arguments</param>
+		public ModuleGroupsAttribute(params string[] moduleGroups)
 		{
-			this.ModuleGroups = ModuleGroups;
+			ModuleGroups = moduleGroups;
 		}
 	}
 }
diff --git a/Engine/Source/Programs/UnrealBuildTool/System/ObsoleteOverrideAttribute.cs b/Engine/Source/Programs/UnrealBuildTool/System/ObsoleteOverrideAttribute.cs
--- a/Engine/Source/Programs/UnrealBuildTool/System/ObsoleteOverrideAttribute.cs
+++ b/Engine/Source/Programs/UnrealBuildTool/System/ObsoleteOverrideAttribute.cs
@@ -13,15 +13,15 @@ namespace UnrealBuildTool
 		/// <summary>
 		/// Message to display to the user if the method is overridden.
 		/// </summary>
-		public string Message;
+		public string Message { get; }
 
 		/// <summary>
 		/// Constructor
 		/// </summary>
-		/// <param name="Message">Message to display to the user if the method is overridden</param>
-		public ObsoleteOverrideAttribute(string Message)
+		/// <param name="message">Message to display to the user if the method is overridden</param>
+		public ObsoleteOverrideAttribute(string message)
 		{
-			this.Message = Message;
+			Message = message;
 		}
 	}
 }
diff --git a/Engine/Source/Programs/UnrealBuildTool/System/SupportedConfigurationsAttribute.cs b/Engine/Source/Programs/UnrealBuildTool/System/SupportedConfigurationsAttribute.cs
--- a/Engine/Source/Programs/UnrealBuildTool/System/SupportedConfigurationsAttribute.cs
+++ b/Engine/Source/Programs/UnrealBuildTool/System/SupportedConfigurationsAttribute.cs
@@ -13,15 +13,15 @@ namespace UnrealBuildTool
 		/// <summary>
 		/// Array of supported platforms
 		/// </summary>
-		public readonly UnrealTargetConfiguration[] Configurations;
+		public UnrealTargetConfiguration[] Configurations { get; }
 
 		/// <summary>
 		/// Initialize the attribute with a list of configurations
 		/// </summary>
-		/// <param name="Configurations">Variable-length array of configuration arguments</param>
-		public SupportedConfigurationsAttribute(params UnrealTargetConfiguration[] Configurations)
+		/// <param name="configurations">Variable-length array of configuration arguments</param>
+		public SupportedConfigurationsAttribute(params UnrealTargetConfiguration[] configurations)
 		{
-			this.Configurations = Configurations;
+			Configurations = configurations;
 		}
 	}
 }
diff --git a/Engine/Source/Programs/UnrealBuildTool/System/SupportedPlatformGroupsAttribute.cs b/Engine/Source/Programs/UnrealBuildTool/System/SupportedPlatformGroupsAttribute.cs
--- a/Engine/Source/Programs/UnrealBuildTool/System/SupportedPlatformGroupsAttribute.cs
+++ b/Engine/Source/Programs/UnrealBuildTool/System/SupportedPlatformGroupsAttribute.cs
@@ -15,34 +15,35 @@ namespace UnrealBuildTool
 		/// <summary>
 		/// Initialize the attribute with a list of platform groups
 		/// </summary>
-		/// <param name="PlatformGroups">Variable-length array of platform group arguments</param>
-		public SupportedPlatformGroupsAttribute(params string[] PlatformGroups) : base(GetPlatformsForGroups(PlatformGroups))
+		/// <param name="platformGroups">Variable-length array of platform group arguments</param>
+		[System.Diagnostics.CodeAnalysis.SuppressMessage("Design", "CA1019:Define accessors for attribute arguments", Justification = "Unneeded")]
+		public SupportedPlatformGroupsAttribute(params string[] platformGroups) : base(GetPlatformsForGroups(platformGroups))
 		{
 		}
 
-		private static string[] GetPlatformsForGroups(params string[] PlatformGroups)
+		private static string[] GetPlatformsForGroups(params string[] platformGroups)
 		{
-			HashSet<UnrealTargetPlatform> SupportedPlatforms = new();
+			HashSet<UnrealTargetPlatform> supportedPlatforms = new();
 			try
 			{
-				foreach (string Name in PlatformGroups)
+				foreach (string name in platformGroups)
 				{
-					if (UnrealPlatformGroup.TryParse(Name, out UnrealPlatformGroup Group))
+					if (UnrealPlatformGroup.TryParse(name, out UnrealPlatformGroup group))
 					{
-						SupportedPlatforms.UnionWith(UnrealTargetPlatform.GetValidPlatforms().Where(x => x.IsInGroup(Group)));
+						supportedPlatforms.UnionWith(UnrealTargetPlatform.GetValidPlatforms().Where(x => x.IsInGroup(group)));
 						continue;
 					}
-					throw new BuildException(String.Format("The platform group name {0} is not a valid platform group name. Valid names are ({1})", Name,
+					throw new BuildException(String.Format("The platform group name {0} is not a valid platform group name. Valid names are ({1})", name,
 						String.Join(",", UnrealPlatformGroup.GetValidGroupNames())));
 				}
 			}
-			catch (BuildException Ex)
+			catch (BuildException ex)
 			{
-				EpicGames.Core.ExceptionUtils.AddContext(Ex, $"while parsing a SupportedPlatformGroups attribute '{String.Join(',', PlatformGroups)}'");
+				EpicGames.Core.ExceptionUtils.AddContext(ex, $"while parsing a SupportedPlatformGroups attribute '{String.Join(',', platformGroups)}'");
 				throw;
 			}
 
-			return SupportedPlatforms.Select(x => x.ToString()).ToArray();
+			return supportedPlatforms.Select(x => x.ToString()).ToArray();
 		}
 	}
 }
diff --git a/Engine/Source/Programs/UnrealBuildTool/System/SupportedPlatformsAttribute.cs b/Engine/Source/Programs/UnrealBuildTool/System/SupportedPlatformsAttribute.cs
--- a/Engine/Source/Programs/UnrealBuildTool/System/SupportedPlatformsAttribute.cs
+++ b/Engine/Source/Programs/UnrealBuildTool/System/SupportedPlatformsAttribute.cs
@@ -14,21 +14,22 @@ namespace UnrealBuildTool
 		/// <summary>
 		/// Array of supported platforms
 		/// </summary>
-		public readonly UnrealTargetPlatform[] Platforms;
+		public UnrealTargetPlatform[] Platforms { get; }
 
 		/// <summary>
 		/// Initialize the attribute with a list of platforms
 		/// </summary>
-		/// <param name="Platforms">Variable-length array of platform arguments</param>
-		public SupportedPlatformsAttribute(params string[] Platforms)
+		/// <param name="platforms">Variable-length array of platform arguments</param>
+		[System.Diagnostics.CodeAnalysis.SuppressMessage("Design", "CA1019:Define accessors for attribute arguments", Justification = "Unneeded")]
+		public SupportedPlatformsAttribute(params string[] platforms)
 		{
 			try
 			{
-				this.Platforms = Array.ConvertAll(Platforms, x => UnrealTargetPlatform.Parse(x));
+				Platforms = Array.ConvertAll(platforms, x => UnrealTargetPlatform.Parse(x));
 			}
-			catch (BuildException Ex)
+			catch (BuildException ex)
 			{
-				EpicGames.Core.ExceptionUtils.AddContext(Ex, "while parsing a SupportedPlatforms attribute");
+				EpicGames.Core.ExceptionUtils.AddContext(ex, "while parsing a SupportedPlatforms attribute");
 				throw;
 			}
 		}
@@ -36,10 +37,11 @@ namespace UnrealBuildTool
 		/// <summary>
 		/// Initialize the attribute with all the platforms in a given category
 		/// </summary>
-		/// <param name="Category">Category of platforms to add</param>
-		public SupportedPlatformsAttribute(UnrealPlatformClass Category)
+		/// <param name="category">Category of platforms to add</param>
+		[System.Diagnostics.CodeAnalysis.SuppressMessage("Design", "CA1019:Define accessors for attribute arguments", Justification = "Unneeded")]
+		public SupportedPlatformsAttribute(UnrealPlatformClass category)
 		{
-			Platforms = Utils.GetPlatformsInClass(Category);
+			Platforms = Utils.GetPlatformsInClass(category);
 		}
 	}
 }
diff --git a/Engine/Source/Programs/UnrealBuildTool/System/SupportedTargetTypesAttribute.cs b/Engine/Source/Programs/UnrealBuildTool/System/SupportedTargetTypesAttribute.cs
--- a/Engine/Source/Programs/UnrealBuildTool/System/SupportedTargetTypesAttribute.cs
+++ b/Engine/Source/Programs/UnrealBuildTool/System/SupportedTargetTypesAttribute.cs
@@ -13,15 +13,15 @@ namespace UnrealBuildTool
 		/// <summary>
 		/// Array of supported target types
 		/// </summary>
-		public readonly TargetType[] TargetTypes;
+		public TargetType[] TargetTypes { get; }
 
 		/// <summary>
 		/// Initialize the attribute with a list of target types
 		/// </summary>
-		/// <param name="TargetTypes">Variable-length array of target type arguments</param>
-		public SupportedTargetTypesAttribute(params TargetType[] TargetTypes)
+		/// <param name="targetTypes">Variable-length array of target type arguments</param>
+		public SupportedTargetTypesAttribute(params TargetType[] targetTypes)
 		{
-			this.TargetTypes = TargetTypes;
+			TargetTypes = targetTypes;
 		}
 	}
 }
diff --git a/Engine/Source/Programs/UnrealBuildTool/System/ToolMode.cs b/Engine/Source/Programs/UnrealBuildTool/System/ToolMode.cs
--- a/Engine/Source/Programs/UnrealBuildTool/System/ToolMode.cs
+++ b/Engine/Source/Programs/UnrealBuildTool/System/ToolMode.cs
@@ -67,22 +67,22 @@ namespace UnrealBuildTool
 		/// <summary>
 		/// Name of this mode
 		/// </summary>
-		public string Name;
+		public string Name { get; }
 
 		/// <summary>
 		/// Options for executing this mode
 		/// </summary>
-		public ToolModeOptions Options;
+		public ToolModeOptions Options { get; }
 
 		/// <summary>
 		/// Constructor
 		/// </summary>
-		/// <param name="Name">Name of the mode</param>
-		/// <param name="Options">Options for this mode</param>
-		public ToolModeAttribute(string Name, ToolModeOptions Options)
+		/// <param name="name">Name of the mode</param>
+		/// <param name="options">Options for this mode</param>
+		public ToolModeAttribute(string name, ToolModeOptions options)
 		{
-			this.Name = Name;
-			this.Options = Options;
+			Name = name;
+			Options = options;
 		}
 	}
 
