UnrealBuildTool: Resolve analysis warnings in UnrealBuildTool.Artifacts

#rnx

diff --git a/Engine/Source/Programs/UnrealBuildTool/Artifacts/ActionArtifactCache.cs b/Engine/Source/Programs/UnrealBuildTool/Artifacts/ActionArtifactCache.cs
--- a/Engine/Source/Programs/UnrealBuildTool/Artifacts/ActionArtifactCache.cs
+++ b/Engine/Source/Programs/UnrealBuildTool/Artifacts/ActionArtifactCache.cs
@@ -89,7 +89,7 @@ namespace UnrealBuildTool.Artifacts
 		/// <param name="path">Path to artifact</param>
 		/// <param name="hash">Hash of the artifact</param>
 		/// <returns>The artifact</returns>
-		private ArtifactFile CreateArtifact(ArtifactDirectoryTree tree, string path, IoHash hash)
+		private static ArtifactFile CreateArtifact(ArtifactDirectoryTree tree, string path, IoHash hash)
 		{
 			return new(tree, new Utf8String(path), hash);
 		}
@@ -168,7 +168,7 @@ namespace UnrealBuildTool.Artifacts
 		/// <returns>Action artifact cache object</returns>
 		public static IActionArtifactCache CreateHordeFileCache(DirectoryReference directory, CppDependencyCache cppDependencyCache, ILogger logger)
 		{
-			IArtifactCache artifactCache = HordeStorageArtifactCache.CreateFileCache(directory, /*logger*/ NullLogger.Instance, false);
+			IArtifactCache artifactCache = HordeStorageArtifactCache.CreateFileCache(directory, NullLogger.Instance, false);
 			return new ActionArtifactCache(artifactCache, cppDependencyCache, logger);
 		}
 
@@ -239,7 +239,7 @@ namespace UnrealBuildTool.Artifacts
 					{
 						if (LogCacheMisses)
 						{
-							_logger.LogInformation("Artifact Cache Miss: Content hash different {actionDescription}/{File}", actionDescription, item.FullName);
+							_logger.LogInformation("Artifact Cache Miss: Content hash different {ActionDescription}/{File}", actionDescription, item.FullName);
 						}
 						match = false;
 						break;
@@ -432,7 +432,7 @@ namespace UnrealBuildTool.Artifacts
 				string[] lines = new string[files.Count];
 				for (int index = 0; index < files.Count; index++)
 				{
-					ArtifactFile artifact = directoryMapping.GetArtifact(action, files[index], waits[index].Result);
+					ArtifactFile artifact = directoryMapping.GetArtifact(action, files[index], await waits[index]);
 					lines[index] = $"{GetArtifactTreeName(artifact)} {artifact.Tree} {artifact.Name} {artifact.ContentHash}";
 				}
 				Array.Sort(lines, StringComparer.Ordinal);
@@ -495,7 +495,7 @@ namespace UnrealBuildTool.Artifacts
 		/// <param name="substitutions">Substitutions when a given input is found</param>
 		/// <param name="outputs">Destination list</param>
 		/// <param name="inputs">Source inputs</param>
-		private void AddFileItems(HashSet<FileItem> uniques, Dictionary<FileItem, List<FileItem>>? substitutions, List<FileItem> outputs, IEnumerable<FileItem> inputs)
+		private static void AddFileItems(HashSet<FileItem> uniques, Dictionary<FileItem, List<FileItem>>? substitutions, List<FileItem> outputs, IEnumerable<FileItem> inputs)
 		{
 			if (substitutions != null)
 			{
diff --git a/Engine/Source/Programs/UnrealBuildTool/Artifacts/Artifacts.cs b/Engine/Source/Programs/UnrealBuildTool/Artifacts/Artifacts.cs
--- a/Engine/Source/Programs/UnrealBuildTool/Artifacts/Artifacts.cs
+++ b/Engine/Source/Programs/UnrealBuildTool/Artifacts/Artifacts.cs
@@ -13,7 +13,7 @@ namespace UnrealBuildTool.Artifacts
 	/// <summary>
 	/// Artifacts can exist in different directory roots.
 	/// </summary>
-	public enum ArtifactDirectoryTree : byte
+	public enum ArtifactDirectoryTree
 	{
 
 		/// <summary>
@@ -38,7 +38,7 @@ namespace UnrealBuildTool.Artifacts
 	/// <param name="Tree">Directory tree containing the artifact</param>
 	/// <param name="Name">Name of the artifact</param>
 	/// <param name="ContentHash">Hash of the artifact contents</param>
-	public record struct ArtifactFile(ArtifactDirectoryTree Tree, Utf8String Name, IoHash ContentHash)
+	public readonly record struct ArtifactFile(ArtifactDirectoryTree Tree, Utf8String Name, IoHash ContentHash)
 	{
 		/// <summary>
 		/// The full path of the artifact
@@ -75,7 +75,7 @@ namespace UnrealBuildTool.Artifacts
 		public IArtifactDirectoryMapping? DirectoryMapping { get; set; } = null;
 
 		/// <inheritdoc/>
-		public override int GetHashCode()
+		public readonly override int GetHashCode()
 		{
 			return HashCode.Combine(Key.GetHashCode(), ActionKey.GetHashCode(), Inputs.GetHashCode(), Outputs.GetHashCode());
 		}
diff --git a/Engine/Source/Programs/UnrealBuildTool/Artifacts/HordeStorageArtifactCache.cs b/Engine/Source/Programs/UnrealBuildTool/Artifacts/HordeStorageArtifactCache.cs
--- a/Engine/Source/Programs/UnrealBuildTool/Artifacts/HordeStorageArtifactCache.cs
+++ b/Engine/Source/Programs/UnrealBuildTool/Artifacts/HordeStorageArtifactCache.cs
@@ -11,7 +11,6 @@ using EpicGames.Horde.Storage;
 using EpicGames.Horde.Storage.Bundles;
 using EpicGames.Horde.Storage.Nodes;
 using Microsoft.Extensions.Logging;
-using Microsoft.Extensions.Logging.Abstractions;
 
 namespace UnrealBuildTool.Artifacts
 {
@@ -73,7 +72,6 @@ namespace UnrealBuildTool.Artifacts
 		{
 			LeafChunkedDataNodeOptions leafOptions = new(512 * 1024, 1 * 1024 * 1024, 2 * 1024 * 1024);
 			InteriorChunkedDataNodeOptions interiorOptions = new(1, 10, 20);
-			ChunkingOptions options = new() { LeafOptions = leafOptions, InteriorOptions = interiorOptions };
 
 			using LeafChunkedDataWriter fileWriter = new(writer, leafOptions);
 			int index = 0;
@@ -155,7 +153,7 @@ namespace UnrealBuildTool.Artifacts
 	/// <summary>
 	/// Class for managing artifacts using horde storage
 	/// </summary>
-	public class HordeStorageArtifactCache : IArtifactCache
+	public sealed class HordeStorageArtifactCache : IArtifactCache, IDisposable
 	{
 		/// <summary>
 		/// Defines the theoretical max number of pending actions to write
@@ -167,11 +165,6 @@ namespace UnrealBuildTool.Artifacts
 		/// </summary>
 		private IStorageClient? _store = null;
 
-		/// <summary>
-		/// Logger to be used
-		/// </summary>
-		private readonly ILogger _logger;
-
 		/// <summary>
 		/// Task used to wait on ready state
 		/// </summary>
@@ -206,13 +199,22 @@ namespace UnrealBuildTool.Artifacts
 			private set => Interlocked.Exchange(ref _state, (int)value);
 		}
 
+		/// <inheritdoc/>
+		public void Dispose()
+		{
+			_store?.Dispose();
+			_pendingWritesFlushTask?.Dispose();
+			_pendingWritesFlushTask = null;
+			_semaphore.Dispose();
+		}
+
 		/// <summary>
 		/// Create a memory only cache
 		/// </summary>
 		/// <returns>Storage client instance</returns>
 		public static IArtifactCache CreateMemoryCache(ILogger logger)
 		{
-			HordeStorageArtifactCache cache = new(BundleStorageClient.CreateInMemory(logger), logger)
+			HordeStorageArtifactCache cache = new(BundleStorageClient.CreateInMemory(logger))
 			{
 				State = ArtifactCacheState.Available
 			};
@@ -228,8 +230,8 @@ namespace UnrealBuildTool.Artifacts
 		/// <returns>Storage client instance</returns>
 		public static IArtifactCache CreateFileCache(DirectoryReference directory, ILogger logger, bool cleanDirectory)
 		{
-			HordeStorageArtifactCache cache = new(null, logger);
-			cache._readyTask = Task.Run(() => cache.InitFileCache(directory, NullLogger.Instance, cleanDirectory));
+			HordeStorageArtifactCache cache = new(null);
+			cache._readyTask = Task.Run(() => cache.InitFileCache(directory, logger, cleanDirectory));
 			return cache;
 		}
 
@@ -237,11 +239,9 @@ namespace UnrealBuildTool.Artifacts
 		/// Constructor
 		/// </summary>
 		/// <param name="storage">Storage object to use</param>
-		/// <param name="logger">Logging destination</param>
-		private HordeStorageArtifactCache(IStorageClient? storage, ILogger logger)
+		private HordeStorageArtifactCache(IStorageClient? storage)
 		{
 			_store = storage;
-			_logger = logger;
 			_pendingWrites = new(MaxPendingSize);
 		}
 
@@ -468,7 +468,7 @@ namespace UnrealBuildTool.Artifacts
 					RefName refName = GetRefName(artifactAction.Key);
 
 					// Locate the destination collection for this key
-					ArtifactActionCollectionNode? node = _store!.TryReadRefTargetAsync<ArtifactActionCollectionNode>(refName, default, cancellationToken: cancellationToken).Result;
+					ArtifactActionCollectionNode? node = await _store!.TryReadRefTargetAsync<ArtifactActionCollectionNode>(refName, default, cancellationToken: cancellationToken);
 					node ??= new ArtifactActionCollectionNode();
 
 					// Update the artifact action collection
