UnrealBuildTool: Resolve CA1001, CA1012, CA1018, CA2213

#rnx

diff --git a/Engine/Source/Programs/UnrealBuildTool/Configuration/TargetRules.cs b/Engine/Source/Programs/UnrealBuildTool/Configuration/TargetRules.cs
--- a/Engine/Source/Programs/UnrealBuildTool/Configuration/TargetRules.cs
+++ b/Engine/Source/Programs/UnrealBuildTool/Configuration/TargetRules.cs
@@ -2880,7 +2880,7 @@ namespace UnrealBuildTool
 		/// Constructor.
 		/// </summary>
 		/// <param name="target">Information about the target being built</param>
-		public TargetRules(TargetInfo target)
+		protected TargetRules(TargetInfo target)
 		{
 			DefaultName = target.Name;
 			Platform = target.Platform;
diff --git a/Engine/Source/Programs/UnrealBuildTool/Executors/Experimental/FASTBuild.cs b/Engine/Source/Programs/UnrealBuildTool/Executors/Experimental/FASTBuild.cs
--- a/Engine/Source/Programs/UnrealBuildTool/Executors/Experimental/FASTBuild.cs
+++ b/Engine/Source/Programs/UnrealBuildTool/Executors/Experimental/FASTBuild.cs
@@ -20,6 +20,7 @@ using System.IO;
 using System.Linq;
 using System.Runtime.Serialization;
 using System.Runtime.Versioning;
+using System.Threading;
 using System.Threading.Tasks;
 using EpicGames.Core;
 using Microsoft.Extensions.Logging;
@@ -69,7 +70,7 @@ namespace UnrealBuildTool
 
 	///////////////////////////////////////////////////////////////////////
 
-	class FASTBuild : ActionExecutor
+	sealed class FASTBuild : ActionExecutor
 	{
 		/// <summary>
 		/// Executor to use for local actions
@@ -171,6 +172,13 @@ namespace UnrealBuildTool
 			LocalExecutor = new ParallelExecutor(MaxLocalActions, bAllCores, bCompactOutput, Logger);
 		}
 
+		/// <inheritdoc/>
+		public new void Dispose()
+		{
+			LocalExecutor.Dispose();
+			base.Dispose();
+		}
+
 		public override string Name => "FASTBuild";
 
 		public static string GetExecutableName()
diff --git a/Engine/Source/Programs/UnrealBuildTool/Executors/ImmediateActionQueue.cs b/Engine/Source/Programs/UnrealBuildTool/Executors/ImmediateActionQueue.cs
--- a/Engine/Source/Programs/UnrealBuildTool/Executors/ImmediateActionQueue.cs
+++ b/Engine/Source/Programs/UnrealBuildTool/Executors/ImmediateActionQueue.cs
@@ -803,7 +803,10 @@ namespace UnrealBuildTool
 		public void Dispose()
 		{
 			_cpuUtilizationTimer?.Dispose();
+			_cpuUtilizationTimer = null;
+			CancellationTokenSource.Dispose();
 			ProcessGroup.Dispose();
+			ProgressWriter.Dispose();
 		}
 
 		/// <summary>
diff --git a/Engine/Source/Programs/UnrealBuildTool/Platform/Apple/AppleProjectGenerator.cs b/Engine/Source/Programs/UnrealBuildTool/Platform/Apple/AppleProjectGenerator.cs
--- a/Engine/Source/Programs/UnrealBuildTool/Platform/Apple/AppleProjectGenerator.cs
+++ b/Engine/Source/Programs/UnrealBuildTool/Platform/Apple/AppleProjectGenerator.cs
@@ -40,7 +40,7 @@ namespace UnrealBuildTool
 			return new List<string>(0);
 		}
 
-		private class XcrunRunner
+		private sealed class XcrunRunner : IDisposable
 		{
 			private readonly Dictionary<string, IList<string>> CachedIncludePaths = new Dictionary<string, IList<string>>();
 
@@ -48,6 +48,14 @@ namespace UnrealBuildTool
 			private Process? XcrunProcess;
 			private bool IsReadingIncludesSection;
 
+			/// <inheritdoc/>
+			public void Dispose()
+			{
+				XcrunProcess?.Kill(true);
+				XcrunProcess?.Dispose();
+				XcrunProcess = null;
+			}
+
 			public IList<string> GetAppleSystemIncludePaths(UnrealArch Architecture, UnrealTargetPlatform Platform, ILogger Logger)
 			{
 				if (!UEBuildPlatform.IsPlatformInGroup(Platform, UnrealPlatformGroup.Apple))
diff --git a/Engine/Source/Programs/UnrealBuildTool/Platform/Windows/AppXManifestGeneratorBase.cs b/Engine/Source/Programs/UnrealBuildTool/Platform/Windows/AppXManifestGeneratorBase.cs
--- a/Engine/Source/Programs/UnrealBuildTool/Platform/Windows/AppXManifestGeneratorBase.cs
+++ b/Engine/Source/Programs/UnrealBuildTool/Platform/Windows/AppXManifestGeneratorBase.cs
@@ -71,7 +71,7 @@ namespace UnrealBuildTool
 		/// <summary>
 		/// Create a manifest generator for the given platform variant.
 		/// </summary>
-		public AppXManifestGeneratorBase(UnrealTargetPlatform InPlatform, ILogger InLogger)
+		protected AppXManifestGeneratorBase(UnrealTargetPlatform InPlatform, ILogger InLogger)
 		{
 			Platform = InPlatform;
 			Logger = InLogger;
diff --git a/Engine/Source/Programs/UnrealBuildTool/System/ConfigHierarchy.cs b/Engine/Source/Programs/UnrealBuildTool/System/ConfigHierarchy.cs
--- a/Engine/Source/Programs/UnrealBuildTool/System/ConfigHierarchy.cs
+++ b/Engine/Source/Programs/UnrealBuildTool/System/ConfigHierarchy.cs
@@ -173,7 +173,7 @@ namespace UnrealBuildTool
 	/// <summary>
 	/// Encapsulates a hierarchy of config files, merging sections from them together on request 
 	/// </summary>
-	public class ConfigHierarchy
+	public sealed class ConfigHierarchy : IDisposable
 	{
 		/// <summary>
 		/// Array of 
@@ -199,6 +199,14 @@ namespace UnrealBuildTool
 			this.Files = Files.ToArray();
 		}
 
+		#region IDisposible
+		/// <inheritdoc/>
+		public void Dispose()
+		{
+			NameToSectionLock.Dispose();
+		}
+		#endregion
+
 		/// <summary>
 		/// Names of all sections in all config files
 		/// </summary>
diff --git a/Engine/Source/Programs/UnrealBuildTool/System/ToolMode.cs b/Engine/Source/Programs/UnrealBuildTool/System/ToolMode.cs
--- a/Engine/Source/Programs/UnrealBuildTool/System/ToolMode.cs
+++ b/Engine/Source/Programs/UnrealBuildTool/System/ToolMode.cs
@@ -62,6 +62,7 @@ namespace UnrealBuildTool
 	/// <summary>
 	/// Attribute used to specify options for a UBT mode.
 	/// </summary>
+	[AttributeUsage(AttributeTargets.Class)]
 	sealed class ToolModeAttribute : Attribute
 	{
 		/// <summary>
