UnrealBuildTool: Allow -Wshadow warnings for clang 18.1.3

#rnx

diff --git a/Engine/Source/Programs/UnrealBuildTool/ToolChain/ClangWarnings.cs b/Engine/Source/Programs/UnrealBuildTool/ToolChain/ClangWarnings.cs
--- a/Engine/Source/Programs/UnrealBuildTool/ToolChain/ClangWarnings.cs
+++ b/Engine/Source/Programs/UnrealBuildTool/ToolChain/ClangWarnings.cs
@@ -113,8 +113,7 @@ namespace UnrealBuildTool
 				}
 			}
 
-			// Clang 17 suffers from [REDACTED_PROFILE_URL]/llvm-project/issues/71976 and should not be used as a preferred version until resolved
-			if (ClangVersion >= new VersionNumber(17))
+			if (ClangVersion >= new VersionNumber(17) && ClangVersion < new VersionNumber(18, 1, 3))
 			{
 				Arguments.Add("-Wno-shadow");
 			}
