UnrealBuildTool: Fix unknown telemetry

#rnx

diff --git a/Engine/Source/Programs/UnrealBuildTool/System/Telemetry.cs b/Engine/Source/Programs/UnrealBuildTool/System/Telemetry.cs
--- a/Engine/Source/Programs/UnrealBuildTool/System/Telemetry.cs
+++ b/Engine/Source/Programs/UnrealBuildTool/System/Telemetry.cs
@@ -268,6 +268,7 @@ namespace UnrealBuildTool
 #endif
 
 		const string ConfigCategoryPrefix = "StudioTelemetry.";
+		static readonly string? _hordeUrl = Environment.GetEnvironmentVariable("UE_HORDE_URL");
 
 		// Lazy to defer resolution until an event is first sent
 		readonly Lazy<TelemetryMetadata> _metadata = new(() =>
@@ -326,9 +327,9 @@ namespace UnrealBuildTool
 			}
 
 			// If we are running via horde, always the current horde url as a provider
-			if (!String.IsNullOrEmpty(_metadata.Value.Horde_URL))
+			if (!String.IsNullOrEmpty(_hordeUrl))
 			{
-				AddEndpointFromConfigInternal(engineIni, "StudioTelemetry.Provider.Horde", _metadata.Value.Horde_URL);
+				AddEndpointFromConfigInternal(engineIni, "StudioTelemetry.Provider.Horde", _hordeUrl);
 			}
 		}
 
