UnrealBuildTool: Fix bad format string in LinuxToolChain

#rnx

diff --git a/Engine/Source/Programs/UnrealBuildTool/Platform/Linux/LinuxToolChain.cs b/Engine/Source/Programs/UnrealBuildTool/Platform/Linux/LinuxToolChain.cs
--- a/Engine/Source/Programs/UnrealBuildTool/Platform/Linux/LinuxToolChain.cs
+++ b/Engine/Source/Programs/UnrealBuildTool/Platform/Linux/LinuxToolChain.cs
@@ -1020,7 +1020,7 @@ namespace UnrealBuildTool
 			ArchiveAction.ProducedItems.Add(OutputFile);
 			ArchiveAction.CommandDescription = "Archive";
 			ArchiveAction.StatusDescription = Path.GetFileName(OutputFile.AbsolutePath);
-			ArchiveAction.CommandArguments += String.Format("{1} \"{2}\"", GetArchiveArguments(LinkEnvironment), OutputFile.AbsolutePath);
+			ArchiveAction.CommandArguments += String.Format("{0} \"{1}\"", GetArchiveArguments(LinkEnvironment), OutputFile.AbsolutePath);
 
 			// Add the input files to a response file, and pass the response file on the command-line.
 			List<string> InputFileNames = new List<string>();
