UnrealBuildTool: TargetRules setting  to disable volatile metadata for msvc platforms

#rnx

diff --git a/Engine/Source/Programs/UnrealBuildTool/Platform/Windows/UEBuildWindows.cs b/Engine/Source/Programs/UnrealBuildTool/Platform/Windows/UEBuildWindows.cs
--- a/Engine/Source/Programs/UnrealBuildTool/Platform/Windows/UEBuildWindows.cs
+++ b/Engine/Source/Programs/UnrealBuildTool/Platform/Windows/UEBuildWindows.cs
@@ -521,6 +521,14 @@ namespace UnrealBuildTool
 		[CommandLine("-StrictEnumTypes")]
 		public bool bStrictEnumTypesConformance = false;
 
+		/// <summary>
+		/// Volatile Metadata is enabled by default and improves x64 emulation on arm64, but may come at a small perfomance cost (/volatileMetadata-).
+		/// </summary>
+		[RequiresUniqueBuildEnvironment]
+		[XmlConfigFile(Category = "WindowsPlatform")]
+		[CommandLine("-DisableVolatileMetadata")]
+		public bool bDisableVolatileMetadata { get; set; } = false;
+
 		/// <summary>
 		/// Whether to request the linker create a stripped pdb file as part of the build.
 		/// If enabled the full debug pdb will have the extension .full.pdb
@@ -891,6 +899,8 @@ namespace UnrealBuildTool
 
 		public bool bStrictEnumTypesConformance => Inner.bStrictEnumTypesConformance;
 
+		public bool bDisableVolatileMetadata => Inner.bDisableVolatileMetadata;
+
 		public bool bStripPrivateSymbols => Inner.bStripPrivateSymbols;
 
 		public uint? PdbPageSize => Inner.PdbPageSize;
diff --git a/Engine/Source/Programs/UnrealBuildTool/Platform/Windows/VCToolChain.cs b/Engine/Source/Programs/UnrealBuildTool/Platform/Windows/VCToolChain.cs
--- a/Engine/Source/Programs/UnrealBuildTool/Platform/Windows/VCToolChain.cs
+++ b/Engine/Source/Programs/UnrealBuildTool/Platform/Windows/VCToolChain.cs
@@ -720,6 +720,12 @@ namespace UnrealBuildTool
 				}
 			}
 
+			// Volatile Metadata is enabled by default and improves x64 emulation on arm64, but may come at a small perfomance cost
+			if (Target.WindowsPlatform.Compiler.IsMSVC() && CompileEnvironment.Architecture == UnrealArch.X64 && Target.WindowsPlatform.bDisableVolatileMetadata)
+			{
+				Arguments.Add("/volatileMetadata-");
+			}
+
 			//
 			// LTCG and PGO
 			//
