UnrealBuildTool: Fix bStaticAnalyzerExtensions not being set in copy constructor

#rnx

diff --git a/Engine/Source/Programs/UnrealBuildTool/System/CppCompileEnvironment.cs b/Engine/Source/Programs/UnrealBuildTool/System/CppCompileEnvironment.cs
--- a/Engine/Source/Programs/UnrealBuildTool/System/CppCompileEnvironment.cs
+++ b/Engine/Source/Programs/UnrealBuildTool/System/CppCompileEnvironment.cs
@@ -759,6 +759,8 @@ namespace UnrealBuildTool
 			StaticAnalyzerCheckers = new HashSet<string>(Other.StaticAnalyzerCheckers);
 			StaticAnalyzerDisabledCheckers = new HashSet<string>(Other.StaticAnalyzerDisabledCheckers);
 			StaticAnalyzerAdditionalCheckers = new HashSet<string>(Other.StaticAnalyzerAdditionalCheckers);
+			bStaticAnalyzerExtensions = Other.bStaticAnalyzerExtensions;
+			StaticAnalyzerRulesets = new HashSet<FileReference>(Other.StaticAnalyzerRulesets);
 			bOptimizeCode = Other.bOptimizeCode;
 			bUseAutoRTFMCompiler = Other.bUseAutoRTFMCompiler;
 			bAllowAutoRTFMInstrumentation = Other.bAllowAutoRTFMInstrumentation;
