- Made a script (ExportAppleBuild) for uploading or exporting .xcarchives (we had MakeIPA but this is more general purpose and uses a project param to get ASC params
  - This needs .plist files to control the export/upload process, added some premade default ones in Engine/Build/Xcode
- Refactored some code for the Auth param stuff

diff --git a/Engine/Build/InstalledEngineFilters.xml b/Engine/Build/InstalledEngineFilters.xml
--- a/Engine/Build/InstalledEngineFilters.xml
+++ b/Engine/Build/InstalledEngineFilters.xml
@@ -368,6 +368,9 @@
         
         <!-- Mac build resources (plist entitlement etc.) -->
         Engine/Build/Mac/Resources/...
+		
+		<!-- Some plist files for exporting Mac/IOS builds for distribuion -->
+		Engine/Build/Xcode/...
 
 		<!-- DotNet -->
 		Engine/Binaries/ThirdParty/DotNet/6.0.302/mac-arm64/...
diff --git a/Engine/Source/Programs/AutomationTool/Apple/ApplePlatform.Automation.cs b/Engine/Source/Programs/AutomationTool/Apple/ApplePlatform.Automation.cs
--- a/Engine/Source/Programs/AutomationTool/Apple/ApplePlatform.Automation.cs
+++ b/Engine/Source/Programs/AutomationTool/Apple/ApplePlatform.Automation.cs
@@ -145,45 +145,27 @@ public abstract class ApplePlatform : Platform
 
 	private DirectoryReference GetArchivePath(StageTarget Target, DeploymentContext SC)
 	{
-		DirectoryReference UserDir = new DirectoryReference(Environment.GetFolderPath(Environment.SpecialFolder.UserProfile));
-		DirectoryReference Library = DirectoryReference.Combine(UserDir, "Library/Developer/Xcode/Archives");
-
-		// order date named folders (use creating data, not name, but same thing)
-		List<DirectoryReference> DateDirs = DirectoryReference.EnumerateDirectories(Library).ToList();
-		DateDirs.SortBy(x => Directory.GetCreationTime(x.FullName));
-		DateDirs.Reverse();
-
-		string ArchiveName = Target.Receipt.TargetName;
+		string TargetName = Target.Receipt.TargetName;
 		if (!SC.IsCodeBasedProject)
 		{
 			if (SC.RawProjectPath != null)
 			{
-				ArchiveName = MakeContentOnlyTargetName(Target.Receipt, SC.ShortProjectName);
+				TargetName = MakeContentOnlyTargetName(Target.Receipt, SC.ShortProjectName);
 			}
 		}
 
-		// go through each folder, starting at most recent, looking for an archive for the target
-		foreach (DirectoryReference DateDir in DateDirs)
-		{
-			// find the most recent archive for this target (based on name of target, this ignores Development vs Shipping, but 
-			// since Distribution is meant only for Shipping it's ok
-			string Wildcard = AppleExports.MakeBinaryFileName(ArchiveName, Target.Receipt.Platform, UnrealTargetConfiguration.Development,
-				Target.Receipt.Architectures, UnrealTargetConfiguration.Development, null) + " *.xcarchive";
-
-			Logger.LogInformation("Looking in Xcode archive dir {0} for {1}", DateDir, Wildcard);
+		// find the most recent archive for this target (based on name of target, this ignores Development vs Shipping, but 
+		// since Distribution is meant only for Shipping it's ok
+		string ArchiveName = AppleExports.MakeBinaryFileName(TargetName, Target.Receipt.Platform, UnrealTargetConfiguration.Development,
+		 	Target.Receipt.Architectures, UnrealTargetConfiguration.Development, null) + " *.xcarchive";
 
-			List<DirectoryReference> XcArchives = DirectoryReference.EnumerateDirectories(DateDir, Wildcard).ToList();
-			if (XcArchives.Count > 0)
-			{
-				XcArchives.SortBy(x => Directory.GetCreationTime(x.FullName));
-				DirectoryReference XcArchive = XcArchives.Last();
-
-				Logger.LogInformation("Found xcarchive dir {0}", XcArchive);
-				return XcArchive;
-			}
+		if (ArchiveName != TargetName)
+		{
+			// i don't know why we called MakeBinaryFileName - with Dev it seems to just return the TargetName passed in
+			throw new AutomationException($"GetArchivePath changed name of Target from {TargetName} to {ArchiveName} - tell Josh Adams");
 		}
 
-		return null;
+		return AppleExports.FindLatestXcArchive(TargetName);
 	}
 
 	private DirectoryReference GetFinalAppPath(StageTarget Target, DeploymentContext SC)
diff --git a/Engine/Source/Programs/AutomationTool/Apple/ExportAppleBuild.Automation.cs b/Engine/Source/Programs/AutomationTool/Apple/ExportAppleBuild.Automation.cs
new file mode 100644
--- /dev/null
+++ b/Engine/Source/Programs/AutomationTool/Apple/ExportAppleBuild.Automation.cs
@@ -0,0 +1,93 @@
+// Copyright Epic Games, Inc. All Rights Reserved.
+
+using System;
+using System.Linq;
+using System.IO;
+using System.Collections.Generic;
+using EpicGames.Core;
+using AutomationTool;
+using System.Threading;
+using UnrealBuildBase;
+using Microsoft.Extensions.Logging;
+using UnrealBuildTool;
+
+class ExportAppleBuild : BuildCommand
+{
+	public override ExitCode Execute()
+	{
+		FileReference ProjectFile = ParseProjectParam();
+		if (ProjectFile == null)
+		{
+			Logger.LogError("Project not specified, of found, with -project param");
+			return ExitCode.Error_Arguments;
+		}
+
+		string PlistParam = ParseParamValue("Plist");
+		string XcArchivePathParam = ParseParamValue("XcArchive");
+		string LatestXcArchiveParam = ParseParamValue("LatestXcArchiveForTarget");
+		string ExportPath = ParseParamValue("ExportPath");
+
+		DirectoryReference XcArchive;
+		if (!string.IsNullOrEmpty(LatestXcArchiveParam))
+		{
+			// find the latest archive for the given project name
+			XcArchive = AppleExports.FindLatestXcArchive(LatestXcArchiveParam);
+		}
+		else
+		{
+			XcArchive = new DirectoryReference(XcArchivePathParam);			
+		}
+
+		if (XcArchive == null || !DirectoryReference.Exists(XcArchive))
+		{
+			Logger.LogError("No XCArchive found, with -XCArchive or LatestXCArchiveForTarget params");
+			return ExitCode.Error_Arguments;
+		}
+
+		string PlistPath;
+		if (!string.IsNullOrEmpty(PlistParam) && PlistParam[0] == '/')
+		{
+			PlistPath = PlistParam;
+		}
+		else
+		{
+			// look in project and engine for named plist
+			PlistPath = Path.Combine(ProjectFile.Directory.FullName, "Build", "Xcode", PlistParam + ".plist");
+			if (!File.Exists(PlistPath))
+			{
+				PlistPath = Path.Combine(Unreal.EngineDirectory.FullName, "Build", "Xcode", PlistParam + ".plist");
+			}
+		}
+		if (!File.Exists(PlistPath))
+		{
+			Logger.LogError("No plist options file found, with -plist param");
+			return ExitCode.Error_Arguments;
+		}
+
+		Logger.LogInformation("Project: {Project}, XCArchive: {XC}, Plist {Plist}", ProjectFile, XcArchive, PlistPath);
+
+		// build commandline
+		string CommandLine = $"-exportArchive";
+		CommandLine += $" -archivePath \"{XcArchive}\"";
+		CommandLine += $" -exportOptionsPlist \"{PlistPath}\"";
+		CommandLine += $" -allowProvisioningUpdates";
+		CommandLine += AppleExports.GetXcodeBuildAuthOptions(ProjectFile);
+
+		if (!string.IsNullOrEmpty(ExportPath))
+		{
+			CommandLine += $" -exportPath {ExportPath}";
+		}
+
+		Logger.LogInformation($"Running 'xcodebuild {CommandLine}'...");
+
+		int Return;
+		Utils.RunLocalProcessAndReturnStdOut("/usr/bin/xcodebuild", CommandLine, Logger, out Return);
+
+// xcodebuild -exportArchive -archivePath "[REDACTED_USER_HOME]/Library/Developer/Xcode/Archives/2024-04-06/LyraGame 4-6-24, 2.35 PM.xcarchive" 
+// -exportPath ./private_keys -exportOptionsPlist Engine/Build/Xcode/ExportOptions_UploadTestFlight.plist 
+// -authenticationKeyIssuerID 69a6de71-5e09-47e3-e053-5b8c7c11a4d1 -authenticationKeyPath $PWD/private_keys/AuthKey_HCTZ2FPZ2Y.p8 
+// -authenticationKeyID HCTZ2FPZ2Y
+
+		return (ExitCode)Return;
+	}
+}
diff --git a/Engine/Source/Programs/UnrealBuildTool/Platform/Apple/AppleExports.cs b/Engine/Source/Programs/UnrealBuildTool/Platform/Apple/AppleExports.cs
--- a/Engine/Source/Programs/UnrealBuildTool/Platform/Apple/AppleExports.cs
+++ b/Engine/Source/Programs/UnrealBuildTool/Platform/Apple/AppleExports.cs
@@ -2,6 +2,9 @@
 
 using System;
 using EpicGames.Core;
+using System.Collections.Generic;
+using System.Linq;
+using System.IO;
 using Microsoft.Extensions.Logging;
 using UnrealBuildBase;
 
@@ -155,6 +158,37 @@ namespace UnrealBuildTool
 			Distribute = 4,
 		}
 
+		/// <summary>
+		/// Gets the AppStoreConnect auth options for a given project. Will return empty string if the project isn't set up to use ASC
+		/// </summary>
+		/// <param name="ProjectFile"></param>
+		/// <returns></returns>
+		public static string GetXcodeBuildAuthOptions(FileReference? ProjectFile)
+		{
+			string Options = "";
+			
+			// handle AppStore Connect settings
+			bool bUseAppStoreConnect;
+		
+			ConfigHierarchy SharedPlatformIni = ConfigCache.ReadHierarchy(ConfigHierarchyType.Engine, ProjectFile?.Directory, UnrealTargetPlatform.Mac);
+			SharedPlatformIni.TryGetValue("/Script/MacTargetPlatform.XcodeProjectSettings", "bUseAppStoreConnect", out bUseAppStoreConnect);
+			if (bUseAppStoreConnect)
+			{
+				string? IssuerID, KeyID, KeyPath;
+				if (SharedPlatformIni.TryGetValue("/Script/MacTargetPlatform.XcodeProjectSettings", "AppStoreConnectIssuerID", out IssuerID) &&
+					SharedPlatformIni.TryGetValue("/Script/MacTargetPlatform.XcodeProjectSettings", "AppStoreConnectKeyID", out KeyID) &&
+					SharedPlatformIni.TryGetValue("/Script/MacTargetPlatform.XcodeProjectSettings", "AppStoreConnectKeyPath", out KeyPath))
+				{
+					FileReference KeyFile = ConvertFilePath(ProjectFile?.Directory, KeyPath);
+					Options += $" -authenticationKeyIssuerID {IssuerID}";
+					Options += $" -authenticationKeyID {KeyID}";
+					Options += $" -authenticationKeyPath \"{KeyFile}\"";
+				}
+			}
+
+			return Options;
+		}
+
 		/// <summary>
 		/// Generates a stub xcode project for the given project/platform/target combo, then builds or archives it
 		/// </summary>
@@ -201,22 +235,8 @@ namespace UnrealBuildTool
 				{
 					ExtraOptions += " -allowProvisioningUpdates";
 
-					// handle AppStore Connect settings
-					bool bUseAppStoreConnect;
-					SharedPlatformIni.TryGetValue("/Script/MacTargetPlatform.XcodeProjectSettings", "bUseAppStoreConnect", out bUseAppStoreConnect);
-					if (bUseAppStoreConnect)
-					{
-						string? IssuerID, KeyID, KeyPath;
-						if (SharedPlatformIni.TryGetValue("/Script/MacTargetPlatform.XcodeProjectSettings", "AppStoreConnectIssuerID", out IssuerID) &&
-							SharedPlatformIni.TryGetValue("/Script/MacTargetPlatform.XcodeProjectSettings", "AppStoreConnectKeyID", out KeyID) &&
-							SharedPlatformIni.TryGetValue("/Script/MacTargetPlatform.XcodeProjectSettings", "AppStoreConnectKeyPath", out KeyPath))
-						{
-							FileReference KeyFile = ConvertFilePath(ProjectFile?.Directory, KeyPath);
-							ExtraOptions += $" -authenticationKeyIssuerID {IssuerID}";
-							ExtraOptions += $" -authenticationKeyID {KeyID}";
-							ExtraOptions += $" -authenticationKeyPath \"{KeyFile}\"";
-						}
-					}
+					ExtraOptions += GetXcodeBuildAuthOptions(ProjectFile);
+
 				}
 			}
 
@@ -271,5 +291,38 @@ namespace UnrealBuildTool
 			string StandardBinaryName = UEBuildTarget.MakeBinaryFileName(BinaryName, Platform, Configuration, Architectures, UndecoratedConfiguration, UEBuildBinaryType.Executable);
 			return System.IO.Path.ChangeExtension(StandardBinaryName, Extension);
 		}
+
+		/// <summary>
+		/// FInds the latest .xcarchive for a given Target name (the .xcarchive will start with this name then have a data appended)
+		/// </summary>
+		/// <param name="TargetName"></param>
+		/// <returns></returns>
+		public static DirectoryReference? FindLatestXcArchive(string TargetName)
+		{
+			DirectoryReference UserDir = new DirectoryReference(Environment.GetFolderPath(Environment.SpecialFolder.UserProfile));
+			DirectoryReference Library = DirectoryReference.Combine(UserDir, "Library/Developer/Xcode/Archives");
+
+			// order date named folders (use creating data, not name, but same thing)
+			List<DirectoryReference> DateDirs = DirectoryReference.EnumerateDirectories(Library).ToList();
+			DateDirs.SortBy(x => Directory.GetCreationTime(x.FullName));
+			DateDirs.Reverse();
+
+			// go through each folder, starting at most recent, looking for an archive for the target
+			foreach (DirectoryReference DateDir in DateDirs)
+			{
+				string Wildcard = $"{TargetName} *.xcarchive";
+
+				List<DirectoryReference> XcArchives = DirectoryReference.EnumerateDirectories(DateDir, Wildcard).ToList();
+				if (XcArchives.Count > 0)
+				{
+					XcArchives.SortBy(x => Directory.GetCreationTime(x.FullName));
+					DirectoryReference XcArchive = XcArchives.Last();
+
+					return XcArchive;
+				}
+			}
+
+			return null;
+		}
 	}
 }
