Fix: Update Restricted ASIS plugin conflict with changes from latest release

diff --git a/Engine/Source/Programs/UnrealBuildTool/Platform/Android/UEBuildAndroid.cs b/Engine/Source/Programs/UnrealBuildTool/Platform/Android/UEBuildAndroid.cs
--- a/Engine/Source/Programs/UnrealBuildTool/Platform/Android/UEBuildAndroid.cs
+++ b/Engine/Source/Programs/UnrealBuildTool/Platform/Android/UEBuildAndroid.cs
@@ -49,6 +49,13 @@ namespace UnrealBuildTool
 		/// </summary>
 		[ConfigFile(ConfigHierarchyType.Engine, "/Script/AndroidTargetPlatform.AndroidTargetSettings")]
 		public bool bEnableRayTracing = false;
+
+		/// <summary>
+		/// Enables ASIS plugin and STANDALONE support.
+		/// </summary>
+		[ConfigFile(ConfigHierarchyType.Engine, "/Script/AndroidSingleInstanceServiceEditor.AndroidSingleInstanceServiceRuntimeSettings")]
+		public bool bEnableASISPlugin = false;
+
 	}
 
 	/// <summary>
@@ -88,6 +95,8 @@ namespace UnrealBuildTool
 
 		public bool bEnableRayTracing => Inner.bEnableRayTracing;
 
+		public bool bEnableASISPlugin => Inner.bEnableASISPlugin;
+
 		public AndroidTargetRules TargetRules => Inner;
 
 #pragma warning restore CS1591
@@ -414,24 +423,9 @@ namespace UnrealBuildTool
 		{
 		}
 
-		public static bool IsMakeAAREnabled(FileReference? ProjectFile, ILogger InLogger)
+		public static bool IsMakeAAREnabled(ReadOnlyTargetRules Target)
 		{
-			// look in ini settings for what platforms to compile for
-			DirectoryReference? ProjectDirectory = DirectoryReference.FromFile(ProjectFile);
-			ConfigHierarchy Ini = ConfigCache.ReadHierarchy(ConfigHierarchyType.Engine, DirectoryReference.FromFile(ProjectFile), UnrealTargetPlatform.Android);
-			bool bEnablePlugin = false;
-
-			if (Ini != null)
-			{
-				Ini.GetBool("/Script/AndroidSingleInstanceServiceEditor.AndroidSingleInstanceServiceRuntimeSettings", "bEnablePlugin", out bEnablePlugin);
-				// Only log if enabled to avoid adding extra logging to builds that don't use this feature.
-				if (bEnablePlugin)
-				{
-					InLogger.LogInformation("IsMakeAAREnabled {Ini} and {ProjectFile} , {ProjectDirectory}, {bEnablePlugin}", Ini, ProjectFile, ProjectDirectory, bEnablePlugin);
-				}
-			}
-
-			return bEnablePlugin;
+			return Target.AndroidPlatform.bEnableASISPlugin;
 		}
 
 		public virtual void SetUpSpecificEnvironment(ReadOnlyTargetRules Target, CppCompileEnvironment CompileEnvironment, LinkEnvironment LinkEnvironment, ILogger Logger)
@@ -464,11 +458,7 @@ namespace UnrealBuildTool
 
 			CompileEnvironment.Definitions.Add("WITH_EDITOR=0");
 			CompileEnvironment.Definitions.Add("USE_NULL_RHI=0");
-			if (IsMakeAAREnabled(Target.ProjectFile, Logger))
-			{
-				Logger.LogInformation("SetUpSpecificEnvironment is adding USE_ANDROID_STANDALONE because IsMakeAAREnabled({ProjectFile}) was true!", Target.ProjectFile);
-				CompileEnvironment.Definitions.Add("USE_ANDROID_STANDALONE=1");
-			}
+
 
 			DirectoryReference NdkDir = new DirectoryReference(NDKPath);
 			//CompileEnvironment.SystemIncludePaths.Add(DirectoryReference.Combine(NdkDir, "sources/cxx-stl/llvm-libc++/include"));
@@ -520,6 +510,12 @@ namespace UnrealBuildTool
 				CompileEnvironment.Definitions.Add("RHI_RAYTRACING=1");
 			}
 
+			if (Target.AndroidPlatform.bEnableASISPlugin)
+			{
+				Logger.LogInformation("Compiling with USE_ANDROID_STANDALONE");
+				CompileEnvironment.Definitions.Add("USE_ANDROID_STANDALONE=1");
+			}
+
 			if (Target.bPGOOptimize || Target.bPGOProfile)
 			{
 				Logger.LogInformation("PGO {PgoType} build", Target.bPGOOptimize ? "optimize" : "profile");
diff --git a/Engine/Source/Programs/UnrealBuildTool/Platform/Android/UEDeployAndroid.cs b/Engine/Source/Programs/UnrealBuildTool/Platform/Android/UEDeployAndroid.cs
--- a/Engine/Source/Programs/UnrealBuildTool/Platform/Android/UEDeployAndroid.cs
+++ b/Engine/Source/Programs/UnrealBuildTool/Platform/Android/UEDeployAndroid.cs
@@ -2763,18 +2763,19 @@ namespace UnrealBuildTool
 			StringBuilder Text = new StringBuilder();
 			Text.AppendLine(XML_HEADER);
 
-			bool bIsMakeAAREnabled = AndroidPlatform.IsMakeAAREnabled(ProjectFile, Logger);
-			if (bIsMakeAAREnabled)
-			{
-				// Note: dist was added for SingleInstanceService to allow instant module features, but need feedback if this will cause any issues since according to docs (https://developer.android.com/guide/playcore/feature-delivery/instant) An instant-enabled module cannot use background services. Additionally, such a module cannot send notifications when running in the background.
 
+			bool bIsMakeAAREnabled = false;
+			Ini.GetBool("/Script/AndroidSingleInstanceServiceEditor.AndroidSingleInstanceServiceRuntimeSettings", "bEnableASISPlugin", out bIsMakeAAREnabled);
+
+			if (bIsMakeAAREnabled)
+			{				
 				Text.AppendLine("<manifest xmlns:android=\"http://schemas.android.com/apk/res/android\" xmlns:dist=\"http://schemas.android.com/apk/distribution\" xmlns:tools=\"http://schemas.android.com/tools\"");
 			}
 			else
 			{
 				Text.AppendLine("<manifest xmlns:android=\"http://schemas.android.com/apk/res/android\" xmlns:tools=\"http://schemas.android.com/tools\"");
 			}
-			//			Text.AppendLine(string.Format("          package=\"{0}\"", PackageName));
+
 			if (ExtraManifestNodeTags != null)
 			{
 				foreach (string Line in ExtraManifestNodeTags)
