Localization:
Allow a plugin descriptor to remove the localization target field from a .uplugin file if all localization targets have been removed from a plugin descriptor.

#rnx

diff --git a/Engine/Source/Programs/UnrealBuildTool/System/LocalizationDescriptor.cs b/Engine/Source/Programs/UnrealBuildTool/System/LocalizationDescriptor.cs
--- a/Engine/Source/Programs/UnrealBuildTool/System/LocalizationDescriptor.cs
+++ b/Engine/Source/Programs/UnrealBuildTool/System/LocalizationDescriptor.cs
@@ -178,6 +178,12 @@ namespace UnrealBuildTool
 				JsonObject[] JsonObjects = Targets.Select(X => X.ToJsonObject()).ToArray();
 				InObject.AddOrSetFieldValue(Name, JsonObjects);
 			}
+			else
+			{
+				// The loaded plugin descriptor could have an existing localization target descriptor 
+				// If we remove it in memory, we should try and also remove it from the cached json in this case 
+				InObject.RemoveField(Name);
+			}
 		}
 	}
 }
