- Allow for -ActionTypeFilter to BuildMode, which will make only certain action types to run

diff --git a/Engine/Source/Programs/UnrealBuildTool/Modes/BuildMode.cs b/Engine/Source/Programs/UnrealBuildTool/Modes/BuildMode.cs
--- a/Engine/Source/Programs/UnrealBuildTool/Modes/BuildMode.cs
+++ b/Engine/Source/Programs/UnrealBuildTool/Modes/BuildMode.cs
@@ -97,6 +97,12 @@ namespace UnrealBuildTool
 		[CommandLine("-SaveCrashDumps=")]
 		public DirectoryReference? SaveCrashDumpDirectory = null;
 
+		/// <summary>
+		/// If specified, then only this type of action will execute
+		/// </summary>
+		[CommandLine("-ActionTypeFilter=")]
+		public string? ActionTypeFilter = null;
+
 		/// <summary>
 		/// Main entry point
 		/// </summary>
@@ -246,7 +252,7 @@ namespace UnrealBuildTool
 					// Create the working set provider per group.
 					using (ISourceFileWorkingSet WorkingSet = SourceFileWorkingSet.Create(Unreal.RootDirectory, ProjectDirs, Logger))
 					{
-						await BuildAsync(TargetDescriptors, BuildConfiguration, WorkingSet, Options, WriteOutdatedActionsFile, Logger, bSkipPreBuildTargets);
+						await BuildAsync(TargetDescriptors, BuildConfiguration, WorkingSet, Options, WriteOutdatedActionsFile, Logger, bSkipPreBuildTargets, ActionTypeFilter:ActionTypeFilter);
 					}
 				}
 			}
@@ -376,8 +382,9 @@ namespace UnrealBuildTool
 		/// <param name="WriteOutdatedActionsFile">Files to write the list of outdated actions to (rather than building them)</param>
 		/// <param name="Logger">Logger for output</param>
 		/// <param name="bSkipPreBuildTargets">If true then only the current target descriptors will be built.</param>
+		/// <param name="ActionTypeFilter">If specified, only actions of this type will be run</param>
 		/// <returns>Result from the compilation</returns>
-		public static async Task BuildAsync(List<TargetDescriptor> TargetDescriptors, BuildConfiguration BuildConfiguration, ISourceFileWorkingSet WorkingSet, BuildOptions Options, FileReference? WriteOutdatedActionsFile, ILogger Logger, bool bSkipPreBuildTargets = false)
+		public static async Task BuildAsync(List<TargetDescriptor> TargetDescriptors, BuildConfiguration BuildConfiguration, ISourceFileWorkingSet WorkingSet, BuildOptions Options, FileReference? WriteOutdatedActionsFile, ILogger Logger, bool bSkipPreBuildTargets = false, string? ActionTypeFilter = null)
 		{
 			List<TargetMakefile> TargetMakefiles = new List<TargetMakefile>();
 
@@ -405,7 +412,7 @@ namespace UnrealBuildTool
 				}
 			}
 
-			await BuildAsync(TargetMakefiles.ToArray(), TargetDescriptors, BuildConfiguration, Options, WriteOutdatedActionsFile, Logger);
+			await BuildAsync(TargetMakefiles.ToArray(), TargetDescriptors, BuildConfiguration, Options, WriteOutdatedActionsFile, Logger, ActionTypeFilter);
 		}
 
 		/// <summary>
@@ -417,8 +424,9 @@ namespace UnrealBuildTool
 		/// <param name="Options">Additional options for the build</param>
 		/// <param name="WriteOutdatedActionsFile">Files to write the list of outdated actions to (rather than building them)</param>
 		/// <param name="Logger">Logger for output</param>
+		/// <param name="ActionTypeFilter">If specified, only actions of this type will be run</param>
 		/// <returns>Result from the compilation</returns>
-		internal static async Task BuildAsync(TargetMakefile[] Makefiles, List<TargetDescriptor> TargetDescriptors, BuildConfiguration BuildConfiguration, BuildOptions Options, FileReference? WriteOutdatedActionsFile, ILogger Logger)
+		internal static async Task BuildAsync(TargetMakefile[] Makefiles, List<TargetDescriptor> TargetDescriptors, BuildConfiguration BuildConfiguration, BuildOptions Options, FileReference? WriteOutdatedActionsFile, ILogger Logger, string? ActionTypeFilter=null)
 		{
 			// Execute the build
 			if ((Options & BuildOptions.SkipBuild) == 0)
@@ -664,6 +672,21 @@ namespace UnrealBuildTool
 
 				// Link the action graph again to sort it
 				List<LinkedAction> MergedActionsToExecute = ActionToOutdatedFlag.Where(x => x.Value).Select(x => x.Key).ToList();
+
+				// if we want to only run one type of action, filter it now
+				if (ActionTypeFilter != null)
+				{
+					ActionType Filter;
+					if (Enum.TryParse<ActionType>(ActionTypeFilter, out Filter))
+					{
+						MergedActionsToExecute = MergedActionsToExecute.Where(x => x.ActionType == Filter).ToList();
+					}
+					else
+					{
+						throw new BuildException($"Unknown ActionType for ActionTypeFilter: {ActionTypeFilter}");
+					}
+				}
+
 				ActionGraph.Link(MergedActionsToExecute, Logger);
 
 				// Allow hot reload to override the actions
