UnrealBuildTool: Handle Ctrl-C to gracefully cancel running action queues
Add command line parameter to cancel the action queue after a timeout

diff --git a/Engine/Source/Programs/UnrealBuildTool/Executors/ImmediateActionQueue.cs b/Engine/Source/Programs/UnrealBuildTool/Executors/ImmediateActionQueue.cs
--- a/Engine/Source/Programs/UnrealBuildTool/Executors/ImmediateActionQueue.cs
+++ b/Engine/Source/Programs/UnrealBuildTool/Executors/ImmediateActionQueue.cs
@@ -115,6 +115,14 @@ namespace UnrealBuildTool
 		[XmlConfigFile(Category = "BuildConfiguration")]
 		public int ActionStallReportTime = 0;
 
+		/// <summary>
+		/// Number of second of no completed actions to trigger to terminate the queue.
+		/// If zero, force termination not be enabled.
+		/// </summary>
+		[CommandLine("-ActionStallTerminateTime=")]
+		[XmlConfigFile(Category = "BuildConfiguration")]
+		public int ActionStallTerminateTime = 0;
+
 		/// <summary>
 		/// Running status of the action
 		/// </summary>
@@ -323,10 +331,15 @@ namespace UnrealBuildTool
 		private int _lastActionChange = 1;
 
 		/// <summary>
-		/// If true, a action stall has been reported for the current change count
+		/// If true, an action stall has been reported for the current change count
 		/// </summary>
 		private bool _lastActionStallReported = false;
 
+		/// <summary>
+		/// If true, the queue has been cancelled due to an action stall
+		/// </summary>
+		private bool _lastActionStallCanceled = false;
+
 		/// <summary>
 		/// Time of the last change to the action count.  This is updated by the timer.
 		/// </summary>
@@ -415,6 +428,25 @@ namespace UnrealBuildTool
 
 				_runners.Add(new(ImmediateActionQueueRunnerType.Automatic, ActionPhase.ArtifactCheck, runAction, false, maxActionArtifactCacheTasks, 0));
 			}
+
+			// Cancel the queue when Ctrl-C is pressed
+			Console.CancelKeyPress += CancelKeyPress;
+		}
+
+		/// <summary>
+		/// Event handler for the Console.CancelKeyPress event
+		/// </summary>
+		/// <param name="sender"></param>
+		/// <param name="e"></param>
+		private void CancelKeyPress(object? sender, ConsoleCancelEventArgs e)
+		{
+			Console.CancelKeyPress -= CancelKeyPress;
+			if (!CancellationTokenSource.IsCancellationRequested)
+			{
+				Logger.LogWarning("Canceling actions...");
+				CancellationTokenSource.Cancel();
+				e.Cancel = true;
+			}
 		}
 
 		/// <summary>
@@ -468,7 +500,7 @@ namespace UnrealBuildTool
 						}
 					}
 
-					if (ActionStallReportTime > 0)
+					if (ActionStallReportTime > 0 || ActionStallTerminateTime > 0)
 					{
 						lock (Actions)
 						{
@@ -482,13 +514,19 @@ namespace UnrealBuildTool
 							}
 
 							// Otherwise, if we haven't already generated a report, test for a timeout in seconds and generate one on timeout.
-							else if (!_lastActionStallReported && (DateTime.Now - _lastActionChangeTime).TotalSeconds > ActionStallReportTime)
+							else if (ActionStallReportTime > 0 && !_lastActionStallReported && (DateTime.Now - _lastActionChangeTime).TotalSeconds > ActionStallReportTime)
 							{
 								_lastActionStallReported = true;
 								GenerateStallReport();
 							}
+							else if (ActionStallTerminateTime > 0 && !_lastActionStallCanceled && (DateTime.Now - _lastActionChangeTime).TotalSeconds > ActionStallTerminateTime)
+							{
+								_lastActionStallCanceled = true;
+								CancelStalledActions();
+							}
 						}
 					}
+
 				}, null, 1000, 1000);
 			}
 
@@ -588,7 +626,7 @@ namespace UnrealBuildTool
 					{
 						try
 						{
-							runAction().Wait();
+							runAction().Wait(CancellationToken);
 						}
 						catch (Exception ex)
 						{
@@ -601,7 +639,7 @@ namespace UnrealBuildTool
 						{
 							try
 							{
-								runAction().Wait();
+								runAction().Wait(CancellationToken);
 							}
 							catch (Exception ex)
 							{
@@ -802,6 +840,7 @@ namespace UnrealBuildTool
 		/// </summary>
 		public void Dispose()
 		{
+			Console.CancelKeyPress -= CancelKeyPress;
 			_cpuUtilizationTimer?.Dispose();
 			_cpuUtilizationTimer = null;
 			CancellationTokenSource.Dispose();
@@ -1218,7 +1257,7 @@ namespace UnrealBuildTool
 			}
 
 			IEnumerable<int> CompletedActions = Enumerable.Range(0, Actions.Length)
-				.Where(x => Actions[x].Results != null && Actions[x].Results!.ExecutionTime > TimeSpan.Zero)
+				.Where(x => Actions[x].Results != null && Actions[x].Results!.ExecutionTime > TimeSpan.Zero && Actions[x].Results!.ExitCode != Int32.MaxValue)
 				.OrderByDescending(x => Actions[x].Results!.ExecutionTime)
 				.Take(20);
 
@@ -1330,5 +1369,12 @@ namespace UnrealBuildTool
 				Logger.LogInformation("Queue Counts: Queued = {Queued}, Running = {Running}, Finished = {Finished}, Error = {Error}", queued, running, finished, error);
 			}
 		}
+
+		private void CancelStalledActions()
+		{
+			GenerateStallReport();
+			Logger.LogInformation("Action stall terminate time exceeded, canceling remaining actions...");
+			CancellationTokenSource.Cancel();
+		}
 	}
 }
diff --git a/Engine/Source/Programs/UnrealBuildTool/Executors/UnrealBuildAccelerator/UBAExecutor.cs b/Engine/Source/Programs/UnrealBuildTool/Executors/UnrealBuildAccelerator/UBAExecutor.cs
--- a/Engine/Source/Programs/UnrealBuildTool/Executors/UnrealBuildAccelerator/UBAExecutor.cs
+++ b/Engine/Source/Programs/UnrealBuildTool/Executors/UnrealBuildAccelerator/UBAExecutor.cs
@@ -138,6 +138,17 @@ namespace UnrealBuildTool
 			return BitConverter.ToString(bytes).Replace("-", "", StringComparison.OrdinalIgnoreCase).ToLowerInvariant(); // "1234567890abcdef1234567890abcdef";
 		}
 
+		private async void ActionQueueCanceled(IStorageServer? ubaStorage)
+		{
+			_bIsCancelled = true;
+			_session?.CancelAll();
+			ubaStorage?.SaveCasTable();
+			foreach (IUBAAgentCoordinator coordinator in _agentCoordinators)
+			{
+				await coordinator.CloseAsync();
+			}
+		}
+
 		public override async Task<bool> ExecuteActionsAsync(IEnumerable<LinkedAction> inputActions, Microsoft.Extensions.Logging.ILogger logger, IActionArtifactCache? actionArtifactCache)
 		{
 			if (!inputActions.Any())
@@ -220,17 +231,6 @@ namespace UnrealBuildTool
 			Log.BackupLogFile(ubaTraceFile);
 
 			IStorageServer? ubaStorage = null;
-			void CancelKeyPress(object? sender, ConsoleCancelEventArgs e)
-			{
-				_bIsCancelled = true;
-				_session?.CancelAll();
-				ubaStorage?.SaveCasTable();
-				foreach (IUBAAgentCoordinator coordinator in _agentCoordinators)
-				{
-					coordinator.CloseAsync().Wait(2000); // Give coordinators some time to close (this makes coordinators like horde return resources faster)
-				}
-			}
-			Console.CancelKeyPress += CancelKeyPress;
 
 			try
 			{
@@ -254,7 +254,7 @@ namespace UnrealBuildTool
 							Server.StartServer(UBAConfig.Host, UBAConfig.Port, Crypto);
 						}
 
-						bool success = ExecuteActionsInternal(inputActions, _session, logger, actionArtifactCache);
+						bool success = ExecuteActionsInternal(inputActions, _session, logger, actionArtifactCache, () => ActionQueueCanceled(ubaStorage));
 
 						if (!UBAConfig.bDisableRemote)
 						{
@@ -266,14 +266,12 @@ namespace UnrealBuildTool
 							_session.PrintSummary();
 						}
 
-						return success;
+						return success && !_bIsCancelled;
 					}
 				}
 			}
 			finally
 			{
-				Console.CancelKeyPress -= CancelKeyPress;
-
 				foreach (IUBAAgentCoordinator coordinator in _agentCoordinators)
 				{
 					await coordinator.CloseAsync();
@@ -332,13 +330,14 @@ namespace UnrealBuildTool
 		/// Executes the provided actions
 		/// </summary>
 		/// <returns>True if all the tasks successfully executed, or false if any of them failed.</returns>
-		bool ExecuteActionsInternal(IEnumerable<LinkedAction> inputActions, ISessionServer session, Microsoft.Extensions.Logging.ILogger logger, IActionArtifactCache? actionArtifactCache)
+		bool ExecuteActionsInternal(IEnumerable<LinkedAction> inputActions, ISessionServer session, Microsoft.Extensions.Logging.ILogger logger, IActionArtifactCache? actionArtifactCache, System.Action onCancel)
 		{
 			DateTime startTimeUTC = DateTime.UtcNow;
 			using ImmediateActionQueue queue = CreateActionQueue(inputActions, actionArtifactCache, logger);
 			int actionLimit = Math.Min(NumParallelProcesses, queue.TotalActions);
 			queue.CreateAutomaticRunner(action => RunActionLocal(queue, action), bUseActionWeights, actionLimit, NumParallelProcesses);
 			ImmediateActionQueueRunner remoteRunner = queue.CreateManualRunner(action => RunActionRemote(queue, action));
+			queue.CancellationToken.Register(onCancel);
 
 			// Setup a notification that alerts uba when an artifact has been read from the cache
 			queue.OnArtifactsRead = (action) =>
@@ -614,10 +613,20 @@ namespace UnrealBuildTool
 			}
 		}
 
+		protected void HandleActionCancelled(ImmediateActionQueue queue, ExecuteResults? results, LinkedAction action)
+		{
+			if (_bIsCancelled)
+			{
+				ExecuteResults cancelResults = new(results?.LogLines ?? new(), Int32.MaxValue, results?.ExecutionTime ?? TimeSpan.Zero, results?.ProcessorTime ?? TimeSpan.Zero, results?.AdditionalDescription);
+				queue.OnActionCompleted(action, false, cancelResults);
+			}
+		}
+
 		protected void ActionFinished(ImmediateActionQueue queue, ExecuteResults results, LinkedAction action, FileItem? pchItem = null, IProcess? process = null)
 		{
 			if (_bIsCancelled)
 			{
+				HandleActionCancelled(queue, results, action);
 				return;
 			}
 
@@ -639,6 +648,7 @@ namespace UnrealBuildTool
 		{
 			if (_bIsCancelled)
 			{
+				HandleActionCancelled(queue, null, action);
 				return;
 			}
 
@@ -656,6 +666,7 @@ namespace UnrealBuildTool
 		{
 			if (_bIsCancelled)
 			{
+				HandleActionCancelled(queue, null, action);
 				return;
 			}
 
diff --git a/Engine/Source/Programs/UnrealBuildTool/UnrealBuildTool.cs b/Engine/Source/Programs/UnrealBuildTool/UnrealBuildTool.cs
--- a/Engine/Source/Programs/UnrealBuildTool/UnrealBuildTool.cs
+++ b/Engine/Source/Programs/UnrealBuildTool/UnrealBuildTool.cs
@@ -8,6 +8,7 @@ using System.Linq;
 using System.Reflection;
 using System.Text;
 using System.Text.Json;
+using System.Threading.Tasks;
 using EpicGames.Core;
 using Microsoft.Extensions.Logging;
 using OpenTracing.Util;
@@ -455,13 +456,32 @@ namespace UnrealBuildTool
 			return String.IsNullOrEmpty(extraArgs) ? Array.Empty<string>() : CommandLineArguments.Split(extraArgs);
 		}
 
-		/// <summary>
-		/// Main entry point. Parses any global options and initializes the logging system, then invokes the appropriate command.
-		/// NB: That the entry point is deliberately NOT async, since we have a single-instance mutex that cannot be disposed from a different thread.
-		/// </summary>
-		/// <param name="ArgumentsArray">Command line arguments</param>
-		/// <returns>Zero on success, non-zero on error</returns>
-		private static int Main(string[] ArgumentsArray)
+        /// <summary>
+        /// Event handler for the Console.CancelKeyPress event
+        /// </summary>
+        /// <param name="sender"></param>
+        /// <param name="e"></param>
+        private static async void CancelKeyPress(object? sender, ConsoleCancelEventArgs e)
+        {
+			Console.CancelKeyPress -= CancelKeyPress;
+			Console.WriteLine($"UnrealBuildTool: Ctrl-{(e.SpecialKey == ConsoleSpecialKey.ControlC ? "C" : "Break")} pressed. Exiting...");
+
+            // Delay a few seconds to allow for the process to exit normally
+            await Task.Delay(2000);
+
+			// While the Ctrl-C handler fixes most instances of a zombie process, we still need to 
+			// force an exit from the process to handle _all_ cases.  Ctrl-C should not be a regular event! 
+			// Note: this could be a dotnet (6.0.302) on macOS issue.  Recheck with next release if this is still required.
+			Environment.Exit(-1);
+        }
+
+        /// <summary>
+        /// Main entry point. Parses any global options and initializes the logging system, then invokes the appropriate command.
+        /// NB: That the entry point is deliberately NOT async, since we have a single-instance mutex that cannot be disposed from a different thread.
+        /// </summary>
+        /// <param name="ArgumentsArray">Command line arguments</param>
+        /// <returns>Zero on success, non-zero on error</returns>
+        private static int Main(string[] ArgumentsArray)
 		{
 			FileReference? RunFile = null;
 			DirectoryReference? TempDirectory = null;
@@ -475,16 +495,8 @@ namespace UnrealBuildTool
 			// By putting this in, the Ctrl-C may not be handled immediately, but it shouldn't leave a blocking zombie process
 			if (OperatingSystem.IsMacOS())
 			{
-				Console.CancelKeyPress += delegate
-				{
-					Console.WriteLine("UnrealBuildTool: Ctrl-C pressed. Exiting...");
-
-					// While the Ctrl-C handler fixes most instances of a zombie process, we still need to 
-					// force an exit from the process to handle _all_ cases.  Ctrl-C should not be a regular event! 
-					// Note: this could be a dotnet (6.0.302) on macOS issue.  Recheck with next release if this is still required.
-					Environment.Exit(-1);
-				};
-			}
+                Console.CancelKeyPress += CancelKeyPress;
+            }
 
 			try
 			{
