UnrealBuildTool: Clang does not generate pchast when analyzing

#rnx

diff --git a/Engine/Source/Programs/UnrealBuildTool/Platform/Windows/VCToolChain.cs b/Engine/Source/Programs/UnrealBuildTool/Platform/Windows/VCToolChain.cs
--- a/Engine/Source/Programs/UnrealBuildTool/Platform/Windows/VCToolChain.cs
+++ b/Engine/Source/Programs/UnrealBuildTool/Platform/Windows/VCToolChain.cs
@@ -1534,7 +1534,7 @@ namespace UnrealBuildTool
 				BaseCompileAction.UsingPchFile = CompileEnvironment.PrecompiledHeaderFile;
 				BaseCompileAction.PchThroughHeaderFile = IncludeHeader;
 
-				if (CompileEnvironment.PrecompiledHeaderFile != null && !CompileEnvironment.bUseHeaderUnitsForPch && Target.StaticAnalyzer == StaticAnalyzer.Default && !CompileEnvironment.bDisableStaticAnalysis)
+				if (Target.WindowsPlatform.Compiler.IsMSVC() && CompileEnvironment.PrecompiledHeaderFile != null && !CompileEnvironment.bUseHeaderUnitsForPch && Target.StaticAnalyzer == StaticAnalyzer.Default && !CompileEnvironment.bDisableStaticAnalysis)
 				{
 					BaseCompileAction.AdditionalPrerequisiteItems.Add(FileItem.GetItemByFileReference(new FileReference(CompileEnvironment.PrecompiledHeaderFile.FullName + "ast")));
 				}
@@ -1642,7 +1642,7 @@ namespace UnrealBuildTool
 					CompileAction.CreatePchFile = FileItem.GetItemByFileReference(FileReference.Combine(OutputDir, SourceFile.Location.GetFileName() + PchExtension));
 					CompileAction.PchThroughHeaderFile = FileItem.GetItemByFileReference(CompileEnvironment.PrecompiledHeaderIncludeFilename);
 
-					if (!CompileEnvironment.bUseHeaderUnitsForPch && Target.StaticAnalyzer == StaticAnalyzer.Default && !CompileEnvironment.bDisableStaticAnalysis)
+					if (Target.WindowsPlatform.Compiler.IsMSVC() && !CompileEnvironment.bUseHeaderUnitsForPch && Target.StaticAnalyzer == StaticAnalyzer.Default && !CompileEnvironment.bDisableStaticAnalysis)
 					{
 						CompileAction.AdditionalProducedItems.Add(FileItem.GetItemByFileReference(FileReference.Combine(OutputDir, SourceFile.Location.GetFileName() + PchExtension + "ast")));
 					}
