UnrealBuildTool: Fix pch output when running clang analysis

#rnx

diff --git a/Engine/Source/Programs/UnrealBuildTool/Platform/Windows/VCToolChain.cs b/Engine/Source/Programs/UnrealBuildTool/Platform/Windows/VCToolChain.cs
--- a/Engine/Source/Programs/UnrealBuildTool/Platform/Windows/VCToolChain.cs
+++ b/Engine/Source/Programs/UnrealBuildTool/Platform/Windows/VCToolChain.cs
@@ -1672,7 +1672,7 @@ namespace UnrealBuildTool
 					CompileAction.PreprocessedFile = FileItem.GetItemByFileReference(FileReference.Combine(OutputDir, FileName + ".i"));
 					CompileAction.ResponseFile = FileItem.GetItemByFileReference(GetResponseFileName(CompileEnvironment, CompileAction.PreprocessedFile));
 				}
-				else if (Target.WindowsPlatform.Compiler.IsClang() && Target.StaticAnalyzer == StaticAnalyzer.Default)
+				else if (Target.WindowsPlatform.Compiler.IsClang() && Target.StaticAnalyzer == StaticAnalyzer.Default && CompileEnvironment.PrecompiledHeaderAction != PrecompiledHeaderAction.Create)
 				{
 					// Clang analysis does not actually create an object, use the dependency list as the response filename
 					string DependencyListFilename = FileName + ".d";
@@ -1866,7 +1866,7 @@ namespace UnrealBuildTool
 					CompileAction.WriteResponseFile(Graph, Logger);
 				}
 
-				CompileAction.bIsAnalyzing = Target.StaticAnalyzer != StaticAnalyzer.None;
+				CompileAction.bIsAnalyzing = Target.StaticAnalyzer != StaticAnalyzer.None && !(Target.WindowsPlatform.Compiler.IsClang() && CompileEnvironment.PrecompiledHeaderAction == PrecompiledHeaderAction.Create);
 
 				// Update the output
 				Graph.AddAction(CompileAction);
