Do not use the AutoRTFM compiler if we explicitly want static analysis from another compiler toolchain.

diff --git a/Engine/Source/Programs/UnrealBuildTool/Platform/Windows/UEBuildWindows.cs b/Engine/Source/Programs/UnrealBuildTool/Platform/Windows/UEBuildWindows.cs
--- a/Engine/Source/Programs/UnrealBuildTool/Platform/Windows/UEBuildWindows.cs
+++ b/Engine/Source/Programs/UnrealBuildTool/Platform/Windows/UEBuildWindows.cs
@@ -1157,7 +1157,17 @@ namespace UnrealBuildTool
 
 			if (Target.bUseAutoRTFMCompiler)
 			{
-				Target.WindowsPlatform.Compiler = WindowsCompiler.ClangRTFM;
+				// We check the static analyzer, and only allow using the AutoRTFM compiler if
+				// we are not doing static analysis, or are explicitly wanting to use clang.
+				switch (Target.StaticAnalyzer)
+				{
+					default:
+						break;
+					case StaticAnalyzer.None:
+					case StaticAnalyzer.Clang:
+						Target.WindowsPlatform.Compiler = WindowsCompiler.ClangRTFM;
+						break;
+				}
 			}
 
 			// Set the compiler version if necessary
