- Addd a XcodeProjectSettings ini value for using a premade .dSYM when making .xcarchive's instead of generating it at archival time. Defaults to false
- Fixed the ExportOptions plist files to use "app-store" instead of "app-store-connect" as the latter was only introduced in Xcode15.3 (deprecateding app-store) and does not work in 15.2 or earlier

diff --git a/Engine/Config/BaseEngine.ini b/Engine/Config/BaseEngine.ini
--- a/Engine/Config/BaseEngine.ini
+++ b/Engine/Config/BaseEngine.ini
@@ -3199,6 +3199,7 @@ bEnableRayTracing=false
 bUseModernXcode=true
 bUseAutomaticCodeSigning=true
 bMacSignToRunLocally=false
+bUsePremadeDSYMInXcArchive=false
 CodeSigningTeam=
 CodeSigningPrefix=com.YourCompany
 BundleIdentifier=$(UE_SIGNING_PREFIX).$(UE_PRODUCT_NAME_STRIPPED)
diff --git a/Engine/Source/Programs/UnrealBuildTool/ProjectFiles/Xcode/XcodeProject.cs b/Engine/Source/Programs/UnrealBuildTool/ProjectFiles/Xcode/XcodeProject.cs
--- a/Engine/Source/Programs/UnrealBuildTool/ProjectFiles/Xcode/XcodeProject.cs
+++ b/Engine/Source/Programs/UnrealBuildTool/ProjectFiles/Xcode/XcodeProject.cs
@@ -1462,17 +1462,43 @@ namespace UnrealBuildTool.XcodeProjectXcconfig
 			// (note bInstallOnly which will make this onle run when archiving)
 			List<string> DsymScript = new();
 
+			bool bUsePremadeDSYMInXcArchive;
+			// get ini file for the platform
+			ConfigHierarchy PlatformIni = ConfigCache.ReadHierarchy(ConfigHierarchyType.Engine, UnrealData.ConfigDirectory, Platform);
+			PlatformIni.TryGetValue("/Script/MacTargetPlatform.XcodeProjectSettings", "bUsePremadeDSYMInXcArchive", out bUsePremadeDSYMInXcArchive);
+
+			if (bUsePremadeDSYMInXcArchive)
+			{
+				DsymScript.AddRange(new string[]
+				{
+					"set -e",
+					"",
+					"# Copy a dsym from next to the executable",
+					"SRC_DSYM=\\\"${UE_BINARIES_DIR}/${UE_UBT_BINARY_SUBPATH}.dSYM\\\"",
+					"if [[ -e \\\"${SRC_DSYM}\\\" ]]; then",
+					"  echo Using pre-existing dSYM at \\\"${SRC_DSYM}\\\"",
+					"  ditto \\\"${SRC_DSYM}\\\" \\\"${DWARF_DSYM_FOLDER_PATH}/${DWARF_DSYM_FILE_NAME}\\\"",
+					"fi"
+				});
+			}
+			else
+			{
+				DsymScript.AddRange(new string[]
+				{
+					"set -e",
+					"",
+					"# Run the wrapper dsym generator",
+					"\\\"${UE_ENGINE_DIR}/Build/BatchFiles/Mac/GenerateUniversalDSYM.sh\\\" \\\"${CONFIGURATION_BUILD_DIR}/${EXECUTABLE_PATH}\\\" \\\"${DWARF_DSYM_FOLDER_PATH}/${DWARF_DSYM_FILE_NAME}\\\"",
+				});
+			}
 			DsymScript.AddRange(new string[]
 			{
-				"set -e",
-				"",
-				"# Run the wrapper dsym generator",
-				"\\\"${UE_ENGINE_DIR}/Build/BatchFiles/Mac/GenerateUniversalDSYM.sh\\\" \\\"${CONFIGURATION_BUILD_DIR}/${EXECUTABLE_PATH}\\\" \\\"${DWARF_DSYM_FOLDER_PATH}/${DWARF_DSYM_FILE_NAME}\\\"",
 				"strip -no_code_signature_warning -D \\\"${CONFIGURATION_BUILD_DIR}/${EXECUTABLE_PATH}\\\"",
 				"",
 				"# Remove any unused architectures from dylibs in the .app (param1) that don't match the executable (param2). Also error if a dylib is missing arches",
 				"\\\"${UE_ENGINE_DIR}/Build/BatchFiles/Mac/ThinApp.sh\\\" \\\"${CONFIGURATION_BUILD_DIR}/${CONTENTS_FOLDER_PATH}\\\" \\\"${CONFIGURATION_BUILD_DIR}/${EXECUTABLE_PATH}\\\"",
 			});
+
 			string DsymScriptInput = $"\\\"$(CONFIGURATION_BUILD_DIR)/$(EXECUTABLE_PATH)\\\"";
 			string DsymScriptOutput = $"\\\"$(DWARF_DSYM_FOLDER_PATH)/$(DWARF_DSYM_FILE_NAME)\\\"";
 			XcodeShellScriptBuildPhase DsymScriptPhase = new("Generate dsym for archive, and strip", DsymScript, new string[] { DsymScriptInput }, new string[] { DsymScriptOutput }, bInstallOnly: true);
