- Adding changelist number to version counter so buildmachine builds are globally increasing
- Added ability to have a per-project versioning update script, so licensees can do what they want

diff --git a/Engine/Build/BatchFiles/Mac/UpdateVersionAfterBuild.sh b/Engine/Build/BatchFiles/Mac/UpdateVersionAfterBuild.sh
--- a/Engine/Build/BatchFiles/Mac/UpdateVersionAfterBuild.sh
+++ b/Engine/Build/BatchFiles/Mac/UpdateVersionAfterBuild.sh
@@ -2,12 +2,19 @@
 
 # $1 is path to product directory (engine, program or project)
 # $2 is the platform we are incrementing
+# $3 is the changlist number if known
 
 PRODUCT_NAME=$(basename "$1")
 VERSION_FILE_DIR="$1/Build/$2"
 VERSION_FILE="$VERSION_FILE_DIR/$PRODUCT_NAME.PackageVersionCounter"
 
-VERSION="0.1"
+if [ -z $3 ]; then
+	CL=0
+else
+	CL=$3
+fi
+
+VERSION="$CL.1"
 # increment version in the counter file
 if [ -f "$VERSION_FILE" ]; then
 	VERSION=$(cat "$VERSION_FILE")
@@ -49,8 +56,8 @@ fi
 XCCONFIG_FILE="$1/Intermediate/Build/Versions.xcconfig"
 
 mkdir -p "$1/Intermediate/Build"
-echo "UE_MAC_BUILD_VERSION = $MAC_VERSION" > "$XCCONFIG_FILE"
-echo "UE_IOS_BUILD_VERSION = $IOS_VERSION" >> "$XCCONFIG_FILE"
-echo "UE_TVOS_BUILD_VERSION = $TVOS_VERSION" >> "$XCCONFIG_FILE"
-echo "UE_VISIONOS_BUILD_VERSION = $VISIONOS_VERSION" >> "$XCCONFIG_FILE"
+echo "UE_MAC_BUILD_VERSION = $CL.$MAC_VERSION" > "$XCCONFIG_FILE"
+echo "UE_IOS_BUILD_VERSION = $CL.$IOS_VERSION.$CL" >> "$XCCONFIG_FILE"
+echo "UE_TVOS_BUILD_VERSION = $CL.$TVOS_VERSION.$CL" >> "$XCCONFIG_FILE"
+echo "UE_VISIONOS_BUILD_VERSION = $CL.$VISIONOS_VERSION.$CL" >> "$XCCONFIG_FILE"
 
diff --git a/Engine/Source/Programs/UnrealBuildTool/ToolChain/AppleToolChain.cs b/Engine/Source/Programs/UnrealBuildTool/ToolChain/AppleToolChain.cs
--- a/Engine/Source/Programs/UnrealBuildTool/ToolChain/AppleToolChain.cs
+++ b/Engine/Source/Programs/UnrealBuildTool/ToolChain/AppleToolChain.cs
@@ -360,9 +360,22 @@ namespace UnrealBuildTool
 			FileReference OutputVersionFile = FileReference.Combine(ProductDirectory, "Intermediate/Build/Versions.xcconfig");
 			DestFile = FileItem.GetItemByFileReference(OutputVersionFile);
 
+			// grab a changlist version if we have it to pass to the script to use if desired
+			int Changelist = 0;
+			BuildVersion Version;
+			if (BuildVersion.TryRead(BuildVersion.GetDefaultFileName(), out Version))
+			{
+				Changelist = Version.Changelist;
+			}
+
 			// make path to the script
-			FileItem BundleScript = FileItem.GetItemByFileReference(FileReference.Combine(Unreal.EngineDirectory, "Build/BatchFiles/Mac/UpdateVersionAfterBuild.sh"));
-			UpdateVersionAction.CommandArguments = $"\"{BundleScript.AbsolutePath}\" \"{ProductDirectory}\" {LinkEnvironment.Platform}";
+			FileReference VersionScript = FileReference.Combine(ProductDirectory, "Build/BatchFiles/Mac/UpdateVersionAfterBuild.sh");
+			if (!FileReference.Exists(VersionScript))
+			{
+				VersionScript = FileReference.Combine(Unreal.EngineDirectory, "Build/BatchFiles/Mac/UpdateVersionAfterBuild.sh");
+			}
+			FileItem BundleScript = FileItem.GetItemByFileReference(VersionScript);
+			UpdateVersionAction.CommandArguments = $"\"{BundleScript.AbsolutePath}\" \"{ProductDirectory}\" {LinkEnvironment.Platform} {Changelist}";
 			UpdateVersionAction.PrerequisiteItems.Add(Prerequisite);
 			UpdateVersionAction.PrerequisiteItems.Add(BundleScript);
 			UpdateVersionAction.ProducedItems.Add(DestFile);
