Enable chunkDBs on iOS
Add ExtraFolderToCopyToApp XcodeProjectSettings to add custom folder to BuildPhase

diff --git a/Engine/Source/Programs/UnrealBuildTool/ProjectFiles/Xcode/XcodeProject.cs b/Engine/Source/Programs/UnrealBuildTool/ProjectFiles/Xcode/XcodeProject.cs
--- a/Engine/Source/Programs/UnrealBuildTool/ProjectFiles/Xcode/XcodeProject.cs
+++ b/Engine/Source/Programs/UnrealBuildTool/ProjectFiles/Xcode/XcodeProject.cs
@@ -1581,6 +1581,24 @@ namespace UnrealBuildTool.XcodeProjectXcconfig
 					}
 				}
 			}
+
+			string? ExtraFolderToCopyToApp;
+			ConfigHierarchy PlatformIni = ConfigCache.ReadHierarchy(ConfigHierarchyType.Engine, UnrealData.ConfigDirectory, Platform);
+			PlatformIni.TryGetValue("/Script/MacTargetPlatform.XcodeProjectSettings", "ExtraFolderToCopyToApp", out ExtraFolderToCopyToApp);
+
+			if (ExtraFolderToCopyToApp != null)
+			{
+				List<string> ResourcesToBundlePaths = new List<string>()
+				{
+					ExtraFolderToCopyToApp,
+				};
+
+				string? ResourcesToBundlePath = UnrealData.FindFile(ResourcesToBundlePaths, Platform, false);
+				if (ResourcesToBundlePath != null)
+				{
+					ResourcesBuildPhase.AddFolderResource(new DirectoryReference(ResourcesToBundlePath), "Resources");
+				}
+			}
 		}
 		protected void ProcessFrameworks(XcodeResourcesBuildPhase ResourcesBuildPhase, XcodeProject Project, XcodeProjectFile ProjectFile, ILogger Logger)
 		{
