- Fixed nullable error

diff --git a/Engine/Source/Programs/UnrealBuildTool/ToolChain/AppleToolChain.cs b/Engine/Source/Programs/UnrealBuildTool/ToolChain/AppleToolChain.cs
--- a/Engine/Source/Programs/UnrealBuildTool/ToolChain/AppleToolChain.cs
+++ b/Engine/Source/Programs/UnrealBuildTool/ToolChain/AppleToolChain.cs
@@ -362,7 +362,7 @@ namespace UnrealBuildTool
 
 			// grab a changlist version if we have it to pass to the script to use if desired
 			int Changelist = 0;
-			BuildVersion Version;
+			BuildVersion? Version;
 			if (BuildVersion.TryRead(BuildVersion.GetDefaultFileName(), out Version))
 			{
 				Changelist = Version.Changelist;
