UBA: Prevent double dispose for agent coordinator

#rnx

diff --git a/Engine/Source/Programs/UnrealBuildTool/Executors/UnrealBuildAccelerator/UBAAgentCoordinatorHorde.cs b/Engine/Source/Programs/UnrealBuildTool/Executors/UnrealBuildAccelerator/UBAAgentCoordinatorHorde.cs
--- a/Engine/Source/Programs/UnrealBuildTool/Executors/UnrealBuildAccelerator/UBAAgentCoordinatorHorde.cs
+++ b/Engine/Source/Programs/UnrealBuildTool/Executors/UnrealBuildAccelerator/UBAAgentCoordinatorHorde.cs
@@ -108,10 +108,11 @@ namespace UnrealBuildTool
 
 			if (_client != null)
 			{
+				IComputeClient client = _client;
+				_client = null;
 				// Set CPU resource need to zero (will also expire on the server if not updated)
 				await UpdateCpuCoreNeedAsync(0);
-				await _client.DisposeAsync();
-				_client = null;
+				await client.DisposeAsync();
 			}
 
 			await _serviceProvider.DisposeAsync();
@@ -610,12 +611,12 @@ namespace UnrealBuildTool
 			{
 				_timer?.Change(Timeout.Infinite, Timeout.Infinite);
 
-				if (_cancellationSource!.IsCancellationRequested)
+				if (_cancellationSource!.IsCancellationRequested || _hordeSessionTask == null)
 				{
 					return;
 				}
 
-				UBAHordeSession? hordeSession = await _hordeSessionTask!;
+				UBAHordeSession? hordeSession = await _hordeSessionTask;
 
 				if (hordeSession == null)
 				{
@@ -712,6 +713,7 @@ namespace UnrealBuildTool
 			}
 
 			UBAHordeSession? hordeSession = await _hordeSessionTask;
+			_hordeSessionTask = null;
 			if (hordeSession != null)
 			{
 				await hordeSession.DisposeAsync();
