UBA: More cancel tweaks

#rnx

diff --git a/Engine/Source/Programs/UnrealBuildTool/Executors/ImmediateActionQueue.cs b/Engine/Source/Programs/UnrealBuildTool/Executors/ImmediateActionQueue.cs
--- a/Engine/Source/Programs/UnrealBuildTool/Executors/ImmediateActionQueue.cs
+++ b/Engine/Source/Programs/UnrealBuildTool/Executors/ImmediateActionQueue.cs
@@ -1098,7 +1098,7 @@ namespace UnrealBuildTool
 				// Canceled
 				if (exitCode == Int32.MaxValue)
 				{
-					Logger.LogInformation("[{CompletedActions}/{TotalActions}] {Description} canceled", completedActions, totalActions, description);
+					//Logger.LogInformation("[{CompletedActions}/{TotalActions}] {Description} canceled", completedActions, totalActions, description);
 					return;
 				}
 
diff --git a/Engine/Source/Programs/UnrealBuildTool/Executors/UnrealBuildAccelerator/UBAExecutor.cs b/Engine/Source/Programs/UnrealBuildTool/Executors/UnrealBuildAccelerator/UBAExecutor.cs
--- a/Engine/Source/Programs/UnrealBuildTool/Executors/UnrealBuildAccelerator/UBAExecutor.cs
+++ b/Engine/Source/Programs/UnrealBuildTool/Executors/UnrealBuildAccelerator/UBAExecutor.cs
@@ -491,6 +491,11 @@ namespace UnrealBuildTool
 
 			return () =>
 			{
+				if (_bIsCancelled)
+				{
+					HandleActionCancelled(queue, null, action);
+					return Task.CompletedTask;
+				}
 				bool enableDetour = !ForceLocalNoDetour(action) && action.bCanExecuteInUBA && !_forcedRetryActions.ContainsKey(action);
 
 				ProcessStartInfo startInfo = GetActionStartInfo(action, out FileItem? pchItem);
@@ -531,6 +536,12 @@ namespace UnrealBuildTool
 
 			return () =>
 			{
+				if (_bIsCancelled)
+				{
+					HandleActionCancelled(queue, null, action);
+					return Task.CompletedTask;
+				}
+
 				uint knownInputsCount = 0;
 				byte[]? knownInputs = null;
 				if (UBAConfig.bUseKnownInputs)
