- Copy uecommandline.txt into the .app root

diff --git a/Engine/Source/Programs/UnrealBuildTool/ProjectFiles/Xcode/XcodeProject.cs b/Engine/Source/Programs/UnrealBuildTool/ProjectFiles/Xcode/XcodeProject.cs
--- a/Engine/Source/Programs/UnrealBuildTool/ProjectFiles/Xcode/XcodeProject.cs
+++ b/Engine/Source/Programs/UnrealBuildTool/ProjectFiles/Xcode/XcodeProject.cs
@@ -1449,6 +1449,10 @@ namespace UnrealBuildTool.XcodeProjectXcconfig
 					"",
 					$"echo \\\"Syncing ${{STAGED_DIR}}{SyncSourceSubdir} to ${{CONFIGURATION_BUILD_DIR}}/${{CONTENTS_FOLDER_PATH}}{SyncDestSubdir}\\\"",
 					$"rsync -a --delete {ExcludeString} \\\"${{STAGED_DIR}}{SyncSourceSubdir}/\\\" \\\"${{CONFIGURATION_BUILD_DIR}}/${{CONTENTS_FOLDER_PATH}}{SyncDestSubdir}\\\"",
+					"if [[ -e \\\"${STAGED_DIR}/uecommandline.txt\\\" ]]; then",
+					"  echo \\\"Syncing ${STAGED_DIR}/uecommandline.txt to ${CONFIGURATION_BUILD_DIR}/${CONTENTS_FOLDER_PATH}\\\"",
+					"  cp \\\"${STAGED_DIR}/uecommandline.txt\\\" \\\"${CONFIGURATION_BUILD_DIR}/${CONTENTS_FOLDER_PATH}\\\"",
+					"fi"
 				});
 			}
 
