Add default implementation for GetMakePriBinaryPath that uses the Microsoft Platform SDK.

#rnx

diff --git a/Engine/Source/Programs/UnrealBuildTool/Platform/Windows/AppXManifestGeneratorBase.cs b/Engine/Source/Programs/UnrealBuildTool/Platform/Windows/AppXManifestGeneratorBase.cs
--- a/Engine/Source/Programs/UnrealBuildTool/Platform/Windows/AppXManifestGeneratorBase.cs
+++ b/Engine/Source/Programs/UnrealBuildTool/Platform/Windows/AppXManifestGeneratorBase.cs
@@ -4,6 +4,7 @@ using System;
 using System.Collections.Generic;
 using System.Diagnostics.CodeAnalysis;
 using System.Linq;
+using System.Runtime.Versioning;
 using System.Text.RegularExpressions;
 using System.Xml.Linq;
 using EpicGames.Core;
@@ -450,7 +451,22 @@ namespace UnrealBuildTool
 		/// <summary>
 		/// Get the path to the makepri.exe tool
 		/// </summary>
-		protected abstract FileReference GetMakePriBinaryPath();
+		[SupportedOSPlatform("windows")]
+		protected virtual FileReference GetMakePriBinaryPath()
+		{
+			if (!MicrosoftPlatformSDK.TryGetWindowsSdkDir(null, Logger, out VersionNumber? SdkVersion, out DirectoryReference? SdkDir))
+			{
+				throw new BuildException("Cannot get default Windows Sdk directory");
+			}
+
+			FileReference MakePriPath = FileReference.Combine(SdkDir, "bin", SdkVersion.ToString(), "x64", "makepri.exe");
+			if (!FileReference.Exists(MakePriPath))
+			{
+				throw new BuildException($"{MakePriPath} - file not found");
+			}
+
+			return MakePriPath;
+		}
 
 		/// <summary>
 		/// Get any additional platform-specific parameters for makepri.exe
