static analysis fix: can only access WindowsSDK from Windows platforms.

 - don't want to tag the entire system with [SupportedOSPlatform] because it will touch a massive amount of files

#rnx

diff --git a/Engine/Source/Programs/UnrealBuildTool/Platform/Windows/AppXManifestGeneratorBase.cs b/Engine/Source/Programs/UnrealBuildTool/Platform/Windows/AppXManifestGeneratorBase.cs
--- a/Engine/Source/Programs/UnrealBuildTool/Platform/Windows/AppXManifestGeneratorBase.cs
+++ b/Engine/Source/Programs/UnrealBuildTool/Platform/Windows/AppXManifestGeneratorBase.cs
@@ -4,6 +4,7 @@ using System;
 using System.Collections.Generic;
 using System.Diagnostics.CodeAnalysis;
 using System.Linq;
+using System.Runtime.InteropServices;
 using System.Runtime.Versioning;
 using System.Text.RegularExpressions;
 using System.Xml.Linq;
@@ -451,9 +452,13 @@ namespace UnrealBuildTool
 		/// <summary>
 		/// Get the path to the makepri.exe tool
 		/// </summary>
-		[SupportedOSPlatform("windows")]
 		protected virtual FileReference GetMakePriBinaryPath()
 		{
+			if (!RuntimeInformation.IsOSPlatform(OSPlatform.Windows))
+			{
+				throw new BuildException("unsupported platform");
+			}			
+			
 			if (!MicrosoftPlatformSDK.TryGetWindowsSdkDir(null, Logger, out VersionNumber? SdkVersion, out DirectoryReference? SdkDir))
 			{
 				throw new BuildException("Cannot get default Windows Sdk directory");
