Clang Static Analysis build
Remove validate toolchains from ue5-main

#rnx

diff --git a/Engine/Source/Programs/UnrealBuildTool/Configuration/ModuleRules.cs b/Engine/Source/Programs/UnrealBuildTool/Configuration/ModuleRules.cs
--- a/Engine/Source/Programs/UnrealBuildTool/Configuration/ModuleRules.cs
+++ b/Engine/Source/Programs/UnrealBuildTool/Configuration/ModuleRules.cs
@@ -1010,18 +1010,43 @@ namespace UnrealBuildTool
 		///    'clang -Xclang -analyzer-checker-help-alpha' 
 		/// for the list of experimental checkers.
 		/// </summary>
-		public HashSet<string> StaticAnalyzerCheckers = new();
+		public HashSet<string> StaticAnalyzerCheckers { get; } = new();
 
 		/// <summary>
 		/// The static analyzer default checkers that should be disabled. Unused if StaticAnalyzerCheckers is populated. This is only supported for Clang.
-		/// This overrides the default disabled checkers, which are deadcode.DeadStores and security.FloatLoopCounter
 		/// See https://clang.llvm.org/docs/analyzer/checkers.html for a full list. Or run:
 		///    'clang -Xclang -analyzer-checker-help' 
 		/// or: 
 		///    'clang -Xclang -analyzer-checker-help-alpha' 
 		/// for the list of experimental checkers.
 		/// </summary>
-		public HashSet<string> StaticAnalyzerDisabledCheckers = new() { "deadcode.DeadStores", "security.FloatLoopCounter" };
+		public HashSet<string> StaticAnalyzerDisabledCheckers { get; } = new() {
+			// Work in progress, remove disabled checkers as all warnings are fixed
+			"core.CallAndMessage",
+			"core.DivideZero",
+			"core.NonNullParamChecker",
+			"core.NullDereference",
+			"core.StackAddressEscape",
+			"core.UndefinedBinaryOperatorResult",
+			"core.uninitialized.ArraySubscript",
+			"core.uninitialized.Assign",
+			"core.uninitialized.Branch",
+			"core.uninitialized.UndefReturn",
+			"cplusplus.Move",
+			"cplusplus.NewDelete",
+			"cplusplus.NewDeleteLeaks",
+			"cplusplus.PlacementNew",
+			"cplusplus.PureVirtualCall",
+			"unix.cstring.NullArg",
+			"unix.Malloc",
+
+			// Android
+			"core.VLASize",
+
+			// Needs evalulation
+			"deadcode.DeadStores",
+			"security.FloatLoopCounter",
+		};
 
 		/// <summary>
 		/// The static analyzer non-default checkers that should be enabled. Unused if StaticAnalyzerCheckers is populated. This is only supported for Clang.
@@ -1031,7 +1056,7 @@ namespace UnrealBuildTool
 		///    'clang -Xclang -analyzer-checker-help-alpha' 
 		/// for the list of experimental checkers.
 		/// </summary>
-		public HashSet<string> StaticAnalyzerAdditionalCheckers = new();
+		public HashSet<string> StaticAnalyzerAdditionalCheckers { get; } = new();
 
 		private bool? bUseUnityOverride;
 		/// <summary>
