Reenable security.FloatLoopCounter

#rnx

diff --git a/Engine/Source/Programs/UnrealBuildTool/Configuration/ModuleRules.cs b/Engine/Source/Programs/UnrealBuildTool/Configuration/ModuleRules.cs
--- a/Engine/Source/Programs/UnrealBuildTool/Configuration/ModuleRules.cs
+++ b/Engine/Source/Programs/UnrealBuildTool/Configuration/ModuleRules.cs
@@ -1045,7 +1045,6 @@ namespace UnrealBuildTool
 
 			// Needs evalulation
 			"deadcode.DeadStores",
-			"security.FloatLoopCounter",
 		};
 
 		/// <summary>
