UBT: Fix warnings from building in Analyze configuration.

diff --git a/Engine/Source/Programs/UnrealBuildTool/Modes/ClReproMode.cs b/Engine/Source/Programs/UnrealBuildTool/Modes/ClReproMode.cs
--- a/Engine/Source/Programs/UnrealBuildTool/Modes/ClReproMode.cs
+++ b/Engine/Source/Programs/UnrealBuildTool/Modes/ClReproMode.cs
@@ -73,7 +73,7 @@ namespace UnrealBuildTool.Modes
 
 			if (!String.IsNullOrEmpty(depends.Data.PCH))
 			{
-				logger.LogError("PCH repro not currently supported, please use -NoPCH when compiling", inputDepFile);
+				logger.LogError("PCH repro not currently supported, please use -NoPCH when compiling ({Dep})", inputDepFile);
 				return Task.FromResult(1);
 			}
 
diff --git a/Engine/Source/Programs/UnrealBuildTool/Preprocessor/Preprocessor.cs b/Engine/Source/Programs/UnrealBuildTool/Preprocessor/Preprocessor.cs
--- a/Engine/Source/Programs/UnrealBuildTool/Preprocessor/Preprocessor.cs
+++ b/Engine/Source/Programs/UnrealBuildTool/Preprocessor/Preprocessor.cs
@@ -1400,7 +1400,7 @@ namespace UnrealBuildTool
 		/// <param name="inputIdx">First token index in the input token list. Set to the last uncopied token index on return.</param>
 		/// <param name="outputTokens">List to recieve the output tokens</param>
 		/// <returns>True if a balanced expression was read, or false if the end of the list was encountered before finding a matching token</returns>
-		bool ReadBalancedToken(List<Token> inputTokens, ref int inputIdx, List<Token> outputTokens)
+		static bool ReadBalancedToken(List<Token> inputTokens, ref int inputIdx, List<Token> outputTokens)
 		{
 			// Copy a single token to the output list
 			Token token = inputTokens[inputIdx++];
@@ -1440,7 +1440,7 @@ namespace UnrealBuildTool
 		/// <param name="outputTokens">List to recieve the output tokens</param>
 		/// <param name="context">The context that the parser is in</param>
 		/// <returns>True if a balanced expression was read, or false if the end of the list was encountered before finding a matching token</returns>
-		bool ReadBalancedToken(IEnumerator<Token> inputEnumerator, List<Token> outputTokens, PreprocessorContext context)
+		static bool ReadBalancedToken(IEnumerator<Token> inputEnumerator, List<Token> outputTokens, PreprocessorContext context)
 		{
 			// Copy a single token to the output list
 			Token token = inputEnumerator.Current;
diff --git a/Engine/Source/Programs/UnrealBuildTool/ProjectFiles/Xcode/XcodeProject.cs b/Engine/Source/Programs/UnrealBuildTool/ProjectFiles/Xcode/XcodeProject.cs
--- a/Engine/Source/Programs/UnrealBuildTool/ProjectFiles/Xcode/XcodeProject.cs
+++ b/Engine/Source/Programs/UnrealBuildTool/ProjectFiles/Xcode/XcodeProject.cs
@@ -198,7 +198,7 @@ namespace UnrealBuildTool.XcodeProjectXcconfig
 							}
 							else
 							{
-								throw ex;
+								throw;
 							}
 						}
 					}
diff --git a/Engine/Source/Programs/UnrealBuildTool/System/NativeProjects.cs b/Engine/Source/Programs/UnrealBuildTool/System/NativeProjects.cs
--- a/Engine/Source/Programs/UnrealBuildTool/System/NativeProjects.cs
+++ b/Engine/Source/Programs/UnrealBuildTool/System/NativeProjects.cs
@@ -497,7 +497,7 @@ namespace UnrealBuildTool
 						// Ignore any optional plugins
 						if (Reference.bOptional)
 						{
-							Logger.LogDebug("Ignored optional reference to '%s' plugin; plugin was not found.", Reference.Name);
+							Logger.LogDebug("Ignored optional reference to '{Plugin}' plugin; plugin was not found.", Reference.Name);
 							continue;
 						}
 
diff --git a/Engine/Source/Programs/UnrealBuildTool/System/TargetMakefile.cs b/Engine/Source/Programs/UnrealBuildTool/System/TargetMakefile.cs
--- a/Engine/Source/Programs/UnrealBuildTool/System/TargetMakefile.cs
+++ b/Engine/Source/Programs/UnrealBuildTool/System/TargetMakefile.cs
@@ -690,8 +690,8 @@ namespace UnrealBuildTool
 				// Check if the arguments are different
 				if (!Enumerable.SequenceEqual(Makefile.AdditionalArguments!, Arguments))
 				{
-					Logger.LogDebug("Old command line arguments:\n", String.Join(' ', Makefile.AdditionalArguments!));
-					Logger.LogDebug("New command line arguments:\n", String.Join(' ', Arguments));
+					Logger.LogDebug("Old command line arguments:\n{Args}", String.Join(' ', Makefile.AdditionalArguments!));
+					Logger.LogDebug("New command line arguments:\n{Args}", String.Join(' ', Arguments));
 					ReasonNotLoaded = "command line arguments changed";
 					return null;
 				}
@@ -725,8 +725,8 @@ namespace UnrealBuildTool
 				string CurrentExternalMetadata = UEBuildPlatform.GetBuildPlatform(Platform).GetExternalBuildMetadata(ProjectFile);
 				if (String.Compare(CurrentExternalMetadata, Makefile.ExternalMetadata, StringComparison.Ordinal) != 0)
 				{
-					Logger.LogDebug("Old metadata:\n", Makefile.ExternalMetadata);
-					Logger.LogDebug("New metadata:\n", CurrentExternalMetadata);
+					Logger.LogDebug("Old metadata:\n{Metadata}", Makefile.ExternalMetadata);
+					Logger.LogDebug("New metadata:\n{Metadata}", CurrentExternalMetadata);
 					ReasonNotLoaded = "build metadata has changed";
 					return null;
 				}
