UnrealBuildTool: Keep track of root paths per action

diff --git a/Engine/Source/Programs/UnrealBuildTool/Platform/Android/AndroidToolChain.cs b/Engine/Source/Programs/UnrealBuildTool/Platform/Android/AndroidToolChain.cs
--- a/Engine/Source/Programs/UnrealBuildTool/Platform/Android/AndroidToolChain.cs
+++ b/Engine/Source/Programs/UnrealBuildTool/Platform/Android/AndroidToolChain.cs
@@ -55,8 +55,6 @@ namespace UnrealBuildTool
 			}
 		}
 
-		protected FileReference? ProjectFile;
-
 		// Version string from the Android specific build of clang. E.g in Android (6317467 based on r365631c1) clang version 9.0.8
 		// this would be 6317467)
 		protected static string? AndroidClangBuild;
@@ -246,7 +244,7 @@ namespace UnrealBuildTool
 
 		protected override ClangToolChainInfo GetToolChainInfo()
 		{
-			return new ClangToolChainInfo(FileReference.FromString(ClangPath)!, FileReference.FromString(ArPathArm64)!, Logger);
+			return new ClangToolChainInfo(DirectoryReference.FromString(AndroidPlatformSDK.GetNDKRoot()), FileReference.FromString(ClangPath)!, FileReference.FromString(ArPathArm64)!, Logger);
 		}
 
 		public static string GetGLESVersion(bool bBuildForES31)
@@ -1408,6 +1406,7 @@ namespace UnrealBuildTool
 
 			// Create an action that invokes the linker.
 			Action LinkAction = Graph.CreateAction(ActionType.Link);
+			LinkAction.RootPaths.AddRange(GetEnvironmentBasePaths(LinkEnvironment));
 			LinkAction.WorkingDirectory = Unreal.EngineSourceDirectory;
 
 			if (LinkEnvironment.bIsBuildingLibrary)
diff --git a/Engine/Source/Programs/UnrealBuildTool/Platform/IOS/IOSToolChain.cs b/Engine/Source/Programs/UnrealBuildTool/Platform/IOS/IOSToolChain.cs
--- a/Engine/Source/Programs/UnrealBuildTool/Platform/IOS/IOSToolChain.cs
+++ b/Engine/Source/Programs/UnrealBuildTool/Platform/IOS/IOSToolChain.cs
@@ -76,7 +76,7 @@ namespace UnrealBuildTool
 		{
 			FileReference CompilerPath = FileReference.Combine(Settings.ToolchainDir, IOSCompiler);
 			FileReference ArchiverPath = FileReference.Combine(Settings.ToolchainDir, IOSArchiver);
-			return new AppleToolChainInfo(CompilerPath, ArchiverPath, Logger);
+			return new AppleToolChainInfo(IOSToolChainSettings.XcodeDeveloperDir, CompilerPath, ArchiverPath, Logger);
 		}
 
 		public override string GetSDKVersion()
@@ -398,6 +398,7 @@ namespace UnrealBuildTool
 
 			// Create an action that invokes the linker.
 			Action LinkAction = Graph.CreateAction(ActionType.Link);
+			LinkAction.RootPaths.AddRange(GetEnvironmentBasePaths(LinkEnvironment));
 
 			// RPC utility parameters are in terms of the Mac side
 			LinkAction.WorkingDirectory = GetMacDevSrcRoot();
diff --git a/Engine/Source/Programs/UnrealBuildTool/Platform/Linux/LinuxToolChain.cs b/Engine/Source/Programs/UnrealBuildTool/Platform/Linux/LinuxToolChain.cs
--- a/Engine/Source/Programs/UnrealBuildTool/Platform/Linux/LinuxToolChain.cs
+++ b/Engine/Source/Programs/UnrealBuildTool/Platform/Linux/LinuxToolChain.cs
@@ -25,7 +25,7 @@ namespace UnrealBuildTool
 			public FileReference BreakpadEncoder { get; init; }
 
 			public LinuxToolChainInfo(DirectoryReference? BaseLinuxPath, DirectoryReference? MultiArchRoot, FileReference Clang, FileReference Archiver, FileReference Objcopy, ILogger Logger)
-				: base(Clang, Archiver, Logger)
+				: base(BaseLinuxPath, Clang, Archiver, Logger)
 			{
 				this.BaseLinuxPath = BaseLinuxPath;
 				this.MultiArchRoot = MultiArchRoot;
@@ -1137,6 +1137,7 @@ namespace UnrealBuildTool
 
 			// Create an action that invokes the linker.
 			Action LinkAction = Graph.CreateAction(ActionType.Link);
+			LinkAction.RootPaths.AddRange(GetEnvironmentBasePaths(LinkEnvironment));
 			LinkAction.WorkingDirectory = Unreal.EngineSourceDirectory;
 
 			string LinkCommandString;
@@ -1581,6 +1582,7 @@ namespace UnrealBuildTool
 				{
 					// Create the action to relink the library. This actions does not overwrite the source file so it can be executed in parallel
 					Action RelinkAction = Graph.CreateAction(ActionType.Link);
+					RelinkAction.RootPaths.AddRange(GetEnvironmentBasePaths(LinkEnvironment));
 					RelinkAction.WorkingDirectory = LinkAction.WorkingDirectory;
 					RelinkAction.StatusDescription = LinkAction.StatusDescription;
 					RelinkAction.CommandDescription = "Relink";
diff --git a/Engine/Source/Programs/UnrealBuildTool/Platform/Mac/MacToolChain.cs b/Engine/Source/Programs/UnrealBuildTool/Platform/Mac/MacToolChain.cs
--- a/Engine/Source/Programs/UnrealBuildTool/Platform/Mac/MacToolChain.cs
+++ b/Engine/Source/Programs/UnrealBuildTool/Platform/Mac/MacToolChain.cs
@@ -71,7 +71,7 @@ namespace UnrealBuildTool
 		{
 			FileReference CompilerPath = FileReference.Combine(Settings.ToolchainDir, MacCompiler);
 			FileReference ArchiverPath = FileReference.Combine(Settings.ToolchainDir, MacArchiver);
-			return new AppleToolChainInfo(CompilerPath, ArchiverPath, Logger);
+			return new AppleToolChainInfo(MacToolChainSettings.XcodeDeveloperDir, CompilerPath, ArchiverPath, Logger);
 		}
 
 		public static DirectoryReference FindProductDirectory(FileReference? ProjectFile, DirectoryReference BinaryDir, string? NameIfProgram)
@@ -685,6 +685,7 @@ namespace UnrealBuildTool
 
 			// Create an action that invokes the linker.
 			Action LinkAction = Graph.CreateAction(ActionType.Link);
+			LinkAction.RootPaths.AddRange(GetEnvironmentBasePaths(LinkEnvironment));
 
 			FileReference LinkerPath = LinkEnvironment.bIsBuildingLibrary ? Info.Archiver : Info.Clang;
 
diff --git a/Engine/Source/Programs/UnrealBuildTool/Platform/Windows/PVSToolChain.cs b/Engine/Source/Programs/UnrealBuildTool/Platform/Windows/PVSToolChain.cs
--- a/Engine/Source/Programs/UnrealBuildTool/Platform/Windows/PVSToolChain.cs
+++ b/Engine/Source/Programs/UnrealBuildTool/Platform/Windows/PVSToolChain.cs
@@ -640,6 +640,17 @@ namespace UnrealBuildTool
 			Lines.Add(String.Format("Using PVS-Studio {0} at {1} with analysis mode {2} ({3})", AnalyzerVersion, AnalyzerFile, (uint)Settings.ModeFlags, Settings.ModeFlags.ToString()));
 		}
 
+		protected override IEnumerable<DirectoryItem> GetEnvironmentBasePaths(CppCompileEnvironment CompileEnvironment)
+		{
+			yield return DirectoryItem.GetItemByDirectoryReference(AnalyzerFile.Directory);
+			yield return DirectoryItem.GetItemByDirectoryReference(Unreal.EngineDirectory);
+			if (ProjectFile != null && (!CompileEnvironment.bUseSharedBuildEnvironment || CompileEnvironment.AllIncludePath.Any(x => x.IsUnderDirectory(ProjectFile.Directory))))
+			{
+				yield return DirectoryItem.GetItemByDirectoryReference(ProjectFile.Directory);
+			}
+			yield return DirectoryItem.GetItemByDirectoryReference(Unreal.RootDirectory);
+		}
+
 		static Version GetAnalyzerVersion(FileReference AnalyzerPath)
 		{
 			string Output = String.Empty;
diff --git a/Engine/Source/Programs/UnrealBuildTool/Platform/Windows/VCCompileAction.cs b/Engine/Source/Programs/UnrealBuildTool/Platform/Windows/VCCompileAction.cs
--- a/Engine/Source/Programs/UnrealBuildTool/Platform/Windows/VCCompileAction.cs
+++ b/Engine/Source/Programs/UnrealBuildTool/Platform/Windows/VCCompileAction.cs
@@ -176,6 +176,11 @@ namespace UnrealBuildTool
 		/// </summary>
 		public List<FileItem> DeleteItems { get; } = new List<FileItem>();
 
+		/// <summary>
+		/// Root paths for this action (generally engine root project root, toolchain root, sdk root)
+		/// </summary>
+		public List<DirectoryItem> RootPaths { get; } = new List<DirectoryItem>();
+
 		/// <inheritdoc/>
 		public bool bCanExecuteRemotely { get; set; }
 
@@ -201,6 +206,7 @@ namespace UnrealBuildTool
 		#region Implementation of IAction
 
 		IEnumerable<FileItem> IExternalAction.DeleteItems => DeleteItems;
+		IEnumerable<DirectoryItem> IExternalAction.RootPaths => RootPaths;
 		public DirectoryReference WorkingDirectory => Unreal.EngineSourceDirectory;
 		string IExternalAction.CommandDescription
 		{
@@ -383,6 +389,7 @@ namespace UnrealBuildTool
 			AdditionalPrerequisiteItems = new List<FileItem>(InAction.AdditionalPrerequisiteItems);
 			AdditionalProducedItems = new List<FileItem>(InAction.AdditionalProducedItems);
 			DeleteItems = new List<FileItem>(InAction.DeleteItems);
+			RootPaths = new List<DirectoryItem>(InAction.RootPaths);
 		}
 
 		/// <summary>
@@ -426,6 +433,7 @@ namespace UnrealBuildTool
 			AdditionalPrerequisiteItems = Reader.ReadList(() => Reader.ReadFileItem())!;
 			AdditionalProducedItems = Reader.ReadList(() => Reader.ReadFileItem())!;
 			DeleteItems = Reader.ReadList(() => Reader.ReadFileItem())!;
+			RootPaths = Reader.ReadList(() => Reader.ReadDirectoryItem())!;
 		}
 
 		/// <inheritdoc/>
@@ -466,6 +474,7 @@ namespace UnrealBuildTool
 			Writer.WriteList(AdditionalPrerequisiteItems, Item => Writer.WriteFileItem(Item));
 			Writer.WriteList(AdditionalProducedItems, Item => Writer.WriteFileItem(Item));
 			Writer.WriteList(DeleteItems, Item => Writer.WriteFileItem(Item));
+			Writer.WriteList(RootPaths, Item => Writer.WriteDirectoryItem(Item));
 		}
 
 		/// <summary>
diff --git a/Engine/Source/Programs/UnrealBuildTool/Platform/Windows/VCToolChain.cs b/Engine/Source/Programs/UnrealBuildTool/Platform/Windows/VCToolChain.cs
--- a/Engine/Source/Programs/UnrealBuildTool/Platform/Windows/VCToolChain.cs
+++ b/Engine/Source/Programs/UnrealBuildTool/Platform/Windows/VCToolChain.cs
@@ -1493,10 +1493,44 @@ namespace UnrealBuildTool
 			}
 		}
 
+		protected override IEnumerable<DirectoryItem> GetEnvironmentBasePaths(CppCompileEnvironment CompileEnvironment)
+		{
+			yield return DirectoryItem.GetItemByDirectoryReference(EnvVars.CompilerDir);
+			if (EnvVars.CompilerDir != EnvVars.ToolChainDir)
+			{
+				yield return DirectoryItem.GetItemByDirectoryReference(EnvVars.ToolChainDir);
+			}
+			yield return DirectoryItem.GetItemByDirectoryReference(EnvVars.WindowsSdkDir);
+			yield return DirectoryItem.GetItemByDirectoryReference(Unreal.EngineDirectory);
+			if (Target.ProjectFile != null && (!CompileEnvironment.bUseSharedBuildEnvironment || CompileEnvironment.AllIncludePath.Any(x => x.IsUnderDirectory(Target.ProjectFile.Directory))))
+			{
+				yield return DirectoryItem.GetItemByDirectoryReference(Target.ProjectFile.Directory);
+			}
+			yield return DirectoryItem.GetItemByDirectoryReference(Unreal.RootDirectory);
+		}
+
+		protected override IEnumerable<DirectoryItem> GetEnvironmentBasePaths(LinkEnvironment LinkEnvironment)
+		{
+			yield return DirectoryItem.GetItemByDirectoryReference(EnvVars.CompilerDir);
+			if (EnvVars.CompilerDir != EnvVars.ToolChainDir)
+			{
+				yield return DirectoryItem.GetItemByDirectoryReference(EnvVars.ToolChainDir);
+			}
+			yield return DirectoryItem.GetItemByDirectoryReference(EnvVars.WindowsSdkDir);
+			yield return DirectoryItem.GetItemByDirectoryReference(Unreal.EngineDirectory);
+			if (Target.ProjectFile != null && LinkEnvironment.InputFiles.Any(x => x.Location.IsUnderDirectory(Target.ProjectFile.Directory)))
+			{
+				yield return DirectoryItem.GetItemByDirectoryReference(Target.ProjectFile.Directory);
+			}
+			yield return DirectoryItem.GetItemByDirectoryReference(Unreal.RootDirectory);
+		}
+
 		private VCCompileAction CreateBaseCompileAction(CppCompileEnvironment CompileEnvironment)
 		{
 			VCCompileAction BaseCompileAction = new VCCompileAction(EnvVars);
 
+			BaseCompileAction.RootPaths.AddRange(GetEnvironmentBasePaths(CompileEnvironment));
+
 			// Add additional response files
 			foreach (FileItem AdditionalRsp in CompileEnvironment.AdditionalResponseFiles)
 			{
@@ -2095,6 +2129,7 @@ namespace UnrealBuildTool
 			foreach (FileItem RCFile in InputFiles)
 			{
 				Action CompileAction = Graph.CreateAction(ActionType.Compile);
+				CompileAction.RootPaths.AddRange(GetEnvironmentBasePaths(CompileEnvironment));
 				CompileAction.CommandDescription = "Resource";
 				CompileAction.WorkingDirectory = Unreal.EngineSourceDirectory;
 				CompileAction.CommandPath = EnvVars.ResourceCompilerPath;
@@ -2654,6 +2689,7 @@ namespace UnrealBuildTool
 
 			// Create an action that invokes the linker.
 			Action LinkAction = Graph.CreateAction(ActionType.Link);
+			LinkAction.RootPaths.AddRange(GetEnvironmentBasePaths(LinkEnvironment));
 			string ReadableArch = UnrealArchitectureConfig.ForPlatform(LinkEnvironment.Platform).ConvertToReadableArchitecture(LinkEnvironment.Architecture);
 			LinkAction.CommandDescription += $"Link [{ReadableArch}]";
 			LinkAction.WorkingDirectory = Unreal.EngineSourceDirectory;
diff --git a/Engine/Source/Programs/UnrealBuildTool/System/Action.cs b/Engine/Source/Programs/UnrealBuildTool/System/Action.cs
--- a/Engine/Source/Programs/UnrealBuildTool/System/Action.cs
+++ b/Engine/Source/Programs/UnrealBuildTool/System/Action.cs
@@ -95,6 +95,11 @@ namespace UnrealBuildTool
 		/// </summary>
 		IEnumerable<FileItem> DeleteItems { get; }
 
+		/// <summary>
+		/// Root paths for this action (generally engine root project root, toolchain root, sdk root)
+		/// </summary>
+		IEnumerable<DirectoryItem> RootPaths { get; }
+
 		/// <summary>
 		/// For C++ source files, specifies a dependency list file used to check changes to header files
 		/// </summary>
@@ -205,17 +210,23 @@ namespace UnrealBuildTool
 		/// <summary>
 		/// Every file this action depends on.  These files need to exist and be up to date in order for this action to even be considered
 		/// </summary>
-		public SortedSet<FileItem> PrerequisiteItems { get; set; } = new SortedSet<FileItem>();
+		public SortedSet<FileItem> PrerequisiteItems { get; set; } = new();
 
 		/// <summary>
 		/// The files that this action produces after completing
 		/// </summary>
-		public SortedSet<FileItem> ProducedItems { get; set; } = new SortedSet<FileItem>();
+		public SortedSet<FileItem> ProducedItems { get; set; } = new();
 
 		/// <summary>
 		/// Items that should be deleted before running this action
 		/// </summary>
-		public SortedSet<FileItem> DeleteItems { get; set; } = new SortedSet<FileItem>();
+		public SortedSet<FileItem> DeleteItems { get; set; } = new();
+
+		/// <summary>
+		/// Root paths for this action (generally engine root project root, toolchain root, sdk root)
+        /// Order needs to be mainted, so a list is used
+		/// </summary>
+		public List<DirectoryItem> RootPaths { get; set; } = new();
 
 		/// <summary>
 		/// For C++ source files, specifies a dependency list file used to check changes to header files
@@ -300,8 +311,9 @@ namespace UnrealBuildTool
 		IEnumerable<FileItem> IExternalAction.PrerequisiteItems => PrerequisiteItems;
 		IEnumerable<FileItem> IExternalAction.ProducedItems => ProducedItems;
 		IEnumerable<FileItem> IExternalAction.DeleteItems => DeleteItems;
+        IEnumerable<DirectoryItem> IExternalAction.RootPaths => RootPaths;
 
-		public Action(ActionType InActionType)
+        public Action(ActionType InActionType)
 		{
 			ActionType = InActionType;
 
@@ -316,9 +328,10 @@ namespace UnrealBuildTool
 		{
 			ActionType = InOther.ActionType;
 			ArtifactMode = InOther.ArtifactMode;
-			PrerequisiteItems = new SortedSet<FileItem>(InOther.PrerequisiteItems);
-			ProducedItems = new SortedSet<FileItem>(InOther.ProducedItems);
-			DeleteItems = new SortedSet<FileItem>(InOther.DeleteItems);
+			PrerequisiteItems = new(InOther.PrerequisiteItems);
+			ProducedItems = new(InOther.ProducedItems);
+			DeleteItems = new(InOther.DeleteItems);
+            RootPaths = new(InOther.RootPaths);
 			DependencyListFile = InOther.DependencyListFile;
 			WorkingDirectory = InOther.WorkingDirectory;
 			CommandPath = InOther.CommandPath;
@@ -358,7 +371,8 @@ namespace UnrealBuildTool
 			PrerequisiteItems = Reader.ReadSortedSet(() => Reader.ReadFileItem())!;
 			ProducedItems = Reader.ReadSortedSet(() => Reader.ReadFileItem())!;
 			DeleteItems = Reader.ReadSortedSet(() => Reader.ReadFileItem())!;
-			DependencyListFile = Reader.ReadFileItem();
+            RootPaths = Reader.ReadList(() => Reader.ReadDirectoryItem())!;
+            DependencyListFile = Reader.ReadFileItem();
 			bUseActionHistory = Reader.ReadBool();
 			bIsHighPriority = Reader.ReadBool();
 			Weight = Reader.ReadDouble();
@@ -388,6 +402,7 @@ namespace UnrealBuildTool
 			Writer.WriteSortedSet(PrerequisiteItems, Item => Writer.WriteFileItem(Item));
 			Writer.WriteSortedSet(ProducedItems, Item => Writer.WriteFileItem(Item));
 			Writer.WriteSortedSet(DeleteItems, Item => Writer.WriteFileItem(Item));
+            Writer.WriteList(RootPaths, Item => Writer.WriteDirectoryItem(Item));
 			Writer.WriteFileItem(DependencyListFile);
 			Writer.WriteBool(bUseActionHistory);
 			Writer.WriteBool(bIsHighPriority);
@@ -402,116 +417,102 @@ namespace UnrealBuildTool
 		{
 			Action Action = new Action(Object.GetEnumField<ActionType>("Type"));
 
-			ArtifactMode ArtifactMode;
-			if (Object.TryGetEnumField("ArtifactMode", out ArtifactMode))
+			if (Object.TryGetEnumField("ArtifactMode", out ArtifactMode ArtifactMode))
 			{
 				Action.ArtifactMode = ArtifactMode;
 			}
 
-			string? WorkingDirectory;
-			if (Object.TryGetStringField("WorkingDirectory", out WorkingDirectory))
+			if (Object.TryGetStringField("WorkingDirectory", out string? WorkingDirectory))
 			{
 				Action.WorkingDirectory = new DirectoryReference(WorkingDirectory);
 			}
 
-			string? CommandPath;
-			if (Object.TryGetStringField("CommandPath", out CommandPath))
+			if (Object.TryGetStringField("CommandPath", out string? CommandPath))
 			{
 				Action.CommandPath = new FileReference(CommandPath);
 			}
 
-			string? CommandArguments;
-			if (Object.TryGetStringField("CommandArguments", out CommandArguments))
+			if (Object.TryGetStringField("CommandArguments", out string? CommandArguments))
 			{
 				Action.CommandArguments = CommandArguments;
 			}
 
-			string? CommandVersion;
-			if (Object.TryGetStringField("CommandVersion", out CommandVersion))
+			if (Object.TryGetStringField("CommandVersion", out string? CommandVersion))
 			{
 				Action.CommandVersion = CommandVersion;
 			}
 
-			string? CommandDescription;
-			if (Object.TryGetStringField("CommandDescription", out CommandDescription))
+			if (Object.TryGetStringField("CommandDescription", out string? CommandDescription))
 			{
 				Action.CommandDescription = CommandDescription;
 			}
 
-			string? StatusDescription;
-			if (Object.TryGetStringField("StatusDescription", out StatusDescription))
+			if (Object.TryGetStringField("StatusDescription", out string? StatusDescription))
 			{
 				Action.StatusDescription = StatusDescription;
 			}
 
-			bool bCanExecuteRemotely;
-			if (Object.TryGetBoolField("bCanExecuteRemotely", out bCanExecuteRemotely))
+			if (Object.TryGetBoolField("bCanExecuteRemotely", out bool bCanExecuteRemotely))
 			{
 				Action.bCanExecuteRemotely = bCanExecuteRemotely;
 			}
 
-			bool bCanExecuteRemotelyWithSNDBS;
-			if (Object.TryGetBoolField("bCanExecuteRemotelyWithSNDBS", out bCanExecuteRemotelyWithSNDBS))
+			if (Object.TryGetBoolField("bCanExecuteRemotelyWithSNDBS", out bool bCanExecuteRemotelyWithSNDBS))
 			{
 				Action.bCanExecuteRemotelyWithSNDBS = bCanExecuteRemotelyWithSNDBS;
 			}
 
-			bool bCanExecuteRemotelyWithXGE;
-			if (Object.TryGetBoolField("bCanExecuteRemotelyWithXGE", out bCanExecuteRemotelyWithXGE))
+			if (Object.TryGetBoolField("bCanExecuteRemotelyWithXGE", out bool bCanExecuteRemotelyWithXGE))
 			{
 				Action.bCanExecuteRemotelyWithXGE = bCanExecuteRemotelyWithXGE;
 			}
 
-			bool bCanExecuteInUBA;
-			if (Object.TryGetBoolField("bCanExecuteInUBA", out bCanExecuteInUBA))
+			if (Object.TryGetBoolField("bCanExecuteInUBA", out bool bCanExecuteInUBA))
 			{
 				Action.bCanExecuteInUBA = bCanExecuteInUBA;
 			}
 
-			bool bIsGCCCompiler;
-			if (Object.TryGetBoolField("bIsGCCCompiler", out bIsGCCCompiler))
+			if (Object.TryGetBoolField("bIsGCCCompiler", out bool bIsGCCCompiler))
 			{
 				Action.bIsGCCCompiler = bIsGCCCompiler;
 			}
 
-			bool bShouldOutputStatusDescription;
-			if (Object.TryGetBoolField("bShouldOutputStatusDescription", out bShouldOutputStatusDescription))
+			if (Object.TryGetBoolField("bShouldOutputStatusDescription", out bool bShouldOutputStatusDescription))
 			{
 				Action.bShouldOutputStatusDescription = bShouldOutputStatusDescription;
 			}
 
-			bool bProducesImportLibrary;
-			if (Object.TryGetBoolField("bProducesImportLibrary", out bProducesImportLibrary))
+			if (Object.TryGetBoolField("bProducesImportLibrary", out bool bProducesImportLibrary))
 			{
 				Action.bProducesImportLibrary = bProducesImportLibrary;
 			}
 
-			string[]? PrerequisiteItems;
-			if (Object.TryGetStringArrayField("PrerequisiteItems", out PrerequisiteItems))
+			if (Object.TryGetStringArrayField("PrerequisiteItems", out string[]? PrerequisiteItems))
 			{
 				Action.PrerequisiteItems.UnionWith(PrerequisiteItems.Select(x => FileItem.GetItemByPath(x)));
 			}
 
-			string[]? ProducedItems;
-			if (Object.TryGetStringArrayField("ProducedItems", out ProducedItems))
+			if (Object.TryGetStringArrayField("ProducedItems", out string[]? ProducedItems))
 			{
 				Action.ProducedItems.UnionWith(ProducedItems.Select(x => FileItem.GetItemByPath(x)));
 			}
 
-			string[]? DeleteItems;
-			if (Object.TryGetStringArrayField("DeleteItems", out DeleteItems))
+			if (Object.TryGetStringArrayField("DeleteItems", out string[]? DeleteItems))
 			{
 				Action.DeleteItems.UnionWith(DeleteItems.Select(x => FileItem.GetItemByPath(x)));
 			}
 
-			string? DependencyListFile;
-			if (Object.TryGetStringField("DependencyListFile", out DependencyListFile))
+			if (Object.TryGetStringArrayField("RootPaths", out string[]? RootPaths))
+			{
+                Action.RootPaths.AddRange(RootPaths.Select(x => DirectoryItem.GetItemByPath(x)));
+			}
+
+			if (Object.TryGetStringField("DependencyListFile", out string? DependencyListFile))
 			{
 				Action.DependencyListFile = FileItem.GetItemByPath(DependencyListFile);
 			}
 
-			double Weight;
-			if (Object.TryGetDoubleField("Weight", out Weight))
+			if (Object.TryGetDoubleField("Weight", out double Weight))
 			{
 				Action.Weight = Weight;
 			}
@@ -603,6 +604,13 @@ namespace UnrealBuildTool
 			}
 			Writer.WriteArrayEnd();
 
+			Writer.WriteArrayStart("RootPaths");
+			foreach (DirectoryItem RootPath in Action.RootPaths)
+			{
+				Writer.WriteValue(RootPath.FullName);
+			}
+			Writer.WriteArrayEnd();
+
 			if (Action.DependencyListFile != null)
 			{
 				Writer.WriteValue("DependencyListFile", Action.DependencyListFile.AbsolutePath);
@@ -688,6 +696,7 @@ namespace UnrealBuildTool
 		public IEnumerable<FileItem> PrerequisiteItems => Inner.PrerequisiteItems;
 		public IEnumerable<FileItem> ProducedItems => Inner.ProducedItems;
 		public IEnumerable<FileItem> DeleteItems => Inner.DeleteItems;
+		public IEnumerable<DirectoryItem> RootPaths => Inner.RootPaths;
 		public FileItem? DependencyListFile => Inner.DependencyListFile;
 		public DirectoryReference WorkingDirectory => Inner.WorkingDirectory;
 		public FileReference CommandPath => Inner.CommandPath;
diff --git a/Engine/Source/Programs/UnrealBuildTool/System/CppCompileEnvironment.cs b/Engine/Source/Programs/UnrealBuildTool/System/CppCompileEnvironment.cs
--- a/Engine/Source/Programs/UnrealBuildTool/System/CppCompileEnvironment.cs
+++ b/Engine/Source/Programs/UnrealBuildTool/System/CppCompileEnvironment.cs
@@ -1,6 +1,7 @@
 // Copyright Epic Games, Inc. All Rights Reserved.
 
 using System.Collections.Generic;
+using System.Linq;
 using EpicGames.Core;
 using UnrealBuildBase;
 
@@ -540,6 +541,11 @@ namespace UnrealBuildTool
 		/// </summary>
 		public HashSet<DirectoryReference> SharedSystemIncludePaths;
 
+		/// <summary>
+		/// Enumerable of all possible include paths
+		/// </summary>
+		public IEnumerable<DirectoryReference> AllIncludePath => SharedUserIncludePaths.Concat(UserIncludePaths).Concat(SharedSystemIncludePaths).Concat(SystemIncludePaths);
+
 		/// <summary>
 		/// List of paths to search for compiled module interface (*.ifc) files
 		/// </summary>
diff --git a/Engine/Source/Programs/UnrealBuildTool/ToolChain/AppleToolChain.cs b/Engine/Source/Programs/UnrealBuildTool/ToolChain/AppleToolChain.cs
--- a/Engine/Source/Programs/UnrealBuildTool/ToolChain/AppleToolChain.cs
+++ b/Engine/Source/Programs/UnrealBuildTool/ToolChain/AppleToolChain.cs
@@ -241,8 +241,8 @@ namespace UnrealBuildTool
 	{
 		protected class AppleToolChainInfo : ClangToolChainInfo
 		{
-			public AppleToolChainInfo(FileReference Clang, FileReference Archiver, ILogger Logger)
-				: base(Clang, Archiver, Logger)
+			public AppleToolChainInfo(DirectoryReference DeveloperDir, FileReference Clang, FileReference Archiver, ILogger Logger)
+				: base(DeveloperDir, Clang, Archiver, Logger)
 			{
 			}
 
@@ -253,7 +253,6 @@ namespace UnrealBuildTool
 
 		public Lazy<AppleToolChainSettings> ToolChainSettings;
 
-		protected FileReference? ProjectFile;
 		public readonly ReadOnlyTargetRules? Target;
 
 		// cache some ini settings
diff --git a/Engine/Source/Programs/UnrealBuildTool/ToolChain/ClangToolChain.cs b/Engine/Source/Programs/UnrealBuildTool/ToolChain/ClangToolChain.cs
--- a/Engine/Source/Programs/UnrealBuildTool/ToolChain/ClangToolChain.cs
+++ b/Engine/Source/Programs/UnrealBuildTool/ToolChain/ClangToolChain.cs
@@ -126,6 +126,7 @@ namespace UnrealBuildTool
 		protected class ClangToolChainInfo
 		{
 			protected ILogger Logger { get; init; }
+			public DirectoryReference? BasePath { get; init; }
 			public FileReference Clang { get; init; }
 			public FileReference Archiver { get; init; }
 
@@ -133,19 +134,21 @@ namespace UnrealBuildTool
 			public string ClangVersionString => LazyClangVersionString.Value;
 			public string ArchiverVersionString => LazyArchiverVersionString.Value;
 
-			Lazy<Version> LazyClangVersion;
-			Lazy<string> LazyClangVersionString;
-			Lazy<string> LazyArchiverVersionString;
+			readonly Lazy<Version> LazyClangVersion;
+			readonly Lazy<string> LazyClangVersionString;
+			readonly Lazy<string> LazyArchiverVersionString;
 
 			/// <summary>
 			/// Constructor for ClangToolChainInfo
 			/// </summary>
+			/// <param name="BasePath">The base path to the clang sdk root, if available</param>
 			/// <param name="Clang">The path to the compiler</param>
 			/// <param name="Archiver">The path to the archiver</param>
 			/// <param name="Logger">Logging interface</param>
-			public ClangToolChainInfo(FileReference Clang, FileReference Archiver, ILogger Logger)
+			public ClangToolChainInfo(DirectoryReference? BasePath, FileReference Clang, FileReference Archiver, ILogger Logger)
 			{
 				this.Logger = Logger;
+				this.BasePath = BasePath;
 				this.Clang = Clang;
 				this.Archiver = Archiver;
 
@@ -1094,9 +1097,38 @@ namespace UnrealBuildTool
 			Graph.AddAction(new ClangSpecificFileAction(SourceDir, OutputDir, Action, GraphBuilder.ContentLines));
 		}
 
+		protected override IEnumerable<DirectoryItem> GetEnvironmentBasePaths(CppCompileEnvironment CompileEnvironment)
+		{
+			if (GetToolChainInfo().BasePath != null)
+			{
+				yield return DirectoryItem.GetItemByDirectoryReference(GetToolChainInfo().BasePath!);
+			}
+			yield return DirectoryItem.GetItemByDirectoryReference(Unreal.EngineDirectory);
+			if (ProjectFile != null && (!CompileEnvironment.bUseSharedBuildEnvironment || CompileEnvironment.AllIncludePath.Any(x => x.IsUnderDirectory(ProjectFile.Directory))))
+			{
+				yield return DirectoryItem.GetItemByDirectoryReference(ProjectFile.Directory);
+			}
+			yield return DirectoryItem.GetItemByDirectoryReference(Unreal.RootDirectory);
+		}
+
+		protected override IEnumerable<DirectoryItem> GetEnvironmentBasePaths(LinkEnvironment LinkEnvironment)
+		{
+			if (GetToolChainInfo().BasePath != null)
+			{
+				yield return DirectoryItem.GetItemByDirectoryReference(GetToolChainInfo().BasePath!);
+			}
+			yield return DirectoryItem.GetItemByDirectoryReference(Unreal.EngineDirectory);
+			if (ProjectFile != null && LinkEnvironment.InputFiles.Any(x => x.Location.IsUnderDirectory(ProjectFile.Directory)))
+			{
+				yield return DirectoryItem.GetItemByDirectoryReference(ProjectFile.Directory);
+			}
+			yield return DirectoryItem.GetItemByDirectoryReference(Unreal.RootDirectory);
+		}
+
 		protected virtual Action CompileCPPFile(CppCompileEnvironment CompileEnvironment, FileItem SourceFile, DirectoryReference OutputDir, string ModuleName, IActionGraphBuilder Graph, IReadOnlyCollection<string> GlobalArguments, CPPOutput Result)
 		{
 			Action CompileAction = Graph.CreateAction(ActionType.Compile);
+			CompileAction.RootPaths.AddRange(GetEnvironmentBasePaths(CompileEnvironment));
 
 			// If we are using the AutoRTFM compiler, we make the compile action depend on the version of the compiler itself.
 			// This lets us update the compiler (which might not cause a version update of the compiler, which instead tracks
@@ -1167,6 +1199,7 @@ namespace UnrealBuildTool
 			if (PreprocessDepends)
 			{
 				Action PrepassAction = Graph.CreateAction(ActionType.Compile);
+				PrepassAction.RootPaths.AddRange(CompileAction.RootPaths);
 				PrepassAction.PrerequisiteItems.UnionWith(CompileAction.PrerequisiteItems);
 				PrepassAction.PrerequisiteItems.Remove(CompilerResponseFileItem);
 				PrepassAction.CommandDescription = "Preprocess Depends";
diff --git a/Engine/Source/Programs/UnrealBuildTool/ToolChain/ISPCToolChain.cs b/Engine/Source/Programs/UnrealBuildTool/ToolChain/ISPCToolChain.cs
--- a/Engine/Source/Programs/UnrealBuildTool/ToolChain/ISPCToolChain.cs
+++ b/Engine/Source/Programs/UnrealBuildTool/ToolChain/ISPCToolChain.cs
@@ -17,6 +17,14 @@ namespace UnrealBuildTool
 		{
 		}
 
+		protected FileReference? ProjectFile = null;
+
+		public override void SetUpGlobalEnvironment(ReadOnlyTargetRules Target)
+		{
+			base.SetUpGlobalEnvironment(Target);
+			ProjectFile = Target.ProjectFile;
+		}
+
 		/// <summary>
 		/// Get CPU Instruction set targets for ISPC.
 		/// </summary>
@@ -384,6 +392,26 @@ namespace UnrealBuildTool
 			return NormalizeCommandLinePath(Item.Location);
 		}
 
+		protected virtual IEnumerable<DirectoryItem> GetEnvironmentBasePaths(CppCompileEnvironment CompileEnvironment)
+		{
+			yield return DirectoryItem.GetItemByDirectoryReference(Unreal.EngineDirectory);
+			if (ProjectFile != null && (!CompileEnvironment.bUseSharedBuildEnvironment || CompileEnvironment.AllIncludePath.Any(x => x.IsUnderDirectory(ProjectFile.Directory))))
+			{
+				yield return DirectoryItem.GetItemByDirectoryReference(ProjectFile.Directory);
+			}
+			yield return DirectoryItem.GetItemByDirectoryReference(Unreal.RootDirectory);
+		}
+
+		protected virtual IEnumerable<DirectoryItem> GetEnvironmentBasePaths(LinkEnvironment LinkEnvironment)
+		{
+			yield return DirectoryItem.GetItemByDirectoryReference(Unreal.EngineDirectory);
+			if (ProjectFile != null && LinkEnvironment.InputFiles.Any(x => x.Location.IsUnderDirectory(ProjectFile.Directory)))
+			{
+				yield return DirectoryItem.GetItemByDirectoryReference(ProjectFile.Directory);
+			}
+			yield return DirectoryItem.GetItemByDirectoryReference(Unreal.RootDirectory);
+		}
+
 		public override CPPOutput GenerateISPCHeaders(CppCompileEnvironment CompileEnvironment, IEnumerable<FileItem> InputFiles, DirectoryReference OutputDir, IActionGraphBuilder Graph)
 		{
 			CPPOutput Result = new CPPOutput();
@@ -445,9 +473,13 @@ namespace UnrealBuildTool
 				}
 			}
 
+			List<DirectoryItem> RootPaths = new(GetEnvironmentBasePaths(CompileEnvironment));
+
 			foreach (FileItem ISPCFile in InputFiles)
 			{
 				Action CompileAction = Graph.CreateAction(ActionType.Compile);
+				CompileAction.RootPaths.AddRange(RootPaths);
+
 				CompileAction.CommandDescription = $"Generate Header [{ISPCArch}]";
 				CompileAction.WorkingDirectory = Unreal.EngineSourceDirectory;
 				CompileAction.CommandPath = new FileReference(GetISPCHostCompilerPath(BuildHostPlatform.Current.Platform));
@@ -615,9 +647,12 @@ namespace UnrealBuildTool
 				}
 			}
 
+			List<DirectoryItem> RootPaths = new(GetEnvironmentBasePaths(CompileEnvironment));
+
 			foreach (FileItem ISPCFile in InputFiles)
 			{
 				Action CompileAction = Graph.CreateAction(ActionType.Compile);
+				CompileAction.RootPaths.AddRange(RootPaths);
 				CompileAction.CommandDescription = $"Compile [{ISPCArch}]";
 				CompileAction.WorkingDirectory = Unreal.EngineSourceDirectory;
 				CompileAction.CommandPath = new FileReference(GetISPCHostCompilerPath(BuildHostPlatform.Current.Platform));
diff --git a/Engine/Source/Programs/UnrealBuildTool/ToolChain/IWYUToolChain.cs b/Engine/Source/Programs/UnrealBuildTool/ToolChain/IWYUToolChain.cs
--- a/Engine/Source/Programs/UnrealBuildTool/ToolChain/IWYUToolChain.cs
+++ b/Engine/Source/Programs/UnrealBuildTool/ToolChain/IWYUToolChain.cs
@@ -46,7 +46,7 @@ namespace UnrealBuildTool
 			CrossCompilingArguments.Add($"--sysroot=\"{NormalizeCommandLinePath(BaseLinuxPath!)}\"");
 
 			FileReference ClangPath = FileReference.Combine(BaseLinuxPath!, "bin", $"clang++{BuildHostPlatform.Current.BinarySuffix}");
-			ClangToolChainInfo CompilerToolChainInfo = new ClangToolChainInfo(ClangPath, null!, Logger);
+			ClangToolChainInfo CompilerToolChainInfo = new ClangToolChainInfo(BaseLinuxPath, ClangPath, null!, Logger);
 
 			// starting with clang 16.x the directory naming changed to include major version only
 			string ClangVersionString = (CompilerToolChainInfo.ClangVersion.Major >= 16) ? CompilerToolChainInfo.ClangVersion.Major.ToString() : CompilerToolChainInfo.ClangVersion.ToString();
@@ -56,7 +56,7 @@ namespace UnrealBuildTool
 			string DevPath = ""; //@"include-what-you-use-0.19\vs_projects\bin\RelWithDebInfo";
 
 			FileReference IWYUPath = FileReference.Combine(Unreal.EngineDirectory, RelativePathToIWYUDirectory, DevPath, @"include-what-you-use.exe");
-			return new ClangToolChainInfo(IWYUPath!, null!, Logger);
+			return new ClangToolChainInfo(IWYUPath.Directory, IWYUPath, null!, Logger);
 		}
 
 		protected override string GetFileNameFromExtension(string AbsolutePath, string Extension)
