UnrealBuildTool: All cplusplus.PureVirtualCall are fixed, reenable

#rnx

diff --git a/Engine/Source/Programs/UnrealBuildTool/Configuration/ModuleRules.cs b/Engine/Source/Programs/UnrealBuildTool/Configuration/ModuleRules.cs
--- a/Engine/Source/Programs/UnrealBuildTool/Configuration/ModuleRules.cs
+++ b/Engine/Source/Programs/UnrealBuildTool/Configuration/ModuleRules.cs
@@ -1036,7 +1036,6 @@ namespace UnrealBuildTool
 			"cplusplus.NewDelete",
 			"cplusplus.NewDeleteLeaks",
 			"cplusplus.PlacementNew",
-			"cplusplus.PureVirtualCall",
 			"unix.cstring.NullArg",
 			"unix.Malloc",
 
