UnrealBuildTool: Allow invalidation if UBA compression is enabled

#rnx

diff --git a/Engine/Source/Programs/UnrealBuildTool/Configuration/ReadOnlyTargetRules.cs b/Engine/Source/Programs/UnrealBuildTool/Configuration/ReadOnlyTargetRules.cs
--- a/Engine/Source/Programs/UnrealBuildTool/Configuration/ReadOnlyTargetRules.cs
+++ b/Engine/Source/Programs/UnrealBuildTool/Configuration/ReadOnlyTargetRules.cs
@@ -402,6 +402,8 @@ namespace UnrealBuildTool
 
 		public bool bWithAssembly => Inner.bWithAssembly;
 
+		public bool bAllowUbaCompression => Inner.bAllowUbaCompression;
+
 		public StaticAnalyzer StaticAnalyzer => Inner.StaticAnalyzer;
 
 		public StaticAnalyzerOutputType StaticAnalyzerOutputType => Inner.StaticAnalyzerOutputType;
diff --git a/Engine/Source/Programs/UnrealBuildTool/Configuration/TargetRules.cs b/Engine/Source/Programs/UnrealBuildTool/Configuration/TargetRules.cs
--- a/Engine/Source/Programs/UnrealBuildTool/Configuration/TargetRules.cs
+++ b/Engine/Source/Programs/UnrealBuildTool/Configuration/TargetRules.cs
@@ -1960,6 +1960,15 @@ namespace UnrealBuildTool
 		[CommandLine("-WithAssembly")]
 		public bool bWithAssembly { get; set; }
 
+		/// <summary>
+		/// Experimental: Store object (.obj) compressed on disk. Requires UBA to link, currently MSVC only. Toggling this flag will invalidate MSVC actions.
+		/// Warning, this option is not currently compatitable with PGO or the the cl-clang linker as those are not detoured and linking will fail.
+		/// </summary>
+		[XmlConfigFile(Category = "UnrealBuildAccelerator")]
+		[CommandLine("-UBAStoreObjFilesCompressed", Value = "true", MarkUsed = false)]
+		[RequiresUniqueBuildEnvironment]
+		public bool bAllowUbaCompression { get; set; }
+
 		/// <summary>
 		/// Whether static code analysis should be enabled.
 		/// </summary>
diff --git a/Engine/Source/Programs/UnrealBuildTool/Executors/UnrealBuildAccelerator/UBAExecutor.cs b/Engine/Source/Programs/UnrealBuildTool/Executors/UnrealBuildAccelerator/UBAExecutor.cs
--- a/Engine/Source/Programs/UnrealBuildTool/Executors/UnrealBuildAccelerator/UBAExecutor.cs
+++ b/Engine/Source/Programs/UnrealBuildTool/Executors/UnrealBuildAccelerator/UBAExecutor.cs
@@ -77,6 +77,10 @@ namespace UnrealBuildTool
 		private void PrintConfiguration()
 		{
 			_threadedLogger.LogInformation("  Storage capacity {StoreCapacityGb}Gb", UBAConfig.StoreCapacityGb);
+			if (UBAConfig.bStoreObjFilesCompressed)
+			{
+				_threadedLogger.LogInformation("  Object Compression Allowed");
+			}
 		}
 
 		private async Task WriteActionOutputFileAsync(IEnumerable<LinkedAction> inputActions)
diff --git a/Engine/Source/Programs/UnrealBuildTool/Platform/Windows/VCCompileAction.cs b/Engine/Source/Programs/UnrealBuildTool/Platform/Windows/VCCompileAction.cs
--- a/Engine/Source/Programs/UnrealBuildTool/Platform/Windows/VCCompileAction.cs
+++ b/Engine/Source/Programs/UnrealBuildTool/Platform/Windows/VCCompileAction.cs
@@ -47,7 +47,7 @@ namespace UnrealBuildTool
 		/// <summary>
 		/// The version of the compiler being used
 		/// </summary>
-		public string ToolChainVersion { get; }
+		public string ToolChainVersion { get; set; }
 
 		/// <summary>
 		/// Source file to compile
diff --git a/Engine/Source/Programs/UnrealBuildTool/Platform/Windows/VCToolChain.cs b/Engine/Source/Programs/UnrealBuildTool/Platform/Windows/VCToolChain.cs
--- a/Engine/Source/Programs/UnrealBuildTool/Platform/Windows/VCToolChain.cs
+++ b/Engine/Source/Programs/UnrealBuildTool/Platform/Windows/VCToolChain.cs
@@ -1617,6 +1617,11 @@ namespace UnrealBuildTool
 			// When compiling with SN-DBS, modules that contain a #import must be built locally
 			BaseCompileAction.bCanExecuteRemotelyWithSNDBS = BaseCompileAction.bCanExecuteRemotely && !CompileEnvironment.bBuildLocallyWithSNDBS;
 
+			if (Target.bAllowUbaCompression)
+			{
+				BaseCompileAction.ToolChainVersion = $"{BaseCompileAction.ToolChainVersion} Compressed";
+			}
+
 			return BaseCompileAction;
 		}
 
@@ -2747,6 +2752,11 @@ namespace UnrealBuildTool
 				LinkAction.bCanExecuteInUBA = false;
 			}
 
+			if (Target.bAllowUbaCompression)
+			{
+				LinkAction.CommandVersion = $"{LinkAction.CommandVersion} Compressed";
+			}
+
 			Logger.LogDebug("     Linking: {StatusDescription}", LinkAction.StatusDescription);
 			Logger.LogDebug("     Command: {CommandArguments}", LinkAction.CommandArguments);
 
