[UBT]
* Fixed a bug where UHT folders belonging to modules referenced through public includes which recursively referenced more modules with public include dependencies. Code was not recursively traversing public include dependencies.

diff --git a/Engine/Source/Programs/UnrealBuildTool/System/ExternalExecution.cs b/Engine/Source/Programs/UnrealBuildTool/System/ExternalExecution.cs
--- a/Engine/Source/Programs/UnrealBuildTool/System/ExternalExecution.cs
+++ b/Engine/Source/Programs/UnrealBuildTool/System/ExternalExecution.cs
@@ -468,10 +468,9 @@ namespace UnrealBuildTool
 
 			// Set Module.bHasUObjects for any IncludePathModules not already processed.
 			// This is necessary to keep include paths consistent between targets built with -AllModules and without
-			HashSet<UEBuildModuleCPP> IncludePathModules = new HashSet<UEBuildModuleCPP>();
-			IncludePathModules.UnionWith(ModulesToGenerateHeadersFor.SelectMany(x => x.PrivateIncludePathModules ?? new()).OfType<UEBuildModuleCPP>().Where(x => !x.bHasUObjects));
-			IncludePathModules.UnionWith(ModulesToGenerateHeadersFor.SelectMany(x => x.PublicIncludePathModules ?? new()).OfType<UEBuildModuleCPP>().Where(x => !x.bHasUObjects));
-			IncludePathModules.ExceptWith(ModulesToGenerateHeadersFor);
+			// Note that PublicIncludePathModules can be recursive so we need to traverse the entire chain
+			List<UEBuildModuleCPP> IncludePathModules = new();
+			CollectModulesOnlyIncluded(ModulesToGenerateHeadersFor, IncludePathModules);
 			if (IncludePathModules.Count > 0)
 			{
 				Dictionary<UEBuildModuleCPP, UHTModuleInfo> IncludePathInfo = new();
@@ -499,6 +498,53 @@ namespace UnrealBuildTool
 			}
 		}
 
+		static void CollectModulesOnlyIncluded(HashSet<UEBuildModuleCPP> HandledModules, List<UEBuildModule> IncludedModules, List<UEBuildModuleCPP> OutList)
+		{
+			foreach (var pi in IncludedModules)
+			{
+				if (pi is UEBuildModuleCPP IncludedModule)
+				{
+					if (IncludedModule.bHasUObjects)
+					{
+						continue;
+					}
+
+					if (!HandledModules.Add(IncludedModule))
+					{
+						continue;
+					}
+
+					OutList.Add(IncludedModule);
+
+					if (IncludedModule.PublicIncludePathModules != null)
+					{
+						CollectModulesOnlyIncluded(HandledModules, IncludedModule.PublicIncludePathModules, OutList);
+					}
+				}
+			}
+		}
+
+		static void CollectModulesOnlyIncluded(IEnumerable<UEBuildModuleCPP> ModulesToGenerateHeadersFor, List<UEBuildModuleCPP> OutList)
+		{
+			HashSet<UEBuildModuleCPP> HandledModules = new();
+			foreach (var Module in ModulesToGenerateHeadersFor)
+			{
+				HandledModules.Add(Module);
+			}
+
+			foreach (var Module in ModulesToGenerateHeadersFor)
+			{
+				if (Module.PublicIncludePathModules != null)
+				{
+					CollectModulesOnlyIncluded(HandledModules, Module.PublicIncludePathModules, OutList);
+				}
+				if (Module.PrivateIncludePathModules != null)
+				{
+					CollectModulesOnlyIncluded(HandledModules, Module.PrivateIncludePathModules, OutList);
+				}
+			}
+		}
+
 		static void SetupUObjectModule(UHTModuleInfo ModuleInfo, IReadOnlySet<string> ExcludedFolders, SourceFileMetadataCache MetadataCache, ThreadPoolWorkQueue Queue)
 		{
 			foreach (DirectoryReference ModuleDirectory in ModuleInfo.ModuleDirectories)
