UnrealBuildTool: Handle UBA exitcodes as crashes even if there is output, and retry them

#rnx

diff --git a/Engine/Source/Programs/UnrealBuildTool/Executors/UnrealBuildAccelerator/UBAExecutor.cs b/Engine/Source/Programs/UnrealBuildTool/Executors/UnrealBuildAccelerator/UBAExecutor.cs
--- a/Engine/Source/Programs/UnrealBuildTool/Executors/UnrealBuildAccelerator/UBAExecutor.cs
+++ b/Engine/Source/Programs/UnrealBuildTool/Executors/UnrealBuildAccelerator/UBAExecutor.cs
@@ -506,7 +506,7 @@ namespace UnrealBuildTool
 				using (IProcess process = _session!.RunProcess(startInfo, false, null, enableDetour))
 				{
 					Interlocked.Add(ref _localProcessedActions, 1);
-					if (process.ExitCode != 0 && UBAConfig.bForcedRetry)
+					if (process.ExitCode != 0 && UBAConfig.bForcedRetry || (process.ExitCode >= 9000 && process.ExitCode < 10000))
 					{
 						_threadedLogger.LogWarning("{Description} {StatusDescription}: Exited with error code {ExitCode}. This action will retry without UBA", action.CommandDescription, action.StatusDescription, process.ExitCode);
 						_forcedRetryActions.AddOrUpdate(action, false, (k, v) => false);
@@ -603,6 +603,11 @@ namespace UnrealBuildTool
 						RemoteActionFailedCrash(queue, action, e.ExitCode, e.ExecutingHost ?? "Unknown", "C1001");
 						return;
 					}
+					else if (e.ExitCode >= 9000 && e.ExitCode < 10000)
+					{
+						RemoteActionFailedCrash(queue, action, e.ExitCode, e.ExecutingHost ?? "Unknown", "UBA error");
+						return;
+					}
 
 					string additionalDescription = $"[RemoteExecutor: {e.ExecutingHost}]";
 					TimeSpan processorTime = e.TotalProcessorTime;
