[UBT] Added IExternalAction.bShouldOutputLog to prevent certain process to spam console output

#rnx

diff --git a/Engine/Source/Programs/UnrealBuildTool/Executors/ImmediateActionQueue.cs b/Engine/Source/Programs/UnrealBuildTool/Executors/ImmediateActionQueue.cs
--- a/Engine/Source/Programs/UnrealBuildTool/Executors/ImmediateActionQueue.cs
+++ b/Engine/Source/Programs/UnrealBuildTool/Executors/ImmediateActionQueue.cs
@@ -1160,7 +1160,7 @@ namespace UnrealBuildTool
 				s_previousLineLength = message.Length;
 
 				_writeToolOutput(message);
-				if (logLines != null)
+				if (logLines != null && action.bShouldOutputLog)
 				{
 					foreach (string Line in logLines.Skip(action.bShouldOutputStatusDescription ? 0 : 1))
 					{
@@ -1377,4 +1377,4 @@ namespace UnrealBuildTool
 			CancellationTokenSource.Cancel();
 		}
 	}
-}
+}
\ No newline at end of file
diff --git a/Engine/Source/Programs/UnrealBuildTool/Platform/Windows/VCCompileAction.cs b/Engine/Source/Programs/UnrealBuildTool/Platform/Windows/VCCompileAction.cs
--- a/Engine/Source/Programs/UnrealBuildTool/Platform/Windows/VCCompileAction.cs
+++ b/Engine/Source/Programs/UnrealBuildTool/Platform/Windows/VCCompileAction.cs
@@ -201,6 +201,9 @@ namespace UnrealBuildTool
 
 		/// <inheritdoc/>
 		public double Weight { get; set; } = 1.0;
+
+		/// <inheritdoc/>
+		public bool bShouldOutputLog { get; set; } = true;
 		#endregion
 
 		#region Implementation of IAction
diff --git a/Engine/Source/Programs/UnrealBuildTool/System/Action.cs b/Engine/Source/Programs/UnrealBuildTool/System/Action.cs
--- a/Engine/Source/Programs/UnrealBuildTool/System/Action.cs
+++ b/Engine/Source/Programs/UnrealBuildTool/System/Action.cs
@@ -166,6 +166,12 @@ namespace UnrealBuildTool
 		/// </summary>
 		bool bShouldOutputStatusDescription { get; }
 
+		/// <summary>
+		/// Whether we should output log of this action, whether executed locally or remotely.  This is useful for actions that spam
+		/// logs to console output with non critical data
+		/// </summary>
+		bool bShouldOutputLog { get; }
+
 		/// <summary>
 		/// True if any libraries produced by this action should be considered 'import libraries'
 		/// </summary>
@@ -294,6 +300,9 @@ namespace UnrealBuildTool
 		/// </summary>
 		public bool bShouldOutputStatusDescription { get; set; } = true;
 
+		/// <inheritdoc/>
+		public bool bShouldOutputLog { get; set; } = true;
+
 		/// <summary>
 		/// True if any libraries produced by this action should be considered 'import libraries'
 		/// </summary>
@@ -345,6 +354,7 @@ namespace UnrealBuildTool
 			bCanExecuteInUBA = InOther.bCanExecuteInUBA;
 			bIsGCCCompiler = InOther.bIsGCCCompiler;
 			bShouldOutputStatusDescription = InOther.bShouldOutputStatusDescription;
+			bShouldOutputLog = InOther.bShouldOutputLog;
 			bProducesImportLibrary = InOther.bProducesImportLibrary;
 			bUseActionHistory = InOther.bUseActionHistory;
 			bIsHighPriority = InOther.bIsHighPriority;
@@ -367,6 +377,7 @@ namespace UnrealBuildTool
 			bCanExecuteInUBA = Reader.ReadBool();
 			bIsGCCCompiler = Reader.ReadBool();
 			bShouldOutputStatusDescription = Reader.ReadBool();
+			bShouldOutputLog = Reader.ReadBool();
 			bProducesImportLibrary = Reader.ReadBool();
 			PrerequisiteItems = Reader.ReadSortedSet(() => Reader.ReadFileItem())!;
 			ProducedItems = Reader.ReadSortedSet(() => Reader.ReadFileItem())!;
@@ -398,6 +409,7 @@ namespace UnrealBuildTool
 			Writer.WriteBool(bCanExecuteInUBA);
 			Writer.WriteBool(bIsGCCCompiler);
 			Writer.WriteBool(bShouldOutputStatusDescription);
+			Writer.WriteBool(bShouldOutputLog);
 			Writer.WriteBool(bProducesImportLibrary);
 			Writer.WriteSortedSet(PrerequisiteItems, Item => Writer.WriteFileItem(Item));
 			Writer.WriteSortedSet(ProducedItems, Item => Writer.WriteFileItem(Item));
@@ -482,6 +494,11 @@ namespace UnrealBuildTool
 				Action.bShouldOutputStatusDescription = bShouldOutputStatusDescription;
 			}
 
+			if (Object.TryGetBoolField("bShouldOutputLog", out bool bShouldOutputLog)) 
+			{
+				Action.bShouldOutputLog= bShouldOutputLog;
+			}
+
 			if (Object.TryGetBoolField("bProducesImportLibrary", out bool bProducesImportLibrary))
 			{
 				Action.bProducesImportLibrary = bProducesImportLibrary;
@@ -580,6 +597,7 @@ namespace UnrealBuildTool
 			Writer.WriteValue("bCanExecuteInUBA", Action.bCanExecuteInUBA);
 			Writer.WriteValue("bIsGCCCompiler", Action.bIsGCCCompiler);
 			Writer.WriteValue("bShouldOutputStatusDescription", Action.bShouldOutputStatusDescription);
+			Writer.WriteValue("bShouldOutputLog", Action.bShouldOutputLog);
 			Writer.WriteValue("bProducesImportLibrary", Action.bProducesImportLibrary);
 			Writer.WriteValue("Weight", Action.Weight);
 
@@ -710,6 +728,7 @@ namespace UnrealBuildTool
 		public bool bCanExecuteInUBA => Inner.bCanExecuteInUBA;
 		public bool bIsGCCCompiler => Inner.bIsGCCCompiler;
 		public bool bShouldOutputStatusDescription => Inner.bShouldOutputStatusDescription;
+		public bool bShouldOutputLog => Inner.bShouldOutputLog;
 		public bool bProducesImportLibrary => Inner.bProducesImportLibrary;
 		public bool bUseActionHistory => Inner.bUseActionHistory;
 		public bool bIsHighPriority => IsHighPriority != 0;
