[UBT]
* Fixed so Classes/Public/Internal/Private includes are only added if there are files directly or indirectly under directory. Change is to make rsp files deterministic even if p4 left empty directories behind

diff --git a/Engine/Source/Programs/Shared/EpicGames.Build/System/DirectoryLookupCache.cs b/Engine/Source/Programs/Shared/EpicGames.Build/System/DirectoryLookupCache.cs
--- a/Engine/Source/Programs/Shared/EpicGames.Build/System/DirectoryLookupCache.cs
+++ b/Engine/Source/Programs/Shared/EpicGames.Build/System/DirectoryLookupCache.cs
@@ -20,6 +20,12 @@ namespace UnrealBuildBase
 			return DirectoryItem.GetItemByDirectoryReference(Directory).Exists;
 		}
 
+		static public bool DirectoryExistsAndContainsFiles(DirectoryReference Directory, SearchOption searchOption = SearchOption.TopDirectoryOnly)
+		{
+			DirectoryItem item = DirectoryItem.GetItemByDirectoryReference(Directory);
+			return item.Exists && item.ContainsFiles(searchOption);
+		}
+
 		static public IEnumerable<FileReference> EnumerateFiles(DirectoryReference Directory)
 		{
 			return DirectoryItem.GetItemByDirectoryReference(Directory).EnumerateFiles().Select(x => x.Location);
diff --git a/Engine/Source/Programs/UnrealBuildTool/Configuration/UEBuildModuleCPP.cs b/Engine/Source/Programs/UnrealBuildTool/Configuration/UEBuildModuleCPP.cs
--- a/Engine/Source/Programs/UnrealBuildTool/Configuration/UEBuildModuleCPP.cs
+++ b/Engine/Source/Programs/UnrealBuildTool/Configuration/UEBuildModuleCPP.cs
@@ -257,14 +257,14 @@ namespace UnrealBuildTool
 
 				// Add the 'classes' directory, if it exists
 				DirectoryReference ClassesDirectory = DirectoryReference.Combine(ModuleDir, "Classes");
-				if (DirectoryLookupCache.DirectoryExists(ClassesDirectory))
+				if (DirectoryLookupCache.DirectoryExistsAndContainsFiles(ClassesDirectory, SearchOption.AllDirectories))
 				{
 					PublicIncludePaths.Add(ClassesDirectory);
 				}
 
 				// Add all the public directories
 				DirectoryReference PublicDirectory = DirectoryReference.Combine(ModuleDir, "Public");
-				if (DirectoryLookupCache.DirectoryExists(PublicDirectory))
+				if (DirectoryLookupCache.DirectoryExistsAndContainsFiles(PublicDirectory, SearchOption.AllDirectories))
 				{
 					PublicIncludePaths.Add(PublicDirectory);
 
@@ -274,14 +274,14 @@ namespace UnrealBuildTool
 
 				// Add the 'internal' directory, if it exists
 				DirectoryReference InternalDirectory = DirectoryReference.Combine(ModuleDir, "Internal");
-				if (DirectoryLookupCache.DirectoryExists(InternalDirectory))
+				if (DirectoryLookupCache.DirectoryExistsAndContainsFiles(InternalDirectory, SearchOption.AllDirectories))
 				{
 					InternalIncludePaths.Add(InternalDirectory);
 				}
 
 				// Add the base private directory for this module
 				DirectoryReference PrivateDirectory = DirectoryReference.Combine(ModuleDir, "Private");
-				if (DirectoryLookupCache.DirectoryExists(PrivateDirectory))
+				if (DirectoryLookupCache.DirectoryExistsAndContainsFiles(PrivateDirectory, SearchOption.AllDirectories))
 				{
 					PrivateIncludePaths.Add(PrivateDirectory);
 				}
