[UBT]
* Added additional parameter to TypeLibrary logic because some type library compiles creates tli files

diff --git a/Engine/Plugins/Runtime/Database/ADOSupport/Source/ADOSupport/ADOSupport.Build.cs b/Engine/Plugins/Runtime/Database/ADOSupport/Source/ADOSupport/ADOSupport.Build.cs
--- a/Engine/Plugins/Runtime/Database/ADOSupport/Source/ADOSupport/ADOSupport.Build.cs
+++ b/Engine/Plugins/Runtime/Database/ADOSupport/Source/ADOSupport/ADOSupport.Build.cs
@@ -16,7 +16,7 @@ namespace UnrealBuildTool.Rules
 				Target.StaticAnalyzer != StaticAnalyzer.PVSStudio)
 			{
 				string MsAdo15 = Path.Combine(Environment.GetFolderPath(Environment.SpecialFolder.CommonProgramFiles), "System", "ADO", "msado15.dll");
-				TypeLibraries.Add(new TypeLibrary(MsAdo15, "rename(\"EOF\", \"ADOEOF\")", "msado15.tlh"));
+				TypeLibraries.Add(new TypeLibrary(MsAdo15, "rename(\"EOF\", \"ADOEOF\")", "msado15.tlh", "msado15.tli"));
 			}
 
 			PublicDependencyModuleNames.AddRange(
diff --git a/Engine/Source/Programs/UnrealBuildTool/Configuration/ModuleRules.cs b/Engine/Source/Programs/UnrealBuildTool/Configuration/ModuleRules.cs
--- a/Engine/Source/Programs/UnrealBuildTool/Configuration/ModuleRules.cs
+++ b/Engine/Source/Programs/UnrealBuildTool/Configuration/ModuleRules.cs
@@ -542,17 +542,24 @@ namespace UnrealBuildTool
 			/// </summary>
 			public string Header { get; init; }
 
+			/// <summary>
+			/// Name of the output include
+			/// </summary>
+			public string? Include { get; init; }
+
 			/// <summary>
 			/// Constructor
 			/// </summary>
 			/// <param name="fileName">Name of the type library. Follows the same conventions as the filename parameter in the MSVC #import directive.</param>
 			/// <param name="attributes">Additional attributes for the import directive</param>
 			/// <param name="header">Name of the output header</param>
-			public TypeLibrary(string fileName, string attributes, string header)
+			/// <param name="include">Name of the output include</param>
+			public TypeLibrary(string fileName, string attributes, string header, string? include = null)
 			{
 				FileName = fileName;
 				Attributes = attributes;
 				Header = header;
+				Include = include;
 			}
 		}
 
diff --git a/Engine/Source/Programs/UnrealBuildTool/Configuration/UEBuildModule.cs b/Engine/Source/Programs/UnrealBuildTool/Configuration/UEBuildModule.cs
--- a/Engine/Source/Programs/UnrealBuildTool/Configuration/UEBuildModule.cs
+++ b/Engine/Source/Programs/UnrealBuildTool/Configuration/UEBuildModule.cs
@@ -684,6 +684,10 @@ namespace UnrealBuildTool
 				foreach (ModuleRules.TypeLibrary TypeLibrary in Rules.TypeLibraries)
 				{
 					AdditionalPrerequisites.Add(FileItem.GetItemByFileReference(FileReference.Combine(IntermediateDirectory, TypeLibrary.Header)));
+					if (!String.IsNullOrEmpty(TypeLibrary.Include))
+					{
+						AdditionalPrerequisites.Add(FileItem.GetItemByFileReference(FileReference.Combine(IntermediateDirectory, TypeLibrary.Include)));
+					}
 				}
 			}
 		}
@@ -931,7 +935,8 @@ namespace UnrealBuildTool
 			foreach (ModuleRules.TypeLibrary TypeLibrary in Rules.TypeLibraries)
 			{
 				FileReference OutputFile = FileReference.Combine(IntermediateDirectory, TypeLibrary.Header);
-				ToolChain.GenerateTypeLibraryHeader(CompileEnvironment, TypeLibrary, OutputFile, Graph);
+				FileReference? OutputInclude = TypeLibrary.Include != null ? FileReference.Combine(IntermediateDirectory, TypeLibrary.Include) : null;
+				ToolChain.GenerateTypeLibraryHeader(CompileEnvironment, TypeLibrary, OutputFile, OutputInclude, Graph);
 			}
 
 			return new List<FileItem>();
diff --git a/Engine/Source/Programs/UnrealBuildTool/Platform/Windows/VCToolChain.cs b/Engine/Source/Programs/UnrealBuildTool/Platform/Windows/VCToolChain.cs
--- a/Engine/Source/Programs/UnrealBuildTool/Platform/Windows/VCToolChain.cs
+++ b/Engine/Source/Programs/UnrealBuildTool/Platform/Windows/VCToolChain.cs
@@ -2261,7 +2261,7 @@ namespace UnrealBuildTool
 			return Result;
 		}
 
-		public override void GenerateTypeLibraryHeader(CppCompileEnvironment CompileEnvironment, ModuleRules.TypeLibrary TypeLibrary, FileReference OutputFile, IActionGraphBuilder Graph)
+		public override void GenerateTypeLibraryHeader(CppCompileEnvironment CompileEnvironment, ModuleRules.TypeLibrary TypeLibrary, FileReference OutputFile, FileReference? OutputHeader, IActionGraphBuilder Graph)
 		{
 			// Create the input file
 			StringBuilder Contents = new StringBuilder();
@@ -2335,6 +2335,8 @@ namespace UnrealBuildTool
 			CompileAction.PrerequisiteItems.Add(BatchFile);
 			CompileAction.ProducedItems.Add(ObjectFile);
 			CompileAction.ProducedItems.Add(FileItem.GetItemByFileReference(OutputFile));
+			if (OutputHeader != null)
+				CompileAction.ProducedItems.Add(FileItem.GetItemByFileReference(OutputHeader));
 			CompileAction.DeleteItems.Add(FileItem.GetItemByFileReference(OutputFile));
 			CompileAction.StatusDescription = TypeLibrary.Header;
 			CompileAction.WorkingDirectory = Unreal.EngineSourceDirectory;
@@ -2343,6 +2345,7 @@ namespace UnrealBuildTool
 			CompileAction.CommandVersion = EnvVars.ToolChainVersion.ToString();
 			CompileAction.bShouldOutputStatusDescription = false;
 			CompileAction.bCanExecuteRemotely = false; // Incompatible with remote distribution
+			CompileAction.RootPaths.AddRange(GetEnvironmentBasePaths(CompileEnvironment));
 		}
 
 		public override IEnumerable<string> GetGlobalCommandLineArgs(CppCompileEnvironment CompileEnvironment)
diff --git a/Engine/Source/Programs/UnrealBuildTool/ToolChain/UEToolChain.cs b/Engine/Source/Programs/UnrealBuildTool/ToolChain/UEToolChain.cs
--- a/Engine/Source/Programs/UnrealBuildTool/ToolChain/UEToolChain.cs
+++ b/Engine/Source/Programs/UnrealBuildTool/ToolChain/UEToolChain.cs
@@ -107,7 +107,7 @@ namespace UnrealBuildTool
 			return Result;
 		}
 
-		public virtual void GenerateTypeLibraryHeader(CppCompileEnvironment CompileEnvironment, ModuleRules.TypeLibrary TypeLibrary, FileReference OutputFile, IActionGraphBuilder Graph)
+		public virtual void GenerateTypeLibraryHeader(CppCompileEnvironment CompileEnvironment, ModuleRules.TypeLibrary TypeLibrary, FileReference OutputFile, FileReference? OutputHeader, IActionGraphBuilder Graph)
 		{
 			throw new NotSupportedException("This platform does not support type libraries.");
 		}
