[UBT]
* Fixed so max artifact task count can be set from the outside
* Cosmetic fix for logging

diff --git a/Engine/Source/Programs/UnrealBuildTool/Executors/ImmediateActionQueue.cs b/Engine/Source/Programs/UnrealBuildTool/Executors/ImmediateActionQueue.cs
--- a/Engine/Source/Programs/UnrealBuildTool/Executors/ImmediateActionQueue.cs
+++ b/Engine/Source/Programs/UnrealBuildTool/Executors/ImmediateActionQueue.cs
@@ -365,7 +365,7 @@ namespace UnrealBuildTool
 		/// </summary>
 		private IActionArtifactCache? _actionArtifactCache;
 
-		static ExecuteResults s_copiedFromCacheResults = new(new List<string>(), 0, TimeSpan.Zero, TimeSpan.Zero, "copied from cache");
+		static ExecuteResults s_copiedFromCacheResults = new(new List<string>(), 0, TimeSpan.Zero, TimeSpan.Zero, "[Cache]");
 
 		/// <summary>
 		/// Construct a new instance of the action queue
@@ -992,7 +992,7 @@ namespace UnrealBuildTool
 		}
 
 		/// <summary>
-		/// Returns the number of queued actions left (not including ArtifactCheck actions)
+		/// Returns the number of queued actions left (including queued actions that will do artifact check first)
 		/// Note, this method is lockless and will not always return accurate count
 		/// </summary>
 		/// <param name="filterFunc">Optional function to filter out actions. Return false if action should not be included</param>
@@ -1003,7 +1003,7 @@ namespace UnrealBuildTool
 			for (int actionIndex = _firstPendingAction; actionIndex != Actions.Length; ++actionIndex)
 			{
 
-				if (Actions[actionIndex].Status != ActionStatus.Queued || Actions[actionIndex].Phase != ActionPhase.Compile)
+				if (Actions[actionIndex].Status != ActionStatus.Queued)
 				{
 					continue;
 				}
diff --git a/Engine/Source/Programs/UnrealBuildTool/Executors/ParallelExecutor.cs b/Engine/Source/Programs/UnrealBuildTool/Executors/ParallelExecutor.cs
--- a/Engine/Source/Programs/UnrealBuildTool/Executors/ParallelExecutor.cs
+++ b/Engine/Source/Programs/UnrealBuildTool/Executors/ParallelExecutor.cs
@@ -163,11 +163,12 @@ namespace UnrealBuildTool
 		/// </summary>
 		/// <param name="actionsToExecute">Actions to be executed</param>
 		/// <param name="actionArtifactCache">Artifact cache</param>
+		/// <param name="maxActionArtifactCacheTasks">Max artifact tasks that can execute in parallel</param>
 		/// <param name="logger">Logging interface</param>
 		/// <returns>Action queue</returns>
-		public ImmediateActionQueue CreateActionQueue(IEnumerable<LinkedAction> actionsToExecute, IActionArtifactCache? actionArtifactCache, ILogger logger)
+		public ImmediateActionQueue CreateActionQueue(IEnumerable<LinkedAction> actionsToExecute, IActionArtifactCache? actionArtifactCache, int maxActionArtifactCacheTasks, ILogger logger)
 		{
-			return new(actionsToExecute, actionArtifactCache, NumParallelProcesses, "Compiling C++ source code...", x => WriteToolOutput(x), () => FlushToolOutput(), logger)
+			return new(actionsToExecute, actionArtifactCache, maxActionArtifactCacheTasks, "Compiling C++ source code...", x => WriteToolOutput(x), () => FlushToolOutput(), logger)
 			{
 				ShowCompilationTimes = bShowCompilationTimes,
 				ShowCPUUtilization = bShowCPUUtilization,
@@ -194,7 +195,7 @@ namespace UnrealBuildTool
 			bool useAutomaticQueue = true;
 			if (useAutomaticQueue)
 			{
-				using ImmediateActionQueue queue = CreateActionQueue(ActionsToExecute, actionArtifactCache, Logger);
+				using ImmediateActionQueue queue = CreateActionQueue(ActionsToExecute, actionArtifactCache, NumParallelProcesses, Logger);
 				int actionLimit = Math.Min(NumParallelProcesses, queue.TotalActions);
 				queue.CreateAutomaticRunner(action => RunAction(queue, action), bUseActionWeights, actionLimit, NumParallelProcesses);
 				queue.Start();
@@ -206,7 +207,7 @@ namespace UnrealBuildTool
 			}
 			else
 			{
-				using ImmediateActionQueue queue = CreateActionQueue(ActionsToExecute, actionArtifactCache, Logger);
+				using ImmediateActionQueue queue = CreateActionQueue(ActionsToExecute, actionArtifactCache, NumParallelProcesses, Logger);
 				int actionLimit = Math.Min(NumParallelProcesses, queue.TotalActions);
 				ImmediateActionQueueRunner runner = queue.CreateManualRunner(action => RunAction(queue, action), bUseActionWeights, actionLimit, actionLimit);
 				queue.Start();
diff --git a/Engine/Source/Programs/UnrealBuildTool/Executors/UnrealBuildAccelerator/UBAExecutor.cs b/Engine/Source/Programs/UnrealBuildTool/Executors/UnrealBuildAccelerator/UBAExecutor.cs
--- a/Engine/Source/Programs/UnrealBuildTool/Executors/UnrealBuildAccelerator/UBAExecutor.cs
+++ b/Engine/Source/Programs/UnrealBuildTool/Executors/UnrealBuildAccelerator/UBAExecutor.cs
@@ -110,6 +110,7 @@ namespace UnrealBuildTool
 				await writer.WriteLineAsync($"arg: {action.CommandArguments}");
 				await writer.WriteLineAsync($"dir: {action.WorkingDirectory}");
 				await writer.WriteLineAsync($"desc: {action.StatusDescription}");
+				// TODO: Add cache roots
 				if (action.Weight != 1.0f)
 				{
 					await writer.WriteLineAsync($"weight: {action.Weight}");
@@ -337,7 +338,8 @@ namespace UnrealBuildTool
 		bool ExecuteActionsInternal(IEnumerable<LinkedAction> inputActions, ISessionServer session, Microsoft.Extensions.Logging.ILogger logger, IActionArtifactCache? actionArtifactCache, System.Action onCancel)
 		{
 			DateTime startTimeUTC = DateTime.UtcNow;
-			using ImmediateActionQueue queue = CreateActionQueue(inputActions, actionArtifactCache, logger);
+			int maxActionArtifactCacheTasks = 64;
+			using ImmediateActionQueue queue = CreateActionQueue(inputActions, actionArtifactCache, maxActionArtifactCacheTasks, logger);
 			int actionLimit = Math.Min(NumParallelProcesses, queue.TotalActions);
 			queue.CreateAutomaticRunner(action => RunActionLocal(queue, action), bUseActionWeights, actionLimit, NumParallelProcesses);
 			ImmediateActionQueueRunner remoteRunner = queue.CreateManualRunner(action => RunActionRemote(queue, action));
