[UBA]
* Integrated uba cache client in to UBT. This is the first iteration and is only really tested on UnrealEditor win64 development -nopch.

Test by starting a UbaCacheService.exe somewhere and then compile using -UbaCache=<thatmachine> -UbaWriteCache

diff --git a/Engine/Source/Programs/Shared/EpicGames.UBA/ICacheClient.cs b/Engine/Source/Programs/Shared/EpicGames.UBA/ICacheClient.cs
new file mode 100644
--- /dev/null
+++ b/Engine/Source/Programs/Shared/EpicGames.UBA/ICacheClient.cs
@@ -0,0 +1,74 @@
+// Copyright Epic Games, Inc. All Rights Reserved.
+
+namespace EpicGames.UBA
+{
+	/// <summary>
+	/// Base interface for root paths used by cache system to normalize paths
+	/// </summary>
+	public interface IRootPaths : IBaseInterface
+	{
+		/// <summary>
+		/// Register roots used to normalize paths in caches
+		/// </summary>
+		public abstract bool RegisterRoot(string path, bool includeInKey);
+
+		/// <summary>
+		/// Register system roots used to normalize paths in caches
+		/// </summary>
+		public abstract bool RegisterSystemRoots();
+
+		/// <summary>
+		/// Create root paths instance
+		/// </summary>
+		public static IRootPaths Create(ILogger logger)
+		{
+			return new RootPathsImpl(logger);
+		}
+	}
+
+	/// <summary>
+	/// Base interface for a cache client
+	/// </summary>
+	public interface ICacheClient : IBaseInterface
+	{
+		/// <summary>
+		/// Connect to cache client
+		/// </summary>
+		/// <param name="host">Cache server address</param>
+		/// <param name="port">Cache server port</param>
+		/// <returns>True if successful</returns>
+		public abstract bool Connect(string host, int port);
+
+		/// <summary>
+		/// Write to cache
+		/// </summary>
+		/// <param name="rootPaths">RootPath instance</param>
+		/// <param name="bucket">Bucket to store cache entry</param>
+		/// <param name="process">Process</param>
+		/// <param name="inputs">Input files</param>
+		/// <param name="inputsSize">Input files size</param>
+		/// <param name="outputs">Output files</param>
+		/// <param name="outputsSize">Output files size</param>
+		/// <returns>True if successful</returns>
+		public abstract bool WriteToCache(IRootPaths rootPaths, uint bucket, IProcess process, byte[] inputs, uint inputsSize, byte[] outputs, uint outputsSize);
+
+		/// <summary>
+		/// Fetch from cache
+		/// </summary>
+		/// <param name="rootPaths">RootPath instance</param>
+		/// <param name="bucket">Bucket to search for cache entry</param>
+		/// <param name="info">Process start info</param>
+		/// <returns>True if successful</returns>
+		public abstract bool FetchFromCache(IRootPaths rootPaths, uint bucket, ProcessStartInfo info);
+
+		/// <summary>
+		/// Create a ICacheClient object
+		/// </summary>
+		/// <param name="session">The session</param>
+		/// <returns>The ICacheClient</returns>
+		public static ICacheClient CreateCacheClient(ISessionServer session)
+		{
+			return new CacheClientImpl(session);
+		}
+	}
+}
diff --git a/Engine/Source/Programs/Shared/EpicGames.UBA/Impl/CacheClientImpl.cs b/Engine/Source/Programs/Shared/EpicGames.UBA/Impl/CacheClientImpl.cs
new file mode 100644
--- /dev/null
+++ b/Engine/Source/Programs/Shared/EpicGames.UBA/Impl/CacheClientImpl.cs
@@ -0,0 +1,143 @@
+// Copyright Epic Games, Inc. All Rights Reserved.
+
+using System;
+using System.Runtime.InteropServices;
+
+namespace EpicGames.UBA
+{
+	internal class RootPathsImpl : IRootPaths
+	{
+		IntPtr _handle = IntPtr.Zero;
+
+		#region DllImport
+		[DllImport("UbaHost", CharSet = CharSet.Auto)]
+		static extern IntPtr RootPaths_Create(IntPtr logWriter);
+
+		[DllImport("UbaHost", CharSet = CharSet.Auto)]
+		static extern void RootPaths_Destroy(IntPtr rootPaths);
+
+		[DllImport("UbaHost", CharSet = CharSet.Auto)]
+		static extern bool RootPaths_RegisterRoot(IntPtr rootPaths, string root, bool includeInKey);
+
+		[DllImport("UbaHost", CharSet = CharSet.Auto)]
+		static extern bool RootPaths_RegisterSystemRoots(IntPtr rootPaths);
+		#endregion
+
+		#region IDisposable
+		~RootPathsImpl() => Dispose(false);
+
+		public void Dispose()
+		{
+			Dispose(true);
+			GC.SuppressFinalize(this);
+		}
+
+		protected virtual void Dispose(bool disposing)
+		{
+			if (disposing)
+			{
+			}
+
+			if (_handle != IntPtr.Zero)
+			{
+				RootPaths_Destroy(_handle);
+				_handle = IntPtr.Zero;
+			}
+		}
+		#endregion
+
+		public RootPathsImpl(ILogger logger)
+		{
+			_handle = RootPaths_Create(logger.GetHandle());
+		}
+
+		public bool RegisterRoot(string path, bool includeInKey)
+		{
+			return RootPaths_RegisterRoot(_handle, path, includeInKey);
+		}
+
+		public bool RegisterSystemRoots()
+		{
+			return RootPaths_RegisterSystemRoots(_handle);
+		}
+
+		public IntPtr GetHandle() => _handle;
+	}
+
+	internal class CacheClientImpl : ICacheClient
+	{
+		IntPtr _handle = IntPtr.Zero;
+		public delegate void ExitCallback(IntPtr userData, IntPtr handle);
+
+		#region DllImport
+		[DllImport("UbaHost", CharSet = CharSet.Auto)]
+		static extern IntPtr ProcessStartInfo_Create(string application, string arguments, string workingDir, string description, uint priorityClass, ulong outputStatsThresholdMs, bool trackInputs, string logFile, ExitCallback? exit);
+
+		[DllImport("UbaHost", CharSet = CharSet.Auto)]
+		static extern void ProcessStartInfo_Destroy(IntPtr server);
+
+		[DllImport("UbaHost", CharSet = CharSet.Auto)]
+		static extern IntPtr CacheClient_Create(IntPtr session);
+
+		[DllImport("UbaHost", CharSet = CharSet.Auto)]
+		static extern bool CacheClient_Connect(IntPtr cacheClient, string host, int port);
+
+		[DllImport("UbaHost", CharSet = CharSet.Auto)]
+		static extern bool CacheClient_WriteToCache(IntPtr cacheClient, IntPtr rootPaths, uint bucket, IntPtr info, byte[] inputs, uint inputsSize, byte[] outputs, uint outputsSize);
+
+		[DllImport("UbaHost", CharSet = CharSet.Auto)]
+		static extern bool CacheClient_FetchFromCache(IntPtr cacheClient, IntPtr rootPaths, uint bucket, IntPtr info);
+
+		[DllImport("UbaHost", CharSet = CharSet.Auto)]
+		static extern void CacheClient_Destroy(IntPtr cacheClient);
+		#endregion
+
+		#region IDisposable
+		~CacheClientImpl() => Dispose(false);
+
+		public void Dispose()
+		{
+			Dispose(true);
+			GC.SuppressFinalize(this);
+		}
+
+		protected virtual void Dispose(bool disposing)
+		{
+			if (disposing)
+			{
+			}
+
+			if (_handle != IntPtr.Zero)
+			{
+				CacheClient_Destroy(_handle);
+				_handle = IntPtr.Zero;
+			}
+		}
+		#endregion
+
+		public CacheClientImpl(ISessionServer session)
+		{
+			_handle = CacheClient_Create(session.GetHandle());
+		}
+
+		public bool Connect(string host, int port)
+		{
+			return CacheClient_Connect(_handle, host, port);
+		}
+
+		public bool WriteToCache(IRootPaths rootPaths, uint bucket, IProcess process, byte[] inputs, uint inputsSize, byte[] outputs, uint outputsSize)
+		{
+			return CacheClient_WriteToCache(_handle, rootPaths.GetHandle(), bucket, process.GetHandle(), inputs, inputsSize, outputs, outputsSize);
+		}
+
+		public bool FetchFromCache(IRootPaths rootPaths, uint bucket, ProcessStartInfo info)
+		{
+			IntPtr si = ProcessStartInfo_Create(info.Application, info.Arguments, info.WorkingDirectory, info.Description, (uint)info.Priority, info.OutputStatsThresholdMs, info.TrackInputs, info.LogFile ?? String.Empty, null);
+			bool result = CacheClient_FetchFromCache(_handle, rootPaths.GetHandle(), bucket, si);
+			ProcessStartInfo_Destroy(si);
+			return result;
+		}
+
+		public IntPtr GetHandle() => _handle;
+	}
+}
diff --git a/Engine/Source/Programs/UnrealBuildTool/Executors/UnrealBuildAccelerator/UBAExecutor.cs b/Engine/Source/Programs/UnrealBuildTool/Executors/UnrealBuildAccelerator/UBAExecutor.cs
--- a/Engine/Source/Programs/UnrealBuildTool/Executors/UnrealBuildAccelerator/UBAExecutor.cs
+++ b/Engine/Source/Programs/UnrealBuildTool/Executors/UnrealBuildAccelerator/UBAExecutor.cs
@@ -3,6 +3,7 @@
 using System;
 using System.Collections.Concurrent;
 using System.Collections.Generic;
+using System.IO;
 using System.Linq;
 using System.Management;
 using System.Runtime.InteropServices;
@@ -24,6 +25,8 @@ namespace UnrealBuildTool
 		public string Crypto { get; private set; } = String.Empty;
 		public IServer? Server { get; private set; }
 		ISessionServer? _session;
+		ICacheClient? _cacheClient;
+		EpicGames.UBA.ILogger? _ubaLogger;
 		readonly List<IUBAAgentCoordinator> _agentCoordinators = new();
 		DirectoryReference? _rootDirRef;
 		bool _bIsCancelled;
@@ -143,6 +146,44 @@ namespace UnrealBuildTool
 			return BitConverter.ToString(bytes).Replace("-", "", StringComparison.OrdinalIgnoreCase).ToLowerInvariant(); // "1234567890abcdef1234567890abcdef";
 		}
 
+		class UBAArtifactCache : IArtifactCache
+		{
+			public ArtifactCacheState State { get => ArtifactCacheState.Available; }
+			public Task<ArtifactCacheState> WaitForReadyAsync() => Task.FromResult(ArtifactCacheState.Available);
+			public Task<ArtifactAction[]> QueryArtifactActionsAsync(IoHash[] partialKeys, CancellationToken cancellationToken) => Task.FromResult(Array.Empty<ArtifactAction>());
+			public Task<bool[]?> QueryArtifactOutputsAsync(ArtifactAction[] artifactActions, CancellationToken cancellationToken) => Task.FromResult<bool[]?>(null);
+			public Task SaveArtifactActionsAsync(ArtifactAction[] artifactActions, CancellationToken cancellationToken) => Task.CompletedTask;
+			public Task FlushChangesAsync(CancellationToken cancellationToken) => Task.CompletedTask;
+		}
+
+		class UBAActionArtifactCache : IActionArtifactCache
+		{
+			UBAExecutor _executor;
+			UBAArtifactCache _cache = new UBAArtifactCache();
+			public UBAActionArtifactCache(UBAExecutor executor) {  _executor = executor; }
+			public IArtifactCache ArtifactCache { get => _cache; }
+			public bool EnableReads { get => true; set { } }
+			public bool EnableWrites { get => true; set { } }
+			public bool LogCacheMisses { get => true; set { } }
+			public DirectoryReference? EngineRoot { get => null; set { } }
+			public DirectoryReference[]? DirectoryRoots { get => null; set { } }
+			public Task<bool> CompleteActionFromCacheAsync(LinkedAction action, CancellationToken cancellationToken)
+			{
+				return Task.Factory.StartNew(() =>
+				{
+					ProcessStartInfo startInfo = _executor.GetActionStartInfo(action, out FileItem? pchItem);
+					uint bucket = _executor.GetActionCacheBucket(action);
+					using (IRootPaths rootPaths = _executor.GetActionRootPaths(action))
+					{
+						return _executor._cacheClient!.FetchFromCache(rootPaths, bucket, startInfo);
+					}
+				}, cancellationToken, TaskCreationOptions.LongRunning | TaskCreationOptions.PreferFairness, TaskScheduler.Default);
+			}
+
+			public Task ActionCompleteAsync(LinkedAction action, CancellationToken cancellationToken) => Task.CompletedTask;
+			public Task FlushChangesAsync(CancellationToken cancellationToken) => Task.CompletedTask;
+		}
+
 		private async void ActionQueueCanceled(IStorageServer? ubaStorage)
 		{
 			_bIsCancelled = true;
@@ -247,10 +288,30 @@ namespace UnrealBuildTool
 				using EpicGames.UBA.ILogger ubaLogger = EpicGames.UBA.ILogger.CreateLogger(logger);
 				using (Server = IServer.CreateServer(UBAConfig.MaxWorkers, UBAConfig.SendSize, ubaLogger, UBAConfig.bUseQuic))
 				{
+					_ubaLogger = ubaLogger;
 					using IStorageServer ubaStorageServer = IStorageServer.CreateStorageServer(Server, ubaLogger, new StorageServerCreateInfo(_rootDirRef.FullName, ((ulong)UBAConfig.StoreCapacityGb) * 1000 * 1000 * 1000, !UBAConfig.bStoreRaw, UBAConfig.Zone));
-					using ISessionServerCreateInfo serverCreateInfo = ISessionServerCreateInfo.CreateSessionServerCreateInfo(ubaStorageServer, Server, ubaLogger, new SessionServerCreateInfo(_rootDirRef.FullName, ubaTraceFile.FullName, UBAConfig.bDisableCustomAlloc, false, UBAConfig.bResetCas, UBAConfig.bWriteToDisk, UBAConfig.bDetailedTrace, !UBAConfig.bDisableWaitOnMem, UBAConfig.bAllowKillOnMem, UBAConfig.bStoreObjFilesCompressed));
+					using ISessionServerCreateInfo serverCreateInfo = ISessionServerCreateInfo.CreateSessionServerCreateInfo(ubaStorageServer, Server, ubaLogger, new SessionServerCreateInfo(_rootDirRef.FullName, ubaTraceFile.FullName.Replace('\\', '/'), UBAConfig.bDisableCustomAlloc, false, UBAConfig.bResetCas, UBAConfig.bWriteToDisk, UBAConfig.bDetailedTrace, !UBAConfig.bDisableWaitOnMem, UBAConfig.bAllowKillOnMem, UBAConfig.bStoreObjFilesCompressed));
 					using (_session = ISessionServer.CreateSessionServer(serverCreateInfo))
+					using (_cacheClient = ICacheClient.CreateCacheClient(_session))
 					{
+						if (!String.IsNullOrEmpty(UBAConfig.CacheServer))
+						{
+							string[] nameAndPort = UBAConfig.CacheServer.Split(':');
+							int port = 1347;
+							if (nameAndPort.Length > 1)
+							{
+								port = int.Parse(nameAndPort[1]);
+							}
+
+							if (_cacheClient.Connect(nameAndPort[0], port))
+							{
+								actionArtifactCache = new UBAActionArtifactCache(this);
+							}
+							else
+							{
+								logger.LogInformation($"Timed out trying to connect to cache server on {nameAndPort[0]}:{port}. Cache will be disabled");
+							}
+						}
 
 						ubaStorage = ubaStorageServer;
 
@@ -338,25 +399,12 @@ namespace UnrealBuildTool
 		bool ExecuteActionsInternal(IEnumerable<LinkedAction> inputActions, ISessionServer session, Microsoft.Extensions.Logging.ILogger logger, IActionArtifactCache? actionArtifactCache, System.Action onCancel)
 		{
 			DateTime startTimeUTC = DateTime.UtcNow;
-			int maxActionArtifactCacheTasks = 64;
-			using ImmediateActionQueue queue = CreateActionQueue(inputActions, actionArtifactCache, maxActionArtifactCacheTasks, logger);
+			using ImmediateActionQueue queue = CreateActionQueue(inputActions, actionArtifactCache, UBAConfig.CacheMaxWorkers, logger);
 			int actionLimit = Math.Min(NumParallelProcesses, queue.TotalActions);
 			queue.CreateAutomaticRunner(action => RunActionLocal(queue, action), bUseActionWeights, actionLimit, NumParallelProcesses);
 			ImmediateActionQueueRunner remoteRunner = queue.CreateManualRunner(action => RunActionRemote(queue, action));
 			queue.CancellationToken.Register(onCancel);
 
-			// Setup a notification that alerts uba when an artifact has been read from the cache
-			queue.OnArtifactsRead = (action) =>
-			{
-				HashSet<DirectoryItem> refreshedDirectories = new();
-				foreach (FileItem output in action.ProducedItems)
-				{
-					if (refreshedDirectories.Add(output.Directory))
-					{
-						session.RefreshDirectories(output.Directory.FullName);
-					}
-				}
-			};
 
 			// Start the queue
 			queue.Start();
@@ -465,27 +513,190 @@ namespace UnrealBuildTool
 				Priority = ProcessPriority,
 				OutputStatsThresholdMs = (uint)UBAConfig.OutputStatsThresholdMs,
 				UserData = action,
-				Description = action.StatusDescription,
+				Description = $"{action.StatusDescription} ({action.CommandDescription})",
 				Configuration = action.bIsGCCCompiler ? EpicGames.UBA.ProcessStartInfo.CommonProcessConfigs.CompileClang : EpicGames.UBA.ProcessStartInfo.CommonProcessConfigs.CompileMsvc,
 				LogFile = UBAConfig.bLogEnabled ? action.Inner.ProducedItems.First().Location.GetFileName() : null,
 			};
 
-			bool usingLtcg = true; // ltcg linking checks what pch was used and it seems like it needs to be identical in other ways than timestamp
-			pchItem = action.Inner.ProducedItems.FirstOrDefault(item => item.Name.EndsWith(".pch", StringComparison.OrdinalIgnoreCase));
-			if (pchItem != null)
+			pchItem = null;
+
+			if (_cacheClient != null && UBAConfig.bWriteCache && action.ArtifactMode.HasFlag(ArtifactMode.Enabled))
+			{
+				startInfo.TrackInputs = true;
+			}
+			else
 			{
-				startInfo.Priority = System.Diagnostics.ProcessPriorityClass.AboveNormal;
-				if (!usingLtcg && action.ArtifactMode.HasFlag(ArtifactMode.PropagateInputs))
+				// TODO: Revisit this code. This was added to make non-deterministic pch deterministic from a caskey perspective
+				// It is not used atm.
+
+				bool usingLtcg = true; // ltcg linking checks what pch was used and it seems like it needs to be identical in other ways than timestamp
+				pchItem = action.Inner.ProducedItems.FirstOrDefault(item => item.Name.EndsWith(".pch", StringComparison.OrdinalIgnoreCase));
+				if (pchItem != null)
 				{
-					startInfo.TrackInputs = true;
+					startInfo.Priority = System.Diagnostics.ProcessPriorityClass.AboveNormal;
+					if (!usingLtcg && action.ArtifactMode.HasFlag(ArtifactMode.PropagateInputs))
+					{
+						startInfo.TrackInputs = true;
+					}
+					else
+					{
+						pchItem = null;
+					}
 				}
-				else
+			}
+			return startInfo;
+		}
+
+		uint GetActionCacheBucket(LinkedAction action)
+		{
+			if (action.Target == null)
+			{
+				return 0;
+			}
+			string platform = action.Target.Platform.ToString();
+			byte[] platformBytes = System.Text.Encoding.ASCII.GetBytes(platform);
+			return (uint)IoHash.Compute(platformBytes).GetHashCode();
+		}
+
+		IRootPaths GetActionRootPaths(LinkedAction action)
+		{
+			IRootPaths rootPaths = IRootPaths.Create(_ubaLogger!);
+			foreach (DirectoryItem root in action.RootPaths)
+			{
+				rootPaths.RegisterRoot(root.FullName + "\\", true);
+			}
+
+			DirectoryReference? autoSdkDir;
+			if (UEBuildPlatformSDK.TryGetHostPlatformAutoSDKDir(out autoSdkDir))
+			{
+				rootPaths.RegisterRoot(autoSdkDir.FullName + "\\", true);
+			}
+
+			rootPaths.RegisterSystemRoots();
+			return rootPaths;
+		}
+
+		public class DepsFile
+		{
+			public class DepsData
+			{
+				public string? Source { get; init; }
+				public string? PCH { get; init; }
+				public SortedSet<string>? Includes { get; init; }
+			}
+			public string? Version { get; init; }
+			public DepsData? Data { get; init; }
+		}
+
+		bool WriteToCache(LinkedAction action, IProcess process)
+		{
+			// Collect all inputs for action
+			// We use prerequisite items plus what we find in dependency list file if it exists.
+
+			using MemoryStream inputsMemory = new(1024);
+			using (BinaryWriter writer = new(inputsMemory, System.Text.Encoding.UTF8, true))
+			{
+				writer.Write(action.CommandPath.FullName);
+
+				foreach (FileItem f in action.PrerequisiteItems)
+				{
+					if (f.HasExtension(".lib")) // It seems like .lib files can change without dependencies relink
+					{
+						continue;
+					}
+					
+					writer.Write(f.FullName);
+				}
+
+				if (action.DependencyListFile != null)
+				{
+					if (action.DependencyListFile.HasExtension(".json"))
+					{
+						DepsFile? deps;
+						try
+						{
+							using (System.IO.FileStream fstream = FileReference.Open(action.DependencyListFile.Location, System.IO.FileMode.Open))
+							{
+								deps = System.Text.Json.JsonSerializer.Deserialize<DepsFile>(fstream);
+							}
+						}
+						catch (Exception e)
+						{
+							_threadedLogger.LogError("Unable to open/deserialize {Dep} - {Ex}", action.DependencyListFile, e.ToString());
+							return false;
+						}
+						if (deps == null || deps.Data == null || String.IsNullOrEmpty(deps.Data.Source))
+						{
+							_threadedLogger.LogError("Unable to deserialize {Dep}", action.DependencyListFile);
+							return false;
+						}
+
+						if (deps.Data.Includes != null)
+						{
+							foreach (string f in deps.Data.Includes)
+							{
+								writer.Write(f);
+							}
+						}
+					}
+					else if (action.DependencyListFile.HasExtension(".d"))
+					{
+						using (System.IO.FileStream fstream = FileReference.Open(action.DependencyListFile.Location, System.IO.FileMode.Open))
+						{
+							StreamReader reader = new(fstream);
+							reader.ReadLine(); // Skip first which is the .o file
+							while (true)
+							{
+								string? line = reader.ReadLine();
+								if (line == null)
+								{
+									break;
+								}
+								if (line.EndsWith(" \\"))
+								{
+									line = line.Substring(0, line.Length - 2);
+								}
+								line = line.Trim();
+								string path = DirectoryReference.Combine(action.WorkingDirectory, line).FullName;
+								writer.Write(path);
+							}
+						}
+					}
+					else
+					{
+						using (System.IO.FileStream fstream = FileReference.Open(action.DependencyListFile.Location, System.IO.FileMode.Open))
+						{
+							StreamReader reader = new(fstream);
+							while (true)
+							{
+								string? line = reader.ReadLine();
+								if (line == null)
+								{
+									break;
+								}
+								string path = DirectoryReference.Combine(action.WorkingDirectory, line.Replace('/', '\\').Replace("\\\\", "\\")).FullName;
+								writer.Write(path);
+							}
+						}
+					}
+				}
+			}
+
+			// Collect all outputs for action
+
+			using MemoryStream outputsMemory = new(1024);
+			using (BinaryWriter writer = new(outputsMemory, System.Text.Encoding.UTF8, true))
+			{
+				foreach (FileItem f in action.ProducedItems)
 				{
-					pchItem = null;
+					writer.Write(f.FullName);
 				}
 			}
 
-			return startInfo;
+			uint bucket = GetActionCacheBucket(action);
+			using IRootPaths rootPaths = GetActionRootPaths(action);
+
+			return _cacheClient!.WriteToCache(rootPaths, bucket, process, inputsMemory.GetBuffer(), (uint)inputsMemory.Position, outputsMemory.GetBuffer(), (uint)outputsMemory.Position);
 		}
 
 		Func<Task>? RunActionLocal(ImmediateActionQueue queue, LinkedAction action)
@@ -521,6 +732,11 @@ namespace UnrealBuildTool
 						_session!.RegisterNewFiles(action.ProducedItems.Where(x => FileReference.Exists(x.Location)).Select(x => x.FullName).ToArray());
 					}
 
+					if (startInfo.TrackInputs && UBAConfig.bWriteCache && enableDetour && process.ExitCode == 0 && _cacheClient != null)
+					{
+						WriteToCache(action, process);
+					}
+
 					TimeSpan processorTime = process.TotalProcessorTime;
 					TimeSpan executionTime = process.TotalWallTime;
 					List<string> logLines = process.LogLines;
@@ -611,12 +827,19 @@ namespace UnrealBuildTool
 						return;
 					}
 
+					IProcess process = (IProcess)s;
+
+					if (startInfo.TrackInputs && UBAConfig.bWriteCache && process.ExitCode == 0 && _cacheClient != null)
+					{
+						WriteToCache(action, process);
+					}
+
 					string additionalDescription = $"[RemoteExecutor: {e.ExecutingHost}]";
 					TimeSpan processorTime = e.TotalProcessorTime;
 					TimeSpan executionTime = e.TotalWallTime;
 					List<string> logLines = e.LogLines;
 					logLines.RemoveAll((line) => line.StartsWith("   Creating library ", StringComparison.OrdinalIgnoreCase) && line.EndsWith(".exp", StringComparison.OrdinalIgnoreCase));
-					ActionFinished(queue, new ExecuteResults(logLines, e.ExitCode, executionTime, processorTime, additionalDescription), action, pchItem, s as IProcess);
+					ActionFinished(queue, new ExecuteResults(logLines, e.ExitCode, executionTime, processorTime, additionalDescription), action, pchItem, process);
 				}, action.Weight, knownInputs, knownInputsCount);
 				return Task.CompletedTask;
 			};
diff --git a/Engine/Source/Programs/UnrealBuildTool/Executors/UnrealBuildAccelerator/UnrealBuildAcceleratorConfig.cs b/Engine/Source/Programs/UnrealBuildTool/Executors/UnrealBuildAccelerator/UnrealBuildAcceleratorConfig.cs
--- a/Engine/Source/Programs/UnrealBuildTool/Executors/UnrealBuildAccelerator/UnrealBuildAcceleratorConfig.cs
+++ b/Engine/Source/Programs/UnrealBuildTool/Executors/UnrealBuildAccelerator/UnrealBuildAcceleratorConfig.cs
@@ -220,5 +220,26 @@ namespace UnrealBuildTool
 		[XmlConfigFile(Category = "UnrealBuildAccelerator")]
 		[CommandLine("-UBADetailedLog", Value = "true")]
 		public bool bDetailedLog { get; set; } = false;
+
+		/// <summary>
+		/// Address of the uba cache service. Will automatically use cache if connected
+		/// </summary>
+		[XmlConfigFile(Category = "UnrealBuildAccelerator", Name = "Cache")]
+		[CommandLine("-UBACache=")]
+		public string? CacheServer { get; set; }
+
+		/// <summary>
+		/// Set cache to write instead of fetch
+		/// </summary>
+		[XmlConfigFile(Category = "UnrealBuildAccelerator", Name = "WriteCache")]
+		[CommandLine("-UBAWriteCache")]
+		public bool bWriteCache { get; set; }
+
+		/// <summary>
+		/// Max number of cache download tasks that can execute in parallel
+		/// </summary>
+		[XmlConfigFile(Category = "UnrealBuildAccelerator", Name = "CacheMaxWorkers")]
+		[CommandLine("-UBACacheMaxWorkers")]
+		public int CacheMaxWorkers { get; set; } = 32;
 	}
 }
\ No newline at end of file
diff --git a/Engine/Source/Programs/UnrealBuildTool/Platform/Windows/VCToolChain.cs b/Engine/Source/Programs/UnrealBuildTool/Platform/Windows/VCToolChain.cs
--- a/Engine/Source/Programs/UnrealBuildTool/Platform/Windows/VCToolChain.cs
+++ b/Engine/Source/Programs/UnrealBuildTool/Platform/Windows/VCToolChain.cs
@@ -269,6 +269,7 @@ namespace UnrealBuildTool
 			string UsingPchFilePath = NormalizeCommandLinePath(UsingPchFile);
 			if (UsingPchFile.Name.EndsWith(".ifc"))
 			{
+				Arguments.Add("/wd4127"); // with header units pragma warning disable is not propagated and this warning is very spammy so explicitly turned off here. Hoping ms will change so warning disable is propgated from HU
 				Arguments.Add($"/translateInclude");
 				Arguments.Add($"/headerUnit:quote {PchThroughHeaderFilePath}={UsingPchFilePath}");
 			}
@@ -1529,7 +1530,26 @@ namespace UnrealBuildTool
 		{
 			VCCompileAction BaseCompileAction = new VCCompileAction(EnvVars);
 
-			BaseCompileAction.RootPaths.AddRange(GetEnvironmentBasePaths(CompileEnvironment));
+			IEnumerable<DirectoryItem> rootPaths = GetEnvironmentBasePaths(CompileEnvironment);
+
+			// TODO: Revisit this code. We want to use d1trimfile to make outputs machine independent.
+			// but want to make sure we're not causing any frustration for devs (since __FILE__ will show a relative path with lines below)
+			// Also need to find the equivalent for clang
+#if false
+			if (!Target.WindowsPlatform.Compiler.IsClang())
+			{
+				foreach (DirectoryItem rootPath in rootPaths)
+				{
+					string pathName = rootPath.FullName;
+					if (pathName.Contains(' '))
+						BaseCompileAction.Arguments.Add($"\"/d1trimfile:{pathName}\\\"");
+					else
+						BaseCompileAction.Arguments.Add($"/d1trimfile:{pathName}\\");
+				}
+			}
+#endif
+
+			BaseCompileAction.RootPaths.AddRange(rootPaths);
 
 			// Add additional response files
 			foreach (FileItem AdditionalRsp in CompileEnvironment.AdditionalResponseFiles)
@@ -1748,11 +1768,10 @@ namespace UnrealBuildTool
 					}
 				}
 
-				if (CompileEnvironment.bDeterministic && !Target.WindowsPlatform.Compiler.IsClang())
 				{
 					if (CompileEnvironment.PrecompiledHeaderAction == PrecompiledHeaderAction.Create)
 					{
-						CompileAction.ArtifactMode |= ArtifactMode.PropagateInputs;
+						CompileAction.ArtifactMode |= ArtifactMode.Enabled | ArtifactMode.PropagateInputs;
 					}
 					else
 					{
@@ -2141,6 +2160,7 @@ namespace UnrealBuildTool
 				CompileAction.StatusDescription = Path.GetFileName(RCFile.AbsolutePath);
 				CompileAction.PrerequisiteItems.UnionWith(CompileEnvironment.ForceIncludeFiles);
 				CompileAction.PrerequisiteItems.UnionWith(CompileEnvironment.AdditionalPrerequisites);
+				CompileAction.ArtifactMode = ArtifactMode.Enabled;
 
 				// Resource tool can run remotely if possible
 				CompileAction.bCanExecuteRemotely = true;
@@ -2346,6 +2366,7 @@ namespace UnrealBuildTool
 			CompileAction.bShouldOutputStatusDescription = false;
 			CompileAction.bCanExecuteRemotely = false; // Incompatible with remote distribution
 			CompileAction.RootPaths.AddRange(GetEnvironmentBasePaths(CompileEnvironment));
+			CompileAction.ArtifactMode = ArtifactMode.Enabled;
 		}
 
 		public override IEnumerable<string> GetGlobalCommandLineArgs(CppCompileEnvironment CompileEnvironment)
@@ -2715,6 +2736,7 @@ namespace UnrealBuildTool
 			LinkAction.ProducedItems.UnionWith(ProducedItems);
 			LinkAction.PrerequisiteItems.UnionWith(PrerequisiteItems);
 			LinkAction.StatusDescription = Path.GetFileName(OutputFile.AbsolutePath);
+			LinkAction.ArtifactMode = ArtifactMode.Enabled;
 
 			// VS 15.3+ does not touch lib files if they do not contain any modifications, but we need to ensure the timestamps are updated to avoid repeatedly building them.
 			if (bBuildImportLibraryOnly || (LinkEnvironment.bHasExports && !bIsBuildingLibraryOrImportLibrary))
diff --git a/Engine/Source/Programs/UnrealBuildTool/ToolChain/ClangToolChain.cs b/Engine/Source/Programs/UnrealBuildTool/ToolChain/ClangToolChain.cs
--- a/Engine/Source/Programs/UnrealBuildTool/ToolChain/ClangToolChain.cs
+++ b/Engine/Source/Programs/UnrealBuildTool/ToolChain/ClangToolChain.cs
@@ -1195,6 +1195,8 @@ namespace UnrealBuildTool
 				CompileEnvironment.PrecompiledHeaderAction != PrecompiledHeaderAction.Create ||
 				CompileEnvironment.bAllowRemotelyCompiledPCHs;
 
+			CompileAction.ArtifactMode = ArtifactMode.Enabled;
+
 			// Two-pass compile where the preprocessor is run first to output the dependency list
 			if (PreprocessDepends)
 			{
diff --git a/Engine/Source/Programs/UnrealBuildTool/ToolChain/ISPCToolChain.cs b/Engine/Source/Programs/UnrealBuildTool/ToolChain/ISPCToolChain.cs
--- a/Engine/Source/Programs/UnrealBuildTool/ToolChain/ISPCToolChain.cs
+++ b/Engine/Source/Programs/UnrealBuildTool/ToolChain/ISPCToolChain.cs
@@ -666,6 +666,8 @@ namespace UnrealBuildTool
 				// TODO: Remove, might work
 				CompileAction.bCanExecuteRemotelyWithSNDBS = false;
 
+				CompileAction.ArtifactMode = ArtifactMode.Enabled;
+
 				List<string> Arguments = new List<string>();
 
 				// Add the ISPC file to be compiled.
