Add Xcode indexing capability to iOS workspaces

diff --git a/Engine/Source/Programs/UnrealBuildTool/ProjectFiles/ProjectFile.cs b/Engine/Source/Programs/UnrealBuildTool/ProjectFiles/ProjectFile.cs
--- a/Engine/Source/Programs/UnrealBuildTool/ProjectFiles/ProjectFile.cs
+++ b/Engine/Source/Programs/UnrealBuildTool/ProjectFiles/ProjectFile.cs
@@ -249,7 +249,7 @@ namespace UnrealBuildTool
 		/// </summary>
 		/// <param name="Module">The module to add</param>
 		/// <param name="CompileEnvironment">Compile environment for this module</param>
-		public virtual void AddModule(UEBuildModuleCPP Module, CppCompileEnvironment CompileEnvironment)
+		public virtual void AddModuleForIntelliSense(UEBuildModuleCPP Module, CppCompileEnvironment CompileEnvironment)
 		{
 			AddIntelliSensePreprocessorDefinitions(CompileEnvironment.Definitions);
 			AddIntelliSenseIncludePaths(SystemIncludePaths, CompileEnvironment.SystemIncludePaths);
diff --git a/Engine/Source/Programs/UnrealBuildTool/ProjectFiles/ProjectFileGenerator.cs b/Engine/Source/Programs/UnrealBuildTool/ProjectFiles/ProjectFileGenerator.cs
--- a/Engine/Source/Programs/UnrealBuildTool/ProjectFiles/ProjectFileGenerator.cs
+++ b/Engine/Source/Programs/UnrealBuildTool/ProjectFiles/ProjectFileGenerator.cs
@@ -475,6 +475,14 @@ namespace UnrealBuildTool
 			return ShouldGenerateIntelliSenseData();
 		}
 
+		/// <summary>
+		/// Platforms we should generate IntelliSense data for, default is only build host platform.
+		/// </summary>
+		protected virtual List<UnrealTargetPlatform> GetIntelliSensePlatforms()
+		{
+			return new() { BuildHostPlatform.Current.Platform };
+		}
+
 		/// <summary>
 		/// Allows each project generator to indicate whether target rules should be used to explicitly enable or disable plugins.
 		/// Default is false - since usually not needed for project generation unless project files indicate whether referenced plugins should be built or not.
@@ -2139,9 +2147,8 @@ namespace UnrealBuildTool
 						ProjectFile TargetProjectFile = Targets[TargetIndex].Item1;
 						ProjectTarget CurTarget = Targets[TargetIndex].Item2;
 
-						// Ignore projects for platforms we can't build on this host
-						UnrealTargetPlatform IntellisensePlatform = BuildHostPlatform.Current.Platform;
-						if (!CurTarget.SupportedPlatforms.Any(x => x == IntellisensePlatform))
+						// Ignore projects for platforms we can't build for
+						if (!CurTarget.SupportedPlatforms.Any(x => GetIntelliSensePlatforms().Contains(x)))
 						{
 							lock (Progress)
 							{
@@ -2162,47 +2169,53 @@ namespace UnrealBuildTool
 
 						try
 						{
-							// Get the architecture from the target platform
-							UnrealArchitectures DefaultArchitecture = UnrealArchitectureConfig.ForPlatform(IntellisensePlatform).ActiveArchitectures(CurTarget.UnrealProjectFilePath, CurTarget.Name);
+							foreach (UnrealTargetPlatform IntellisensePlatform in GetIntelliSensePlatforms())
+							{
+								if (!CurTarget.SupportedPlatforms.Contains(IntellisensePlatform))
+									continue;
 
-							// Create the target descriptor
-							TargetDescriptor TargetDesc = new TargetDescriptor(CurTarget.UnrealProjectFilePath, CurTarget.Name, IntellisensePlatform, UnrealTargetConfiguration.Development, DefaultArchitecture, new CommandLineArguments(NewArguments.ToArray()));
-							TargetDesc.IntermediateEnvironment = UnrealIntermediateEnvironment.GenerateProjectFiles;
+								// Get the architecture from the target platform
+								UnrealArchitectures DefaultArchitecture = UnrealArchitectureConfig.ForPlatform(IntellisensePlatform).ActiveArchitectures(CurTarget.UnrealProjectFilePath, CurTarget.Name);
 
-							// Create the target
-							UEBuildTarget Target = UEBuildTarget.Create(TargetDesc, false, false, bUsePrecompiled, UnrealIntermediateEnvironment.GenerateProjectFiles, Logger);
+								// Create the target descriptor
+								TargetDescriptor TargetDesc = new TargetDescriptor(CurTarget.UnrealProjectFilePath, CurTarget.Name, IntellisensePlatform, UnrealTargetConfiguration.Development, DefaultArchitecture, new CommandLineArguments(NewArguments.ToArray()));
+								TargetDesc.IntermediateEnvironment = UnrealIntermediateEnvironment.GenerateProjectFiles;
 
-							AddTargetForIntellisense(Target, Logger);
+								// Create the target
+								UEBuildTarget Target = UEBuildTarget.Create(TargetDesc, false, false, bUsePrecompiled, UnrealIntermediateEnvironment.GenerateProjectFiles, Logger);
 
-							// If the project generator just cares about the result of UEBuildTarget.Create, skip generating the compile environments.
-							if (ShouldGenerateIntelliSenseCompileEnvironments())
-							{
-								// Generate a compile environment for each module in the binary
-								CppCompileEnvironment GlobalCompileEnvironment = Target.CreateCompileEnvironmentForProjectFiles(Logger);
-								foreach (UEBuildBinary Binary in Target.Binaries)
+								AddTargetForIntellisense(Target, Logger);
+
+								// If the project generator just cares about the result of UEBuildTarget.Create, skip generating the compile environments.
+								if (ShouldGenerateIntelliSenseCompileEnvironments())
 								{
-									CppCompileEnvironment BinaryCompileEnvironment = Binary.CreateBinaryCompileEnvironment(GlobalCompileEnvironment);
-									foreach (UEBuildModuleCPP Module in Binary.Modules.OfType<UEBuildModuleCPP>())
+									// Generate a compile environment for each module in the binary
+									CppCompileEnvironment GlobalCompileEnvironment = Target.CreateCompileEnvironmentForProjectFiles(Logger);
+									foreach (UEBuildBinary Binary in Target.Binaries)
 									{
-										ProjectFile? ProjectFileForIDE;
-										if (ModuleToEditorProjectFileMap.TryGetValue(Module.RulesFile, out ProjectFileForIDE) && ProjectFileForIDE == TargetProjectFile)
+										CppCompileEnvironment BinaryCompileEnvironment = Binary.CreateBinaryCompileEnvironment(GlobalCompileEnvironment);
+										foreach (UEBuildModuleCPP Module in Binary.Modules.OfType<UEBuildModuleCPP>())
 										{
-											CppCompileEnvironment ModuleCompileEnvironment = Module.CreateCompileEnvironmentForIntellisense(Target.Rules, BinaryCompileEnvironment, Logger);
-											lock (ProjectFileForIDE)
+											ProjectFile? ProjectFileForIDE;
+											if (ModuleToEditorProjectFileMap.TryGetValue(Module.RulesFile, out ProjectFileForIDE) && ProjectFileForIDE == TargetProjectFile)
 											{
-												ProjectFileForIDE.AddModule(Module, ModuleCompileEnvironment);
+												CppCompileEnvironment ModuleCompileEnvironment = Module.CreateCompileEnvironmentForIntellisense(Target.Rules, BinaryCompileEnvironment, Logger);
+												lock (ProjectFileForIDE)
+												{
+													ProjectFileForIDE.AddModuleForIntelliSense(Module, ModuleCompileEnvironment);
+												}
 											}
 										}
 									}
-								}
 
-								// If we're generating project files, then go ahead and wipe out the existing UBTMakefile for every target, to make sure that
-								// it gets a full dependency scan next time.
-								// NOTE: This is just a safeguard and doesn't have to be perfect.  We also check for newer project file timestamps in LoadUBTMakefile()
-								FileReference MakefileLocation = TargetMakefile.GetLocation(TargetDesc.ProjectFile, TargetDesc.Name, TargetDesc.Platform, TargetDesc.Architectures, TargetDesc.Configuration, TargetDesc.IntermediateEnvironment);
-								if (FileReference.Exists(MakefileLocation))
-								{
-									FileReference.Delete(MakefileLocation);
+									// If we're generating project files, then go ahead and wipe out the existing UBTMakefile for every target, to make sure that
+									// it gets a full dependency scan next time.
+									// NOTE: This is just a safeguard and doesn't have to be perfect.  We also check for newer project file timestamps in LoadUBTMakefile()
+									FileReference MakefileLocation = TargetMakefile.GetLocation(TargetDesc.ProjectFile, TargetDesc.Name, TargetDesc.Platform, TargetDesc.Architectures, TargetDesc.Configuration, TargetDesc.IntermediateEnvironment);
+									if (FileReference.Exists(MakefileLocation))
+									{
+										FileReference.Delete(MakefileLocation);
+									}
 								}
 							}
 						}
@@ -2438,8 +2451,8 @@ namespace UnrealBuildTool
 					else
 					{
 						Debug.Assert(bAllowMultiModuleReference, "ProjectFileGenerator assert", $"Unexpected multi projects for module {CurModuleFile.GetFileName()}");
-						// e.g. QAGame module would be add to both QAGame.xcodeproj and QAGameEditor.xcodeproj, use the editor one for module map
-						ProjectFile? EditorProjectFile = ProjectFiles.FirstOrDefault(x => x.ProjectTargets.Any(x => x.TargetRules!.Type == TargetType.Editor));
+						// e.g. QAGame module would be add to both QAGame.xcodeproj and QAGameEditor.xcodeproj, use the game project for module map, this way IOS workspace can have indexing
+						ProjectFile? EditorProjectFile = ProjectFiles.FirstOrDefault(x => x.ProjectTargets.Any(x => x.TargetRules!.Type == TargetType.Game));
 						if (EditorProjectFile != null)
 						{
 							ModuleToEditorProjectFileMap[CurModuleFile] = EditorProjectFile;
diff --git a/Engine/Source/Programs/UnrealBuildTool/ProjectFiles/VisualStudio/VCProject.cs b/Engine/Source/Programs/UnrealBuildTool/ProjectFiles/VisualStudio/VCProject.cs
--- a/Engine/Source/Programs/UnrealBuildTool/ProjectFiles/VisualStudio/VCProject.cs
+++ b/Engine/Source/Programs/UnrealBuildTool/ProjectFiles/VisualStudio/VCProject.cs
@@ -495,9 +495,9 @@ namespace UnrealBuildTool
 		}
 
 		/// <inheritdoc/>
-		public override void AddModule(UEBuildModuleCPP Module, CppCompileEnvironment CompileEnvironment)
+		public override void AddModuleForIntelliSense(UEBuildModuleCPP Module, CppCompileEnvironment CompileEnvironment)
 		{
-			base.AddModule(Module, CompileEnvironment);
+			base.AddModuleForIntelliSense(Module, CompileEnvironment);
 
 			if (Settings.bUsePerFileIntellisense)
 			{
diff --git a/Engine/Source/Programs/UnrealBuildTool/ProjectFiles/Xcode/XcodeProject.File.cs b/Engine/Source/Programs/UnrealBuildTool/ProjectFiles/Xcode/XcodeProject.File.cs
--- a/Engine/Source/Programs/UnrealBuildTool/ProjectFiles/Xcode/XcodeProject.File.cs
+++ b/Engine/Source/Programs/UnrealBuildTool/ProjectFiles/Xcode/XcodeProject.File.cs
@@ -183,7 +183,7 @@ namespace UnrealBuildTool.XcodeProjectXcconfig
 		//public string ProductRefGroupGuid = XcodeProjectFileGenerator.MakeXcodeGuid();
 
 		internal Dictionary<string, XcodeFileGroup> Groups = new();
-		internal Dictionary<XcodeSourceFile, FileReference?> BuildableFilesToResponseFile = new();
+		internal Dictionary<XcodeSourceFile, Dictionary<UnrealTargetPlatform, FileReference>> BuildableFilesToResponseFile = new();
 		internal List<Tuple<XcodeSourceFile, string>> BuildableResourceFiles = new();
 		internal List<XcodeSourceFile> AllFiles = new();
 
@@ -261,7 +261,7 @@ namespace UnrealBuildTool.XcodeProjectXcconfig
 						if (!ExcludedFolders.Any(x => File.Reference.ContainsName(x, BuildFileOffset)))
 						{
 							// will fill in the response file later
-							BuildableFilesToResponseFile[File] = null;
+							BuildableFilesToResponseFile[File] = new();
 						}
 					}
 				}
@@ -342,16 +342,30 @@ namespace UnrealBuildTool.XcodeProjectXcconfig
 			}
 		}
 
-		public void Write(StringBuilder Content)
+		public void Write(StringBuilder Content, UnrealTargetPlatform? Platform)
 		{
 			Content.WriteLine("/* Begin PBXBuildFile section */");
-			foreach (KeyValuePair<XcodeSourceFile, FileReference?> Pair in BuildableFilesToResponseFile)
+			if (Platform == null)
+			{
+				// Shared platform project should only index for Mac
+				Platform = UnrealTargetPlatform.Mac;
+			}
+			foreach (KeyValuePair<XcodeSourceFile, Dictionary<UnrealTargetPlatform, FileReference>> Pair in BuildableFilesToResponseFile)
 			{
 				string CompileFlags = "";
-				if (Pair.Value != null)
+				if (Pair.Value.Count() > 0)
 				{
-					CompileFlags = $" settings = {{COMPILER_FLAGS = \"@{Pair.Value.FullName}\"; }};";
+					if (Pair.Value.ContainsKey(Platform.Value))
+					{
+						CompileFlags = $" settings = {{COMPILER_FLAGS = \"@{Pair.Value[Platform.Value].FullName}\"; }};";
+					}
+					else if (Pair.Value.ContainsKey(UnrealTargetPlatform.Mac))
+					{
+						// We don't have the correct response file for this platform, fall back to use the Mac version
+						CompileFlags = $" settings = {{COMPILER_FLAGS = \"@{Pair.Value[UnrealTargetPlatform.Mac].FullName}\"; }};";
+					}
 				}
+				
 				Content.WriteLine($"\t\t{Pair.Key.FileGuid} = {{isa = PBXBuildFile; fileRef = {Pair.Key.FileRefGuid};{CompileFlags} }}; /* {Pair.Key.Reference.GetFileName()} */");
 			}
 			foreach (Tuple<XcodeSourceFile, string> Resource in BuildableResourceFiles)
diff --git a/Engine/Source/Programs/UnrealBuildTool/ProjectFiles/Xcode/XcodeProject.cs b/Engine/Source/Programs/UnrealBuildTool/ProjectFiles/Xcode/XcodeProject.cs
--- a/Engine/Source/Programs/UnrealBuildTool/ProjectFiles/Xcode/XcodeProject.cs
+++ b/Engine/Source/Programs/UnrealBuildTool/ProjectFiles/Xcode/XcodeProject.cs
@@ -15,7 +15,7 @@ using UnrealBuildBase;
 Here's how this works:
 
   * An XcodeProjectFile (subclass of generic ProjectFile class) is created, along with it - UnrealData and XcodeFileCollection objects are made
-  * High level code calls AddModule() which this code will use to cache information about the Modules in the project (including build settings, etc)
+  * High level code calls AddModuleForIntelliSense() which this code will use to cache information about the Modules in the project (including build settings, etc)
     * These are used to determine what source files can be indexed together (we use native xcode code compilation for indexing, so we make compiling succesful for best index)
     * A few #defines are removed or modified (FOO_API, etc) which would otherwise make every module a separate target
   * High level code then calls WriteProjectFile() which is the meat of all this
@@ -101,10 +101,10 @@ namespace UnrealBuildTool.XcodeProjectXcconfig
 		public UnrealBatchedFiles(UnrealData UnrealData, int Index, UEBuildModuleCPP Module)
 		{
 			this.Module = Module;
-			ResponseFile = FileReference.Combine(UnrealData.XcodeProjectFileLocation.ParentDirectory!, "ResponseFiles", $"{UnrealData.ProductName}{Index}.response");
+			ResponseFile = FileReference.Combine(UnrealData.XcodeProjectFileLocation.ParentDirectory!, "ResponseFiles", Module.Rules.Target.Platform.ToString(), $"{UnrealData.ProductName}{Index}.response");
 		}
 
-		public void GenerateResponseFile()
+		public void GenerateResponseFile(UnrealTargetPlatform Platform)
 		{
 			StringBuilder ResponseFileContents = new();
 			ResponseFileContents.Append("-isystem");
@@ -341,7 +341,7 @@ namespace UnrealBuildTool.XcodeProjectXcconfig
 
 		public List<UnrealBuildConfig> AllConfigs = new();
 
-		public List<UnrealBatchedFiles> BatchedFiles = new();
+		public Dictionary<UnrealTargetPlatform, List<UnrealBatchedFiles>> BatchedFiles = new();
 
 		public List<string> ExtraPreBuildScriptLines = new();
 
@@ -630,11 +630,15 @@ namespace UnrealBuildTool.XcodeProjectXcconfig
 			return FinalPath;
 		}
 
-		public void AddModule(UEBuildModuleCPP Module, CppCompileEnvironment CompileEnvironment)
+		public void AddModuleForIntelliSense(UEBuildModuleCPP Module, CppCompileEnvironment CompileEnvironment)
 		{
-			// one batched files per module
-			UnrealBatchedFiles FileBatch = new UnrealBatchedFiles(this, BatchedFiles.Count + 1, Module);
-			BatchedFiles.Add(FileBatch);
+			// one batched files per module per platform
+			if (!BatchedFiles.ContainsKey(Module.Rules.Target.Platform))
+			{
+				BatchedFiles[Module.Rules.Target.Platform] = new();
+			}
+			UnrealBatchedFiles FileBatch = new UnrealBatchedFiles(this, BatchedFiles[Module.Rules.Target.Platform].Count + 1, Module);
+			BatchedFiles[Module.Rules.Target.Platform].Add(FileBatch);
 
 			if (CompileEnvironment.ForceIncludeFiles.Count == 0)
 			{
@@ -2152,7 +2156,7 @@ namespace UnrealBuildTool.XcodeProjectXcconfig
 		private UnrealData UnrealData;
 
 		// just take the Project since it has everything we need, and is needed when adding target dependencies
-		public XcodeIndexTarget(XcodeProject Project)
+		public XcodeIndexTarget(XcodeProject Project, UnrealTargetPlatform Platform)
 			: base(XcodeTarget.Type.Index, Project.UnrealData)
 		{
 			UnrealData = Project.UnrealData;
@@ -2169,10 +2173,10 @@ namespace UnrealBuildTool.XcodeProjectXcconfig
 			BuildPhases.Add(SourcesBuildPhase);
 			References.Add(SourcesBuildPhase);
 
-			foreach (KeyValuePair<XcodeSourceFile, FileReference?> Pair in Project.FileCollection.BuildableFilesToResponseFile)
+			foreach (KeyValuePair<XcodeSourceFile, Dictionary<UnrealTargetPlatform, FileReference>> Pair in Project.FileCollection.BuildableFilesToResponseFile)
 			{
 				// only add files that found a moduleto be part of (since we can't build without the build settings that come from a module)
-				if (Pair.Value != null)
+				if (Pair.Value.ContainsKey(Platform))
 				{
 					SourcesBuildPhase.AddFile(Pair.Key);
 				}
@@ -2244,11 +2248,26 @@ namespace UnrealBuildTool.XcodeProjectXcconfig
 			References.Add(ProjectBuildConfigs);
 
 			// make an indexing target if we aren't just a run-only project, and it has buildable source files
-			if (!XcodeProjectFileGenerator.bGenerateRunOnlyProject && UnrealData.BatchedFiles.Count != 0)
+			if (!XcodeProjectFileGenerator.bGenerateRunOnlyProject)
 			{
-				// index isn't a dependency of run, it's simply a target that xcode will find to index from
-				XcodeIndexTarget IndexTarget = new XcodeIndexTarget(this);
-				References.Add(IndexTarget);
+				if (Platform != null)
+				{
+					// Per platform project, check if we have files to index for that platform
+					if (UnrealData.BatchedFiles.ContainsKey(Platform.Value) && UnrealData.BatchedFiles[Platform.Value].Count() > 0)
+					{
+						XcodeIndexTarget IndexTarget = new XcodeIndexTarget(this, Platform.Value);
+						References.Add(IndexTarget);
+					}
+				}
+				else
+				{
+					// Shared platform project, just try to index Mac target
+					if (UnrealData.BatchedFiles[UnrealTargetPlatform.Mac].Count() > 0)
+					{
+						XcodeIndexTarget IndexTarget = new XcodeIndexTarget(this, UnrealTargetPlatform.Mac);
+						References.Add(IndexTarget);
+					}
+				}
 			}
 		}
 
@@ -2363,7 +2382,7 @@ namespace UnrealBuildTool.XcodeProjectXcconfig
 				bHasCheckedForLegacy = true;
 				if (ProjectTargets.Count == 0)
 				{
-					throw new BuildException("Expected to have a target before AddModule is called");
+					throw new BuildException("Expected to have a target before AddModuleForIntelliSense is called");
 				}
 
 				UnrealData.InitializeUProjectFileLocation(this);
@@ -2379,17 +2398,17 @@ namespace UnrealBuildTool.XcodeProjectXcconfig
 			}
 		}
 
-		public override void AddModule(UEBuildModuleCPP Module, CppCompileEnvironment CompileEnvironment)
+		public override void AddModuleForIntelliSense(UEBuildModuleCPP Module, CppCompileEnvironment CompileEnvironment)
 		{
 			ConditionalCreateLegacyProject();
 
-			if (LegacyProjectFile != null)
+			if (LegacyProjectFile != null)// && Module.Rules.Target.Platform == UnrealTargetPlatform.Mac)
 			{
-				LegacyProjectFile.AddModule(Module, CompileEnvironment);
+				LegacyProjectFile.AddModuleForIntelliSense(Module, CompileEnvironment);
 				return;
 			}
 
-			UnrealData.AddModule(Module, CompileEnvironment);
+			UnrealData.AddModuleForIntelliSense(Module, CompileEnvironment);
 		}
 
 		/// <summary>
@@ -2416,21 +2435,27 @@ namespace UnrealBuildTool.XcodeProjectXcconfig
 			}
 
 			// write out the response files for each batch now that everything is done
-			foreach (UnrealBatchedFiles Batch in UnrealData.BatchedFiles)
+			foreach (KeyValuePair<UnrealTargetPlatform, List<UnrealBatchedFiles>> Pair in UnrealData.BatchedFiles)
 			{
-				Batch.GenerateResponseFile();
+				foreach(UnrealBatchedFiles Batch in Pair.Value)
+				{
+					Batch.GenerateResponseFile(Pair.Key);
+				}
 			}
 		}
 
 		private void AddFileToBatch(XcodeSourceFile File, XcodeFileCollection FileCollection)
 		{
-			foreach (UnrealBatchedFiles Batch in UnrealData.BatchedFiles)
+			foreach (KeyValuePair<UnrealTargetPlatform, List<UnrealBatchedFiles>> Pair in UnrealData.BatchedFiles)
 			{
-				if (Batch.Module.ContainsFile(File.Reference))
+				foreach (UnrealBatchedFiles Batch in Pair.Value)
 				{
-					Batch.Files.Add(File);
-					FileCollection.BuildableFilesToResponseFile[File] = Batch.ResponseFile;
-					return;
+					if (Batch.Module.ContainsFile(File.Reference))
+					{
+						Batch.Files.Add(File);
+						FileCollection.BuildableFilesToResponseFile[File][Pair.Key] = Batch.ResponseFile;
+						break;
+					}
 				}
 			}
 		}
@@ -2517,7 +2542,7 @@ namespace UnrealBuildTool.XcodeProjectXcconfig
 					Content.WriteLine(1, "objects = {");
 
 					// write out the list of files and groups
-					FileCollection.Write(Content);
+					FileCollection.Write(Content, Platform);
 
 					// now write out the project node and its recursive dependent nodes
 					XcodeProjectNode.WriteNodeAndReferences(Content, RootProject, Logger);
@@ -2602,7 +2627,8 @@ namespace UnrealBuildTool.XcodeProjectXcconfig
 
 			StringBuilder Content = new StringBuilder();
 			Content.WriteLine(0, "{");
-			FileCollection.Write(Content);
+			// @todo figure out how to deal with platform
+			// FileCollection.Write(Content, Platform);
 			XcodeProjectNode.WriteNodeAndReferences(Content, BuildTarget, Logger);
 			XcodeProjectNode.WriteNodeAndReferences(Content, IndexTarget, Logger);
 			XcodeProjectNode.WriteNodeAndReferences(Content, BuildDependency, Logger);
diff --git a/Engine/Source/Programs/UnrealBuildTool/ProjectFiles/Xcode/XcodeProjectFileGenerator.cs b/Engine/Source/Programs/UnrealBuildTool/ProjectFiles/Xcode/XcodeProjectFileGenerator.cs
--- a/Engine/Source/Programs/UnrealBuildTool/ProjectFiles/Xcode/XcodeProjectFileGenerator.cs
+++ b/Engine/Source/Programs/UnrealBuildTool/ProjectFiles/Xcode/XcodeProjectFileGenerator.cs
@@ -80,6 +80,24 @@ namespace UnrealBuildTool
 		/// </summary>
 		public static bool bGenerateRunOnlyProject = false;
 
+		/// <summary>
+		/// Platforms we should generate IntelliSense data for, Xcode projects should index for each Apple platform.
+		/// </summary>
+		protected override List<UnrealTargetPlatform> GetIntelliSensePlatforms()
+		{
+			if (bMakeProjectPerTarget)
+			{
+				// For now just allow indexing to work on Mac and IOS, enabling too many will make GPF and workspaces even slower
+				return new() { UnrealTargetPlatform.Mac, UnrealTargetPlatform.IOS };
+			}
+			else
+			{
+				// If sharing projects across platform targets, we are unable to index separately (due to generated headers being platform dependent)
+				// Choose Mac to be the platform for indexing
+				return new() { UnrealTargetPlatform.Mac };
+			}
+		}
+
 		public XcodeProjectFileGenerator(FileReference? InOnlyGameProject, CommandLineArguments CommandLine)
 			: base(InOnlyGameProject)
 		{
