Add WITH_CLIENT_CODE macro similar to WITH_SERVER_CODE.
#rnx

diff --git a/Engine/Source/Programs/UnrealBuildTool/Configuration/ReadOnlyTargetRules.cs b/Engine/Source/Programs/UnrealBuildTool/Configuration/ReadOnlyTargetRules.cs
--- a/Engine/Source/Programs/UnrealBuildTool/Configuration/ReadOnlyTargetRules.cs
+++ b/Engine/Source/Programs/UnrealBuildTool/Configuration/ReadOnlyTargetRules.cs
@@ -225,6 +225,8 @@ namespace UnrealBuildTool
 
 		public bool bUseInlining => Inner.bUseInlining;
 
+		public bool bWithClientCode => Inner.bWithClientCode;
+
 		public bool bWithServerCode => Inner.bWithServerCode;
 
 		public bool bFNameOutlineNumber => Inner.bFNameOutlineNumber;
diff --git a/Engine/Source/Programs/UnrealBuildTool/Configuration/TargetRules.cs b/Engine/Source/Programs/UnrealBuildTool/Configuration/TargetRules.cs
--- a/Engine/Source/Programs/UnrealBuildTool/Configuration/TargetRules.cs
+++ b/Engine/Source/Programs/UnrealBuildTool/Configuration/TargetRules.cs
@@ -1261,6 +1261,17 @@ namespace UnrealBuildTool
 		[CommandLine("-stack-protect")]
 		public bool bEnableStackProtection { get; set; }
 
+		/// <summary>
+		/// Compile client-only code.
+		/// </summary>
+		[RequiresUniqueBuildEnvironment]
+		public bool bWithClientCode
+		{
+			get => bWithClientCodeOverride ?? (Type != TargetType.Server);
+			set => bWithClientCodeOverride = value;
+		}
+		private bool? bWithClientCodeOverride;
+
 		/// <summary>
 		/// Compile server-only code.
 		/// </summary>
diff --git a/Engine/Source/Programs/UnrealBuildTool/Configuration/UEBuildTarget.cs b/Engine/Source/Programs/UnrealBuildTool/Configuration/UEBuildTarget.cs
--- a/Engine/Source/Programs/UnrealBuildTool/Configuration/UEBuildTarget.cs
+++ b/Engine/Source/Programs/UnrealBuildTool/Configuration/UEBuildTarget.cs
@@ -5182,6 +5182,9 @@ namespace UnrealBuildTool
 				GlobalCompileEnvironment.Definitions.Add("WITH_EDITORONLY_DATA=0");
 			}
 
+			// Check if client-only code should be compiled out.
+			GlobalCompileEnvironment.Definitions.Add(String.Format("WITH_CLIENT_CODE={0}", Rules.bWithClientCode ? 1 : 0));
+
 			// Check if server-only code should be compiled out.
 			GlobalCompileEnvironment.Definitions.Add(String.Format("WITH_SERVER_CODE={0}", Rules.bWithServerCode ? 1 : 0));
 
