UnrealBuildTool: Set -fno-delete-null-pointer-checks when compiling ms platforms with clang

diff --git a/Engine/Source/Programs/UnrealBuildTool/Platform/Windows/VCToolChain.cs b/Engine/Source/Programs/UnrealBuildTool/Platform/Windows/VCToolChain.cs
--- a/Engine/Source/Programs/UnrealBuildTool/Platform/Windows/VCToolChain.cs
+++ b/Engine/Source/Programs/UnrealBuildTool/Platform/Windows/VCToolChain.cs
@@ -349,6 +349,9 @@ namespace UnrealBuildTool
 				// compiler version of 19.xx).
 				Arguments.Add($"-fms-compatibility-version=19.{EnvVars.ToolChainVersion.GetComponent(1)}");
 
+				// We have 'this' vs nullptr comparisons that get optimized away for newer versions of Clang, which is undesirable until we refactor these checks.
+				Arguments.Add("-fno-delete-null-pointer-checks");
+
 				if (Target.StaticAnalyzer == StaticAnalyzer.Default && CompileEnvironment.PrecompiledHeaderAction != PrecompiledHeaderAction.Create && !CompileEnvironment.bDisableStaticAnalysis)
 				{
 					Arguments.Add("-Wno-unused-command-line-argument");
