[UBT]
* Removed full path in generated pch cpp file and /Yc option inside rsp file (they both need to be changed to match). Needed to make pch builds deterministic between machines

diff --git a/Engine/Source/Programs/UnrealBuildTool/Platform/Windows/VCToolChain.cs b/Engine/Source/Programs/UnrealBuildTool/Platform/Windows/VCToolChain.cs
--- a/Engine/Source/Programs/UnrealBuildTool/Platform/Windows/VCToolChain.cs
+++ b/Engine/Source/Programs/UnrealBuildTool/Platform/Windows/VCToolChain.cs
@@ -247,7 +247,7 @@ namespace UnrealBuildTool
 
 		public static void AddCreatePchFile(List<string> Arguments, FileItem PchThroughHeaderFile, FileItem CreatePchFile)
 		{
-			string PchThroughHeaderFilePath = NormalizeCommandLinePath(PchThroughHeaderFile);
+			string PchThroughHeaderFilePath = PchThroughHeaderFile.Location.GetFileName();
 			string CreatePchFilePath = NormalizeCommandLinePath(CreatePchFile);
 			if (CreatePchFile.Name.EndsWith(".ifc"))
 			{
@@ -1694,8 +1694,8 @@ namespace UnrealBuildTool
 				if (CompileEnvironment.PrecompiledHeaderAction == PrecompiledHeaderAction.Create)
 				{
 					// Generate a CPP File that just includes the precompiled header.
-					string PrecompiledHeaderIncludeFilenameString = NormalizeCommandLinePath(CompileEnvironment.PrecompiledHeaderIncludeFilename!);
-					string PchCppFile = $"// Compiler: {EnvVars.CompilerVersion}\n#include \"{PrecompiledHeaderIncludeFilenameString.Replace('\\', '/')}\"\r\n";
+					string PrecompiledHeaderIncludeFilenameString = CompileEnvironment.PrecompiledHeaderIncludeFilename!.GetFileName();
+					string PchCppFile = $"// Compiler: {EnvVars.CompilerVersion}\n#include \"{PrecompiledHeaderIncludeFilenameString}\"\r\n";
 					CompileAction.SourceFile = FileItem.GetItemByFileReference(CompileEnvironment.PrecompiledHeaderIncludeFilename!.ChangeExtension(".cpp"));
 					Graph.CreateIntermediateTextFile(CompileAction.SourceFile, PchCppFile);
 
