[UBT]
* Made CppDependencyCache.ReadDependencyInfo "internal" instead of private to make it possible to shared code for parsing dep files produced by various compilers
* Removed implementation in UBAExecutor.cs and used this instead

diff --git a/Engine/Source/Programs/UnrealBuildTool/Executors/UnrealBuildAccelerator/UBAExecutor.cs b/Engine/Source/Programs/UnrealBuildTool/Executors/UnrealBuildAccelerator/UBAExecutor.cs
--- a/Engine/Source/Programs/UnrealBuildTool/Executors/UnrealBuildAccelerator/UBAExecutor.cs
+++ b/Engine/Source/Programs/UnrealBuildTool/Executors/UnrealBuildAccelerator/UBAExecutor.cs
@@ -610,74 +610,10 @@ namespace UnrealBuildTool
 
 				if (action.DependencyListFile != null)
 				{
-					if (action.DependencyListFile.HasExtension(".json"))
+					CppDependencyCache.DependencyInfo info = CppDependencyCache.ReadDependencyInfo(action.DependencyListFile);
+					foreach (FileItem f in info.Files)
 					{
-						DepsFile? deps;
-						try
-						{
-							using (System.IO.FileStream fstream = FileReference.Open(action.DependencyListFile.Location, System.IO.FileMode.Open))
-							{
-								deps = System.Text.Json.JsonSerializer.Deserialize<DepsFile>(fstream);
-							}
-						}
-						catch (Exception e)
-						{
-							_threadedLogger.LogError("Unable to open/deserialize {Dep} - {Ex}", action.DependencyListFile, e.ToString());
-							return false;
-						}
-						if (deps == null || deps.Data == null || String.IsNullOrEmpty(deps.Data.Source))
-						{
-							_threadedLogger.LogError("Unable to deserialize {Dep}", action.DependencyListFile);
-							return false;
-						}
-
-						if (deps.Data.Includes != null)
-						{
-							foreach (string f in deps.Data.Includes)
-							{
-								writer.Write(f);
-							}
-						}
-					}
-					else if (action.DependencyListFile.HasExtension(".d"))
-					{
-						using (System.IO.FileStream fstream = FileReference.Open(action.DependencyListFile.Location, System.IO.FileMode.Open))
-						{
-							StreamReader reader = new(fstream);
-							reader.ReadLine(); // Skip first which is the .o file
-							while (true)
-							{
-								string? line = reader.ReadLine();
-								if (line == null)
-								{
-									break;
-								}
-								if (line.EndsWith(" \\"))
-								{
-									line = line.Substring(0, line.Length - 2);
-								}
-								line = line.Trim();
-								string path = DirectoryReference.Combine(action.WorkingDirectory, line).FullName;
-								writer.Write(path);
-							}
-						}
-					}
-					else
-					{
-						using (System.IO.FileStream fstream = FileReference.Open(action.DependencyListFile.Location, System.IO.FileMode.Open))
-						{
-							StreamReader reader = new(fstream);
-							while (true)
-							{
-								string? line = reader.ReadLine();
-								if (line == null)
-								{
-									break;
-								}
-								string path = DirectoryReference.Combine(action.WorkingDirectory, line.Replace('/', '\\').Replace("\\\\", "\\")).FullName;
-								writer.Write(path);
-							}
-						}
+						writer.Write(f.FullName);
 					}
 				}
 			}
diff --git a/Engine/Source/Programs/UnrealBuildTool/System/CppDependencyCache.cs b/Engine/Source/Programs/UnrealBuildTool/System/CppDependencyCache.cs
--- a/Engine/Source/Programs/UnrealBuildTool/System/CppDependencyCache.cs
+++ b/Engine/Source/Programs/UnrealBuildTool/System/CppDependencyCache.cs
@@ -22,7 +22,7 @@ namespace UnrealBuildTool
 		/// <summary>
 		/// Contents of a single dependency file
 		/// </summary>
-		class DependencyInfo
+		internal class DependencyInfo
 		{
 			public long LastWriteTimeUtc;
 			public string? ProducedModule;
@@ -420,7 +420,7 @@ namespace UnrealBuildTool
 		/// </summary>
 		/// <param name="InputFile">The file to read from</param>
 		/// <returns>List of included dependencies</returns>
-		static DependencyInfo ReadDependencyInfo(FileItem InputFile)
+		internal static DependencyInfo ReadDependencyInfo(FileItem InputFile)
 		{
 			if (InputFile.HasExtension(".d"))
 			{
