Use optional arguments and Extend/Expand pattern for Macros to reduce LLT xml metadata file footprint
.xml metadata files are generated exclusively using TestMode, from values of TestMetadata set from the module class
RunUBT -Mode=Test -GenerateMetadata
To support additional NDA platforms, build modules must be created.
Cleanup unused metadata files from the previous generation system.

diff --git a/Engine/Build/LowLevelTests.xml b/Engine/Build/LowLevelTests.xml
--- a/Engine/Build/LowLevelTests.xml
+++ b/Engine/Build/LowLevelTests.xml
@@ -24,7 +24,22 @@
 	<Macro Name="Catch2BuildAppendExtraCMakeArgs" Arguments="ExtraArgValue">
 		<Property Name="CMakeAdditionalArgs" Value="$(CMakeAdditionalArgs) $(ExtraArgValue)"/>
 	</Macro>
-	<Macro Name="SelectAgentType" Arguments="Platform;BuildPlatform"/>
+	<Macro Name="SelectBuildPlatformAndAgentType" Arguments="Platform">
+		<Do If="'$(Platform)' == 'Linux' Or '$(Platform)' == 'LinuxArm64'">
+			<Property Name="AgentType" Value="Linux"/>
+			<Property Name="BuildPlatform" Value="Linux"/>
+		</Do>
+		<Do If="'$(Platform)' == 'Mac' Or '$(Platform)' == 'iOS'">
+			<Property Name="AgentType" Value="Mac"/>
+			<Property Name="BuildPlatform" Value="Mac"/>
+		</Do>
+		<Do If="'$(Platform)' == 'Win64' Or '$(Platform)' == 'Android'">
+			<Property Name="AgentType" Value="Win64"/>
+			<Property Name="BuildPlatform" Value="Win64"/>
+		</Do>
+	</Macro>
+
+	<Macro Name ="SetupAllPlatformProperties"/>
 
 	<Macro Name="SetupPlatformProperties" Arguments="Platform;BuildPlatform" OptionalArguments="AlternateDevicePool">
 		<Property Name="AllTestablePlatforms" Value="$(AllTestablePlatforms);$(Platform)"/>
@@ -41,16 +56,9 @@
 			<Property Name="AlternateDevicePoolPlatformsProperty" Value="$(AlternateDevicePoolPlatformsProperty);$(AlternateDevicePool)"/>
 		</Do>
 	</Macro>
-	<Macro Name="SetupAllPlatformProperties"/>
-
-	<!-- All test metadata scripts along with platform specific versions for both public and internal paths -->
-	<Include Script="LowLevelTests/*.xml"/>
-	<Include Script="../Platforms/*/Build/LowLevelTests/*.xml"/>
-	<Include Script="../Restricted/NotForLicensees/Build/LowLevelTests/*.xml"/>
-	<Include Script="../Restricted/NotForLicensees/Platforms/*/Build/LowLevelTests/*.xml"/>
-
-	<Expand Name="SetupAllPlatformProperties"/>
 
+	<Macro Name="RunAllTests"/>
+	
 	<Option Name="BuildCatch2Libs" DefaultValue="false" Description="True if building Catch2 static libraries"/>
 	<Option Name="TestCatch2Libs" DefaultValue="false" Description="True if additionally building Catch2 with extras and tests"/>
 	<Option Name="Catch2LibVariation" DefaultValue="VS2019" Description="Catch2 library variation, e.g. VS2022"/>
@@ -90,7 +98,7 @@
 		<Property Name="TestsAllSupportedPlatforms" Value="false" If="$(TestPlatform$(Platform))"/>
 	</ForEach>
 	
-	<Option Name="TestArgs" DefaultValue="" Description="Extra test arguments."/>
+	<Option Name="TestArgs" DefaultValue="" Description="Initial test arguments, applies to all tests."/>
 
 	<Option Name="GauntletDevicePool" DefaultValue="UE4" Description="Gauntlet Device Pool to use to run target platform tests"/>
 	<Option Name="AlternateGauntletDevicePool" DefaultValue="UE5" Description="Alternate Gauntlet Device Pool to use to run target platform tests"/>
@@ -120,103 +128,121 @@
 	<Property Name="CheckedRunAfterStep" Value="false"/>
 	<Property Name="CheckedRunAfterStep" Value="True" If="'$(PreflightChange)' == '' OR $(RunAfterSteps)"/>
 
-	<Macro Name="ExpandAfterSteps" Arguments="Platform;TestName">
+	<Macro Name="TestAfterStep" Arguments="Platform;TestName;ShortName">
 		<Node Name="$(TestName) AfterStep $(Platform)">
 			<Expand Name="$(TestName)AfterSteps" Platform="$(Platform)" />	
 		</Node>
-		<Label Category="After" Name="$($(TestName)Short) After $(Platform)" Requires="$(TestName) AfterStep $(Platform)"/>
+		<Label Category="After" Name="$(ShortName) After $(Platform)" Requires="$(TestName) AfterStep $(Platform)"/>
 		<Property Name="AfterTestNodes" Value="$(AfterTestNodes);$(TestName) AfterStep $(Platform)" If="$(RunAllTests) Or $(Run$(TestName)Tests)"/>
 	</Macro>
 	
-	<Macro Name="BuildAndStageTest" Arguments="Platform;TestName">
-		<Property Name="LocalAchiveDirectory" Value="$(RootDir)/$($(TestName)Target)"/>
-		<Do If="'$($(TestName)StagesWithProjectFile)' == 'True'">
-			<Command Name="BuildCookRun" Arguments="-Build -NoBootstrapExe -NoCodeSign -SkipCook -Stage -NoSubmit -Archive -ArchiveDirectory=$(LocalAchiveDirectory) -project=$($(TestName)Target) -configuration=$(Configuration) -platform=$(Platform) $($(Platform)$(TestName)ExtraCompilationArgs)"/>
+	<Macro Name="BuildAndStageTest" Arguments="Platform;TestName;StagesWithProjectFile;TargetName;ExtraCompilationArgs">
+		<Property Name="LocalAchiveDirectory" Value="$(RootDir)/$(TargetName)"/>
+		<Do If="'$(StagesWithProjectFile)' == 'True'">
+			<Command Name="BuildCookRun" Arguments="-Build -NoBootstrapExe -NoCodeSign -SkipCook -Stage -NoSubmit -Archive -ArchiveDirectory=$(LocalAchiveDirectory) -project=$(TargetName) -configuration=$(Configuration) -platform=$(Platform) $(ExtraCompilationArgs)"/>
 		</Do>
-		<Do If="'$($(TestName)StagesWithProjectFile)' != 'True'">
-			<Compile Target="$($(TestName)Target)" Configuration="$(Configuration)" Platform="$(Platform)" AllowXGE="false" Arguments="-Mode=Test $($(Platform)$(TestName)ExtraCompilationArgs)"/>
+		<Do If="'$(StagesWithProjectFile)' != 'True'">
+			<Compile Target="$(TargetName)" Configuration="$(Configuration)" Platform="$(Platform)" AllowXGE="false" Arguments="-Mode=Test $(ExtraCompilationArgs)"/>
 		</Do>
 	</Macro>
-	<Macro Name="DeployAndTest" Arguments="TestName;Platform;BuildPlatform">
-		<Property Name="PlatformStageDirName" Value="$(Platform)"/>
-		<Property Name="PlatformStageDirName" Value="Windows" If="'$(Platform)' == 'Win64'"/>
-		<Property Name="DevicePool" Value="$(GauntletDevicePool)" />
-		<Property Name="DevicePool" Value="$(AlternateGauntletDevicePool)" If="ContainsItem('$(AlternateDevicePoolPlatforms)', '$(Platform)', ';')" />
-		<Property Name="TestArgsExt" Value="$(TestArgs)"/>
-		<Property Name="TestArgsExt" Value="$(TestArgsExt) -deviceurl=$(GauntletDeviceService) -devicepool=$(DevicePool)" If="'$(GauntletDeviceService)' != ''"/>
-		<Property Name="TestModeSuffix" Value=""/>
-		<Property Name="TestTags" Value=""/>
-		<Property Name="TestTags" Value="$($(Platform)$(TestName)Tags)" If="'$($(Platform)$(TestName)Tags)' != ''"/>
-		<Property Name="TestTags" Value="$(OverrideTags)" If="'$(OverrideTags)' != ''"/>
-		<Do If="'$(Platform)' == 'Android'">
-			<Property Name="TestArgsExt" Value="$(TestArgsExt) -cleandevice"/>
-		</Do>
-		<Property Name="TestArgsExt" Value="$(TestArgsExt) -tags=&quot;$(TestTags)&quot;"/>
-		<Property Name="TestArgsExt" Value="$(TestArgsExt) --buildmachine" If="$(IsBuildMachine)"/>
-		<Property Name="TestArgsExt" Value="$(TestArgsExt) -timeout=$(PerTestTimeout)"/>
-		<Property Name="TestArgsExt" Value="$(TestArgsExt) $($(TestName)GauntletArgs)"/>
-		<Do If="'$($(TestName)ExtraArgs)' != ''">
-			<Property Name="TestArgsExt" Value="$(TestArgsExt) -extra-args=&quot;$($(TestName)ExtraArgs)&quot;"/>
-		</Do>
-		<Property Name="TestArgsExt" Value="$(TestArgsExt) $($(Platform)ExtraArgs)"/>
-		<Property Name="TestArgsExt" Value="$(TestArgsExt) -VerifyLogin" If="$(EnableDeviceLoginVerification)"/>
-		<Do If="'$(Device)' != ''">
-			<Property Name="TestArgsExt" Value="$(TestArgsExt) -device=$(Device) "/>
-		</Do>
-		<Do If="'$(AttachToDebugger)' != 'false'">
-			<Property Name="TestArgsExt" Value="$(TestArgsExt) -AttachToDebugger "/>
-		</Do>
-		<Do If="'$($(TestName)ReportType)' != ''">
-			<Property Name="TestArgsExt" Value="$(TestArgsExt) -ReportType=$($(TestName)ReportType)"/>
-		</Do>
-		<Do If="'$($(TestName)UsesCatch2)' != 'True'">
-			<Property Name="TestArgsExt" Value=""/>
+	
+	<!-- Deploy and run test, one agent per test per platform -->
+	<Macro Name="DeployAndTest" Arguments="TestName;ShortName;TargetName;BinaryRelativePath;Platform;ReportType" OptionalArguments="Deactivated;StagesWithProjectFile;Tags;PlatformExtraArgs;ExtraArgs;GauntletArgs;PlatformGauntletArgs;RunUnsupported;HasAfterSteps;UsesCatch2;RunContainerized;ExtraCompilationArgs">
+		<Do If="'$(Deactivated)' == 'True'">
+			<Warning Message="Test $(TestName) set as deactivated, will build but will not run."/>
 		</Do>
 
-		<Property Name="TestArgsExt" Value="$(TestArgsExt) -SkipStage -Build=&quot;$(RootDir)\$($(TestName)Target)\$(PlatformStageDirName)&quot;" If="'$($(TestName)StagesWithProjectFile)' == 'True'"/>
-		<Property Name="TestArgsExt" Value="$(TestArgsExt) -build=&quot;$($(TestName)BinariesRelative)/$(Platform)/$($(TestName)Target)$(TestModeSuffix)&quot;" If="'$($(TestName)StagesWithProjectFile)' != 'True'"/>
+		<Property Name="AgentType" Value=""/>
+		<Property Name="BuildPlatform" Value=""/>
 
-		<Property Name="AppRunsContainerized" Value="$($(Platform)$(TestName)RunContainerized)"/>
+		<Expand Name="SelectBuildPlatformAndAgentType" Platform="$(Platform)"/>
+		
+		<Agent Name="Low Level Tests $(Platform)" Type="$(AgentType)">
+			<Property Name="TestRequirements" Value=""/>
+			<Do If="$(BuildCatch2Libs)">
+				<Property Name="TestRequirements" Value="$(TestRequirements);Catch2 Build Library $(Platform)"/>
+			</Do>
+			<Node Name="$(ShortName) Tests $(Platform)" Requires="$(TestRequirements)">
+				<Do If="'$(Deactivated)' == ''">
+					<Expand Name="BuildAndStageTest" TestName="$(TestName)" Platform="$(Platform)" TargetName="$(TargetName)" StagesWithProjectFile="$(StagesWithProjectFile)" ExtraCompilationArgs="$(ExtraCompilationArgs)"/>
+				</Do>
+				<Do If="'$(RunUnsupported)' == '' AND '$(Deactivated)' == ''">
+					<Property Name="PlatformStageDirName" Value="$(Platform)"/>
+					<Property Name="PlatformStageDirName" Value="Windows" If="'$(Platform)' == 'Win64'"/>
+					<Property Name="DevicePool" Value="$(GauntletDevicePool)" />
+					<Property Name="DevicePool" Value="$(AlternateGauntletDevicePool)" If="ContainsItem('$(AlternateDevicePoolPlatforms)', '$(Platform)', ';')" />
+					<Property Name="TestArgsExt" Value="$(TestArgs)"/>
+					<Property Name="TestArgsExt" Value="$(TestArgsExt) -deviceurl=$(GauntletDeviceService) -devicepool=$(DevicePool)" If="'$(GauntletDeviceService)' != ''"/>
+					<Property Name="TestModeSuffix" Value=""/>
+					<Property Name="TestTags" Value=""/>
+					<Property Name="TestTags" Value="$(Tags)" If="'$(Tags)' != ''"/>
+					<Property Name="TestTags" Value="$(OverrideTags)" If="'$(OverrideTags)' != ''"/>
+					<Property Name="TestArgsExt" Value="$(TestArgsExt) -cleandevice" If="'$(Platform)' == 'Android'"/>
+					<Property Name="TestArgsExt" Value="$(TestArgsExt) -tags=&quot;$(TestTags)&quot;"/>
+					<Property Name="TestArgsExt" Value="$(TestArgsExt) --buildmachine" If="$(IsBuildMachine)"/>
+					<Property Name="TestArgsExt" Value="$(TestArgsExt) -timeout=$(PerTestTimeout)"/>
+					<Property Name="TestArgsExt" Value="$(TestArgsExt) $(GauntletArgs)" If="'$(GauntletArgs)' != ''"/>
+					<Property Name="TestArgsExt" Value="$(TestArgsExt) $(PlatformGauntletArgs)" If="'$(PlatformGauntletArgs)' != ''"/>
+					<Property Name="TestArgsExt" Value="$(TestArgsExt) -extra-args=&quot;$(ExtraArgs)&quot;" If="'$(ExtraArgs)' != ''"/>
+					<Property Name="TestArgsExt" Value="$(TestArgsExt) $(PlatformExtraArgs)" If="'$(PlatformExtraArgs)' != ''"/>
+					<Property Name="TestArgsExt" Value="$(TestArgsExt) -VerifyLogin" If="$(EnableDeviceLoginVerification)"/>
+					<Property Name="TestArgsExt" Value="$(TestArgsExt) -device=$(Device) " If="'$(Device)' != ''"/>
+					<Property Name="TestArgsExt" Value="$(TestArgsExt) -AttachToDebugger " If="'$(AttachToDebugger)' != 'False'"/>
+					<Property Name="TestArgsExt" Value="$(TestArgsExt) -ReportType=$(ReportType)" If="'$(ReportType)' != ''"/>
+					<Do If="'$(UsesCatch2)' != ''">
+						<Property Name="TestArgsExt" Value=""/>
+					</Do>
 
-		<Do If="'$(AppRunsContainerized)' == 'True'">
-			<Property Name="TestTargetLowercase" Value=""/>
-			<StringOp Input="$($(TestName)Target)" Method="ToLower" Output="TestTargetLowercase"/>
+					<Property Name="TestArgsExt" Value="$(TestArgsExt) -SkipStage -Build=&quot;$(RootDir)\$(TargetName)\$(PlatformStageDirName)&quot;" If="'$(StagesWithProjectFile)' == 'True'"/>
+					<Property Name="TestArgsExt" Value="$(TestArgsExt) -build=&quot;$(BinaryRelativePath)/$(Platform)/$(TargetName)$(TestModeSuffix)&quot;" If="'$(StagesWithProjectFile)' != 'True'"/>
 
-			<Property Name="BinariesRelativeForwardSlash" Value=""/>
-			<StringOp Input="$($(TestName)BinariesRelative)" Method="Replace" Arguments="\;/" Output="BinariesRelativeForwardSlash"/>
+					<Property Name="AppRunsContainerized" Value="$(RunContainerized)"/>
 
-			<!-- Builds docker image -->
-			<Expand Name="$(Platform)DockerCleanup" AppNameLowercase="$(TestTargetLowercase)"/>
-			<Expand Name="$(Platform)DockerSetup" AppNameLowercase="$(TestTargetLowercase)" AppDir="$(BinariesRelativeForwardSlash)/$(Platform)/$($(TestName)Target)$(TestModeSuffix)"/>
-			<Property Name="TestArgsExt" Value="$(TestArgsExt) -containerized"/>
-		</Do>
+					<Do If="'$(AppRunsContainerized)' == 'True'">
+						<Property Name="TestTargetLowercase" Value=""/>
+						<StringOp Input="$(TargetName)" Method="ToLower" Output="TestTargetLowercase"/>
 
-		<Do If="'$(Platform)' == 'Android'">
-			<Expand Name="SelectAdbKeysPath" BuildPlatform="$(BuildPlatform)"/>
-			<Property Name="TestArgsExt" Value="$(TestArgsExt) -adbkeys=$(AdbKeysPath) -captureoutput"/>
-		</Do>
-		
-		<Command Name="RunLowLevelTests"
-				Arguments="-test=LowLevelTests -testapp=$($(TestName)Target)$(TestModeSuffix) -platform=$(Platform) -configuration=$(Configuration) $(TestArgsExt) -VeryVerbose"/>
+						<Property Name="BinariesRelativeForwardSlash" Value=""/>
+						<StringOp Input="$(BinaryRelativePath)" Method="Replace" Arguments="\;/" Output="BinariesRelativeForwardSlash"/>
 
-		<Do If="'$(AppRunsContainerized)' == 'True'">
-			<Property Name="TestTargetLowercase" Value=""/>
-			<StringOp Input="$($(TestName)Target)" Method="ToLower" Output="TestTargetLowercase"/>
+						<!-- Builds docker image -->
+						<Expand Name="$(Platform)DockerCleanup" AppNameLowercase="$(TestTargetLowercase)"/>
+						<Expand Name="$(Platform)DockerSetup" AppNameLowercase="$(TestTargetLowercase)" AppDir="$(BinariesRelativeForwardSlash)/$(Platform)/$(TargetName)$(TestModeSuffix)"/>
+						<Property Name="TestArgsExt" Value="$(TestArgsExt) -containerized"/>
+					</Do>
 
-			<!-- Stop any running docker container and remove image -->
-			<Expand Name="$(Platform)DockerCleanup" AppNameLowercase="$(TestTargetLowercase)"/>
-		</Do>
+					<Do If="'$(Platform)' == 'Android'">
+						<Expand Name="SelectAdbKeysPath" BuildPlatform="$(BuildPlatform)"/>
+						<Property Name="TestArgsExt" Value="$(TestArgsExt) -adbkeys=$(AdbKeysPath) -captureoutput"/>
+					</Do>
+
+					<Command Name="RunLowLevelTests"
+							Arguments="-test=LowLevelTests -testapp=$(TargetName)$(TestModeSuffix) -platform=$(Platform) -configuration=$(Configuration) $(TestArgsExt) -VeryVerbose"/>
+
+					<Do If="'$(AppRunsContainerized)' == 'True'">
+						<Property Name="TestTargetLowercase" Value=""/>
+						<StringOp Input="$(TargetName)" Method="ToLower" Output="TestTargetLowercase"/>
+
+						<!-- Stop any running docker container and remove image -->
+						<Expand Name="$(Platform)DockerCleanup" AppNameLowercase="$(TestTargetLowercase)"/>
+					</Do>
+				</Do>
+			</Node>
+			<Label Category="Test" Name="$(ShortName) $(Platform)" Requires="$(ShortName) Tests $(Platform)"/>
+			<Property Name="TestNodes" Value="$(TestNodes);$(ShortName) Tests $(Platform)" If="$(RunAllTests) Or $(Run$(TestName)Tests)"/>
+				
+			<Do If="'$(CheckedRunAfterStep)' == 'True' AND '$(HasAfterSteps)' == 'True'">
+				<Expand Name="TestAfterStep" TestName="$(TestName)" ShortName="$(ShortName)" Platform="$(Platform)"/>
+			</Do>
+		</Agent>
 	</Macro>
 
+	<!-- Catch2 build targets -->
 	<ForEach Name="BuildPlatform" Values="$(BuildPlatforms)">
 		<ForEach Name="Platform" Values="$(PlatformsBuiltFrom$(BuildPlatform))">
 			<Property Name="AgentType" Value="$(BuildPlatform)"/>
-			<Expand Name="SelectAgentType" Platform="$(Platform)" BuildPlatform="$(BuildPlatform)" />
-			<Property Name="BuildAgentType" Value="$(AgentType)"/>
-			<Do If="'$(Platform)' == 'Linux' Or '$(Platform)' == 'LinuxArm64'">
-				<Property Name="BuildAgentType" Value="Win64"/>
-			</Do>
-			<Agent Name="Catch2 Build $(Platform)" Type="$(BuildAgentType)">
+			<Expand Name="SelectBuildPlatformAndAgentType" Platform="$(Platform)" />
+			<Agent Name="Catch2 Build $(Platform)" Type="$(AgentType)">
 				<Node Name="Catch2 Build Library $(Platform)">
 					<Property Name="Generator" Value="Makefile"/>
 					<Property Name="TargetConfigs" Value="debug+release"/>
@@ -253,42 +279,16 @@
 				</Node>
 				<Property Name="Catch2BuildNodes" Value="$(Catch2BuildNodes);Catch2 Build Library $(Platform)"/>
 			</Agent>
-			<Agent Name="Low Level Tests $(Platform)" Type="$(AgentType)">
-				<ForEach Name="TestName" Values="$(TestNames)">
-					<Property Name="TestSupportedOnPlatform" Value="false"/>
-					<Property Name="TestRunsOnPlatform" Value="false"/>
-					<Property Name="RunsTestOnPlatform" Value="false"/>
-					
-					<Property Name="TestSupportedOnPlatform" Value="true" If="Contains('$($(TestName)SupportedPlatforms)', '$(Platform)')"/>
-					<Property Name="TestRunsOnPlatform" Value="true" If="($(TestSupportedOnPlatform) And $(TestsAllSupportedPlatforms)) Or $(TestPlatform$(Platform))"/>
-					<Property Name="RunsTestOnPlatform" Value="true" If="!$($(TestName)Disabled) And $(TestRunsOnPlatform) And ($(Run$(TestName)Tests) Or $(RunsAllTests))"/>
-
-					<Do If="$(RunsTestOnPlatform)">
-						<Do If="!$(TestSupportedOnPlatform)">
-							<Warning Message="Low level test $(TestName) not yet tested for support on platform $(Platform)."/>
-						</Do>
-
-						<Property Name="TestRequirements" Value=""/>
-						<Do If="$(BuildCatch2Libs)">
-							<Property Name="TestRequirements" Value="$(TestRequirements);Catch2 Build Library $(Platform)"/>
-						</Do>
-						<Node Name="$($(TestName)Short) Tests $(Platform)" Requires="$(TestRequirements)">
-							<Expand Name="BuildAndStageTest" TestName="$(TestName)" Platform="$(Platform)"/>
-							<Do If="'$($(Platform)$(TestName)RunSupported)' == 'True'">
-								<Expand Name="DeployAndTest" TestName="$(TestName)" Platform="$(Platform)" BuildPlatform="$(BuildPlatform)"/>
-							</Do>
-						</Node>
-						<Label Category="Test" Name="$($(TestName)Short) $(Platform)" Requires="$($(TestName)Short) Tests $(Platform)"/>
-						<Property Name="TestNodes" Value="$(TestNodes);$($(TestName)Short) Tests $(Platform)" If="$(RunAllTests) Or $(Run$(TestName)Tests)"/>
-						
-						<Do If="'$(CheckedRunAfterStep)' == 'True' AND '$($(TestName)HasAfterSteps)' == 'True'">
-							<Expand Name="ExpandAfterSteps" TestName="$(TestName)" Platform="$(Platform)" />
-						</Do>
-					</Do>
-				</ForEach>
-			</Agent>
 		</ForEach>
 	</ForEach>
+
+	<!-- All test metadata scripts along with platform specific versions for both public and internal paths -->
+	<Include Script="LowLevelTests/*.xml"/>
+	<Include Script="../Platforms/*/Build/LowLevelTests/*.xml"/>
+	<Include Script="../Restricted/NotForLicensees/Build/LowLevelTests/*.xml"/>
+	<Include Script="../Restricted/NotForLicensees/Platforms/*/Build/LowLevelTests/*.xml"/>
+
+	<Expand Name="RunAllTests"/>
 	
 	<Aggregate Name="Run Low Level Tests" Requires="$(TestNodes);$(AfterTestNodes)"/>
 	<Aggregate Name="Build Catch2 Libraries" Requires="$(Catch2BuildNodes)"/>
diff --git a/Engine/Build/LowLevelTests/ChasoUserDataPT.xml b/Engine/Build/LowLevelTests/ChasoUserDataPT.xml
new file mode 100644
--- /dev/null
+++ b/Engine/Build/LowLevelTests/ChasoUserDataPT.xml
@@ -0,0 +1,7 @@
+﻿<?xml version="1.0" encoding="utf-8"?>
+<BuildGraph xmlns:xsi="http://www.w3.org/2001/XMLSchema-instance" xsi:schemaLocation="http://www.epicgames.com/BuildGraph ../../Build/Graph/Schema.xsd" xmlns="http://www.epicgames.com/BuildGraph">
+  <Option Name="RunChasoUserDataPTTests" DefaultValue="False" Description="Run Chaos User Data PT Tests" />
+  <Extend Name="RunAllTests">
+    <Expand Name="DeployAndTest" Platform="Win64" TestName="ChasoUserDataPT" ShortName="Chaos User Data PT" TargetName="ChaosUserDataPTTests" BinaryRelativePath="Engine\Binaries" ReportType="console" />
+  </Extend>
+</BuildGraph>
\ No newline at end of file
diff --git a/Engine/Build/LowLevelTests/EventLoop.xml b/Engine/Build/LowLevelTests/EventLoop.xml
--- a/Engine/Build/LowLevelTests/EventLoop.xml
+++ b/Engine/Build/LowLevelTests/EventLoop.xml
@@ -1,49 +1,7 @@
 ﻿<?xml version="1.0" encoding="utf-8"?>
 <BuildGraph xmlns:xsi="http://www.w3.org/2001/XMLSchema-instance" xsi:schemaLocation="http://www.epicgames.com/BuildGraph ../../Build/Graph/Schema.xsd" xmlns="http://www.epicgames.com/BuildGraph">
-  <Property Name="TestNames" Value="$(TestNames);EventLoop" />
-  <Property Name="EventLoopDisabled" Value="False" />
-  <Property Name="EventLoopShort" Value="EventLoop" />
-  <Property Name="EventLoopStagesWithProjectFile" Value="False" />
-  <Property Name="EventLoopTarget" Value="EventLoopUnitTests" />
-  <Property Name="EventLoopBinariesRelative" Value="Engine\Binaries" />
-  <Property Name="EventLoopReportType" Value="console" />
-  <Property Name="EventLoopGauntletArgs" Value="" />
-  <Property Name="EventLoopExtraArgs" />
-  <Property Name="EventLoopInitialExtraArgs" />
-  <Property Name="EventLoopHasAfterSteps" Value="False" />
-  <Property Name="EventLoopUsesCatch2" Value="True" />
   <Option Name="RunEventLoopTests" DefaultValue="False" Description="Run EventLoop Tests" />
-  <Property Name="EventLoopSupportedPlatforms" Value="Win64" />
-  <Property Name="Win64EventLoopTags" Value="" />
-  <Property Name="Win64EventLoopExtraCompilationArgs" Value="" />
-  <Property Name="Win64EventLoopRunSupported" Value="True" />
-  <Property Name="MacEventLoopTags" Value="" />
-  <Property Name="MacEventLoopExtraCompilationArgs" Value="" />
-  <Property Name="MacEventLoopRunSupported" Value="True" />
-  <Property Name="IOSEventLoopTags" Value="" />
-  <Property Name="IOSEventLoopExtraCompilationArgs" Value="" />
-  <Property Name="IOSEventLoopRunSupported" Value="False" />
-  <Property Name="AndroidEventLoopTags" Value="" />
-  <Property Name="AndroidEventLoopExtraCompilationArgs" Value="" />
-  <Property Name="AndroidEventLoopRunSupported" Value="False" />
-  <Property Name="LinuxEventLoopTags" Value="" />
-  <Property Name="LinuxEventLoopExtraCompilationArgs" Value="" />
-  <Property Name="LinuxEventLoopRunSupported" Value="True" />
-  <Property Name="LinuxArm64EventLoopTags" Value="" />
-  <Property Name="LinuxArm64EventLoopExtraCompilationArgs" Value="" />
-  <Property Name="LinuxArm64EventLoopRunSupported" Value="True" />
-  <Property Name="TVOSEventLoopTags" Value="" />
-  <Property Name="TVOSEventLoopExtraCompilationArgs" Value="" />
-  <Property Name="TVOSEventLoopRunSupported" Value="False" />
-  <Property Name="VisionOSEventLoopTags" Value="" />
-  <Property Name="VisionOSEventLoopExtraCompilationArgs" Value="" />
-  <Property Name="VisionOSEventLoopRunSupported" Value="False" />
-  <Property Name="Win64EventLoopRunContainerized" Value="False" />
-  <Property Name="MacEventLoopRunContainerized" Value="False" />
-  <Property Name="IOSEventLoopRunContainerized" Value="False" />
-  <Property Name="AndroidEventLoopRunContainerized" Value="False" />
-  <Property Name="LinuxEventLoopRunContainerized" Value="False" />
-  <Property Name="LinuxArm64EventLoopRunContainerized" Value="False" />
-  <Property Name="TVOSEventLoopRunContainerized" Value="False" />
-  <Property Name="VisionOSEventLoopRunContainerized" Value="False" />
+  <Extend Name="RunAllTests">
+    <Expand Name="DeployAndTest" Platform="Win64" TestName="EventLoop" ShortName="EventLoop" TargetName="EventLoopUnitTests" BinaryRelativePath="Engine\Binaries" ReportType="console" />
+  </Extend>
 </BuildGraph>
\ No newline at end of file
diff --git a/Engine/Build/LowLevelTests/Foundation.xml b/Engine/Build/LowLevelTests/Foundation.xml
--- a/Engine/Build/LowLevelTests/Foundation.xml
+++ b/Engine/Build/LowLevelTests/Foundation.xml
@@ -1,50 +1,11 @@
 ﻿<?xml version="1.0" encoding="utf-8"?>
 <BuildGraph xmlns:xsi="http://www.w3.org/2001/XMLSchema-instance" xsi:schemaLocation="http://www.epicgames.com/BuildGraph ../../Build/Graph/Schema.xsd" xmlns="http://www.epicgames.com/BuildGraph">
-  <Property Name="TestNames" Value="$(TestNames);Foundation" />
-  <Property Name="FoundationUsesCatch2" Value="True" />
-  <Property Name="FoundationDisabled" Value="False" />
-  <Property Name="FoundationShort" Value="Foundation" />
-  <Property Name="FoundationStagesWithProjectFile" Value="False" />
-  <Property Name="FoundationTarget" Value="FoundationTests" />
-  <Property Name="FoundationBinariesRelative" Value="Engine\Binaries" />
-  <Property Name="FoundationReportType" Value="xml" />
-  <Property Name="FoundationGauntletArgs" Value="" />
-  <Property Name="FoundationExtraArgs" />
-  <Property Name="FoundationInitialExtraArgs" />
-  <Property Name="FoundationHasAfterSteps" Value="False" />
-  <Property Name="FoundationUsesCatch2" Value="True" />
   <Option Name="RunFoundationTests" DefaultValue="False" Description="Run Foundation Tests" />
-  <Property Name="FoundationSupportedPlatforms" Value="Win64;Linux;Mac;Android;IOS" />
-  <Property Name="Win64FoundationTags" Value="" />
-  <Property Name="MacFoundationTags" Value="" />
-  <Property Name="IOSFoundationTags" Value="" />
-  <Property Name="AndroidFoundationTags" Value="~[Perf]~[Slow]~[AndroidSkip]" />
-  <Property Name="LinuxFoundationTags" Value="~[.]~[Slow]" />
-  <Property Name="LinuxArm64FoundationTags" Value="" />
-  <Property Name="TVOSFoundationTags" Value="" />
-  <Property Name="Win64FoundationExtraCompilationArgs" Value="-allmodules" />
-  <Property Name="MacFoundationExtraCompilationArgs" Value="-allmodules" />
-  <Property Name="IOSFoundationExtraCompilationArgs" Value="-allmodules" />
-  <Property Name="AndroidFoundationExtraCompilationArgs" Value="-allmodules -architectures=arm64" />
-  <Property Name="LinuxFoundationExtraCompilationArgs" Value="-allmodules" />
-  <Property Name="LinuxArm64FoundationExtraCompilationArgs" Value="-allmodules" />
-  <Property Name="TVOSFoundationExtraCompilationArgs" Value="-allmodules" />
-  <Property Name="Win64FoundationRunSupported" Value="True" />
-  <Property Name="MacFoundationRunSupported" Value="True" />
-  <Property Name="IOSFoundationRunSupported" Value="False" />
-  <Property Name="AndroidFoundationRunSupported" Value="True" />
-  <Property Name="LinuxFoundationRunSupported" Value="True" />
-  <Property Name="LinuxArm64FoundationRunSupported" Value="True" />
-  <Property Name="TVOSFoundationRunSupported" Value="False" />
-  <Property Name="VisionOSFoundationTags" Value="" />
-  <Property Name="VisionOSFoundationExtraCompilationArgs" Value="-allmodules" />
-  <Property Name="VisionOSFoundationRunSupported" Value="False" />
-  <Property Name="Win64FoundationRunContainerized" Value="False" />
-  <Property Name="MacFoundationRunContainerized" Value="False" />
-  <Property Name="IOSFoundationRunContainerized" Value="False" />
-  <Property Name="AndroidFoundationRunContainerized" Value="False" />
-  <Property Name="LinuxFoundationRunContainerized" Value="False" />
-  <Property Name="LinuxArm64FoundationRunContainerized" Value="False" />
-  <Property Name="TVOSFoundationRunContainerized" Value="False" />
-  <Property Name="VisionOSFoundationRunContainerized" Value="False" />
+  <Extend Name="RunAllTests">
+    <Expand Name="DeployAndTest" Platform="Win64" TestName="Foundation" ShortName="Foundation" TargetName="FoundationTests" BinaryRelativePath="Engine\Binaries" ReportType="xml" ExtraCompilationArgs="-allmodules" />
+    <Expand Name="DeployAndTest" Platform="Linux" TestName="Foundation" ShortName="Foundation" TargetName="FoundationTests" BinaryRelativePath="Engine\Binaries" ReportType="xml" Tags="~[.]~[Slow]" ExtraCompilationArgs="-allmodules" />
+    <Expand Name="DeployAndTest" Platform="Mac" TestName="Foundation" ShortName="Foundation" TargetName="FoundationTests" BinaryRelativePath="Engine\Binaries" ReportType="xml" ExtraCompilationArgs="-allmodules" />
+    <Expand Name="DeployAndTest" Platform="Android" TestName="Foundation" ShortName="Foundation" TargetName="FoundationTests" BinaryRelativePath="Engine\Binaries" ReportType="xml" Tags="~[Perf]~[Slow]~[AndroidSkip]" ExtraCompilationArgs="-allmodules -architectures=arm64" />
+    <Expand Name="DeployAndTest" Platform="IOS" TestName="Foundation" ShortName="Foundation" TargetName="FoundationTests" BinaryRelativePath="Engine\Binaries" ReportType="xml" ExtraCompilationArgs="-allmodules" RunUnsupported="True" />
+  </Extend>
 </BuildGraph>
\ No newline at end of file
diff --git a/Engine/Build/LowLevelTests/GameplayGraph.xml b/Engine/Build/LowLevelTests/GameplayGraph.xml
new file mode 100644
--- /dev/null
+++ b/Engine/Build/LowLevelTests/GameplayGraph.xml
@@ -0,0 +1,7 @@
+﻿<?xml version="1.0" encoding="utf-8"?>
+<BuildGraph xmlns:xsi="http://www.w3.org/2001/XMLSchema-instance" xsi:schemaLocation="http://www.epicgames.com/BuildGraph ../../Build/Graph/Schema.xsd" xmlns="http://www.epicgames.com/BuildGraph">
+  <Option Name="RunGameplayGraphTests" DefaultValue="False" Description="Run GameplayGraph Tests" />
+  <Extend Name="RunAllTests">
+    <Expand Name="DeployAndTest" Platform="Win64" TestName="GameplayGraph" ShortName="GameplayGraph" TargetName="GameplayGraphTests" BinaryRelativePath="Engine\Binaries" ReportType="console" />
+  </Extend>
+</BuildGraph>
\ No newline at end of file
diff --git a/Engine/Build/LowLevelTests/General.xml b/Engine/Build/LowLevelTests/General.xml
--- a/Engine/Build/LowLevelTests/General.xml
+++ b/Engine/Build/LowLevelTests/General.xml
@@ -1,23 +1,5 @@
 <?xml version="1.0" encoding="utf-8"?>
 <BuildGraph xmlns:xsi="http://www.w3.org/2001/XMLSchema-instance" xsi:schemaLocation="http://www.epicgames.com/BuildGraph ../../Build/Graph/Schema.xsd" xmlns="http://www.epicgames.com/BuildGraph">
-	<Property Name="Win64ExtraArgs" Value=""/>
-	<Property Name="LinuxExtraArgs" Value=""/>
-	<Property Name="LinuxArm64ExtraArgs" Value=""/>
-	<Property Name="MacExtraArgs" Value=""/>
-	<Property Name="IOSExtraArgs" Value=""/>
-	<Property Name="AndroidExtraArgs" Value=""/>
-	<Property Name="Win64RequiresSpecialAgent" Value="False"/>
-	<Property Name="LinuxRequiresSpecialAgent" Value="False"/>
-	<Property Name="LinuxArm64RequiresSpecialAgent" Value="False"/>
-	<Property Name="MacRequiresSpecialAgent" Value="False"/>
-	<Property Name="IOSRequiresSpecialAgent" Value="False"/>
-	<Property Name="AndroidRequiresSpecialAgent" Value="False"/>
-	<Property Name="Win64RequiresCatch2BuildExtraCMakeArgs" Value="False"/>
-	<Property Name="LinuxRequiresCatch2BuildExtraCMakeArgs" Value="False"/>
-	<Property Name="LinuxArm64RequiresCatch2BuildExtraCMakeArgs" Value="False"/>
-	<Property Name="MacRequiresCatch2BuildExtraCMakeArgs" Value="False"/>
-	<Property Name="IOSRequiresCatch2BuildExtraCMakeArgs" Value="False"/>
-	<Property Name="AndroidRequiresCatch2BuildExtraCMakeArgs" Value="False"/>
 	<Extend Name="Catch2BuildAppendExtraCMakeArgsPlatform">
 		<Do If="'$(Platform)' == 'Win64'">
 			<Expand Name="Catch2BuildAppendExtraCMakeArgs"
diff --git a/Engine/Build/LowLevelTests/HeadlessChaos.xml b/Engine/Build/LowLevelTests/HeadlessChaos.xml
--- a/Engine/Build/LowLevelTests/HeadlessChaos.xml
+++ b/Engine/Build/LowLevelTests/HeadlessChaos.xml
@@ -1,50 +1,7 @@
 ﻿<?xml version="1.0" encoding="utf-8"?>
 <BuildGraph xmlns:xsi="http://www.w3.org/2001/XMLSchema-instance" xsi:schemaLocation="http://www.epicgames.com/BuildGraph ../../Build/Graph/Schema.xsd" xmlns="http://www.epicgames.com/BuildGraph">
-  <Property Name="TestNames" Value="$(TestNames);HeadlessChaos" />
-  <Property Name="HeadlessChaosUsesCatch2" Value="False" />
-  <Property Name="HeadlessChaosDisabled" Value="False" />
-  <Property Name="HeadlessChaosShort" Value="Headless Chaos" />
-  <Property Name="HeadlessChaosStagesWithProjectFile" Value="False" />
-  <Property Name="HeadlessChaosTarget" Value="HeadlessChaos" />
-  <Property Name="HeadlessChaosBinariesRelative" Value="Engine\Binaries" />
-  <Property Name="HeadlessChaosReportType" Value="console" />
-  <Property Name="HeadlessChaosGauntletArgs" Value="" />
-  <Property Name="HeadlessChaosExtraArgs" />
-  <Property Name="HeadlessChaosInitialExtraArgs" />
-  <Property Name="HeadlessChaosHasAfterSteps" Value="False" />
-  <Property Name="HeadlessChaosUsesCatch2" Value="False" />
   <Option Name="RunHeadlessChaosTests" DefaultValue="False" Description="Run Headless Chaos Tests" />
-  <Property Name="HeadlessChaosSupportedPlatforms" Value="Win64" />
-  <Property Name="Win64HeadlessChaosTags" Value="" />
-  <Property Name="MacHeadlessChaosTags" Value="" />
-  <Property Name="IOSHeadlessChaosTags" Value="" />
-  <Property Name="AndroidHeadlessChaosTags" Value="" />
-  <Property Name="LinuxHeadlessChaosTags" Value="" />
-  <Property Name="LinuxArm64HeadlessChaosTags" Value="" />
-  <Property Name="TVOSHeadlessChaosTags" Value="" />
-  <Property Name="Win64HeadlessChaosExtraCompilationArgs" Value="" />
-  <Property Name="MacHeadlessChaosExtraCompilationArgs" Value="" />
-  <Property Name="IOSHeadlessChaosExtraCompilationArgs" Value="" />
-  <Property Name="AndroidHeadlessChaosExtraCompilationArgs" Value="" />
-  <Property Name="LinuxHeadlessChaosExtraCompilationArgs" Value="" />
-  <Property Name="LinuxArm64HeadlessChaosExtraCompilationArgs" Value="" />
-  <Property Name="TVOSHeadlessChaosExtraCompilationArgs" Value="" />
-  <Property Name="Win64HeadlessChaosRunSupported" Value="True" />
-  <Property Name="MacHeadlessChaosRunSupported" Value="True" />
-  <Property Name="IOSHeadlessChaosRunSupported" Value="False" />
-  <Property Name="AndroidHeadlessChaosRunSupported" Value="False" />
-  <Property Name="LinuxHeadlessChaosRunSupported" Value="True" />
-  <Property Name="LinuxArm64HeadlessChaosRunSupported" Value="True" />
-  <Property Name="TVOSHeadlessChaosRunSupported" Value="False" />
-  <Property Name="VisionOSHeadlessChaosTags" Value="" />
-  <Property Name="VisionOSHeadlessChaosExtraCompilationArgs" Value="" />
-  <Property Name="VisionOSHeadlessChaosRunSupported" Value="False" />
-  <Property Name="Win64HeadlessChaosRunContainerized" Value="False" />
-  <Property Name="MacHeadlessChaosRunContainerized" Value="False" />
-  <Property Name="IOSHeadlessChaosRunContainerized" Value="False" />
-  <Property Name="AndroidHeadlessChaosRunContainerized" Value="False" />
-  <Property Name="LinuxHeadlessChaosRunContainerized" Value="False" />
-  <Property Name="LinuxArm64HeadlessChaosRunContainerized" Value="False" />
-  <Property Name="TVOSHeadlessChaosRunContainerized" Value="False" />
-  <Property Name="VisionOSHeadlessChaosRunContainerized" Value="False" />
+  <Extend Name="RunAllTests">
+    <Expand Name="DeployAndTest" Platform="Win64" TestName="HeadlessChaos" ShortName="Headless Chaos" TargetName="HeadlessChaos" BinaryRelativePath="Engine\Binaries" ReportType="console" UsesCatch2="False" />
+  </Extend>
 </BuildGraph>
\ No newline at end of file
diff --git a/Engine/Build/LowLevelTests/MathCore.xml b/Engine/Build/LowLevelTests/MathCore.xml
--- a/Engine/Build/LowLevelTests/MathCore.xml
+++ b/Engine/Build/LowLevelTests/MathCore.xml
@@ -1,48 +1,11 @@
 ﻿<?xml version="1.0" encoding="utf-8"?>
 <BuildGraph xmlns:xsi="http://www.w3.org/2001/XMLSchema-instance" xsi:schemaLocation="http://www.epicgames.com/BuildGraph ../../Build/Graph/Schema.xsd" xmlns="http://www.epicgames.com/BuildGraph">
-  <Property Name="TestNames" Value="$(TestNames);MathCore" />
-  <Property Name="MathCoreDisabled" Value="False" />
-  <Property Name="MathCoreShort" Value="MathCore" />
-  <Property Name="MathCoreStagesWithProjectFile" Value="False" />
-  <Property Name="MathCoreTarget" Value="MathCoreTests" />
-  <Property Name="MathCoreBinariesRelative" Value="Engine\Binaries" />
-  <Property Name="MathCoreReportType" Value="xml" />
-  <Property Name="MathCoreGauntletArgs" Value="" />
-  <Property Name="MathCoreExtraArgs" />
-  <Property Name="MathCoreHasAfterSteps" Value="False" />
-  <Property Name="MathCoreUsesCatch2" Value="True" />
   <Option Name="RunMathCoreTests" DefaultValue="False" Description="Run MathCore Tests" />
-  <Property Name="MathCoreSupportedPlatforms" Value="Win64;Linux;Mac;Android;IOS" />
-  <Property Name="Win64MathCoreTags" Value="" />
-  <Property Name="Win64MathCoreExtraCompilationArgs" Value="-allmodules" />
-  <Property Name="Win64MathCoreRunSupported" Value="True" />
-  <Property Name="Win64MathCoreRunContainerized" Value="False" />
-  <Property Name="MacMathCoreTags" Value="" />
-  <Property Name="MacMathCoreExtraCompilationArgs" Value="-allmodules" />
-  <Property Name="MacMathCoreRunSupported" Value="True" />
-  <Property Name="MacMathCoreRunContainerized" Value="False" />
-  <Property Name="IOSMathCoreTags" Value="" />
-  <Property Name="IOSMathCoreExtraCompilationArgs" Value="-allmodules" />
-  <Property Name="IOSMathCoreRunSupported" Value="False" />
-  <Property Name="IOSMathCoreRunContainerized" Value="False" />
-  <Property Name="AndroidMathCoreTags" Value="" />
-  <Property Name="AndroidMathCoreExtraCompilationArgs" Value="-allmodules -architectures=arm64" />
-  <Property Name="AndroidMathCoreRunSupported" Value="True" />
-  <Property Name="AndroidMathCoreRunContainerized" Value="False" />
-  <Property Name="LinuxMathCoreTags" Value="" />
-  <Property Name="LinuxMathCoreExtraCompilationArgs" Value="-allmodules" />
-  <Property Name="LinuxMathCoreRunSupported" Value="True" />
-  <Property Name="LinuxMathCoreRunContainerized" Value="False" />
-  <Property Name="LinuxArm64MathCoreTags" Value="" />
-  <Property Name="LinuxArm64MathCoreExtraCompilationArgs" Value="-allmodules" />
-  <Property Name="LinuxArm64MathCoreRunSupported" Value="True" />
-  <Property Name="LinuxArm64MathCoreRunContainerized" Value="False" />
-  <Property Name="TVOSMathCoreTags" Value="" />
-  <Property Name="TVOSMathCoreExtraCompilationArgs" Value="-allmodules" />
-  <Property Name="TVOSMathCoreRunSupported" Value="False" />
-  <Property Name="TVOSMathCoreRunContainerized" Value="False" />
-  <Property Name="VisionOSMathCoreTags" Value="" />
-  <Property Name="VisionOSMathCoreExtraCompilationArgs" Value="-allmodules" />
-  <Property Name="VisionOSMathCoreRunSupported" Value="False" />
-  <Property Name="VisionOSMathCoreRunContainerized" Value="False" />
+  <Extend Name="RunAllTests">
+    <Expand Name="DeployAndTest" Platform="Win64" TestName="MathCore" ShortName="MathCore" TargetName="MathCoreTests" BinaryRelativePath="Engine\Binaries" ReportType="xml" ExtraCompilationArgs="-allmodules" />
+    <Expand Name="DeployAndTest" Platform="Linux" TestName="MathCore" ShortName="MathCore" TargetName="MathCoreTests" BinaryRelativePath="Engine\Binaries" ReportType="xml" ExtraCompilationArgs="-allmodules" />
+    <Expand Name="DeployAndTest" Platform="Mac" TestName="MathCore" ShortName="MathCore" TargetName="MathCoreTests" BinaryRelativePath="Engine\Binaries" ReportType="xml" ExtraCompilationArgs="-allmodules" />
+    <Expand Name="DeployAndTest" Platform="Android" TestName="MathCore" ShortName="MathCore" TargetName="MathCoreTests" BinaryRelativePath="Engine\Binaries" ReportType="xml" ExtraCompilationArgs="-allmodules -architectures=arm64" />
+    <Expand Name="DeployAndTest" Platform="IOS" TestName="MathCore" ShortName="MathCore" TargetName="MathCoreTests" BinaryRelativePath="Engine\Binaries" ReportType="xml" ExtraCompilationArgs="-allmodules" RunUnsupported="True" />
+  </Extend>
 </BuildGraph>
\ No newline at end of file
diff --git a/Engine/Build/LowLevelTests/NetworkPredictionPlugin.xml b/Engine/Build/LowLevelTests/NetworkPredictionPlugin.xml
--- a/Engine/Build/LowLevelTests/NetworkPredictionPlugin.xml
+++ b/Engine/Build/LowLevelTests/NetworkPredictionPlugin.xml
@@ -1,50 +1,8 @@
 ﻿<?xml version="1.0" encoding="utf-8"?>
 <BuildGraph xmlns:xsi="http://www.w3.org/2001/XMLSchema-instance" xsi:schemaLocation="http://www.epicgames.com/BuildGraph ../../Build/Graph/Schema.xsd" xmlns="http://www.epicgames.com/BuildGraph">
-  <Property Name="TestNames" Value="$(TestNames);NetworkPredictionPlugin" />
-  <Property Name="NetworkPredictionPluginUsesCatch2" Value="True" />
-  <Property Name="NetworkPredictionPluginDisabled" Value="False" />
-  <Property Name="NetworkPredictionPluginShort" Value="Net Prediction" />
-  <Property Name="NetworkPredictionPluginStagesWithProjectFile" Value="False" />
-  <Property Name="NetworkPredictionPluginTarget" Value="NetworkPredictionTests" />
-  <Property Name="NetworkPredictionPluginBinariesRelative" Value="Engine\Binaries" />
-  <Property Name="NetworkPredictionPluginReportType" Value="console" />
-  <Property Name="NetworkPredictionPluginGauntletArgs" Value="" />
-  <Property Name="NetworkPredictionPluginExtraArgs" />
-  <Property Name="NetworkPredictionPluginInitialExtraArgs" />
-  <Property Name="NetworkPredictionPluginHasAfterSteps" Value="False" />
-  <Property Name="NetworkPredictionPluginUsesCatch2" Value="True" />
   <Option Name="RunNetworkPredictionPluginTests" DefaultValue="False" Description="Run Net Prediction Tests" />
-  <Property Name="NetworkPredictionPluginSupportedPlatforms" Value="Win64;Linux" />
-  <Property Name="Win64NetworkPredictionPluginTags" Value="" />
-  <Property Name="MacNetworkPredictionPluginTags" Value="" />
-  <Property Name="IOSNetworkPredictionPluginTags" Value="" />
-  <Property Name="AndroidNetworkPredictionPluginTags" Value="" />
-  <Property Name="LinuxNetworkPredictionPluginTags" Value="" />
-  <Property Name="LinuxArm64NetworkPredictionPluginTags" Value="" />
-  <Property Name="TVOSNetworkPredictionPluginTags" Value="" />
-  <Property Name="Win64NetworkPredictionPluginExtraCompilationArgs" Value="" />
-  <Property Name="MacNetworkPredictionPluginExtraCompilationArgs" Value="" />
-  <Property Name="IOSNetworkPredictionPluginExtraCompilationArgs" Value="" />
-  <Property Name="AndroidNetworkPredictionPluginExtraCompilationArgs" Value="" />
-  <Property Name="LinuxNetworkPredictionPluginExtraCompilationArgs" Value="" />
-  <Property Name="LinuxArm64NetworkPredictionPluginExtraCompilationArgs" Value="" />
-  <Property Name="TVOSNetworkPredictionPluginExtraCompilationArgs" Value="" />
-  <Property Name="Win64NetworkPredictionPluginRunSupported" Value="True" />
-  <Property Name="MacNetworkPredictionPluginRunSupported" Value="True" />
-  <Property Name="IOSNetworkPredictionPluginRunSupported" Value="False" />
-  <Property Name="AndroidNetworkPredictionPluginRunSupported" Value="False" />
-  <Property Name="LinuxNetworkPredictionPluginRunSupported" Value="True" />
-  <Property Name="LinuxArm64NetworkPredictionPluginRunSupported" Value="True" />
-  <Property Name="TVOSNetworkPredictionPluginRunSupported" Value="False" />
-  <Property Name="VisionOSNetworkPredictionPluginTags" Value="" />
-  <Property Name="VisionOSNetworkPredictionPluginExtraCompilationArgs" Value="" />
-  <Property Name="VisionOSNetworkPredictionPluginRunSupported" Value="False" />
-  <Property Name="Win64NetworkPredictionPluginRunContainerized" Value="False" />
-  <Property Name="MacNetworkPredictionPluginRunContainerized" Value="False" />
-  <Property Name="IOSNetworkPredictionPluginRunContainerized" Value="False" />
-  <Property Name="AndroidNetworkPredictionPluginRunContainerized" Value="False" />
-  <Property Name="LinuxNetworkPredictionPluginRunContainerized" Value="False" />
-  <Property Name="LinuxArm64NetworkPredictionPluginRunContainerized" Value="False" />
-  <Property Name="TVOSNetworkPredictionPluginRunContainerized" Value="False" />
-  <Property Name="VisionOSNetworkPredictionPluginRunContainerized" Value="False" />
+  <Extend Name="RunAllTests">
+    <Expand Name="DeployAndTest" Platform="Win64" TestName="NetworkPredictionPlugin" ShortName="Net Prediction" TargetName="NetworkPredictionTests" BinaryRelativePath="Engine\Binaries" ReportType="console" />
+    <Expand Name="DeployAndTest" Platform="Linux" TestName="NetworkPredictionPlugin" ShortName="Net Prediction" TargetName="NetworkPredictionTests" BinaryRelativePath="Engine\Binaries" ReportType="console" />
+  </Extend>
 </BuildGraph>
\ No newline at end of file
diff --git a/Engine/Build/LowLevelTests/ReplicationSystem.xml b/Engine/Build/LowLevelTests/ReplicationSystem.xml
--- a/Engine/Build/LowLevelTests/ReplicationSystem.xml
+++ b/Engine/Build/LowLevelTests/ReplicationSystem.xml
@@ -1,50 +1,9 @@
 ﻿<?xml version="1.0" encoding="utf-8"?>
 <BuildGraph xmlns:xsi="http://www.w3.org/2001/XMLSchema-instance" xsi:schemaLocation="http://www.epicgames.com/BuildGraph ../../Build/Graph/Schema.xsd" xmlns="http://www.epicgames.com/BuildGraph">
-  <Property Name="TestNames" Value="$(TestNames);ReplicationSystem" />
-  <Property Name="ReplicationSystemUsesCatch2" Value="True" />
-  <Property Name="ReplicationSystemDisabled" Value="False" />
-  <Property Name="ReplicationSystemShort" Value="Replication System" />
-  <Property Name="ReplicationSystemStagesWithProjectFile" Value="False" />
-  <Property Name="ReplicationSystemTarget" Value="ReplicationSystemLowLevelTests" />
-  <Property Name="ReplicationSystemBinariesRelative" Value="Engine\Binaries" />
-  <Property Name="ReplicationSystemReportType" Value="console" />
-  <Property Name="ReplicationSystemGauntletArgs" Value="" />
-  <Property Name="ReplicationSystemExtraArgs" />
-  <Property Name="ReplicationSystemInitialExtraArgs" />
-  <Property Name="ReplicationSystemHasAfterSteps" Value="False" />
-  <Property Name="ReplicationSystemUsesCatch2" Value="True" />
   <Option Name="RunReplicationSystemTests" DefaultValue="False" Description="Run Replication System Tests" />
-  <Property Name="ReplicationSystemSupportedPlatforms" Value="Win64;Linux;LinuxArm64" />
-  <Property Name="Win64ReplicationSystemTags" Value="" />
-  <Property Name="MacReplicationSystemTags" Value="" />
-  <Property Name="IOSReplicationSystemTags" Value="" />
-  <Property Name="AndroidReplicationSystemTags" Value="" />
-  <Property Name="LinuxReplicationSystemTags" Value="" />
-  <Property Name="LinuxArm64ReplicationSystemTags" Value="" />
-  <Property Name="TVOSReplicationSystemTags" Value="" />
-  <Property Name="Win64ReplicationSystemExtraCompilationArgs" Value="" />
-  <Property Name="MacReplicationSystemExtraCompilationArgs" Value="" />
-  <Property Name="IOSReplicationSystemExtraCompilationArgs" Value="" />
-  <Property Name="AndroidReplicationSystemExtraCompilationArgs" Value="" />
-  <Property Name="LinuxReplicationSystemExtraCompilationArgs" Value="" />
-  <Property Name="LinuxArm64ReplicationSystemExtraCompilationArgs" Value="" />
-  <Property Name="TVOSReplicationSystemExtraCompilationArgs" Value="" />
-  <Property Name="Win64ReplicationSystemRunSupported" Value="True" />
-  <Property Name="MacReplicationSystemRunSupported" Value="True" />
-  <Property Name="IOSReplicationSystemRunSupported" Value="False" />
-  <Property Name="AndroidReplicationSystemRunSupported" Value="False" />
-  <Property Name="LinuxReplicationSystemRunSupported" Value="True" />
-  <Property Name="LinuxArm64ReplicationSystemRunSupported" Value="True" />
-  <Property Name="TVOSReplicationSystemRunSupported" Value="False" />
-  <Property Name="VisionOSReplicationSystemTags" Value="" />
-  <Property Name="VisionOSReplicationSystemExtraCompilationArgs" Value="" />
-  <Property Name="VisionOSReplicationSystemRunSupported" Value="False" />
-  <Property Name="Win64ReplicationSystemRunContainerized" Value="False" />
-  <Property Name="MacReplicationSystemRunContainerized" Value="False" />
-  <Property Name="IOSReplicationSystemRunContainerized" Value="False" />
-  <Property Name="AndroidReplicationSystemRunContainerized" Value="False" />
-  <Property Name="LinuxReplicationSystemRunContainerized" Value="False" />
-  <Property Name="LinuxArm64ReplicationSystemRunContainerized" Value="True" />
-  <Property Name="TVOSReplicationSystemRunContainerized" Value="False" />
-  <Property Name="VisionOSReplicationSystemRunContainerized" Value="False" />
+  <Extend Name="RunAllTests">
+    <Expand Name="DeployAndTest" Platform="Win64" TestName="ReplicationSystem" ShortName="Replication System" TargetName="ReplicationSystemLowLevelTests" BinaryRelativePath="Engine\Binaries" ReportType="console" />
+    <Expand Name="DeployAndTest" Platform="Linux" TestName="ReplicationSystem" ShortName="Replication System" TargetName="ReplicationSystemLowLevelTests" BinaryRelativePath="Engine\Binaries" ReportType="console" />
+    <Expand Name="DeployAndTest" Platform="LinuxArm64" TestName="ReplicationSystem" ShortName="Replication System" TargetName="ReplicationSystemLowLevelTests" BinaryRelativePath="Engine\Binaries" ReportType="console" RunContainerized="True" />
+  </Extend>
 </BuildGraph>
\ No newline at end of file
diff --git a/Engine/Build/LowLevelTests/StateGraph.xml b/Engine/Build/LowLevelTests/StateGraph.xml
--- a/Engine/Build/LowLevelTests/StateGraph.xml
+++ b/Engine/Build/LowLevelTests/StateGraph.xml
@@ -1,48 +1,7 @@
 ﻿<?xml version="1.0" encoding="utf-8"?>
 <BuildGraph xmlns:xsi="http://www.w3.org/2001/XMLSchema-instance" xsi:schemaLocation="http://www.epicgames.com/BuildGraph ../../Build/Graph/Schema.xsd" xmlns="http://www.epicgames.com/BuildGraph">
-  <Property Name="TestNames" Value="$(TestNames);StateGraph" />
-  <Property Name="StateGraphDisabled" Value="False" />
-  <Property Name="StateGraphShort" Value="State Graph" />
-  <Property Name="StateGraphStagesWithProjectFile" Value="False" />
-  <Property Name="StateGraphTarget" Value="StateGraphTests" />
-  <Property Name="StateGraphBinariesRelative" Value="Engine\Binaries" />
-  <Property Name="StateGraphReportType" Value="console" />
-  <Property Name="StateGraphGauntletArgs" Value="" />
-  <Property Name="StateGraphExtraArgs" />
-  <Property Name="StateGraphHasAfterSteps" Value="False" />
-  <Property Name="StateGraphUsesCatch2" Value="True" />
   <Option Name="RunStateGraphTests" DefaultValue="False" Description="Run State Graph Tests" />
-  <Property Name="StateGraphSupportedPlatforms" Value="Win64" />
-  <Property Name="Win64StateGraphTags" Value="" />
-  <Property Name="Win64StateGraphExtraCompilationArgs" Value="" />
-  <Property Name="Win64StateGraphRunSupported" Value="True" />
-  <Property Name="MacStateGraphTags" Value="" />
-  <Property Name="MacStateGraphExtraCompilationArgs" Value="" />
-  <Property Name="MacStateGraphRunSupported" Value="True" />
-  <Property Name="IOSStateGraphTags" Value="" />
-  <Property Name="IOSStateGraphExtraCompilationArgs" Value="" />
-  <Property Name="IOSStateGraphRunSupported" Value="False" />
-  <Property Name="AndroidStateGraphTags" Value="" />
-  <Property Name="AndroidStateGraphExtraCompilationArgs" Value="" />
-  <Property Name="AndroidStateGraphRunSupported" Value="False" />
-  <Property Name="LinuxStateGraphTags" Value="" />
-  <Property Name="LinuxStateGraphExtraCompilationArgs" Value="" />
-  <Property Name="LinuxStateGraphRunSupported" Value="True" />
-  <Property Name="LinuxArm64StateGraphTags" Value="" />
-  <Property Name="LinuxArm64StateGraphExtraCompilationArgs" Value="" />
-  <Property Name="LinuxArm64StateGraphRunSupported" Value="True" />
-  <Property Name="TVOSStateGraphTags" Value="" />
-  <Property Name="TVOSStateGraphExtraCompilationArgs" Value="" />
-  <Property Name="TVOSStateGraphRunSupported" Value="False" />
-  <Property Name="VisionOSStateGraphTags" Value="" />
-  <Property Name="VisionOSStateGraphExtraCompilationArgs" Value="" />
-  <Property Name="VisionOSStateGraphRunSupported" Value="False" />
-  <Property Name="Win64StateGraphRunContainerized" Value="False" />
-  <Property Name="MacStateGraphRunContainerized" Value="False" />
-  <Property Name="IOSStateGraphRunContainerized" Value="False" />
-  <Property Name="AndroidStateGraphRunContainerized" Value="False" />
-  <Property Name="LinuxStateGraphRunContainerized" Value="False" />
-  <Property Name="LinuxArm64StateGraphRunContainerized" Value="False" />
-  <Property Name="TVOSStateGraphRunContainerized" Value="False" />
-  <Property Name="VisionOSStateGraphRunContainerized" Value="False" />
+  <Extend Name="RunAllTests">
+    <Expand Name="DeployAndTest" Platform="Win64" TestName="StateGraph" ShortName="State Graph" TargetName="StateGraphTests" BinaryRelativePath="Engine\Binaries" ReportType="console" />
+  </Extend>
 </BuildGraph>
\ No newline at end of file
diff --git a/Engine/Build/LowLevelTests/SwitchboardListener.xml b/Engine/Build/LowLevelTests/SwitchboardListener.xml
new file mode 100644
--- /dev/null
+++ b/Engine/Build/LowLevelTests/SwitchboardListener.xml
@@ -0,0 +1,7 @@
+﻿<?xml version="1.0" encoding="utf-8"?>
+<BuildGraph xmlns:xsi="http://www.w3.org/2001/XMLSchema-instance" xsi:schemaLocation="http://www.epicgames.com/BuildGraph ../../Build/Graph/Schema.xsd" xmlns="http://www.epicgames.com/BuildGraph">
+  <Option Name="RunSwitchboardListenerTests" DefaultValue="False" Description="Run SwitchboardListener Tests" />
+  <Extend Name="RunAllTests">
+    <Expand Name="DeployAndTest" Platform="Win64" TestName="SwitchboardListener" ShortName="SwitchboardListener" TargetName="SwitchboardListenerTests" BinaryRelativePath="Engine\Binaries" ReportType="console" />
+  </Extend>
+</BuildGraph>
\ No newline at end of file
diff --git a/Engine/Build/LowLevelTests/ToolMenus.xml b/Engine/Build/LowLevelTests/ToolMenus.xml
new file mode 100644
--- /dev/null
+++ b/Engine/Build/LowLevelTests/ToolMenus.xml
@@ -0,0 +1,9 @@
+﻿<?xml version="1.0" encoding="utf-8"?>
+<BuildGraph xmlns:xsi="http://www.w3.org/2001/XMLSchema-instance" xsi:schemaLocation="http://www.epicgames.com/BuildGraph ../../Build/Graph/Schema.xsd" xmlns="http://www.epicgames.com/BuildGraph">
+  <Option Name="RunToolMenusTests" DefaultValue="False" Description="Run ToolMenus Tests" />
+  <Extend Name="RunAllTests">
+    <Expand Name="DeployAndTest" Platform="Win64" TestName="ToolMenus" ShortName="ToolMenus" TargetName="ToolMenusTests" BinaryRelativePath="Engine\Binaries" ReportType="xml" ExtraCompilationArgs="-allmodules" />
+    <Expand Name="DeployAndTest" Platform="Linux" TestName="ToolMenus" ShortName="ToolMenus" TargetName="ToolMenusTests" BinaryRelativePath="Engine\Binaries" ReportType="xml" ExtraCompilationArgs="-allmodules" />
+    <Expand Name="DeployAndTest" Platform="Mac" TestName="ToolMenus" ShortName="ToolMenus" TargetName="ToolMenusTests" BinaryRelativePath="Engine\Binaries" ReportType="xml" ExtraCompilationArgs="-allmodules" />
+  </Extend>
+</BuildGraph>
\ No newline at end of file
diff --git a/Engine/Build/LowLevelTests/Web.xml b/Engine/Build/LowLevelTests/Web.xml
--- a/Engine/Build/LowLevelTests/Web.xml
+++ b/Engine/Build/LowLevelTests/Web.xml
@@ -1,49 +1,11 @@
 ﻿<?xml version="1.0" encoding="utf-8"?>
 <BuildGraph xmlns:xsi="http://www.w3.org/2001/XMLSchema-instance" xsi:schemaLocation="http://www.epicgames.com/BuildGraph ../../Build/Graph/Schema.xsd" xmlns="http://www.epicgames.com/BuildGraph">
-	<Property Name="TestNames" Value="$(TestNames);Web" />
-	<Property Name="WebDisabled" Value="False" />
-	<Property Name="WebShort" Value="Web" />
-	<Property Name="WebStagesWithProjectFile" Value="False" />
-	<Property Name="WebTarget" Value="WebTests" />
-	<Property Name="WebBinariesRelative" Value="Engine\Binaries" />
-	<Property Name="WebReportType" Value="xml" />
-	<Property Name="WebGauntletArgs" Value="" />
-	<Property Name="WebExtraArgs" Value="" />
-	<Property Name="WebInitialExtraArgs" />
-	<Property Name="WebHasAfterSteps" Value="True" />
-	<Property Name="WebUsesCatch2" Value="True" />
-	<Option Name="RunWebTests" DefaultValue="False" Description="Run Web Tests" />
-	<Property Name="WebSupportedPlatforms" Value="Win64;Linux;Android;IOS;Mac" />
-	<Property Name="Win64WebTags" Value="" />
-	<Property Name="MacWebTags" Value="" />
-	<Property Name="IOSWebTags" Value="" />
-	<Property Name="AndroidWebTags" Value="" />
-	<Property Name="LinuxWebTags" Value="" />
-	<Property Name="LinuxArm64WebTags" Value="" />
-	<Property Name="TVOSWebTags" Value="" />
-	<Property Name="Win64WebExtraCompilationArgs" Value="-allmodules" />
-	<Property Name="MacWebExtraCompilationArgs" Value="-allmodules" />
-	<Property Name="IOSWebExtraCompilationArgs" Value="-allmodules" />
-	<Property Name="AndroidWebExtraCompilationArgs" Value="-allmodules -architectures=arm64" />
-	<Property Name="LinuxWebExtraCompilationArgs" Value="-allmodules" />
-	<Property Name="LinuxArm64WebExtraCompilationArgs" Value="-allmodules" />
-	<Property Name="TVOSWebExtraCompilationArgs" Value="-allmodules" />
-	<Property Name="Win64WebRunSupported" Value="True" />
-	<Property Name="MacWebRunSupported" Value="True" />
-	<Property Name="IOSWebRunSupported" Value="False" />
-	<Property Name="AndroidWebRunSupported" Value="False" />
-	<Property Name="LinuxWebRunSupported" Value="True" />
-	<Property Name="LinuxArm64WebRunSupported" Value="True" />
-	<Property Name="TVOSWebRunSupported" Value="False" />
-	<Property Name="VisionOSWebTags" Value="" />
-	<Property Name="VisionOSWebExtraCompilationArgs" Value="-allmodules" />
-	<Property Name="VisionOSWebRunSupported" Value="False" />
-	<Property Name="Win64WebRunContainerized" Value="False" />
-	<Property Name="MacWebRunContainerized" Value="False" />
-	<Property Name="IOSWebRunContainerized" Value="False" />
-	<Property Name="AndroidWebRunContainerized" Value="False" />
-	<Property Name="LinuxWebRunContainerized" Value="False" />
-	<Property Name="LinuxArm64WebRunContainerized" Value="False" />
-	<Property Name="TVOSWebRunContainerized" Value="False" />
-	<Property Name="VisionOSWebRunContainerized" Value="False" />
-</BuildGraph>
+  <Option Name="RunWebTests" DefaultValue="False" Description="Run Web Tests" />
+  <Extend Name="RunAllTests">
+    <Expand Name="DeployAndTest" Platform="Win64" TestName="Web" ShortName="Web" TargetName="WebTests" BinaryRelativePath="Engine\Binaries" ReportType="xml" HasAfterSteps="True" ExtraCompilationArgs="-allmodules" />
+    <Expand Name="DeployAndTest" Platform="Linux" TestName="Web" ShortName="Web" TargetName="WebTests" BinaryRelativePath="Engine\Binaries" ReportType="xml" HasAfterSteps="True" ExtraCompilationArgs="-allmodules" />
+    <Expand Name="DeployAndTest" Platform="Mac" TestName="Web" ShortName="Web" TargetName="WebTests" BinaryRelativePath="Engine\Binaries" ReportType="xml" HasAfterSteps="True" ExtraCompilationArgs="-allmodules" />
+    <Expand Name="DeployAndTest" Platform="Android" TestName="Web" ShortName="Web" TargetName="WebTests" BinaryRelativePath="Engine\Binaries" ReportType="xml" HasAfterSteps="True" ExtraCompilationArgs="-allmodules  -architectures=arm64" RunUnsupported="True" />
+    <Expand Name="DeployAndTest" Platform="IOS" TestName="Web" ShortName="Web" TargetName="WebTests" BinaryRelativePath="Engine\Binaries" ReportType="xml" HasAfterSteps="True" ExtraCompilationArgs="-allmodules" RunUnsupported="True" />
+  </Extend>
+</BuildGraph>
\ No newline at end of file
diff --git a/Engine/Config/BaseEngine.ini b/Engine/Config/BaseEngine.ini
--- a/Engine/Config/BaseEngine.ini
+++ b/Engine/Config/BaseEngine.ini
@@ -3505,9 +3505,6 @@ EnvironmentCollisionMaxShapes = 32
 ; Enable/Disable mechanism that provides a disk overflow store for patch chunks.
 bEnableDiskOverflowStore=true
 
-[LowLevelTestsSettings]
-bUpdateBuildGraphPropertiesFile=false
-
 [/Script/Engine.BlueprintGeneratedClass]
 ; Members of this struct property are inter-dependent and their initialization can be deferred, which
 ; means native defaults can match BP defaults at BP compile time while still requiring post-construct
diff --git a/Engine/Plugins/Experimental/GameplayGraph/Tests/GameplayGraphTests.Build.cs b/Engine/Plugins/Experimental/GameplayGraph/Tests/GameplayGraphTests.Build.cs
--- a/Engine/Plugins/Experimental/GameplayGraph/Tests/GameplayGraphTests.Build.cs
+++ b/Engine/Plugins/Experimental/GameplayGraph/Tests/GameplayGraphTests.Build.cs
@@ -4,6 +4,15 @@ using UnrealBuildTool;
 
 public class GameplayGraphTests : TestModuleRules
 {
+	static GameplayGraphTests()
+	{
+		TestMetadata = new Metadata()
+		{
+			TestName = "GameplayGraph",
+			TestShortName = "GameplayGraph"
+		};
+	}
+
 	public GameplayGraphTests(ReadOnlyTargetRules Target) : base(Target)
 	{
 		PrivateDependencyModuleNames.AddRange(
@@ -12,10 +21,5 @@ public class GameplayGraphTests : TestModuleRules
 				"CoreUObject",
 				"GameplayGraph"
 			});
-
-		UpdateBuildGraphPropertiesFile(new Metadata() {
-			TestName = "GameplayGraph",
-			TestShortName = "GameplayGraph"
-		});
 	}
 }
\ No newline at end of file
diff --git a/Engine/Plugins/Experimental/StateGraph/Tests/StateGraphTests/StateGraphTests.Build.cs b/Engine/Plugins/Experimental/StateGraph/Tests/StateGraphTests/StateGraphTests.Build.cs
--- a/Engine/Plugins/Experimental/StateGraph/Tests/StateGraphTests/StateGraphTests.Build.cs
+++ b/Engine/Plugins/Experimental/StateGraph/Tests/StateGraphTests/StateGraphTests.Build.cs
@@ -4,6 +4,13 @@ using UnrealBuildTool;
 
 public class StateGraphTests : TestModuleRules
 {
+	static StateGraphTests()
+	{
+		TestMetadata = new Metadata();
+		TestMetadata.TestName = "StateGraph";
+		TestMetadata.TestShortName = "State Graph";
+	}
+
 	public StateGraphTests(ReadOnlyTargetRules Target) : base(Target)
 	{
 		OptimizeCode = CodeOptimization.Never;
@@ -13,7 +20,5 @@ public class StateGraphTests : TestModuleRules
 				"Core",
 				"StateGraph"
 			});
-
-		UpdateBuildGraphPropertiesFile(new Metadata() { TestName = "StateGraph", TestShortName = "State Graph"});
 	}
 }
diff --git a/Engine/Source/Developer/LowLevelTestsRunner/README.md b/Engine/Source/Developer/LowLevelTestsRunner/README.md
--- a/Engine/Source/Developer/LowLevelTestsRunner/README.md
+++ b/Engine/Source/Developer/LowLevelTestsRunner/README.md
@@ -85,8 +85,8 @@ Follow these steps to create your explicit test:
 2. Inherit your module class from `TestModuleRules`.
     * If you are writing a test for a plugin, place the new module inside a `Tests` directory at the same level as the plugin's Source directory.
     * If you are building a test module that does not use `Catch2`, inherit the base constructor with the second parameter set to false: `base(Target, false)`.
-3. Call `UpdateBuildGraphPropertiesFile` with a new `Metadata` object argument.
-    * This information is used to generate BuildGraph script test metadata.
+3. Implement a static constructor for the class and initialize `TestMetadata` explicitly.
+    * The dat stored in this object is used to generate BuildGraph script test metadata.
     * For more information about BuildGraph script generation, see the [Generate BuildGraph Script Metadata](#generate-buildgraph-script-metadata-files) section.
 4. Suppose that you have an explicit test module titled `UEModuleTests`. Your explicit test `.Build.cs` file should look similar to this:
 	**UEModuleTests\UEModuleTests.Build.cs**
@@ -94,6 +94,14 @@ Follow these steps to create your explicit test:
 	```csharp
 	public class UEModuleTests : TestModuleRules
 	{
+        sttic UEModuleTests()
+        {
+          TestMetadata = new Metadata();
+		  TestMetadata.TestName = "UEModule";
+		  TestMetadata.TestShortName = "UE Module";
+		  // any other additional properties
+        }
+
         public UEModuleTests(ReadOnlyTargetRules Target) : base(Target)
         {
             PrivateIncludePaths.AddRange(
@@ -104,12 +112,12 @@ Follow these steps to create your explicit test:
 			);
 
 			// Other types of dependencies or module specific logic
-
-			UpdateBuildGraphPropertiesFile(new Metadata("UEModule", "UE Module"));
 		}
 	}
 	```
 	</font>
+
+	You can also extend this module class for a NDA platform, in which case you need to provide a static constructor in that class and call `TestMetadata.SupportedPlatforms.Add(...)` with the platform as parameter - do not re-initialize TestMetadata so that all the properties set in the base class are kept.
 5. Add a test target file (`.Target.cs`) with a class that inherits from `TestTargetRules`.
 6. Override the default compilation flags if necessary.
     * Aim for a minimal, testable module free of default features that don't add value to low-level testing.
@@ -137,18 +145,20 @@ Now you can write C++ test files in the `Private` folder of the module and write
 
 ### Generate BuildGraph Script Metadata Files
 
-If you want to build and run your tests with BuildGraph, you need to enable generation of BuildGraph script metadata files for explicit tests. When generating the IDE solution via `GenerateProjectFiles.bat`, the explicit test modules generate BuildGraph `.xml` files.
+If you want to build and run your tests with BuildGraph, you need generate BuildGraph script metadata files for explicit tests.
+All classes inheriting TestModuleRules have a static property called TestMetadata whose properties are translated into BuildGraph metadata and instructions.
 
-An engine configuration setting conditions this generation. You can set this configuration in `Engine/Config/BaseEngine.ini`:
-
-```xml
-[LowLevelTestsSettings]
-bUpdateBuildGraphPropertiesFile=true
+Run the following command at the root of the engine:
+```
+.\RunUBT.bat -Mode=Test -GenerateMetadata
 ```
 
-When you run `GenerateProjectFiles.bat`, test metadata `.xml` files are generated in the `Build/LowLevelTests/<TEST_NAME>.xml` folder for each test target, where `<TEST_NAME>` is the name of your test target. For NDA platforms, these files are generated under `Platforms/<PLATFORM_NAME>/Build/LowLevelTests/<TEST_NAME>.xml`. An additional `General.xml` file is optionally present next to the test files containing global properties.
+If your tests live inside a project, use `-ProjectFile=` with the absolute path to the *.uproject* file to generate metadata in the project's directory under the `Build\LowLevelTests` folder.
+
+
+This generates test metadata `.xml` files in the `Engine/Build/LowLevelTests/<TEST_NAME>.xml` folder for each test target, where `<TEST_NAME>` is the name of your test target. For NDA platforms, these files are generated under `Platforms/<PLATFORM_NAME>/Build/LowLevelTests/<TEST_NAME>.xml`. An additional `General.xml` file is optionally present next to the test files containing global properties.
 
-If the files already exist, they are updated according to the C#-described `Metadata` object. The folders and files that are accessed by project file generation must be writable. Typically, these files are read-only when under source control, so check them out or make them writeable before generation.
+These files are overwritten every time the command is executed and their contents depends on the TestModuleRules inherited class' `TestMetadata` object. The folders and files that are accessed by project file generation must be writable. Typically, these files are read-only when under source control, so check them out or make them writeable before generation.
 
 > **_TIP_**: To see an example, the directory Engine/Build/LowLevelTests contains an .xml file named Foundation. This is the generated BuildGraph metadata for the Foundation Tests.
 
@@ -156,7 +166,7 @@ If the files already exist, they are updated according to the C#-described `Meta
 
 #### Test Module Rules Reference
 
-The `TestModuleRules` class extends `ModuleRules` with `UpdateBuildGraphPropertiesFile`. `UpdateBuildGraphPropertiesFile` accepts a `Metadata` object which generates BuildGraph test metadata `.xml` files. With a `Metadata` object, you can set the following properties:
+The `TestModuleRules` class extends `ModuleRules` and adds a static `TestMetadata` filed which is used to generate BuildGraph test metadata `.xml` files. With a `Metadata` object type, you can set the following properties:
 * `TestName`: The name of your tests used by the BuildGraph script to generate test-specific properties. This field cannot contain spaces.
 * `TestShortName`: The short name of your tests used for display in the build system. This field can contain spaces.
 * `ReportType`: The Catch2 report type. The most common report types are console and xml. For more information about Catch2 report types, see the external Catch2 documentation.
@@ -545,7 +555,7 @@ You can build and run explicit tests directly from Visual Studio on desktop plat
 
 #### Build
 
-All low level tests must be built usint the UBT mode `-Mode=Test`. This mode performs a pre-analysis of the build target's dependency chain and detects UE core modules such as CoreUObject, Engine, ApplicationCore and UnrealEd and sets appropriate compilation flags. By default `-Mode=Test` is used to build explicit tests - tests that have their own explicitely defined target. To build an implicit test from an existing target use `-Mode=Test -Implicit`: note the additional -Implicit argument here.
+All low level tests must be built using the UBT mode `-Mode=Test`. This mode performs a pre-analysis of the build target's dependency chain and detects UE core modules such as CoreUObject, Engine, ApplicationCore and UnrealEd and sets appropriate compilation flags. By default `-Mode=Test` is used to build explicit tests - tests that have their own explicitely defined target. To build an implicit test from an existing target use `-Mode=Test -Implicit`: note the additional -Implicit argument here.
 
 ##### Explicit Tests
 
diff --git a/Engine/Source/Programs/AutomationTool/LowLevelTests/UpdateGraphsLowLevelTests.Automation.cs b/Engine/Source/Programs/AutomationTool/LowLevelTests/UpdateGraphsLowLevelTests.Automation.cs
--- a/Engine/Source/Programs/AutomationTool/LowLevelTests/UpdateGraphsLowLevelTests.Automation.cs
+++ b/Engine/Source/Programs/AutomationTool/LowLevelTests/UpdateGraphsLowLevelTests.Automation.cs
@@ -59,6 +59,7 @@ namespace LowLevelTests
 			}
 
 			//TODO(SMA): Enable bUpdateBuildGraphPropertiesFile in config and run UpdateBuildGraphPropertiesFile on all relevant .Build.cs files.
+			//Note: The metadata files are now generated/updated using RunUBT.bat -Mode=Test -GenerateMetadata
 
 			return ExitCode.Success;
 		}
diff --git a/Engine/Source/Programs/ChaosUserDataPTTests/ChaosUserDataPTTests.Build.cs b/Engine/Source/Programs/ChaosUserDataPTTests/ChaosUserDataPTTests.Build.cs
--- a/Engine/Source/Programs/ChaosUserDataPTTests/ChaosUserDataPTTests.Build.cs
+++ b/Engine/Source/Programs/ChaosUserDataPTTests/ChaosUserDataPTTests.Build.cs
@@ -4,6 +4,13 @@ using UnrealBuildTool;
 
 public class ChaosUserDataPTTests : TestModuleRules
 {
+	static ChaosUserDataPTTests()
+	{
+		TestMetadata = new Metadata();
+		TestMetadata.TestName = "ChasoUserDataPT";
+		TestMetadata.TestShortName = "Chaos User Data PT";
+	}
+
 	public ChaosUserDataPTTests(ReadOnlyTargetRules Target) : base(Target)
 	{
 		PrivateDependencyModuleNames.AddRange(
diff --git a/Engine/Source/Programs/HeadlessChaos/HeadlessChaos.Build.cs b/Engine/Source/Programs/HeadlessChaos/HeadlessChaos.Build.cs
--- a/Engine/Source/Programs/HeadlessChaos/HeadlessChaos.Build.cs
+++ b/Engine/Source/Programs/HeadlessChaos/HeadlessChaos.Build.cs
@@ -2,9 +2,16 @@
 
 using UnrealBuildTool;
 
-[SupportedPlatforms("Win64")]
 public class HeadlessChaos : TestModuleRules
 {
+	static HeadlessChaos()
+	{
+		TestMetadata = new Metadata();
+		TestMetadata.TestName = "HeadlessChaos";
+		TestMetadata.TestShortName = "Headless Chaos";
+		TestMetadata.UsesCatch2 = false;
+	}
+
 	public HeadlessChaos(ReadOnlyTargetRules Target) : base(Target, false)
 	{
 		PublicIncludePathModuleNames.Add("Launch");
@@ -49,7 +56,5 @@ public class HeadlessChaos : TestModuleRules
 		}
 
 		PrivateDefinitions.Add("CHAOS_INCLUDE_LEVEL_1=1");
-
-		UpdateBuildGraphPropertiesFile(new Metadata() { TestName = "HeadlessChaos", TestShortName = "Headless Chaos", UsesCatch2 = false });
 	}
 }
diff --git a/Engine/Source/Programs/LowLevelTests/FoundationTests/FoundationTests.Build.cs b/Engine/Source/Programs/LowLevelTests/FoundationTests/FoundationTests.Build.cs
--- a/Engine/Source/Programs/LowLevelTests/FoundationTests/FoundationTests.Build.cs
+++ b/Engine/Source/Programs/LowLevelTests/FoundationTests/FoundationTests.Build.cs
@@ -3,45 +3,17 @@ using UnrealBuildTool;
 
 public class FoundationTests : TestModuleRules
 {
-	protected Metadata FoundationTestsMetadata = new Metadata() {
-		TestName = "Foundation",
-		TestShortName = "Foundation",
-		ReportType = "xml",
-		SupportedPlatforms = {
-			UnrealTargetPlatform.Win64,
-			UnrealTargetPlatform.Linux,
-			UnrealTargetPlatform.Mac,
-			UnrealTargetPlatform.Android,
-			UnrealTargetPlatform.IOS } };
-
-	/// <summary>
-	/// Test metadata to be used with BuildGraph
-	/// </summary>
-	public Metadata TestMetadata
-	{ 
-		get { return FoundationTestsMetadata; }
-	}
-
-	public FoundationTests(ReadOnlyTargetRules Target) : base(Target, true)
+	static FoundationTests()
 	{
-		PrivateDependencyModuleNames.AddRange(
-			new string[] {
-				"Core",
-				"Cbor",
-				"CoreUObject",
-				"TelemetryUtils",
-				"AssetRegistry",
-				"Serialization",
-			});
-
-		if (Target.bBuildWithEditorOnlyData)
-		{
-			PrivateDependencyModuleNames.AddRange(
-				new string[] {
-					"DesktopPlatform"
-				});
-		}
-
+		TestMetadata = new Metadata();
+		TestMetadata.TestName = "Foundation";
+		TestMetadata.TestShortName = "Foundation";
+		TestMetadata.ReportType = "xml";
+		TestMetadata.SupportedPlatforms.Add(UnrealTargetPlatform.Linux);
+		TestMetadata.SupportedPlatforms.Add(UnrealTargetPlatform.Mac);
+		TestMetadata.SupportedPlatforms.Add(UnrealTargetPlatform.Android);
+		TestMetadata.SupportedPlatforms.Add(UnrealTargetPlatform.IOS);
+		
 		string PlatformCompilationArgs;
 		foreach (var Platform in UnrealTargetPlatform.GetValidPlatforms())
 		{
@@ -63,7 +35,25 @@ public class FoundationTests : TestModuleRules
 		// Allow Android run for this test
 		// Will remove Android from PlatformsRunUnsupported as more diverse types of tests can run on this platform
 		TestMetadata.PlatformsRunUnsupported.Remove(UnrealTargetPlatform.Android);
+	}
+	public FoundationTests(ReadOnlyTargetRules Target) : base(Target, true)
+	{
+		PrivateDependencyModuleNames.AddRange(
+			new string[] {
+				"Core",
+				"Cbor",
+				"CoreUObject",
+				"TelemetryUtils",
+				"AssetRegistry",
+				"Serialization",
+			});
 
-		UpdateBuildGraphPropertiesFile(TestMetadata);
+		if (Target.bBuildWithEditorOnlyData)
+		{
+			PrivateDependencyModuleNames.AddRange(
+				new string[] {
+					"DesktopPlatform"
+				});
+		}
 	}
 }
\ No newline at end of file
diff --git a/Engine/Source/Programs/LowLevelTests/MathCoreTests/MathCoreTests.Build.cs b/Engine/Source/Programs/LowLevelTests/MathCoreTests/MathCoreTests.Build.cs
--- a/Engine/Source/Programs/LowLevelTests/MathCoreTests/MathCoreTests.Build.cs
+++ b/Engine/Source/Programs/LowLevelTests/MathCoreTests/MathCoreTests.Build.cs
@@ -3,32 +3,16 @@ using UnrealBuildTool;
 
 public class MathCoreTests : TestModuleRules
 {
-	protected Metadata MathCoreTestsMetadata = new Metadata() {
-		TestName = "MathCore",
-		TestShortName = "MathCore",
-		ReportType = "xml",
-		SupportedPlatforms = {
-			UnrealTargetPlatform.Win64,
-			UnrealTargetPlatform.Linux,
-			UnrealTargetPlatform.Mac,
-			UnrealTargetPlatform.Android,
-			UnrealTargetPlatform.IOS } };
-
-	/// <summary>
-	/// Test metadata to be used with BuildGraph
-	/// </summary>
-	public Metadata TestMetadata
-	{ 
-		get { return MathCoreTestsMetadata; }
-	}
-
-	public MathCoreTests(ReadOnlyTargetRules Target) : base(Target, InUsesCatch2:true)
+	static MathCoreTests()
 	{
-		PrivateDependencyModuleNames.AddRange(
-			new string[] {
-				"Core",
-				"MathCore",
-			});
+		TestMetadata = new Metadata();
+		TestMetadata.TestName = "MathCore";
+		TestMetadata.TestShortName = "MathCore";
+		TestMetadata.ReportType = "xml";
+		TestMetadata.SupportedPlatforms.Add(UnrealTargetPlatform.Linux);
+		TestMetadata.SupportedPlatforms.Add(UnrealTargetPlatform.Mac);
+		TestMetadata.SupportedPlatforms.Add(UnrealTargetPlatform.Android);
+		TestMetadata.SupportedPlatforms.Add(UnrealTargetPlatform.IOS);
 
 		string PlatformCompilationArgs;
 		foreach (var Platform in UnrealTargetPlatform.GetValidPlatforms())
@@ -44,6 +28,14 @@ public class MathCoreTests : TestModuleRules
 			TestMetadata.PlatformCompilationExtraArgs.Add(Platform, PlatformCompilationArgs);
 		}
 		TestMetadata.PlatformsRunUnsupported.Remove(UnrealTargetPlatform.Android);
-		UpdateBuildGraphPropertiesFile(TestMetadata);
+	}
+
+	public MathCoreTests(ReadOnlyTargetRules Target) : base(Target, InUsesCatch2:true)
+	{
+		PrivateDependencyModuleNames.AddRange(
+			new string[] {
+				"Core",
+				"MathCore",
+			});
 	}
 }
\ No newline at end of file
diff --git a/Engine/Source/Programs/NetworkPredictionTests/NetworkPredictionTests.Build.cs b/Engine/Source/Programs/NetworkPredictionTests/NetworkPredictionTests.Build.cs
--- a/Engine/Source/Programs/NetworkPredictionTests/NetworkPredictionTests.Build.cs
+++ b/Engine/Source/Programs/NetworkPredictionTests/NetworkPredictionTests.Build.cs
@@ -5,6 +5,14 @@ using UnrealBuildTool;
 [SupportedPlatforms("Win64", "Linux")]
 public class NetworkPredictionTests : TestModuleRules
 {
+	static NetworkPredictionTests()
+	{
+		TestMetadata = new Metadata();
+		TestMetadata.TestName = "NetworkPredictionPlugin";
+		TestMetadata.TestShortName = "Net Prediction";
+		TestMetadata.SupportedPlatforms.Add(UnrealTargetPlatform.Linux);
+	}
+
 	public NetworkPredictionTests(ReadOnlyTargetRules Target) : base(Target)
 	{
 		PrivateDependencyModuleNames.AddRange(
@@ -15,11 +23,5 @@ public class NetworkPredictionTests : TestModuleRules
 				"NetworkPrediction",
 			}
 		);
-
-		UpdateBuildGraphPropertiesFile(new Metadata() {
-			TestName = "NetworkPredictionPlugin",
-			TestShortName = "Net Prediction",
-			SupportedPlatforms = { UnrealTargetPlatform.Win64, UnrealTargetPlatform.Linux }
-		});
 	}
 }
diff --git a/Engine/Source/Programs/PlainPropsTests/PlainPropsTests.Build.cs b/Engine/Source/Programs/PlainPropsTests/PlainPropsTests.Build.cs
--- a/Engine/Source/Programs/PlainPropsTests/PlainPropsTests.Build.cs
+++ b/Engine/Source/Programs/PlainPropsTests/PlainPropsTests.Build.cs
@@ -3,20 +3,25 @@ using UnrealBuildTool;
 
 public class PlainPropsTests : TestModuleRules
 {
-	protected Metadata PlainPropsMetadata = new Metadata() {
-		TestName = "PlainProps",
-		TestShortName = "PlainProps",
-		ReportType = "xml",
-		SupportedPlatforms = {
-			UnrealTargetPlatform.Win64,
-			UnrealTargetPlatform.Linux,
-			UnrealTargetPlatform.Mac} };
+	static PlainPropsTests()
+	{
+		TestMetadata = new Metadata();
+		TestMetadata.TestName = "PlainProps";
+		TestMetadata.TestShortName = "PlainProps";
+		TestMetadata.ReportType = "xml";
+		TestMetadata.SupportedPlatforms.Add(UnrealTargetPlatform.Linux);
+		TestMetadata.SupportedPlatforms.Add(UnrealTargetPlatform.Mac);
+		
+		string PlatformCompilationArgs;
+		foreach (var Platform in UnrealTargetPlatform.GetValidPlatforms())
+		{
+			PlatformCompilationArgs = "-allmodules";
+			TestMetadata.PlatformCompilationExtraArgs.Add(Platform, PlatformCompilationArgs);
+		}
 
-	public Metadata TestMetadata
-	{ 
-		get { return PlainPropsMetadata; }
+		// Platform-specific tags
+		TestMetadata.PlatformTags.Add(UnrealTargetPlatform.Linux, "~[.]~[Slow]");
 	}
-
 	public PlainPropsTests(ReadOnlyTargetRules Target) : base(Target, true)
 	{
 		PrivateDependencyModuleNames.AddRange(
@@ -29,17 +34,5 @@ public class PlainPropsTests : TestModuleRules
 		{
 			PrivateDependencyModuleNames.AddRange(new string[] { "DesktopPlatform" });
 		}
-
-		string PlatformCompilationArgs;
-		foreach (var Platform in UnrealTargetPlatform.GetValidPlatforms())
-		{
-			PlatformCompilationArgs = "-allmodules";
-			TestMetadata.PlatformCompilationExtraArgs.Add(Platform, PlatformCompilationArgs);
-		}
-
-		// Platform-specific tags
-		TestMetadata.PlatformTags.Add(UnrealTargetPlatform.Linux, "~[.]~[Slow]");
-
-		UpdateBuildGraphPropertiesFile(TestMetadata);
 	}
 }
\ No newline at end of file
diff --git a/Engine/Source/Programs/ReplicationSystemLowLevelTests/ReplicationSystemLowLevelTests.Build.cs b/Engine/Source/Programs/ReplicationSystemLowLevelTests/ReplicationSystemLowLevelTests.Build.cs
--- a/Engine/Source/Programs/ReplicationSystemLowLevelTests/ReplicationSystemLowLevelTests.Build.cs
+++ b/Engine/Source/Programs/ReplicationSystemLowLevelTests/ReplicationSystemLowLevelTests.Build.cs
@@ -4,6 +4,16 @@ using UnrealBuildTool;
 
 public class ReplicationSystemLowLevelTests : TestModuleRules
 {
+	static ReplicationSystemLowLevelTests()
+	{
+		TestMetadata = new Metadata();
+		TestMetadata.TestName = "ReplicationSystem";
+		TestMetadata.TestShortName = "Replication System";
+		TestMetadata.SupportedPlatforms.Add(UnrealTargetPlatform.Linux);
+		TestMetadata.SupportedPlatforms.Add(UnrealTargetPlatform.LinuxArm64);
+		TestMetadata.PlatformRunContainerized.Add(UnrealTargetPlatform.LinuxArm64, true);
+	}
+
 	public ReplicationSystemLowLevelTests(ReadOnlyTargetRules Target) : base(Target)
 	{
 		PrivateDependencyModuleNames.AddRange(
@@ -26,12 +36,5 @@ public class ReplicationSystemLowLevelTests : TestModuleRules
 				"ApplicationCore"
 			}
 		);
-
-		UpdateBuildGraphPropertiesFile(new Metadata() {
-			TestName = "ReplicationSystem",
-			TestShortName = "Replication System",
-			SupportedPlatforms = { UnrealTargetPlatform.Win64, UnrealTargetPlatform.Linux, UnrealTargetPlatform.LinuxArm64 },
-			PlatformRunContainerized = { { UnrealTargetPlatform.LinuxArm64, true } }
-		});
 	}
 }
diff --git a/Engine/Source/Programs/SwitchboardListener/SwitchboardListenerTests/SwitchboardListenerTests.Build.cs b/Engine/Source/Programs/SwitchboardListener/SwitchboardListenerTests/SwitchboardListenerTests.Build.cs
--- a/Engine/Source/Programs/SwitchboardListener/SwitchboardListenerTests/SwitchboardListenerTests.Build.cs
+++ b/Engine/Source/Programs/SwitchboardListener/SwitchboardListenerTests/SwitchboardListenerTests.Build.cs
@@ -5,6 +5,13 @@ using UnrealBuildTool;
 
 public class SwitchboardListenerTests : TestModuleRules
 {
+	static SwitchboardListenerTests()
+	{
+		TestMetadata = new Metadata();
+		TestMetadata.TestName = "SwitchboardListener";
+		TestMetadata.TestShortName = "SwitchboardListener";
+	}
+
 	public SwitchboardListenerTests(ReadOnlyTargetRules Target) : base(Target)
 	{
 		PrivateDependencyModuleNames.AddRange(new string[] {
@@ -14,7 +21,5 @@ public class SwitchboardListenerTests : TestModuleRules
 		});
 
 		PrivateIncludePathModuleNames.Add("LowLevelTestsRunner");
-
-		UpdateBuildGraphPropertiesFile(new Metadata() { TestName = "SwitchboardListener", TestShortName = "SwitchboardListener" });
 	}
 }
diff --git a/Engine/Source/Programs/ToolMenusTests/ToolMenusTests.Build.cs b/Engine/Source/Programs/ToolMenusTests/ToolMenusTests.Build.cs
--- a/Engine/Source/Programs/ToolMenusTests/ToolMenusTests.Build.cs
+++ b/Engine/Source/Programs/ToolMenusTests/ToolMenusTests.Build.cs
@@ -3,21 +3,28 @@ using UnrealBuildTool;
 
 public class ToolMenusTests : TestModuleRules
 {
-	protected Metadata ToolMenusTestsMetadata = new Metadata() {
-		TestName = "ToolMenus",
-		TestShortName = "ToolMenus",
-		ReportType = "xml",
-		SupportedPlatforms = {
-			UnrealTargetPlatform.Win64,
-			UnrealTargetPlatform.Linux,
-			UnrealTargetPlatform.Mac } };
-
-	/// <summary>
-	/// Test metadata to be used with BuildGraph
-	/// </summary>
-	public Metadata TestMetadata
-	{ 
-		get { return ToolMenusTestsMetadata; }
+	static ToolMenusTests()
+	{
+		TestMetadata = new Metadata();
+		TestMetadata.TestName = "ToolMenus";
+		TestMetadata.TestShortName = "ToolMenus";
+		TestMetadata.ReportType = "xml";
+		TestMetadata.SupportedPlatforms.Add(UnrealTargetPlatform.Linux);
+		TestMetadata.SupportedPlatforms.Add(UnrealTargetPlatform.Mac);
+		
+		string PlatformCompilationArgs;
+		foreach (var Platform in UnrealTargetPlatform.GetValidPlatforms())
+		{
+			if (Platform == UnrealTargetPlatform.Android)
+			{
+				PlatformCompilationArgs = "-allmodules -architectures=arm64";
+			}
+			else
+			{
+				PlatformCompilationArgs = "-allmodules";
+			}
+			TestMetadata.PlatformCompilationExtraArgs.Add(Platform, PlatformCompilationArgs);
+		}
 	}
 
 	public ToolMenusTests(ReadOnlyTargetRules Target) : base(Target, true)
@@ -38,20 +45,5 @@ public class ToolMenusTests : TestModuleRules
 					"DesktopPlatform"
 				});
 		}
-
-		string PlatformCompilationArgs;
-		foreach (var Platform in UnrealTargetPlatform.GetValidPlatforms())
-		{
-			if (Platform == UnrealTargetPlatform.Android)
-			{
-				PlatformCompilationArgs = "-allmodules -architectures=arm64";
-			}
-			else
-			{
-				PlatformCompilationArgs = "-allmodules";
-			}
-			TestMetadata.PlatformCompilationExtraArgs.Add(Platform, PlatformCompilationArgs);
-		}
-		UpdateBuildGraphPropertiesFile(TestMetadata);
 	}
 }
diff --git a/Engine/Source/Programs/UnrealBuildTool/Configuration/TestModuleRules.cs b/Engine/Source/Programs/UnrealBuildTool/Configuration/TestModuleRules.cs
--- a/Engine/Source/Programs/UnrealBuildTool/Configuration/TestModuleRules.cs
+++ b/Engine/Source/Programs/UnrealBuildTool/Configuration/TestModuleRules.cs
@@ -6,6 +6,7 @@ using System.IO;
 using System.Linq;
 using System.Xml.Linq;
 using EpicGames.Core;
+using Microsoft.Extensions.Logging;
 using UnrealBuildBase;
 
 namespace UnrealBuildTool
@@ -15,9 +16,15 @@ namespace UnrealBuildTool
 	/// </summary>
 	public class TestModuleRules : ModuleRules
 	{
-		private readonly XNamespace BuildGraphNamespace = XNamespace.Get("http://www.epicgames.com/BuildGraph");
-		private readonly XNamespace SchemaInstance = XNamespace.Get("http://www.w3.org/2001/XMLSchema-instance");
-		private readonly XNamespace SchemaLocation = XNamespace.Get("http://www.epicgames.com/BuildGraph ../../Build/Graph/Schema.xsd");
+		private static readonly XNamespace BuildGraphNamespace = XNamespace.Get("http://www.epicgames.com/BuildGraph");
+		private static readonly XNamespace SchemaInstance = XNamespace.Get("http://www.w3.org/2001/XMLSchema-instance");
+		private static readonly XNamespace SchemaLocation = XNamespace.Get("http://www.epicgames.com/BuildGraph ../../Build/Graph/Schema.xsd");
+		private static readonly List<string> RestrictedFoldersNonPlatform = new List<string>() {
+			RestrictedFolder.NotForLicensees.ToString(),
+			RestrictedFolder.NoRedist.ToString(),
+			RestrictedFolder.EpicInternal.ToString(),
+			RestrictedFolder.CarefullyRedist.ToString()
+		};
 
 		private bool bUsesCatch2 = true;
 
@@ -26,6 +33,11 @@ namespace UnrealBuildTool
 		/// </summary>
 		public ModuleRules? TestedModule { get; private set; }
 
+		/// <summary>
+		/// Test metadata, used with BuildGraph only.
+		/// </summary>
+		protected static Metadata TestMetadata = new Metadata();
+
 		/// <summary>
 		/// Constructs a TestModuleRules object as its own test module.
 		/// </summary>
@@ -153,65 +165,29 @@ namespace UnrealBuildTool
 
 #pragma warning disable 8602
 #pragma warning disable 8604
+
 		/// <summary>
-		/// Generates or updates metadata file for LowLevelTests.xml containing test flags: name, short name, target name, relative binaries path, supported platforms etc.
-		/// <paramref name="TestMetadata">The test metadata specifying name, short name etc used to populate the BuildGraph properties file.</paramref>
+		/// Deprecated, test metadata now generated explicitly using -Mode-Test with -GenerateMetadata.
 		/// </summary>
+		/// <param name="TestMetadata"></param>
+		[Obsolete("Use RunUBT -Mode=Test -GenerateMetadata instead")]
 		protected void UpdateBuildGraphPropertiesFile(Metadata TestMetadata)
 		{
-			bool bUpdateBuildGraphPropertiesFile = false;
-			TestTargetRules? TestTargetRules = Target.InnerTestTargetRules;
-			if (TestTargetRules != null)
-			{
-				bUpdateBuildGraphPropertiesFile = TestTargetRules.bUpdateBuildGraphPropertiesFile;
-			}
-
-			bool bIsBuildMachine = Unreal.IsBuildMachine();
-			if (bIsBuildMachine || !bUpdateBuildGraphPropertiesFile || TestMetadata == null)
-			{
-				return;
-			}
-
-			string BaseFolder = GetBaseFolder();
-
-			string GeneratedPropertiesScriptFile;
+		}
 
-			string NonPublicPath = Path.Combine(BaseFolder, "Restricted", "NotForLicensees", "Build", "LowLevelTests", $"{TestMetadata.TestName}.xml");
+		/// <summary>
+		/// Generates or updates metadata file for LowLevelTests.xml containing test flags: name, short name, target name, relative binaries path, supported platforms etc.
+		/// Called by RunUBT.bat -Mode=Test -GenerateMetadata
+		/// </summary>
+		private static void UpdateBuildGraphMetadata(Metadata TestMetadata, string ModuleDirectory, string ModuleName, ILogger Log)
+		{
+			string BaseFolder = GetBaseFolder(ModuleDirectory);
 
 			bool ModuleInRestrictedPath = IsRestrictedPath(ModuleDirectory);
 
-			if (ModuleInRestrictedPath)
-			{
-				GeneratedPropertiesScriptFile = NonPublicPath;
-			}
-			else
-			{
-				GeneratedPropertiesScriptFile = Path.Combine(BaseFolder, "Build", "LowLevelTests", $"{TestMetadata.TestName}.xml");
-			}
-
-			if (!System.IO.File.Exists(GeneratedPropertiesScriptFile))
-			{
-				string? DirGenProps = Path.GetDirectoryName(GeneratedPropertiesScriptFile);
-				if (DirGenProps != null && !System.IO.Directory.Exists(DirGenProps))
-				{
-					System.IO.Directory.CreateDirectory(DirGenProps);
-				}
-				using (FileStream FileStream = System.IO.File.Create(GeneratedPropertiesScriptFile))
-				{
-					XDocument XInitFile = new XDocument(new XElement(BuildGraphNamespace + "BuildGraph", new XAttribute(XNamespace.Xmlns + "xsi", SchemaInstance), new XAttribute(SchemaInstance + "schemaLocation", SchemaLocation)));
-					XInitFile.Root?.Add(
-						new XElement(
-							BuildGraphNamespace + "Property",
-							new XAttribute("Name", "TestNames"),
-							new XAttribute("Value", "$(TestNames);" + TestMetadata.TestName)));
-
-					XInitFile.Save(FileStream);
-				}
-			}
-
 			// All relevant properties
-			string TestTargetName = Target.LaunchModuleName ?? "Launch";
-			string TestBinariesPath = TryGetBinariesPath();
+			string TestTargetName = ModuleName ?? "Launch";
+			string TestBinariesPath = TryGetBinariesPath(ModuleDirectory);
 
 			// Do not save full paths
 			if (Path.IsPathRooted(TestBinariesPath))
@@ -219,43 +195,15 @@ namespace UnrealBuildTool
 				TestBinariesPath = Path.GetRelativePath(Unreal.RootDirectory.FullName, TestBinariesPath);
 			}
 
-			MakeFileWriteable(GeneratedPropertiesScriptFile);
-			XDocument GenPropsDoc = XDocument.Load(GeneratedPropertiesScriptFile);
-			XElement? Root = GenPropsDoc.Root;
-			// First descendant must be TestNames
-			if (Root != null && Root.FirstNode != null)
-			{
-				XElement TestNames = (XElement)Root.FirstNode;
-				if (TestNames != null)
-				{
-					XElement lastUpdatedNode = TestNames;
-
-					InsertOrUpdateTestFlagProperty(ref lastUpdatedNode, TestMetadata.TestName, "Disabled", Convert.ToString(TestMetadata.Disabled));
-					InsertOrUpdateTestFlagProperty(ref lastUpdatedNode, TestMetadata.TestName, "Short", Convert.ToString(TestMetadata.TestShortName));
-					InsertOrUpdateTestFlagProperty(ref lastUpdatedNode, TestMetadata.TestName, "StagesWithProjectFile", Convert.ToString(TestMetadata.StagesWithProjectFile));
-					InsertOrUpdateTestFlagProperty(ref lastUpdatedNode, TestMetadata.TestName, "Target", Convert.ToString(TestTargetName));
-					InsertOrUpdateTestFlagProperty(ref lastUpdatedNode, TestMetadata.TestName, "BinariesRelative", Convert.ToString(TestBinariesPath));
-					InsertOrUpdateTestFlagProperty(ref lastUpdatedNode, TestMetadata.TestName, "ReportType", Convert.ToString(TestMetadata.ReportType));
-					InsertOrUpdateTestFlagProperty(ref lastUpdatedNode, TestMetadata.TestName, "GauntletArgs", Convert.ToString(TestMetadata.InitialExtraArgs) + Convert.ToString(TestMetadata.GauntletArgs));
-					InsertOrUpdateTestFlagProperty(ref lastUpdatedNode, TestMetadata.TestName, "ExtraArgs", Convert.ToString(TestMetadata.ExtraArgs));
-					InsertOrUpdateTestFlagProperty(ref lastUpdatedNode, TestMetadata.TestName, "HasAfterSteps", Convert.ToString(TestMetadata.HasAfterSteps));
-					InsertOrUpdateTestFlagProperty(ref lastUpdatedNode, TestMetadata.TestName, "UsesCatch2", Convert.ToString(TestMetadata.UsesCatch2));
-
-					InsertOrUpdateTestOption(ref lastUpdatedNode, TestMetadata.TestName, $"Run {TestMetadata.TestShortName} Tests", "Run", "Tests", false.ToString());
-
-					InsertOrUpdateTestFlagProperty(ref lastUpdatedNode, TestMetadata.TestName, "SupportedPlatforms", TestMetadata.SupportedPlatforms.Aggregate("", (current, next) => (String.IsNullOrEmpty(current) ? next.ToString() : current + ";" + next.ToString())));
-				}
-			}
-
-			GenPropsDoc.Save(GeneratedPropertiesScriptFile);
-
 			// Platform-specific configurations
 			string GeneratedPropertiesPlatformFile;
 
 			string NonPublicPathPlatform;
 
-			// Generate peroperty file for each platform
-			foreach (UnrealTargetPlatform ValidPlatform in UnrealTargetPlatform.GetValidPlatforms())
+			Dictionary<string, XDocument> SaveAtEnd = new Dictionary<string, XDocument>();
+
+			// Generate peroperty file for each supported platform
+			foreach (UnrealTargetPlatform ValidPlatform in TestMetadata.SupportedPlatforms)
 			{
 				bool IsRestrictedPlatformName = IsPlatformRestricted(ValidPlatform);
 				if (IsRestrictedPlatformName)
@@ -297,26 +245,100 @@ namespace UnrealBuildTool
 				}
 
 				MakeFileWriteable(GeneratedPropertiesPlatformFile);
-				XDocument XInitPlatformFile = XDocument.Load(GeneratedPropertiesPlatformFile);
+				XElement Root;
+				if (!SaveAtEnd.ContainsKey(GeneratedPropertiesPlatformFile))
+				{
+					XDocument XInitPlatformFile = XDocument.Load(GeneratedPropertiesPlatformFile);
+					// Any manually edited elements to keep
+					List<XElement> KeepElements = XInitPlatformFile.Root!.Elements().Where(e => e.Attribute("Name").Value == $"{TestMetadata.TestName}AfterSteps").ToList();
+					XInitPlatformFile.Root!.Elements().Remove();
+					foreach (var Element in KeepElements)
+					{
+						XInitPlatformFile.Root!.Add(Element);
+					}
+					SaveAtEnd.Add(GeneratedPropertiesPlatformFile, XInitPlatformFile);
+				}
+
+				Root = SaveAtEnd[GeneratedPropertiesPlatformFile].Root!;
 
-				// Adding per-test and per-platform tags
+				// Optional metadata, use Expand and set any non-default metadata
+				Dictionary<string, string> ExpandArguments = new Dictionary<string, string>();
+
+				if (!IsRestrictedPlatformName)
+				{
+					InsertOrUpdateTestOption(Root, $"Run{TestMetadata.TestName}Tests", $"Run {TestMetadata.TestShortName} Tests", false.ToString());
+				}
+
+				if (TestMetadata.Deactivated)
+				{
+					ExpandArguments.Add("Deactivated", Convert.ToString(TestMetadata.Deactivated));
+				}
+				ExpandArguments.Add("TestName", Convert.ToString(TestMetadata.TestName));
+				ExpandArguments.Add("ShortName", Convert.ToString(TestMetadata.TestShortName));
+				if (TestMetadata.StagesWithProjectFile)
+				{
+					ExpandArguments.Add("StagesWithProjectFile", Convert.ToString(TestMetadata.StagesWithProjectFile));
+				}
+				ExpandArguments.Add("TargetName", Convert.ToString(TestTargetName));
+				ExpandArguments.Add("BinaryRelativePath", Convert.ToString(TestBinariesPath));
+				ExpandArguments.Add("ReportType", Convert.ToString(TestMetadata.ReportType));
+				if (!string.IsNullOrEmpty(TestMetadata.GauntletArgs))
+				{
+					ExpandArguments.Add("GauntletArgs", Convert.ToString(TestMetadata.InitialExtraArgs) + Convert.ToString(TestMetadata.GauntletArgs));
+				}
+				if(TestMetadata.PlatformGauntletArgs.ContainsKey(ValidPlatform))
+				{
+					ExpandArguments.Add("PlatformGauntletArgs", TestMetadata.PlatformGauntletArgs[ValidPlatform]);
+				}
+				if (!string.IsNullOrEmpty(TestMetadata.ExtraArgs))
+				{
+					ExpandArguments.Add("ExtraArgs", Convert.ToString(TestMetadata.ExtraArgs));
+				}
+				if (TestMetadata.HasAfterSteps)
+				{
+					ExpandArguments.Add("HasAfterSteps", Convert.ToString(TestMetadata.HasAfterSteps));
+				}
+				if (!TestMetadata.UsesCatch2)
+				{
+					ExpandArguments.Add("UsesCatch2", Convert.ToString(TestMetadata.UsesCatch2));
+				}
 				string TagsValue = TestMetadata.PlatformTags.ContainsKey(ValidPlatform) ? TestMetadata.PlatformTags[ValidPlatform] : String.Empty;
-				AppendOrUpdateTestFlagProperty(ref XInitPlatformFile, TestMetadata.TestName, ValidPlatform.ToString(), "Tags", TagsValue);
+				if (!string.IsNullOrEmpty(TagsValue))
+				{
+					ExpandArguments.Add("Tags", TagsValue);
+					
+				}
 
 				string ExtraCompilationArgsValue = TestMetadata.PlatformCompilationExtraArgs.ContainsKey(ValidPlatform) ? TestMetadata.PlatformCompilationExtraArgs[ValidPlatform] : String.Empty;
-				AppendOrUpdateTestFlagProperty(ref XInitPlatformFile, TestMetadata.TestName, ValidPlatform.ToString(), "ExtraCompilationArgs", ExtraCompilationArgsValue);
+				if (!string.IsNullOrEmpty (ExtraCompilationArgsValue))
+				{
+					ExpandArguments.Add("ExtraCompilationArgs", ExtraCompilationArgsValue);
+				}
 
-				string RunSupportedValue = TestMetadata.PlatformsRunUnsupported.Contains(ValidPlatform) ? "False" : "True";
-				AppendOrUpdateTestFlagProperty(ref XInitPlatformFile, TestMetadata.TestName, ValidPlatform.ToString(), "RunSupported", RunSupportedValue);
+				// By default all test supported platforms have run supported, generally only a few don't (e.g. iOS)
+				bool RunUnsupportedPlatform = TestMetadata.PlatformsRunUnsupported.Contains(ValidPlatform);
+				if (RunUnsupportedPlatform)
+				{
+					ExpandArguments.Add("RunUnsupported", "True");
+				}
 
 				string RunContainerizedValue = TestMetadata.PlatformRunContainerized.ContainsKey(ValidPlatform) ? "True" : "False";
-				AppendOrUpdateTestFlagProperty(ref XInitPlatformFile, TestMetadata.TestName, ValidPlatform.ToString(), "RunContainerized", RunContainerizedValue);
+				if (RunContainerizedValue == "True")
+				{
+					ExpandArguments.Add("RunContainerized", RunContainerizedValue);
+				}
+
+				AppendOrUpdateRunAllTestsNode(Root, "DeployAndTest", ValidPlatform.ToString(), ExpandArguments);
+			}
 
-				XInitPlatformFile.Save(GeneratedPropertiesPlatformFile);
+			foreach (KeyValuePair<string, XDocument> KVP in SaveAtEnd)
+			{
+				Log.LogInformation("Saving metadata to {File}", KVP.Key);
+				KVP.Value.Save(KVP.Key);
 			}
 		}
 
-		private string GetBaseFolder()
+		private static string GetBaseFolder(string ModuleDirectory)
 		{
 			string RelativeModulePath = Path.GetRelativePath(Unreal.RootDirectory.FullName, ModuleDirectory);
 			string[] BreadCrumbs = RelativeModulePath.Split(new char[] { '\\', '/' }, StringSplitOptions.RemoveEmptyEntries);
@@ -327,25 +349,17 @@ namespace UnrealBuildTool
 			return Unreal.EngineDirectory.FullName;
 		}
 
-		private bool IsPlatformRestricted(UnrealTargetPlatform Platform)
+		private static bool IsPlatformRestricted(UnrealTargetPlatform Platform)
 		{
 			return RestrictedFolder.GetNames().Contains(Platform.ToString());
 		}
 
-		private bool IsRestrictedPath(string ModuleDirectory)
+		private static bool IsRestrictedPath(string ModuleDirectory)
 		{
-			foreach (string RestrictedFolderName in RestrictedFolder.GetNames())
-			{
-				if (ModuleDirectory.Contains(RestrictedFolderName))
-				{
-					return true;
-				}
-			}
-
-			return false;
+			return ModuleDirectory.Split(new char[] { '/', '\\' }).Intersect(RestrictedFoldersNonPlatform).Count() > 0;
 		}
 
-		private string TryGetBinariesPath()
+		private static string TryGetBinariesPath(string ModuleDirectory)
 		{
 			int SourceFolderIndex = ModuleDirectory.IndexOf("Source");
 			if (SourceFolderIndex < 0)
@@ -355,73 +369,58 @@ namespace UnrealBuildTool
 				{
 					return ModuleDirectory.Substring(0, PluginFolderIndex) + "Binaries";
 				}
-				throw new Exception("Could not detect source folder path for module " + GetType());
+				throw new Exception("Could not detect source folder path for module from directory " + ModuleDirectory);
 			}
 			return ModuleDirectory.Substring(0, SourceFolderIndex) + "Binaries";
 		}
 
-		private void AppendOrUpdateTestFlagProperty(ref XDocument Document, string FlagRadix, string FlagPrefix, string FlagSuffix, string FlagValue)
+		private static void AppendOrUpdateRunAllTestsNode(XElement Root, string MacroName, string Platform, Dictionary<string, string> ExpandArguments)
 		{
-			XElement? Existing = Document.Root!.Elements().Where(element => element.Attribute("Name").Value == FlagPrefix + FlagRadix + FlagSuffix).FirstOrDefault();
-			if (Existing != null)
+			XElement? ExtendNode = Root.Elements().Where(element => element.Name.LocalName == "Extend").FirstOrDefault();
+			if (ExtendNode == null)
 			{
-				Existing!.SetAttributeValue("Value", FlagValue);
+				ExtendNode = new XElement(BuildGraphNamespace + "Extend");
+				ExtendNode.SetAttributeValue("Name", "RunAllTests");
+				Root.Add(ExtendNode);
 			}
-			else
+			
+			XElement? ExpandNode = ExtendNode.Elements().Where(element => element.Attribute("Name").Value == MacroName && element.Attribute("Platform").Value == Platform).FirstOrDefault();
+			if (ExpandNode == null)
 			{
-				XElement ElementAppend = new XElement(BuildGraphNamespace + "Property");
-				ElementAppend.SetAttributeValue("Name", FlagPrefix + FlagRadix + FlagSuffix);
-				ElementAppend.SetAttributeValue("Value", FlagValue);
-				Document.Root!.Add(ElementAppend);
+				ExpandNode = new XElement(BuildGraphNamespace + "Expand");
+				ExpandNode.SetAttributeValue("Name", MacroName);
+				ExpandNode.SetAttributeValue("Platform", Platform);
+				ExtendNode.Add(ExpandNode);
 			}
-		}
-		private void InsertOrUpdateTestFlagProperty(ref XElement ElementUpsertAfter, string TestName, string FlagSuffix, string FlagValue)
-		{
-			IEnumerable<XElement> NextChunk = ElementUpsertAfter.ElementsAfterSelf(BuildGraphNamespace + "Property")
-				.Where(prop => prop.Attribute("Name").Value.EndsWith(FlagSuffix));
-			if (!NextChunk
-				.Where(prop => prop.Attribute("Name").Value == TestName + FlagSuffix).Any())
-			{
-				XElement ElementInsert = new XElement(BuildGraphNamespace + "Property");
-				ElementInsert.SetAttributeValue("Name", TestName + FlagSuffix);
-				ElementInsert.SetAttributeValue("Value", FlagValue);
-				ElementUpsertAfter.AddAfterSelf(ElementInsert);
-			}
-			else
+			foreach (KeyValuePair<string, string> ArgumentAndValue in ExpandArguments)
 			{
-				NextChunk
-					.Where(prop => prop.Attribute("Name").Value == TestName + FlagSuffix).First().SetAttributeValue("Value", FlagValue);
+				ExpandNode!.SetAttributeValue(ArgumentAndValue.Key, ArgumentAndValue.Value);
 			}
-			ElementUpsertAfter = NextChunk.Last();
 		}
 
-		private void InsertOrUpdateTestOption(ref XElement ElementUpsertAfter, string OptionRadix, string Description, string OptionPrefix, string OptionSuffix, string DefaultValue)
+		private static void InsertOrUpdateTestOption(XElement Root, string OptionName, string Description, string DefaultValue)
 		{
-			IEnumerable<XElement> NextChunk = ElementUpsertAfter.ElementsAfterSelf(BuildGraphNamespace + "Option")
-				.Where(prop => prop.Attribute("Name").Value.StartsWith(OptionPrefix) && prop.Attribute("Name").Value.EndsWith(OptionSuffix));
-			if (!NextChunk
-				.Where(prop => prop.Attribute("Name").Value == OptionPrefix + OptionRadix + OptionSuffix).Any())
+			XElement? OptionElementWithName = Root.Elements(BuildGraphNamespace + "Option")
+				.Where(prop => prop.Attribute("Name").Value == OptionName).FirstOrDefault();
+			if (OptionElementWithName == null)
 			{
 				XElement ElementInsert = new XElement(BuildGraphNamespace + "Option");
-				ElementInsert.SetAttributeValue("Name", OptionPrefix + OptionRadix + OptionSuffix);
+				ElementInsert.SetAttributeValue("Name", OptionName);
 				ElementInsert.SetAttributeValue("DefaultValue", DefaultValue);
 				ElementInsert.SetAttributeValue("Description", Description);
-				ElementUpsertAfter.AddAfterSelf(ElementInsert);
+				Root.Add(ElementInsert);
 			}
 			else
 			{
-				XElement ElementUpdate = NextChunk
-					.Where(prop => prop.Attribute("Name").Value == OptionPrefix + OptionRadix + OptionSuffix).First();
-				ElementUpdate.SetAttributeValue("Description", Description);
-				ElementUpdate.SetAttributeValue("DefaultValue", DefaultValue);
+				OptionElementWithName.SetAttributeValue("Description", Description);
+				OptionElementWithName.SetAttributeValue("DefaultValue", DefaultValue);
 			}
-			ElementUpsertAfter = NextChunk.Last();
 		}
 
 #pragma warning restore 8604
 #pragma warning restore 8602
 
-		private void MakeFileWriteable(string InFilePath)
+		private static void MakeFileWriteable(string InFilePath)
 		{
 			System.IO.File.SetAttributes(InFilePath, System.IO.File.GetAttributes(InFilePath) & ~FileAttributes.ReadOnly);
 		}
@@ -459,9 +458,9 @@ namespace UnrealBuildTool
 			public bool StagesWithProjectFile { get; set; }
 
 			/// <summary>
-			/// Is this test disabled?
+			/// Is this test deactivated?
 			/// </summary>
-			public bool Disabled { get; set; }
+			public bool Deactivated { get; set; }
 
 			/// <summary>
 			/// Depercated, use GauntletArgs or ExtraArgs instead to help indicate arguments to launch the test under.
@@ -514,6 +513,16 @@ namespace UnrealBuildTool
 				set => PlatformTagsPrivate = value;
 			}
 
+			private Dictionary<UnrealTargetPlatform, string> PlatformGauntletArgsPrivate = new Dictionary<UnrealTargetPlatform, string>();
+			/// <summary>
+			/// Per-platform gauntlet args.
+			/// </summary>
+			public Dictionary<UnrealTargetPlatform, string> PlatformGauntletArgs
+			{
+				get => PlatformGauntletArgsPrivate;
+				set => PlatformGauntletArgsPrivate = value;
+			}
+
 			private Dictionary<UnrealTargetPlatform, string> PlatformCompilationExtraArgsPrivate = new Dictionary<UnrealTargetPlatform, string>();
 			/// <summary>
 			/// Per-platform extra compilation arguments.
diff --git a/Engine/Source/Programs/UnrealBuildTool/Modes/TestMode.cs b/Engine/Source/Programs/UnrealBuildTool/Modes/TestMode.cs
--- a/Engine/Source/Programs/UnrealBuildTool/Modes/TestMode.cs
+++ b/Engine/Source/Programs/UnrealBuildTool/Modes/TestMode.cs
@@ -2,9 +2,12 @@
 
 using System;
 using System.Collections.Generic;
+using System.IO;
 using System.Linq;
+using System.Reflection;
 using System.Threading.Tasks;
 using EpicGames.Core;
+using Microsoft.Build.Execution;
 using Microsoft.Extensions.Logging;
 using UnrealBuildBase;
 
@@ -41,6 +44,18 @@ namespace UnrealBuildTool
 		[CommandLine("-RebuildTests")]
 		bool bRebuildTests = false;
 
+		/// <summary>
+		/// Updates tests' build graph metadata files for explicit tests
+		/// </summary>
+		[CommandLine("-GenerateMetadata")]
+		bool bGenerateMetadata = false;
+
+        /// <summary>
+        /// Specific project file to run tests on or update test metadata
+        /// </summary>
+        [CommandLine("-Project")]
+        FileReference? ProjectFile = null;
+
 		/// <summary>
 		/// Main entry point
 		/// </summary>
@@ -55,6 +70,30 @@ namespace UnrealBuildTool
 			Arguments.ApplyTo(BuildConfiguration);
 			XmlConfig.ApplyTo(BuildConfiguration);
 
+			if (bGenerateMetadata)
+			{
+				Logger.LogInformation("Retrieving rules...");
+
+                RulesAssembly RulesAssembly;
+
+                if (ProjectFile != null)
+                {
+                    RulesAssembly = RulesCompiler.CreateProjectRulesAssembly(ProjectFile, BuildConfiguration.bUsePrecompiled, BuildConfiguration.bSkipRulesCompile, BuildConfiguration.bForceRulesCompile, Logger);
+                }
+                else
+                {
+                    RulesAssembly = RulesCompiler.CreateEngineRulesAssembly(BuildConfiguration.bUsePrecompiled, BuildConfiguration.bSkipRulesCompile, BuildConfiguration.bForceRulesCompile, Logger);
+                }
+				GenerateMetadataForTestsInAssembly(RulesAssembly, Logger);
+				if (ProjectFile == null && RulesAssembly.Parent != null)
+				{
+					// Non-program rules (includes plugins)
+					GenerateMetadataForTestsInAssembly(RulesAssembly.Parent, Logger);
+				}
+
+				return 0;
+			}
+
 			// Parse all the targets being built
 			List<TargetDescriptor> TargetDescriptors = TargetDescriptor.ParseCommandLine(Arguments, BuildConfiguration, Logger);
 
@@ -94,6 +133,127 @@ namespace UnrealBuildTool
 			return 0;
 		}
 
+		private void GenerateMetadataForTestsInAssembly(RulesAssembly Assembly, ILogger Logger)
+		{
+			Type ModuleRulesType = typeof(TestModuleRules);
+			Dictionary<Type, HashSet<Type>> InheritedGroups = new Dictionary<Type, HashSet<Type>>();
+			HashSet<Type> AllTestModules = Assembly.GetTypes()
+				.Where(t => t.IsAssignableTo(ModuleRulesType))
+				.ToHashSet();
+			foreach(var Type in AllTestModules)
+			{
+				// Find root class that's not TestModuleRules, the while loop always stops because AllTestModules contains only types that inherit from TestModuleRules
+				Type? RootBaseType = Type.BaseType;
+				while(true)
+				{
+					if (RootBaseType != null)
+					{
+						if (RootBaseType.Name == ModuleRulesType.Name)
+						{
+							break;
+						}
+						else if (RootBaseType.BaseType != null && RootBaseType.BaseType.Name != ModuleRulesType.Name)
+						{
+							RootBaseType = RootBaseType.BaseType;
+						}
+						else
+						{
+							break;
+						}
+					}
+					else
+					{
+						break;
+					}
+				}
+				// Add to inheritance group if it's a derived class
+				if (RootBaseType != null && RootBaseType.Name != Type.Name && AllTestModules.Contains(RootBaseType))
+				{
+					if (!InheritedGroups.ContainsKey(RootBaseType))
+					{
+						InheritedGroups[RootBaseType] = new HashSet<Type>();
+					}
+					InheritedGroups[RootBaseType].Add(Type);
+					AllTestModules.Remove(Type); // Removed so that we use InheritedGroups to call static constructors of inherited types to completely initialize the TestMetadata field
+				}
+			}
+			foreach (Type TestModuleType in AllTestModules)
+			{
+				MethodInfo? UpdateBuildGraphMetadata = FindMethodUpwards(TestModuleType, "UpdateBuildGraphMetadata", BindingFlags.Static | BindingFlags.NonPublic); 
+
+				// Call static constructor if exists, sets TestMetadata properties (can be set from derived module rules classes, including for NDA platforms
+				if (UpdateBuildGraphMetadata != null && TestModuleType.TypeInitializer != null)
+				{
+					Logger.LogInformation("Generating test metadata for module {TestModule}", TestModuleType.FullName);
+
+					// Initialize TestMetadata in base class using static constructor
+					TestModuleType.TypeInitializer.Invoke(null, null);
+					if (InheritedGroups.ContainsKey(TestModuleType))
+					{
+						foreach (var InheritedTestModuleType in InheritedGroups[TestModuleType])
+						{
+							if (InheritedTestModuleType.TypeInitializer != null)
+							{
+								// Update TestMetadata on inherited types through static constructor.
+								// UpdateBuildGraphMetadata will have a complete TestMetadata to act on.
+								Logger.LogInformation("Subclass found, update test metadata: {InheritedTestModule}", InheritedTestModuleType.FullName);
+								InheritedTestModuleType.TypeInitializer.Invoke(null, null);
+							}
+						}
+					}
+					FieldInfo? TestMetadataField = FindFieldUpwards(TestModuleType, "TestMetadata", BindingFlags.Static | BindingFlags.NonPublic);
+					if (UpdateBuildGraphMetadata != null && TestMetadataField != null)
+					{
+						TestModuleRules.Metadata MetadataValue = (TestModuleRules.Metadata)TestMetadataField.GetValue(null)!;
+						if (MetadataValue != null && !string.IsNullOrEmpty(MetadataValue.TestName))
+						{
+							UpdateBuildGraphMetadata.Invoke(null,
+                                new object[] {
+                                    MetadataValue,
+                                    Path.GetDirectoryName(RulesCompiler.GetFileNameFromType(TestModuleType)) ?? "",
+                                    TestModuleType.FullName ?? "",
+									Logger
+								});
+						}
+					}
+				}
+			}
+		}
+
+		private MethodInfo? FindMethodUpwards(Type TestType, string MethodName, BindingFlags Flags)
+		{
+			Type CurrentLevelType = TestType;
+			MethodInfo? FindMethod;
+			while (CurrentLevelType != null &&  CurrentLevelType != typeof(ModuleRules))
+			{
+				FindMethod = CurrentLevelType.GetMethod(MethodName, Flags);
+				if (FindMethod == null)
+				{
+					CurrentLevelType = CurrentLevelType.BaseType!;
+				}
+				else
+					return FindMethod;
+			}
+			return null;
+		}
+
+		private FieldInfo? FindFieldUpwards(Type TestType, string FieldName, BindingFlags Flags)
+		{
+			Type CurrentLevelType = TestType;
+			FieldInfo? FindField;
+			while (CurrentLevelType != null && CurrentLevelType != typeof(ModuleRules))
+			{
+				FindField = CurrentLevelType.GetField(FieldName, Flags);
+				if (FindField == null)
+				{
+					CurrentLevelType = CurrentLevelType.BaseType!;
+				}
+				else
+					return FindField;
+			}
+			return null;
+		}
+
 		private void AnalyzeDependencyGraph(TargetDescriptor TargetDescriptor, BuildConfiguration BuildConfiguration, ILogger Logger)
 		{
 			UEBuildTarget Target = UEBuildTarget.Create(TargetDescriptor, BuildConfiguration, Logger);
@@ -146,11 +306,6 @@ namespace UnrealBuildTool
 					}
 				}
 			}
-
-			foreach (UEBuildModule Module in Visited)
-			{
-				SetupTestsTargetForModule(Module, Logger);
-			}
 		}
 
 		private void AnalyzeModuleDependencies(string ModuleName, List<string> ParentChain, UEBuildTarget Target, IncludeChainMap Chains, HashSet<UEBuildModule> Visited, ILogger Logger)
@@ -181,6 +336,7 @@ namespace UnrealBuildTool
 				return;
 			}
 
+			SetupTestsTargetForModule(Module, Logger);
 			Visited.Add(Module);
 
 			List<UEBuildModule> TargetModules = new();
@@ -190,22 +346,22 @@ namespace UnrealBuildTool
 
 		private void SetupTestsTargetForModule(UEBuildModule Module, ILogger Logger)
 		{
-			if (Module.Name == ENGINE_MODULE)
+			if (Module.Name == ENGINE_MODULE && !TestTargetRules.bTestsRequireEngine)
 			{
 				TestTargetRules.bTestsRequireEngine = true;
 				Logger.LogInformation("{Module} module found in dependency graph, will compile against it.", ENGINE_MODULE);
 			}
-			else if (Module.Name == APPLICATION_CORE_MODULE)
+			else if (Module.Name == APPLICATION_CORE_MODULE && !TestTargetRules.bTestsRequireApplicationCore)
 			{
 				TestTargetRules.bTestsRequireApplicationCore = true;
 				Logger.LogInformation("{Module} module found in dependency graph, will compile against it.", APPLICATION_CORE_MODULE);
 			}
-			else if (Module.Name == COREUOBJECT_MODULE)
+			else if (Module.Name == COREUOBJECT_MODULE && !TestTargetRules.bTestsRequireCoreUObject)
 			{
 				TestTargetRules.bTestsRequireCoreUObject = true;
 				Logger.LogInformation("{Module} module found in dependency graph, will compile against it.", COREUOBJECT_MODULE);
 			}
-			else if (Module.Name == EDITOR_MODULE)
+			else if (Module.Name == EDITOR_MODULE && !TestTargetRules.bTestsRequireEditor)
 			{
 				TestTargetRules.bTestsRequireEditor = true;
 				Logger.LogInformation("{Module} module found in dependency graph, will compile against it.", EDITOR_MODULE);
diff --git a/Engine/Source/Programs/WebTests/WebTests.Build.cs b/Engine/Source/Programs/WebTests/WebTests.Build.cs
--- a/Engine/Source/Programs/WebTests/WebTests.Build.cs
+++ b/Engine/Source/Programs/WebTests/WebTests.Build.cs
@@ -4,6 +4,28 @@ using UnrealBuildTool;
 
 public class WebTests : TestModuleRules
 {
+	static WebTests()
+	{
+		TestMetadata = new Metadata();
+		TestMetadata.TestName = "Web";
+		TestMetadata.TestShortName = "Web";
+		TestMetadata.ReportType = "xml";
+		TestMetadata.SupportedPlatforms.Add(UnrealTargetPlatform.Linux);
+		TestMetadata.SupportedPlatforms.Add(UnrealTargetPlatform.Mac);
+		TestMetadata.SupportedPlatforms.Add(UnrealTargetPlatform.Android);
+		TestMetadata.SupportedPlatforms.Add(UnrealTargetPlatform.IOS);
+
+		TestMetadata.HasAfterSteps = true;
+		TestMetadata.PlatformCompilationExtraArgs.Add(UnrealTargetPlatform.Win64, "-allmodules");
+		TestMetadata.PlatformCompilationExtraArgs.Add(UnrealTargetPlatform.Linux, "-allmodules");
+		TestMetadata.PlatformCompilationExtraArgs.Add(UnrealTargetPlatform.LinuxArm64, "-allmodules");
+		TestMetadata.PlatformCompilationExtraArgs.Add(UnrealTargetPlatform.Mac, "-allmodules");
+		TestMetadata.PlatformCompilationExtraArgs.Add(UnrealTargetPlatform.IOS, "-allmodules");
+		TestMetadata.PlatformCompilationExtraArgs.Add(UnrealTargetPlatform.TVOS, "-allmodules");
+		TestMetadata.PlatformCompilationExtraArgs.Add(UnrealTargetPlatform.Android, "-allmodules  -architectures=arm64");
+		TestMetadata.PlatformCompilationExtraArgs.Add(UnrealTargetPlatform.VisionOS, "-allmodules");
+	}
+
 	public WebTests(ReadOnlyTargetRules Target) : base(Target)
 	{
 		PrivateDependencyModuleNames.AddRange(
@@ -15,12 +37,6 @@ public class WebTests : TestModuleRules
 				"WebSockets",
 				"Json"
 			});
-		
-		UpdateBuildGraphPropertiesFile(new Metadata() { 
-			TestName = "Web", 
-			TestShortName = "Web", 
-			ReportType = "xml", 
-		});
 	}
 }
 
diff --git a/Engine/Source/Runtime/Online/Experimental/EventLoopTests/EventLoopUnitTests.Build.cs b/Engine/Source/Runtime/Online/Experimental/EventLoopTests/EventLoopUnitTests.Build.cs
--- a/Engine/Source/Runtime/Online/Experimental/EventLoopTests/EventLoopUnitTests.Build.cs
+++ b/Engine/Source/Runtime/Online/Experimental/EventLoopTests/EventLoopUnitTests.Build.cs
@@ -3,7 +3,14 @@ using UnrealBuildTool;
 
 public class EventLoopUnitTests : TestModuleRules
 {
-	public EventLoopUnitTests(ReadOnlyTargetRules Target) : base(Target)
+    static EventLoopUnitTests()
+    {
+		TestMetadata = new Metadata();
+        TestMetadata.TestName = "EventLoop";
+        TestMetadata.TestShortName = "EventLoop";
+    }
+
+    public EventLoopUnitTests(ReadOnlyTargetRules Target) : base(Target)
 	{
 		OptimizeCode = CodeOptimization.Never;
 
@@ -13,7 +20,5 @@ public class EventLoopUnitTests : TestModuleRules
 				"EventLoop",
 				"Sockets"
 			});
-
-		UpdateBuildGraphPropertiesFile(new Metadata() { TestName = "EventLoop", TestShortName = "EventLoop" });
 	}
 }
\ No newline at end of file
