[UBA]
* Fixed so if no horde server url is set we early out and write a log entry saying horde is disabled because no url

diff --git a/Engine/Source/Programs/UnrealBuildTool/Executors/UnrealBuildAccelerator/UBAAgentCoordinatorHorde.cs b/Engine/Source/Programs/UnrealBuildTool/Executors/UnrealBuildAccelerator/UBAAgentCoordinatorHorde.cs
--- a/Engine/Source/Programs/UnrealBuildTool/Executors/UnrealBuildAccelerator/UBAAgentCoordinatorHorde.cs
+++ b/Engine/Source/Programs/UnrealBuildTool/Executors/UnrealBuildAccelerator/UBAAgentCoordinatorHorde.cs
@@ -373,6 +373,12 @@ namespace UnrealBuildTool
 				return null;
 			}
 
+			if (String.IsNullOrEmpty(hordeConfig.HordeServer))
+			{
+				logger.LogInformation("Horde disabled. Url not set.");
+				return null;
+			}
+
 			Uri? server = (hordeConfig.HordeServer == null) ? null : new Uri(hordeConfig.HordeServer);
 			string? token = hordeConfig.HordeToken;
 
