Fix spelling error
#rnx

diff --git a/Engine/Source/Programs/UnrealBuildTool/Configuration/UEBuildTarget.cs b/Engine/Source/Programs/UnrealBuildTool/Configuration/UEBuildTarget.cs
--- a/Engine/Source/Programs/UnrealBuildTool/Configuration/UEBuildTarget.cs
+++ b/Engine/Source/Programs/UnrealBuildTool/Configuration/UEBuildTarget.cs
@@ -1257,7 +1257,7 @@ namespace UnrealBuildTool
 			// make sure we are allowed to build this platform
 			if (!UEBuildPlatform.IsPlatformAvailable(Descriptor.Platform))
 			{
-				throw new BuildException("Platform {0} is not a valid platform to build. Check that the SDK is installed properly and that you have the necessary platorm support files (DataDrivenPlatformInfo.ini, SDK.json, etc).", Descriptor.Platform);
+				throw new BuildException("Platform {0} is not a valid platform to build. Check that the SDK is installed properly and that you have the necessary platform support files (DataDrivenPlatformInfo.ini, SDK.json, etc).", Descriptor.Platform);
 			}
 
 			RulesAssembly RulesAssembly;
