UnrealBuildTool: Put a mutex around project intellisense generation

#rnx

diff --git a/Engine/Source/Programs/UnrealBuildTool/Configuration/UEBuildModuleCPP.cs b/Engine/Source/Programs/UnrealBuildTool/Configuration/UEBuildModuleCPP.cs
--- a/Engine/Source/Programs/UnrealBuildTool/Configuration/UEBuildModuleCPP.cs
+++ b/Engine/Source/Programs/UnrealBuildTool/Configuration/UEBuildModuleCPP.cs
@@ -1990,6 +1990,9 @@ namespace UnrealBuildTool
 
 		public CppCompileEnvironment CreateCompileEnvironmentForIntellisense(ReadOnlyTargetRules Target, CppCompileEnvironment BaseCompileEnvironment, ILogger Logger)
 		{
+			string MutexName = SingleInstanceMutex.GetUniqueMutexForPath("UnrealBuildTool_CreateCompileEnvironmentForIntellisense", IntermediateDirectory.FullName);
+			using SingleInstanceMutex mutex = new SingleInstanceMutex(MutexName, true);
+
 			CppCompileEnvironment CompileEnvironment = CreateModuleCompileEnvironment(Target, BaseCompileEnvironment, Logger);
 			CompileEnvironment = SetupPrecompiledHeaders(Target, null, CompileEnvironment, new List<FileItem>(), new NullActionGraphBuilder(Logger), Logger);
 			CreateHeaderForDefinitions(CompileEnvironment, IntermediateDirectory, null, new NullActionGraphBuilder(Logger));
