Remove exccesive resource files in Xcode projects, speed up loading and closing

diff --git a/Engine/Source/Programs/UnrealBuildTool/ProjectFiles/Xcode/XcodeProject.cs b/Engine/Source/Programs/UnrealBuildTool/ProjectFiles/Xcode/XcodeProject.cs
--- a/Engine/Source/Programs/UnrealBuildTool/ProjectFiles/Xcode/XcodeProject.cs
+++ b/Engine/Source/Programs/UnrealBuildTool/ProjectFiles/Xcode/XcodeProject.cs
@@ -2419,10 +2419,21 @@ namespace UnrealBuildTool.XcodeProjectXcconfig
 			// process the files that came from UE/cross-platform land
 			SourceFiles.SortBy(x => x.Reference.FullName);
 
+			ConfigHierarchy PlatformIni = ConfigCache.ReadHierarchy(ConfigHierarchyType.Engine, UnrealData.ConfigDirectory, UnrealTargetPlatform.Mac);
+
+			IReadOnlyList<string>? ExcludedDirs;
+			if (!PlatformIni.TryGetValues("/Script/MacTargetPlatform.XcodeProjectSettings", "XCProjectExcludedDirs", out ExcludedDirs))
+			{
+				ExcludedDirs = new List<string>();
+			}
+
 			Dictionary<DirectoryReference, int> BuildFileMap = new();
 			foreach (XcodeSourceFile SourceFile in SourceFiles.OfType<XcodeSourceFile>())
 			{
-				SharedFileCollection.ProcessFile(SourceFile, bIsForBuild: IsGeneratedProject, bIsFolder: false, SourceToBuildFileMap: BuildFileMap);
+				if (!ExcludedDirs.Any(dir => SourceFile.Reference.IsUnderDirectory(DirectoryReference.Combine(SharedFileCollection.RootDirectory, dir))))
+				{
+					SharedFileCollection.ProcessFile(SourceFile, bIsForBuild: IsGeneratedProject, bIsFolder: false, SourceToBuildFileMap: BuildFileMap);
+				}
 			}
 
 			// cache the main group
