UnrealBuildTool: Option to force retry any remote UBA failures locally

#rnx

diff --git a/Engine/Source/Programs/UnrealBuildTool/Executors/UnrealBuildAccelerator/UBAExecutor.cs b/Engine/Source/Programs/UnrealBuildTool/Executors/UnrealBuildAccelerator/UBAExecutor.cs
--- a/Engine/Source/Programs/UnrealBuildTool/Executors/UnrealBuildAccelerator/UBAExecutor.cs
+++ b/Engine/Source/Programs/UnrealBuildTool/Executors/UnrealBuildAccelerator/UBAExecutor.cs
@@ -762,6 +762,11 @@ namespace UnrealBuildTool
 						RemoteActionFailedCrash(queue, action, e.ExitCode, e.ExecutingHost ?? "Unknown", "UBA error");
 						return;
 					}
+					else if (e.ExitCode != 0 && UBAConfig.bForcedRetryRemote)
+					{
+						RemoteActionFailedCrash(queue, action, e.ExitCode, e.ExecutingHost ?? "Unknown", "Force local retry");
+						return;
+					}
 
 					IProcess process = (IProcess)s;
 
diff --git a/Engine/Source/Programs/UnrealBuildTool/Executors/UnrealBuildAccelerator/UnrealBuildAcceleratorConfig.cs b/Engine/Source/Programs/UnrealBuildTool/Executors/UnrealBuildAccelerator/UnrealBuildAcceleratorConfig.cs
--- a/Engine/Source/Programs/UnrealBuildTool/Executors/UnrealBuildAccelerator/UnrealBuildAcceleratorConfig.cs
+++ b/Engine/Source/Programs/UnrealBuildTool/Executors/UnrealBuildAccelerator/UnrealBuildAcceleratorConfig.cs
@@ -32,6 +32,13 @@ namespace UnrealBuildTool
 		[CommandLine("-UBAForcedRetry")]
 		public bool bForcedRetry { get; set; } = false;
 
+		/// <summary>
+		/// When set to true, actions that fail remotely with UBA will be retried locally with UBA.
+		/// </summary>
+		[XmlConfigFile(Category = "UnrealBuildAccelerator")]
+		[CommandLine("-UBAForcedRetryRemote")]
+		public bool bForcedRetryRemote { get; set; } = false;
+
 		/// <summary>
 		/// When set to true, all errors and warnings from UBA will be output at the appropriate severity level to the log (rather than being output as 'information' and attempting to continue regardless).
 		/// </summary>
