Fix output directory issue for Blueprint APKs

#android

diff --git a/Engine/Source/Programs/UnrealBuildTool/Platform/Android/UEDeployAndroid.cs b/Engine/Source/Programs/UnrealBuildTool/Platform/Android/UEDeployAndroid.cs
--- a/Engine/Source/Programs/UnrealBuildTool/Platform/Android/UEDeployAndroid.cs
+++ b/Engine/Source/Programs/UnrealBuildTool/Platform/Android/UEDeployAndroid.cs
@@ -4992,7 +4992,11 @@ namespace UnrealBuildTool
 				string? OutputDirectory = Path.GetDirectoryName(OutputPath);
 				if (OutputDirectory != null && Path.GetRelativePath(DestApkDirectory, OutputDirectory) != OutputDirectory)
 				{
-					DestApkDirectory = OutputDirectory;
+					// only override if not UnrealGame (BluePrint project)
+					if (!Path.GetFileName(OutputPath).StartsWith("UnrealGame"))
+					{
+						DestApkDirectory = OutputDirectory;
+					}
 				}
 				string? DestApkName = null;
 				if (bSkipGradleBuild)
