UnrealBuildTool: Properly log AggregateExceptions

#rnx

diff --git a/Engine/Source/Programs/UnrealBuildTool/System/BuildException.cs b/Engine/Source/Programs/UnrealBuildTool/System/BuildException.cs
--- a/Engine/Source/Programs/UnrealBuildTool/System/BuildException.cs
+++ b/Engine/Source/Programs/UnrealBuildTool/System/BuildException.cs
@@ -1,6 +1,8 @@
 // Copyright Epic Games, Inc. All Rights Reserved.
 
 using System;
+using System.Linq;
+using System.Text.Json;
 using EpicGames.Core;
 using Microsoft.Extensions.Logging;
 
@@ -228,5 +230,57 @@ namespace UnrealBuildTool
 			Logger.LogDebug(this, "{Ex}", ExceptionUtils.FormatExceptionDetails(this));
 		}
 	}
+
+	/// <summary>
+	/// Extension methods for build exceptions
+	/// </summary>
+	public static class BuildExceptionExtensions
+	{
+		/// <summary>
+		/// Log Exception with a provided ILogger
+		/// </summary>
+		/// <param name="ex">The exception to log</param>
+		/// <param name="logger">The ILogger to use to log this exception</param>
+		public static void LogException(this Exception ex, ILogger logger)
+		{
+			if (ex is BuildException buildException)
+			{
+				buildException.LogException(logger);
+			}
+			else if (ex is JsonException jsonException)
+			{
+				FileReference source = new FileReference(jsonException.Path ?? jsonException.Source ?? "unknown");
+				LogValue fileValue = LogValue.SourceFile(source, source.GetFileName());
+				long line = jsonException.LineNumber ?? 0;
+				logger.LogError(KnownLogEvents.Compiler, "{File}({Line}): error: {Message}", fileValue, line, ExceptionUtils.FormatExceptionDetails(ex));
+			}
+			else if (ex is AggregateException aggregateException)
+			{
+				logger.LogError(ex, "Unhandled {Count} aggregate exceptions", aggregateException.InnerExceptions.Count);
+
+				foreach (Exception innerException in aggregateException.InnerExceptions)
+				{
+					innerException.LogException(logger);
+				}
+			}
+			else
+			{
+				logger.LogError(ex, "Unhandled exception: {Ex}", ExceptionUtils.FormatExceptionDetails(ex));
+			}
+		}
+
+		/// <summary>
+		/// Get the CompilationResult for a provided Exception
+		/// </summary>
+		/// <param name="ex">The exception to get the result for</param>
+		/// <returns>CompilationResult</returns>
+		public static CompilationResult GetCompilationResult(this Exception ex)
+		{
+			return (ex as CompilationResultException)?.Result
+				?? (ex.InnerException as CompilationResultException)?.Result
+				?? (ex as AggregateException)?.InnerExceptions.OfType<CompilationResultException>().FirstOrDefault()?.Result
+				?? CompilationResult.OtherCompilationError;
+		}
+	}
 }
 
diff --git a/Engine/Source/Programs/UnrealBuildTool/System/PluginDescriptor.cs b/Engine/Source/Programs/UnrealBuildTool/System/PluginDescriptor.cs
--- a/Engine/Source/Programs/UnrealBuildTool/System/PluginDescriptor.cs
+++ b/Engine/Source/Programs/UnrealBuildTool/System/PluginDescriptor.cs
@@ -385,7 +385,7 @@ namespace UnrealBuildTool
 			}
 			catch (JsonException ex)
 			{
-				throw new JsonException($"{ex.Message} (in {FileName})", ex.Source ?? FileName.FullName, ex.LineNumber, ex.BytePositionInLine, ex);
+				throw new JsonException($"{ex.Message} (in {FileName})", FileName.FullName, ex.LineNumber, ex.BytePositionInLine, ex);
 			}
 		}
 
diff --git a/Engine/Source/Programs/UnrealBuildTool/System/ProjectDescriptor.cs b/Engine/Source/Programs/UnrealBuildTool/System/ProjectDescriptor.cs
--- a/Engine/Source/Programs/UnrealBuildTool/System/ProjectDescriptor.cs
+++ b/Engine/Source/Programs/UnrealBuildTool/System/ProjectDescriptor.cs
@@ -212,9 +212,10 @@ namespace UnrealBuildTool
 		/// <returns>New plugin descriptor</returns>
 		public static ProjectDescriptor FromFile(FileReference FileName)
 		{
-			JsonObject RawObject = JsonObject.Read(FileName);
 			try
 			{
+				JsonObject RawObject = JsonObject.Read(FileName);
+
 				ProjectDescriptor Descriptor = new ProjectDescriptor(RawObject, FileName.Directory, FileName);
 				if (Descriptor.Modules != null)
 				{
@@ -227,7 +228,7 @@ namespace UnrealBuildTool
 			}
 			catch (JsonException ex)
 			{
-				throw new JsonException($"{ex.Message} (in {FileName})", ex.Source ?? FileName.FullName, ex.LineNumber, ex.BytePositionInLine, ex);
+				throw new JsonException($"{ex.Message} (in {FileName})", FileName.FullName, ex.LineNumber, ex.BytePositionInLine, ex);
 			}
 		}
 
diff --git a/Engine/Source/Programs/UnrealBuildTool/System/TargetReceipt.cs b/Engine/Source/Programs/UnrealBuildTool/System/TargetReceipt.cs
--- a/Engine/Source/Programs/UnrealBuildTool/System/TargetReceipt.cs
+++ b/Engine/Source/Programs/UnrealBuildTool/System/TargetReceipt.cs
@@ -547,7 +547,7 @@ namespace UnrealBuildTool
 			BuildVersion? Version;
 			if (!BuildVersion.TryParse(RawObject.GetObjectField("Version"), out Version))
 			{
-				throw new JsonException("Invalid 'Version' field");
+				throw new JsonException("Invalid 'Version' field", Location.FullName, null, null);
 			}
 
 			// Read the project path
diff --git a/Engine/Source/Programs/UnrealBuildTool/UnrealBuildTool.cs b/Engine/Source/Programs/UnrealBuildTool/UnrealBuildTool.cs
--- a/Engine/Source/Programs/UnrealBuildTool/UnrealBuildTool.cs
+++ b/Engine/Source/Programs/UnrealBuildTool/UnrealBuildTool.cs
@@ -7,7 +7,6 @@ using System.IO;
 using System.Linq;
 using System.Reflection;
 using System.Text;
-using System.Text.Json;
 using System.Threading.Tasks;
 using EpicGames.Core;
 using Microsoft.Extensions.Logging;
@@ -732,10 +731,6 @@ namespace UnrealBuildTool
 				{
 					Result = Mode.ExecuteAsync(Arguments, Logger).GetAwaiter().GetResult();
 				}
-				catch (AggregateException AggEx) when (AggEx.InnerExceptions.Count == 1 && AggEx.InnerExceptions.FirstOrDefault() != null)
-				{
-					throw AggEx.InnerExceptions.First();
-				}
 				finally
 				{
 					if ((ModeOptions & ToolModeOptions.ShowExecutionTime) != 0)
@@ -747,42 +742,12 @@ namespace UnrealBuildTool
 				ApplicationResult = (CompilationResult)Result;
 				return Result;
 			}
-			catch (CompilationResultException Ex)
-			{
-				// Used to return a propagate a specific exit code after an error has occurred.
-				Ex.LogException(Logger);
-				ApplicationResult = Ex.Result;
-				return (int)Ex.Result;
-			}
-			catch (BuildLogEventException Ex)
-			{
-				// BuildExceptions should have nicely formatted messages.
-				Ex.LogException(Logger);
-				ApplicationResult = CompilationResult.OtherCompilationError;
-				return (int)CompilationResult.OtherCompilationError;
-			}
-			catch (JsonException Ex)
-			{
-				FileReference source = new FileReference(Ex.Source ?? "unknown");
-				LogValue FileValue = LogValue.SourceFile(source, source.GetFileName());
-				Logger.LogError(KnownLogEvents.Compiler, "{File}({Line}): error:{Message}", FileValue, Ex.LineNumber ?? 0, ExceptionUtils.FormatException(Ex));
-				Logger.LogDebug(KnownLogEvents.Compiler, "{File}({Line}): error:{Message}", FileValue, Ex.LineNumber ?? 0, ExceptionUtils.FormatExceptionDetails(Ex));
-				ApplicationResult = CompilationResult.OtherCompilationError;
-				return (int)CompilationResult.OtherCompilationError;
-			}
-			catch (BuildException Ex)
-			{
-				// BuildExceptions should have nicely formatted messages.
-				Ex.LogException(Logger);
-				return (int)CompilationResult.OtherCompilationError;
-			}
 			catch (Exception Ex)
 			{
-				// Unhandled exception.
-				Logger.LogError(Ex, "Unhandled exception: {Ex}", ExceptionUtils.FormatException(Ex));
-				Logger.LogDebug(Ex, "Unhandled exception: {Ex}", ExceptionUtils.FormatExceptionDetails(Ex));
-				ApplicationResult = CompilationResult.OtherCompilationError;
-				return (int)CompilationResult.OtherCompilationError;
+				Ex.LogException(Logger);
+				// CompilationResultException is used to return a propagate a specific exit code after an error has occurred.
+				ApplicationResult = Ex.GetCompilationResult();
+				return (int)ApplicationResult;
 			}
 			finally
 			{
