UnrealBuildTool: Downgrade some exception logs during UBA remote distribution to debug

#rnx

diff --git a/Engine/Source/Programs/UnrealBuildTool/Executors/UnrealBuildAccelerator/UBAAgentCoordinatorHorde.cs b/Engine/Source/Programs/UnrealBuildTool/Executors/UnrealBuildAccelerator/UBAAgentCoordinatorHorde.cs
--- a/Engine/Source/Programs/UnrealBuildTool/Executors/UnrealBuildAccelerator/UBAAgentCoordinatorHorde.cs
+++ b/Engine/Source/Programs/UnrealBuildTool/Executors/UnrealBuildAccelerator/UBAAgentCoordinatorHorde.cs
@@ -360,7 +360,7 @@ namespace UnrealBuildTool
 				}
 				catch (Exception e)
 				{
-					_logger.Log(_strict ? LogLevel.Error : LogLevel.Information, KnownLogEvents.Systemic_Horde_Compute, e, "Failed updating resource need to {TargetCoreCount} cores", targetCoreCount);
+					_logger.Log(_strict ? LogLevel.Error : LogLevel.Debug, KnownLogEvents.Systemic_Horde_Compute, e, "Failed updating resource need to {TargetCoreCount} cores", targetCoreCount);
 				}
 			}
 		}
@@ -448,7 +448,7 @@ namespace UnrealBuildTool
 						}
 						catch (TimeoutException)
 						{
-							logger.Log(_strict ? LogLevel.Error : LogLevel.Information, KnownLogEvents.Systemic_Horde_Compute, "Waited {Time}s on attach message. Giving up", (int)attachTimeout.TotalSeconds);
+							logger.Log(_strict ? LogLevel.Error : LogLevel.Debug, KnownLogEvents.Systemic_Horde_Compute, "Waited {Time}s on attach message. Giving up", (int)attachTimeout.TotalSeconds);
 							throw;
 						}
 
@@ -550,10 +550,10 @@ namespace UnrealBuildTool
 			{
 				if (!cancellationToken.IsCancellationRequested)
 				{
-					logger.Log(_strict ? LogLevel.Error : LogLevel.Information, KnownLogEvents.Systemic_Horde_Compute, ex, "Exception in worker task: {Ex}", ex.ToString());
+					logger.Log(_strict ? LogLevel.Error : LogLevel.Debug, KnownLogEvents.Systemic_Horde_Compute, ex, "Exception in worker task: {Ex}", ex.ToString());
 
 					// Add additional properties to aid debugging
-					logger.LogInformation(KnownLogEvents.Systemic_Horde_Compute, ex, "UBA agent locator {UBAAgentLocator}", _ubaAgentLocator.ToString());
+					logger.Log(_strict ? LogLevel.Information : LogLevel.Debug, KnownLogEvents.Systemic_Horde_Compute, ex, "UBA agent locator {UBAAgentLocator}", _ubaAgentLocator.ToString());
 				}
 			}
 		}
@@ -683,14 +683,7 @@ namespace UnrealBuildTool
 				{
 					if (!_cancellationSource!.IsCancellationRequested)
 					{
-						if (_ubaConfig.bStrict)
-						{
-							_logger.Log(LogLevel.Error, KnownLogEvents.Systemic_Horde_Compute, ex, "Unable to get worker: {Ex}", ex.ToString());
-						}
-						else
-						{
-							_logger.Log(LogLevel.Information, KnownLogEvents.Systemic_Horde_Compute, "Unable to get worker: {Ex}", ex.ToString());
-						}
+						_logger.Log(_ubaConfig.bStrict ? LogLevel.Warning : LogLevel.Information, KnownLogEvents.Systemic_Horde_Compute, ex, "Unable to get worker: {Ex}", ex.ToString());
 					}
 				}
 
