UnrealBuildTool: Fix /analyze not being applied to -SingleFile when the module contains generated cpp files

diff --git a/Engine/Source/Programs/UnrealBuildTool/Configuration/UEBuildModuleCPP.cs b/Engine/Source/Programs/UnrealBuildTool/Configuration/UEBuildModuleCPP.cs
--- a/Engine/Source/Programs/UnrealBuildTool/Configuration/UEBuildModuleCPP.cs
+++ b/Engine/Source/Programs/UnrealBuildTool/Configuration/UEBuildModuleCPP.cs
@@ -650,6 +650,28 @@ namespace UnrealBuildTool
 			CompileEnvironment.bUseUnity = bModuleUsesUnityBuild;
 			GeneratedCPPCompileEnvironment.bUseUnity = bModuleUsesUnityBuild;
 
+			// Create and register a special action that can be used to compile single files (even when unity is enabled) for generated files
+			if (GeneratedFileItems.Any())
+			{
+				CppCompileEnvironment SingleGeneratedCPPCompileEnvironment = new CppCompileEnvironment(GeneratedCPPCompileEnvironment);
+				SingleGeneratedCPPCompileEnvironment.PrecompiledHeaderAction = PrecompiledHeaderAction.None;
+				foreach (DirectoryReference Directory in GeneratedFileItems.Select(x => x.Location.Directory).Distinct())
+				{
+					ToolChain.CreateSpecificFileAction(SingleGeneratedCPPCompileEnvironment, Directory, IntermediateDirectory, Graph);
+				}
+			}
+
+			// Create and register a special action that can be used to compile single files (even when unity is enabled) for normal files
+			if (ModuleDirectories.Any())
+			{
+				CppCompileEnvironment SingleCompileEnvironment = new CppCompileEnvironment(CompileEnvironment);
+				SingleCompileEnvironment.PrecompiledHeaderAction = PrecompiledHeaderAction.None;
+				foreach (DirectoryReference Directory in ModuleDirectories)
+				{
+					ToolChain.CreateSpecificFileAction(SingleCompileEnvironment, Directory, IntermediateDirectory, Graph);
+				}
+			}
+
 			// Compile Generated CPP Files
 			if (bModuleUsesUnityBuild)
 			{
@@ -1489,18 +1511,6 @@ namespace UnrealBuildTool
 				AdaptiveFiles = new List<FileItem>();
 			}
 
-			// Create and register a special action that can be used to compile single files (even when unity is enabled)
-			if (NormalFiles.Count > 0 || AdaptiveFiles.Count > 0)
-			{
-				PrecompiledHeaderAction OldAction = CompileEnvironment.PrecompiledHeaderAction;
-				CompileEnvironment.PrecompiledHeaderAction = PrecompiledHeaderAction.None;
-				foreach (DirectoryReference Directory in ModuleDirectories)
-				{
-					ToolChain.CreateSpecificFileAction(CompileEnvironment, Directory, IntermediateDirectory, Graph);
-				}
-				CompileEnvironment.PrecompiledHeaderAction = OldAction;
-			}
-
 			CPPOutput OutputFiles = new CPPOutput();
 
 			if (NormalFiles.Count > 0)
diff --git a/Engine/Source/Programs/UnrealBuildTool/Platform/Windows/VCToolChain.cs b/Engine/Source/Programs/UnrealBuildTool/Platform/Windows/VCToolChain.cs
--- a/Engine/Source/Programs/UnrealBuildTool/Platform/Windows/VCToolChain.cs
+++ b/Engine/Source/Programs/UnrealBuildTool/Platform/Windows/VCToolChain.cs
@@ -1565,7 +1565,7 @@ namespace UnrealBuildTool
 
 			AppendCLArguments_Global(CompileEnvironment, BaseCompileAction.Arguments);
 
-			BaseCompileAction.bIsAnalyzing = Target.StaticAnalyzer != StaticAnalyzer.None;
+			BaseCompileAction.bIsAnalyzing = Target.StaticAnalyzer != StaticAnalyzer.Default && !CompileEnvironment.bDisableStaticAnalysis && !(Target.WindowsPlatform.Compiler.IsClang() && CompileEnvironment.PrecompiledHeaderAction == PrecompiledHeaderAction.Create);
 
 			// Add include paths to the argument list.
 			BaseCompileAction.IncludePaths.AddRange(CompileEnvironment.UserIncludePaths);
