UnrealBuildTool: Fix more warnings

#rnx

diff --git a/Engine/Source/Programs/UnrealBuildTool/Configuration/TestModuleRules.cs b/Engine/Source/Programs/UnrealBuildTool/Configuration/TestModuleRules.cs
--- a/Engine/Source/Programs/UnrealBuildTool/Configuration/TestModuleRules.cs
+++ b/Engine/Source/Programs/UnrealBuildTool/Configuration/TestModuleRules.cs
@@ -257,7 +257,7 @@ namespace UnrealBuildTool
 					// Any manually edited elements to keep
 					List<XElement> KeepElements = XInitPlatformFile.Root!.Elements().Where(e => e.Attribute("Name").Value == $"{TestMetadata.TestName}AfterSteps").ToList();
 					XInitPlatformFile.Root!.Elements().Remove();
-					foreach (var Element in KeepElements)
+					foreach (XElement Element in KeepElements)
 					{
 						XInitPlatformFile.Root!.Add(Element);
 					}
@@ -287,7 +287,7 @@ namespace UnrealBuildTool
 				ExpandArguments.Add("TargetName", Convert.ToString(TestTargetName));
 				ExpandArguments.Add("BinaryRelativePath", Convert.ToString(TestBinariesPath));
 				ExpandArguments.Add("ReportType", Convert.ToString(TestMetadata.ReportType));
-				if (!string.IsNullOrEmpty(TestMetadata.GauntletArgs))
+				if (!String.IsNullOrEmpty(TestMetadata.GauntletArgs))
 				{
 					ExpandArguments.Add("GauntletArgs", Convert.ToString(TestMetadata.InitialExtraArgs) + Convert.ToString(TestMetadata.GauntletArgs));
 				}
@@ -295,7 +295,7 @@ namespace UnrealBuildTool
 				{
 					ExpandArguments.Add("PlatformGauntletArgs", TestMetadata.PlatformGauntletArgs[ValidPlatform]);
 				}
-				if (!string.IsNullOrEmpty(TestMetadata.ExtraArgs))
+				if (!String.IsNullOrEmpty(TestMetadata.ExtraArgs))
 				{
 					ExpandArguments.Add("ExtraArgs", Convert.ToString(TestMetadata.ExtraArgs));
 				}
@@ -308,14 +308,14 @@ namespace UnrealBuildTool
 					ExpandArguments.Add("UsesCatch2", Convert.ToString(TestMetadata.UsesCatch2));
 				}
 				string TagsValue = TestMetadata.PlatformTags.ContainsKey(ValidPlatform) ? TestMetadata.PlatformTags[ValidPlatform] : String.Empty;
-				if (!string.IsNullOrEmpty(TagsValue))
+				if (!String.IsNullOrEmpty(TagsValue))
 				{
 					ExpandArguments.Add("Tags", TagsValue);
 					
 				}
 
 				string ExtraCompilationArgsValue = TestMetadata.PlatformCompilationExtraArgs.ContainsKey(ValidPlatform) ? TestMetadata.PlatformCompilationExtraArgs[ValidPlatform] : String.Empty;
-				if (!string.IsNullOrEmpty (ExtraCompilationArgsValue))
+				if (!String.IsNullOrEmpty (ExtraCompilationArgsValue))
 				{
 					ExpandArguments.Add("ExtraCompilationArgs", ExtraCompilationArgsValue);
 				}
diff --git a/Engine/Source/Programs/UnrealBuildTool/Executors/Experimental/FASTBuild.cs b/Engine/Source/Programs/UnrealBuildTool/Executors/Experimental/FASTBuild.cs
--- a/Engine/Source/Programs/UnrealBuildTool/Executors/Experimental/FASTBuild.cs
+++ b/Engine/Source/Programs/UnrealBuildTool/Executors/Experimental/FASTBuild.cs
@@ -325,12 +325,16 @@ namespace UnrealBuildTool
 		//////////////////////////////////////////
 		// Action Helpers
 
-		private ObjectIDGenerator objectIDGenerator = new ObjectIDGenerator();
+		private readonly Dictionary<LinkedAction, long>  _actionToId = new();
+		private long _actionIdIndex = 0;
 
 		private long GetActionID(LinkedAction Action)
 		{
-			bool bFirstTime = false;
-			return objectIDGenerator.GetId(Action, out bFirstTime);
+			if (!_actionToId.ContainsKey(Action))
+			{
+				_actionToId.Add(Action, ++_actionIdIndex);
+			}
+			return _actionToId[Action];
 		}
 
 		private string ActionToActionString(LinkedAction Action)
@@ -338,12 +342,12 @@ namespace UnrealBuildTool
 			return ActionToActionString(GetActionID(Action));
 		}
 
-		private string ActionToActionString(long UniqueId)
+		private static string ActionToActionString(long UniqueId)
 		{
 			return $"Action_{UniqueId}";
 		}
 
-		private string ActionToDependencyString(long UniqueId, string StatusDescription, string? CommandDescription = null, ActionType? ActionType = null)
+		private static string ActionToDependencyString(long UniqueId, string StatusDescription, string? CommandDescription = null, ActionType? ActionType = null)
 		{
 			string? ExtraInfoString = null;
 			if ((CommandDescription != null) && String.IsNullOrEmpty(CommandDescription))
diff --git a/Engine/Source/Programs/UnrealBuildTool/Executors/UnrealBuildAccelerator/UBAExecutor.cs b/Engine/Source/Programs/UnrealBuildTool/Executors/UnrealBuildAccelerator/UBAExecutor.cs
--- a/Engine/Source/Programs/UnrealBuildTool/Executors/UnrealBuildAccelerator/UBAExecutor.cs
+++ b/Engine/Source/Programs/UnrealBuildTool/Executors/UnrealBuildAccelerator/UBAExecutor.cs
@@ -148,7 +148,7 @@ namespace UnrealBuildTool
 
 		class UBAArtifactCache : IArtifactCache
 		{
-			public ArtifactCacheState State { get => ArtifactCacheState.Available; }
+			public ArtifactCacheState State => ArtifactCacheState.Available;
 			public Task<ArtifactCacheState> WaitForReadyAsync() => Task.FromResult(ArtifactCacheState.Available);
 			public Task<ArtifactAction[]> QueryArtifactActionsAsync(IoHash[] partialKeys, CancellationToken cancellationToken) => Task.FromResult(Array.Empty<ArtifactAction>());
 			public Task<bool[]?> QueryArtifactOutputsAsync(ArtifactAction[] artifactActions, CancellationToken cancellationToken) => Task.FromResult<bool[]?>(null);
@@ -300,7 +300,7 @@ namespace UnrealBuildTool
 							int port = 1347;
 							if (nameAndPort.Length > 1)
 							{
-								port = int.Parse(nameAndPort[1]);
+								port = Int32.Parse(nameAndPort[1]);
 							}
 
 							if (_cacheClient.Connect(nameAndPort[0], port))
diff --git a/Engine/Source/Programs/UnrealBuildTool/Modes/TestMode.cs b/Engine/Source/Programs/UnrealBuildTool/Modes/TestMode.cs
--- a/Engine/Source/Programs/UnrealBuildTool/Modes/TestMode.cs
+++ b/Engine/Source/Programs/UnrealBuildTool/Modes/TestMode.cs
@@ -140,7 +140,7 @@ namespace UnrealBuildTool
 			HashSet<Type> AllTestModules = Assembly.GetTypes()
 				.Where(t => t.IsAssignableTo(ModuleRulesType))
 				.ToHashSet();
-			foreach(var Type in AllTestModules)
+			foreach(Type Type in AllTestModules)
 			{
 				// Find root class that's not TestModuleRules, the while loop always stops because AllTestModules contains only types that inherit from TestModuleRules
 				Type? RootBaseType = Type.BaseType;
@@ -190,7 +190,7 @@ namespace UnrealBuildTool
 					TestModuleType.TypeInitializer.Invoke(null, null);
 					if (InheritedGroups.ContainsKey(TestModuleType))
 					{
-						foreach (var InheritedTestModuleType in InheritedGroups[TestModuleType])
+						foreach (Type InheritedTestModuleType in InheritedGroups[TestModuleType])
 						{
 							if (InheritedTestModuleType.TypeInitializer != null)
 							{
@@ -205,7 +205,7 @@ namespace UnrealBuildTool
 					if (UpdateBuildGraphMetadata != null && TestMetadataField != null)
 					{
 						TestModuleRules.Metadata MetadataValue = (TestModuleRules.Metadata)TestMetadataField.GetValue(null)!;
-						if (MetadataValue != null && !string.IsNullOrEmpty(MetadataValue.TestName))
+						if (MetadataValue != null && !String.IsNullOrEmpty(MetadataValue.TestName))
 						{
 							UpdateBuildGraphMetadata.Invoke(null,
                                 new object[] {
diff --git a/Engine/Source/Programs/UnrealBuildTool/System/ExternalExecution.cs b/Engine/Source/Programs/UnrealBuildTool/System/ExternalExecution.cs
--- a/Engine/Source/Programs/UnrealBuildTool/System/ExternalExecution.cs
+++ b/Engine/Source/Programs/UnrealBuildTool/System/ExternalExecution.cs
@@ -500,7 +500,7 @@ namespace UnrealBuildTool
 
 		static void CollectModulesOnlyIncluded(HashSet<UEBuildModuleCPP> HandledModules, List<UEBuildModule> IncludedModules, List<UEBuildModuleCPP> OutList)
 		{
-			foreach (var pi in IncludedModules)
+			foreach (UEBuildModule pi in IncludedModules)
 			{
 				if (pi is UEBuildModuleCPP IncludedModule)
 				{
@@ -527,12 +527,12 @@ namespace UnrealBuildTool
 		static void CollectModulesOnlyIncluded(IEnumerable<UEBuildModuleCPP> ModulesToGenerateHeadersFor, List<UEBuildModuleCPP> OutList)
 		{
 			HashSet<UEBuildModuleCPP> HandledModules = new();
-			foreach (var Module in ModulesToGenerateHeadersFor)
+			foreach (UEBuildModuleCPP Module in ModulesToGenerateHeadersFor)
 			{
 				HandledModules.Add(Module);
 			}
 
-			foreach (var Module in ModulesToGenerateHeadersFor)
+			foreach (UEBuildModuleCPP Module in ModulesToGenerateHeadersFor)
 			{
 				if (Module.PublicIncludePathModules != null)
 				{
