Adding CustomConfig option to ConfigCache.ReadSettings

 - not fixing up all call sites because it may not be appropriate

#rnx

diff --git a/Engine/Source/Programs/UnrealBuildTool/System/ConfigCache.cs b/Engine/Source/Programs/UnrealBuildTool/System/ConfigCache.cs
--- a/Engine/Source/Programs/UnrealBuildTool/System/ConfigCache.cs
+++ b/Engine/Source/Programs/UnrealBuildTool/System/ConfigCache.cs
@@ -449,9 +449,10 @@ namespace UnrealBuildTool
 		/// <param name="ProjectDir">Path to the project directory</param>
 		/// <param name="Platform">The platform being built</param>
 		/// <param name="TargetObject">Object to receive the settings</param>
-		public static void ReadSettings(DirectoryReference? ProjectDir, UnrealTargetPlatform Platform, object TargetObject)
+		/// <param name="CustomConfig">Optional override config directory to search, for support of multiple target types</param>
+		public static void ReadSettings(DirectoryReference? ProjectDir, UnrealTargetPlatform Platform, object TargetObject, string CustomConfig = "")
 		{
-			ReadSettings(ProjectDir, Platform, TargetObject, null, null);
+			ReadSettings(ProjectDir, Platform, TargetObject, null, null, CustomConfig);
 		}
 
 		/// <summary>
@@ -462,7 +463,8 @@ namespace UnrealBuildTool
 		/// <param name="TargetObject">Object to receive the settings</param>
 		/// <param name="ConfigValues">Will be populated with config values that were retrieved. May be null.</param>
 		/// <param name="CmdLineArgs">Command line arguments, if null the application's command line arguments will be used</param>
-		internal static void ReadSettings(DirectoryReference? ProjectDir, UnrealTargetPlatform Platform, object TargetObject, Dictionary<ConfigDependencyKey, IReadOnlyList<string>?>? ConfigValues, CommandLineArguments? CmdLineArgs)
+		/// <param name="CustomConfig">Optional override config directory to search, for support of multiple target types</param>
+		internal static void ReadSettings(DirectoryReference? ProjectDir, UnrealTargetPlatform Platform, object TargetObject, Dictionary<ConfigDependencyKey, IReadOnlyList<string>?>? ConfigValues, CommandLineArguments? CmdLineArgs, string CustomConfig = "")
 		{
 			List<ConfigMember> Members = FindConfigMembersForType(TargetObject.GetType());
 
@@ -473,11 +475,11 @@ namespace UnrealBuildTool
 
 				if (CmdLineArgs == null)
 				{
-					Hierarchy = ReadHierarchy(Member.Attribute.ConfigType, ProjectDir, Platform);
+					Hierarchy = ReadHierarchy(Member.Attribute.ConfigType, ProjectDir, Platform, CustomConfig);
 				}
 				else
 				{
-					Hierarchy = ReadHierarchy(Member.Attribute.ConfigType, ProjectDir, Platform, "", CmdLineArgs);
+					Hierarchy = ReadHierarchy(Member.Attribute.ConfigType, ProjectDir, Platform, CustomConfig, CmdLineArgs);
 				}
 
 				// Get the key name
