UnrealBuildTool: Allow UBA executor to be used if action count is low, as it internally disables remote if the number of actions is insufficent

#rnx

diff --git a/Engine/Source/Programs/UnrealBuildTool/Executors/UnrealBuildAccelerator/UBAExecutor.cs b/Engine/Source/Programs/UnrealBuildTool/Executors/UnrealBuildAccelerator/UBAExecutor.cs
--- a/Engine/Source/Programs/UnrealBuildTool/Executors/UnrealBuildAccelerator/UBAExecutor.cs
+++ b/Engine/Source/Programs/UnrealBuildTool/Executors/UnrealBuildAccelerator/UBAExecutor.cs
@@ -53,6 +53,10 @@ namespace UnrealBuildTool
 				_session?.Dispose();
 				_session = null;
 				_threadedLogger.Dispose();
+				_ubaLogger?.Dispose();
+				_ubaLogger = null;
+				_cacheClient?.Dispose();
+				_cacheClient = null;
 			}
 			base.Dispose(disposing);
 		}
@@ -158,10 +162,10 @@ namespace UnrealBuildTool
 
 		class UBAActionArtifactCache : IActionArtifactCache
 		{
-			UBAExecutor _executor;
-			UBAArtifactCache _cache = new UBAArtifactCache();
+			readonly UBAExecutor _executor;
+			readonly UBAArtifactCache _cache = new UBAArtifactCache();
 			public UBAActionArtifactCache(UBAExecutor executor) {  _executor = executor; }
-			public IArtifactCache ArtifactCache { get => _cache; }
+			public IArtifactCache ArtifactCache => _cache;
 			public bool EnableReads { get => true; set { } }
 			public bool EnableWrites { get => true; set { } }
 			public bool LogCacheMisses { get => true; set { } }
@@ -172,7 +176,7 @@ namespace UnrealBuildTool
 				return Task.Factory.StartNew(() =>
 				{
 					ProcessStartInfo startInfo = _executor.GetActionStartInfo(action, out FileItem? pchItem);
-					uint bucket = _executor.GetActionCacheBucket(action);
+					uint bucket = UBAExecutor.GetActionCacheBucket(action);
 					using (IRootPaths rootPaths = _executor.GetActionRootPaths(action))
 					{
 						return _executor._cacheClient!.FetchFromCache(rootPaths, bucket, startInfo);
@@ -205,6 +209,8 @@ namespace UnrealBuildTool
 			if (inputActions.Count() < NumParallelProcesses && !UBAConfig.bForceBuildAllRemote)
 			{
 				UBAConfig.bDisableRemote = true;
+				UBAConfig.bForceBuildAllRemote = false;
+				UBAConfig.Zone = "local";
 			}
 
 			PrintConfiguration();
@@ -309,7 +315,7 @@ namespace UnrealBuildTool
 							}
 							else
 							{
-								logger.LogInformation($"Timed out trying to connect to cache server on {nameAndPort[0]}:{port}. Cache will be disabled");
+								logger.LogInformation("Timed out trying to connect to cache server on {Name}:{Port}. Cache will be disabled", nameAndPort[0], port);
 							}
 						}
 
@@ -405,7 +411,6 @@ namespace UnrealBuildTool
 			ImmediateActionQueueRunner remoteRunner = queue.CreateManualRunner(action => RunActionRemote(queue, action));
 			queue.CancellationToken.Register(onCancel);
 
-
 			// Start the queue
 			queue.Start();
 
@@ -547,7 +552,7 @@ namespace UnrealBuildTool
 			return startInfo;
 		}
 
-		uint GetActionCacheBucket(LinkedAction action)
+		static uint GetActionCacheBucket(LinkedAction action)
 		{
 			if (action.Target == null)
 			{
diff --git a/Engine/Source/Programs/UnrealBuildTool/System/ActionGraph.cs b/Engine/Source/Programs/UnrealBuildTool/System/ActionGraph.cs
--- a/Engine/Source/Programs/UnrealBuildTool/System/ActionGraph.cs
+++ b/Engine/Source/Programs/UnrealBuildTool/System/ActionGraph.cs
@@ -355,13 +355,13 @@ namespace UnrealBuildTool
 			}
 		}
 
-		private static ActionExecutor? GetRemoteExecutorByName(string Name, BuildConfiguration BuildConfiguration, int ActionCount, List<TargetDescriptor> TargetDescriptors, ILogger Logger)
+		private static ActionExecutor? GetRemoteExecutorByName(string Name, BuildConfiguration BuildConfiguration, int ActionCount, int MinActionsForRemote, List<TargetDescriptor> TargetDescriptors, ILogger Logger)
 		{
 			switch (Name)
 			{
 				case "XGE":
 					{
-						if (BuildConfiguration.bAllowXGE && XGE.IsAvailable(Logger) && ActionCount >= XGE.MinActions)
+						if (ActionCount >= MinActionsForRemote && BuildConfiguration.bAllowXGE && XGE.IsAvailable(Logger) && ActionCount >= XGE.MinActions)
 						{
 							return new XGE(Logger);
 						}
@@ -369,7 +369,7 @@ namespace UnrealBuildTool
 					}
 				case "SNDBS":
 					{
-						if (BuildConfiguration.bAllowSNDBS && SNDBS.IsAvailable(Logger))
+						if (ActionCount >= MinActionsForRemote && BuildConfiguration.bAllowSNDBS && SNDBS.IsAvailable(Logger))
 						{
 							return new SNDBS(TargetDescriptors, Logger);
 						}
@@ -377,7 +377,7 @@ namespace UnrealBuildTool
 					}
 				case "FASTBuild":
 					{
-						if (BuildConfiguration.bAllowFASTBuild && FASTBuild.IsAvailable(Logger))
+						if (ActionCount >= MinActionsForRemote && BuildConfiguration.bAllowFASTBuild && FASTBuild.IsAvailable(Logger))
 						{
 							return new FASTBuild(BuildConfiguration.MaxParallelActions, BuildConfiguration.bAllCores, BuildConfiguration.bCompactOutput, Logger);
 						}
@@ -385,6 +385,7 @@ namespace UnrealBuildTool
 					}
 				case "UBA":
 					{
+						// Intentionally not checking MinActionsForRemote
 						if (BuildConfiguration.bAllowUBAExecutor && UBAExecutor.IsAvailable())
 						{
 							return new UBAExecutor(BuildConfiguration.MaxParallelActions, BuildConfiguration.bAllCores, BuildConfiguration.bCompactOutput, Logger, TargetDescriptors.FirstOrDefault()?.AdditionalArguments);
@@ -404,15 +405,13 @@ namespace UnrealBuildTool
 		/// </summary>
 		private static ActionExecutor SelectExecutor(BuildConfiguration BuildConfiguration, int ActionCount, List<TargetDescriptor> TargetDescriptors, ILogger Logger)
 		{
-			if (ActionCount > ParallelExecutor.GetDefaultNumParallelProcesses(BuildConfiguration.MaxParallelActions, BuildConfiguration.bAllCores, Logger))
+			int MinActionsForRemote = ParallelExecutor.GetDefaultNumParallelProcesses(BuildConfiguration.MaxParallelActions, BuildConfiguration.bAllCores, Logger);
+			foreach (string Name in BuildConfiguration.RemoteExecutorPriority)
 			{
-				foreach (string Name in BuildConfiguration.RemoteExecutorPriority)
+				ActionExecutor? Executor = GetRemoteExecutorByName(Name, BuildConfiguration, ActionCount, MinActionsForRemote, TargetDescriptors, Logger);
+				if (Executor != null)
 				{
-					ActionExecutor? Executor = GetRemoteExecutorByName(Name, BuildConfiguration, ActionCount, TargetDescriptors, Logger);
-					if (Executor != null)
-					{
-						return Executor;
-					}
+					return Executor;
 				}
 			}
 
