UnrealBuildTool: Don't throw an exeception if no Horde URL is provided, and ensure HordeAgents still get the correct settings from the environment

#rnx

diff --git a/Engine/Source/Programs/UnrealBuildTool/Executors/UnrealBuildAccelerator/UBAAgentCoordinatorHorde.cs b/Engine/Source/Programs/UnrealBuildTool/Executors/UnrealBuildAccelerator/UBAAgentCoordinatorHorde.cs
--- a/Engine/Source/Programs/UnrealBuildTool/Executors/UnrealBuildAccelerator/UBAAgentCoordinatorHorde.cs
+++ b/Engine/Source/Programs/UnrealBuildTool/Executors/UnrealBuildAccelerator/UBAAgentCoordinatorHorde.cs
@@ -373,6 +373,12 @@ namespace UnrealBuildTool
 				return null;
 			}
 
+			if (String.IsNullOrEmpty(hordeConfig.HordeServer))
+			{
+				logger.LogInformation("Horde disabled. Url not set.");
+				return null;
+			}
+
 			Uri? server = (hordeConfig.HordeServer == null) ? null : new Uri(hordeConfig.HordeServer);
 			string? token = hordeConfig.HordeToken;
 
@@ -569,6 +575,16 @@ namespace UnrealBuildTool
 			XmlConfig.ApplyTo(HordeConfig);
 			additionalArguments?.ApplyTo(HordeConfig);
 
+			if (String.IsNullOrEmpty(HordeConfig.HordeServer))
+			{
+				HordeConfig.HordeServer = Environment.GetEnvironmentVariable(HordeHttpClient.HordeUrlEnvVarName);
+			}
+
+			if (String.IsNullOrEmpty(HordeConfig.HordeToken))
+			{
+				HordeConfig.HordeToken = Environment.GetEnvironmentVariable(HordeHttpClient.HordeTokenEnvVarName);
+			}
+
 			// Sentry is currently unsupported for non-Windows and non-x64
 			if (!OperatingSystem.IsWindows() || RuntimeInformation.ProcessArchitecture != Architecture.X64)
 			{
