UnrealBuildTool: Fix unset build branch for build machines if the version update task was not run

#rnx

diff --git a/Engine/Source/Programs/UnrealBuildTool/System/Telemetry.cs b/Engine/Source/Programs/UnrealBuildTool/System/Telemetry.cs
--- a/Engine/Source/Programs/UnrealBuildTool/System/Telemetry.cs
+++ b/Engine/Source/Programs/UnrealBuildTool/System/Telemetry.cs
@@ -274,7 +274,14 @@ namespace UnrealBuildTool
 		readonly Lazy<TelemetryMetadata> _metadata = new(() =>
 		{
 			BuildVersion.TryRead(BuildVersion.GetDefaultFileName(), out BuildVersion? version);
+			string buildBranchName = version?.BranchName ?? String.Empty;
+			if (Unreal.IsBuildMachine())
+			{
+				buildBranchName = Environment.GetEnvironmentVariable("uebp_BuildRoot_Escaped") ?? buildBranchName;
+			}
 #pragma warning disable CA1308 // Normalize strings to uppercase
+			buildBranchName = buildBranchName.ToLowerInvariant();
+#pragma warning restore CA1308 // Normalize strings to uppercase
 			return new TelemetryMetadata(
 				version != null ? $"{version.MajorVersion}.{version.MinorVersion}.{version.PatchVersion}-{version.Changelist}+{version.BranchName}" : String.Empty,
 				"UnrealBuildTool",
@@ -283,7 +290,7 @@ namespace UnrealBuildTool
 				UnrealBuildTool.SessionIdentifier,
 				UnrealBuildTool.BuildMode,
 				Get()._descriptor,
-				version?.BranchName?.ToLowerInvariant() ?? String.Empty,
+				buildBranchName,
 				Unreal.IsBuildMachine(),
 				Environment.GetEnvironmentVariable("UE_HORDE_BATCHID"),
 				Environment.GetEnvironmentVariable("UE_HORDE_JOBID"),
@@ -293,7 +300,6 @@ namespace UnrealBuildTool
 				Environment.GetEnvironmentVariable("UE_HORDE_TEMPLATEID"),
 				Environment.GetEnvironmentVariable("UE_HORDE_TEMPLATENAME")
 			);
-#pragma warning restore CA1308 // Normalize strings to uppercase
 		});
 
 		readonly Lazy<HttpClient> _httpClient;
