Throw a build exception when trying to build hybrid content only projects within a Templates folder

diff --git a/Engine/Source/Programs/UnrealBuildTool/Configuration/TargetDescriptor.cs b/Engine/Source/Programs/UnrealBuildTool/Configuration/TargetDescriptor.cs
--- a/Engine/Source/Programs/UnrealBuildTool/Configuration/TargetDescriptor.cs
+++ b/Engine/Source/Programs/UnrealBuildTool/Configuration/TargetDescriptor.cs
@@ -579,6 +579,13 @@ namespace UnrealBuildTool
 						// make a temp target for hybrid content-as-code projects
 						if (TargetProjectFile != null)
 						{
+							// cannot create temp targets for hybrid content only template projects
+							if (NativeProjects.IsHybridContentOnlyProject(TargetProjectFile, Logger) && 
+								TargetProjectFile.ContainsName("Templates", 0))
+							{
+								throw new BuildException("Cannot create target for project in Templates folder {0}. Please copy or move the project out of Templates.", TargetProjectFile.Directory.ToString());
+							}
+
 							NativeProjects.ConditionalMakeTempTargetForHybridProject(TargetProjectFile, new List<UnrealTargetPlatform>(1) { Platform }, Logger);
 						}
 
diff --git a/Engine/Source/Programs/UnrealBuildTool/System/NativeProjects.cs b/Engine/Source/Programs/UnrealBuildTool/System/NativeProjects.cs
--- a/Engine/Source/Programs/UnrealBuildTool/System/NativeProjects.cs
+++ b/Engine/Source/Programs/UnrealBuildTool/System/NativeProjects.cs
@@ -255,12 +255,6 @@ namespace UnrealBuildTool
 			// no reason by default
 			Reason = null;
 
-			// never create temp targets for the Template projects
-			if (UProjectFile.ContainsName("Templates", 0))
-			{
-				return false;
-			}
-
 			bool bHasCode = ProjectHasCode(UProjectFile, bCheckForTempTargets: false);
 			foreach (UnrealTargetPlatform Platform in Platforms)
 			{
