Fix Mac ShaderPlatform settings not being correctly intialized/loaded.

Fix made by [REDACTED_PERSON].

diff --git a/Engine/Source/Programs/UnrealBuildTool/Platform/Mac/UEBuildMac.cs b/Engine/Source/Programs/UnrealBuildTool/Platform/Mac/UEBuildMac.cs
--- a/Engine/Source/Programs/UnrealBuildTool/Platform/Mac/UEBuildMac.cs
+++ b/Engine/Source/Programs/UnrealBuildTool/Platform/Mac/UEBuildMac.cs
@@ -471,9 +471,10 @@ namespace UnrealBuildTool
 				// because remote IOS building needs the new XcodeProject Settings to show up in the editor, we bring in the Mac bits that expose it
 				if (ModuleName == "Engine")
 				{
+					Rules.DynamicallyLoadedModuleNames.Add("MacTargetPlatformSettings");
 					if (bIsPlatformAvailableForTargetWithSDK)
 					{
-						Rules.DynamicallyLoadedModuleNames.AddAll("MacTargetPlatform", "MacTargetPlatformSettings", "MacTargetPlatformControls", "MacPlatformEditor");
+						Rules.DynamicallyLoadedModuleNames.AddAll("MacTargetPlatform", "MacTargetPlatformControls", "MacPlatformEditor");
 					}
 				}
 			}
