UnrealBuildTool: Track cache stats in executor telemetry

#rnx

diff --git a/Engine/Source/Programs/UnrealBuildTool/Executors/Experimental/FASTBuild.cs b/Engine/Source/Programs/UnrealBuildTool/Executors/Experimental/FASTBuild.cs
--- a/Engine/Source/Programs/UnrealBuildTool/Executors/Experimental/FASTBuild.cs
+++ b/Engine/Source/Programs/UnrealBuildTool/Executors/Experimental/FASTBuild.cs
@@ -501,7 +501,7 @@ namespace UnrealBuildTool
 				}
 
 				bool result = ExecuteBffFile(FASTBuildFilePath, Logger);
-				telemetryEvent = new TelemetryExecutorEvent(Name, startTimeUTC, result, Actions.Count(), -1, -1, DateTime.UtcNow);
+				telemetryEvent = new TelemetryExecutorEvent(Name, startTimeUTC, result, Actions.Count(), -1, -1, 0, 0, DateTime.UtcNow);
 			}
 
 			return await LocalExecutor.ExecuteActionsAsync(Actions, Logger, actionArtifactCache);
diff --git a/Engine/Source/Programs/UnrealBuildTool/Executors/ImmediateActionQueue.cs b/Engine/Source/Programs/UnrealBuildTool/Executors/ImmediateActionQueue.cs
--- a/Engine/Source/Programs/UnrealBuildTool/Executors/ImmediateActionQueue.cs
+++ b/Engine/Source/Programs/UnrealBuildTool/Executors/ImmediateActionQueue.cs
@@ -298,7 +298,7 @@ namespace UnrealBuildTool
 		/// <summary>
 		/// Per-second logging of cpu utilization
 		/// </summary>
-		private List<float> _cpuUtilization = new();
+		private readonly List<float> _cpuUtilization = new();
 
 		/// <summary>
 		/// Collection of all actions remaining to be logged
@@ -320,6 +320,16 @@ namespace UnrealBuildTool
 		/// </summary>
 		private int _completedActions = 0;
 
+		/// <summary>
+		/// Tracks the number of completed actions from cache hits.
+		/// </summary>
+		private int _cacheHitActions = 0;
+
+		/// <summary>
+		/// Tracks the number of unsuccessful actions from cache misses.
+		/// </summary>
+		private int _cacheMissActions = 0;
+
 		/// <summary>
 		/// Flags used to track how StartManyActions should run
 		/// </summary>
@@ -363,9 +373,9 @@ namespace UnrealBuildTool
 		/// <summary>
 		/// If set, artifact cache used to retrieve previously compiled results and save new results
 		/// </summary>
-		private IActionArtifactCache? _actionArtifactCache;
+		private readonly IActionArtifactCache? _actionArtifactCache;
 
-		static ExecuteResults s_copiedFromCacheResults = new(new List<string>(), 0, TimeSpan.Zero, TimeSpan.Zero, "[Cache]");
+		static readonly ExecuteResults s_copiedFromCacheResults = new(new List<string>(), 0, TimeSpan.Zero, TimeSpan.Zero, "[Cache]");
 
 		/// <summary>
 		/// Construct a new instance of the action queue
@@ -409,22 +419,24 @@ namespace UnrealBuildTool
 
 			if (readArtifacts)
 			{
-				Func<LinkedAction, Func<Task>> runAction = (LinkedAction action) =>
+				Func<Task> runAction(LinkedAction action)
 				{
 					return new Func<Task>(async () =>
 					{
 						bool success = await _actionArtifactCache!.CompleteActionFromCacheAsync(action, CancellationToken);
 						if (success)
 						{
+							Interlocked.Increment(ref _cacheHitActions);
 							OnArtifactsRead?.Invoke(action);
 							OnActionCompleted(action, success, s_copiedFromCacheResults);
 						}
 						else
 						{
+							Interlocked.Increment(ref _cacheMissActions);
 							RequeueAction(action);
 						}
 					});
-				};
+				}
 
 				_runners.Add(new(ImmediateActionQueueRunnerType.Automatic, ActionPhase.ArtifactCheck, runAction, false, maxActionArtifactCacheTasks, 0));
 			}
@@ -526,7 +538,6 @@ namespace UnrealBuildTool
 							}
 						}
 					}
-
 				}, null, 1000, 1000);
 			}
 
@@ -550,11 +561,15 @@ namespace UnrealBuildTool
 		/// <param name="totalActions">Out parameter, the total number of actions</param>
 		/// <param name="succeededActions">Out parameter, the number of successful actions</param>
 		/// <param name="failedActions">Out parameter, the number of failed actions</param>
-		public void GetActionResultCounts(out int totalActions, out int succeededActions, out int failedActions)
+		/// <param name="cacheHitActions">Out parameter, the number of cache hit actions</param>
+		/// <param name="cacheMissActions">Out parameter, the number of cache miss actions</param>
+		public void GetActionResultCounts(out int totalActions, out int succeededActions, out int failedActions, out int cacheHitActions, out int cacheMissActions)
 		{
 			totalActions = Actions.Length;
 			succeededActions = Actions.Where(x => x.Results?.ExitCode == 0).Count();
 			failedActions = Actions.Where(x => x.Results != null && x.Results.ExitCode != 0).Count();
+			cacheHitActions = _cacheHitActions;
+			cacheMissActions = _cacheMissActions;
 		}
 
 		/// <summary>
@@ -1165,7 +1180,7 @@ namespace UnrealBuildTool
 					foreach (string Line in logLines.Skip(action.bShouldOutputStatusDescription ? 0 : 1))
 					{
 						// suppress library creation messages when writing compact output
-						if (CompactOutput && Line.StartsWith("   Creating library ") && Line.EndsWith(".exp"))
+						if (CompactOutput && Line.StartsWith("   Creating library ", StringComparison.OrdinalIgnoreCase) && Line.EndsWith(".exp", StringComparison.OrdinalIgnoreCase))
 						{
 							continue;
 						}
diff --git a/Engine/Source/Programs/UnrealBuildTool/Executors/ParallelExecutor.cs b/Engine/Source/Programs/UnrealBuildTool/Executors/ParallelExecutor.cs
--- a/Engine/Source/Programs/UnrealBuildTool/Executors/ParallelExecutor.cs
+++ b/Engine/Source/Programs/UnrealBuildTool/Executors/ParallelExecutor.cs
@@ -202,8 +202,8 @@ namespace UnrealBuildTool
 				queue.StartManyActions();
 				result = await queue.RunTillDone();
 
-				queue.GetActionResultCounts(out int totalActions, out int succeededActions, out int failedActions);
-				telemetryEvent = new TelemetryExecutorEvent(Name, startTimeUTC, result, totalActions, succeededActions, failedActions, DateTime.UtcNow);
+				queue.GetActionResultCounts(out int totalActions, out int succeededActions, out int failedActions, out int cacheHitActions, out int cacheMissActions);
+				telemetryEvent = new TelemetryExecutorEvent(Name, startTimeUTC, result, totalActions, succeededActions, failedActions, cacheHitActions, cacheMissActions, DateTime.UtcNow);
 			}
 			else
 			{
@@ -215,8 +215,8 @@ namespace UnrealBuildTool
 				queue.StartManyActions();
 				result = await queue.RunTillDone();
 
-				queue.GetActionResultCounts(out int totalActions, out int succeededActions, out int failedActions);
-				telemetryEvent = new TelemetryExecutorEvent(Name, startTimeUTC, result, totalActions, succeededActions, failedActions, DateTime.UtcNow);
+				queue.GetActionResultCounts(out int totalActions, out int succeededActions, out int failedActions, out int cacheHitActions, out int cacheMissActions);
+				telemetryEvent = new TelemetryExecutorEvent(Name, startTimeUTC, result, totalActions, succeededActions, failedActions, cacheHitActions, cacheMissActions, DateTime.UtcNow);
 			}
 
 			return result;
diff --git a/Engine/Source/Programs/UnrealBuildTool/Executors/SNDBS.cs b/Engine/Source/Programs/UnrealBuildTool/Executors/SNDBS.cs
--- a/Engine/Source/Programs/UnrealBuildTool/Executors/SNDBS.cs
+++ b/Engine/Source/Programs/UnrealBuildTool/Executors/SNDBS.cs
@@ -382,7 +382,7 @@ namespace UnrealBuildTool
 					LocalProcess.WaitForExit();
 					bool result = LocalProcess.ExitCode == 0;
 
-					telemetryEvent = new TelemetryExecutorEvent(Name, startTimeUTC, result, NumActions, -1, -1, DateTime.UtcNow);
+					telemetryEvent = new TelemetryExecutorEvent(Name, startTimeUTC, result, NumActions, -1, -1, 0, 0, DateTime.UtcNow);
 					return result;
 				}
 				catch (Exception Ex)
diff --git a/Engine/Source/Programs/UnrealBuildTool/Executors/UnrealBuildAccelerator/UBAExecutor.cs b/Engine/Source/Programs/UnrealBuildTool/Executors/UnrealBuildAccelerator/UBAExecutor.cs
--- a/Engine/Source/Programs/UnrealBuildTool/Executors/UnrealBuildAccelerator/UBAExecutor.cs
+++ b/Engine/Source/Programs/UnrealBuildTool/Executors/UnrealBuildAccelerator/UBAExecutor.cs
@@ -164,7 +164,7 @@ namespace UnrealBuildTool
 		{
 			readonly UBAExecutor _executor;
 			readonly UBAArtifactCache _cache = new UBAArtifactCache();
-			public UBAActionArtifactCache(UBAExecutor executor) {  _executor = executor; }
+			public UBAActionArtifactCache(UBAExecutor executor) { _executor = executor; }
 			public IArtifactCache ArtifactCache => _cache;
 			public bool EnableReads { get => true; set { } }
 			public bool EnableWrites { get => true; set { } }
@@ -459,8 +459,8 @@ namespace UnrealBuildTool
 
 				bool res = queue.RunTillDone().Result; // Using inline wait to avoid possible thread switch
 
-				queue.GetActionResultCounts(out int totalActions, out int succeededActions, out int failedActions);
-				telemetryEvent = new TelemetryExecutorUBAEvent(Name, startTimeUTC, res, totalActions, succeededActions, failedActions,
+				queue.GetActionResultCounts(out int totalActions, out int succeededActions, out int failedActions, out int cacheHitActions, out int cacheMissActions);
+				telemetryEvent = new TelemetryExecutorUBAEvent(Name, startTimeUTC, res, totalActions, succeededActions, failedActions, cacheHitActions, cacheMissActions,
 					_localProcessedActions, _remoteProcessedActions, _localRetryActions.Count, _forcedRetryActions.Count, DateTime.UtcNow);
 
 				return res;
@@ -609,7 +609,7 @@ namespace UnrealBuildTool
 					{
 						continue;
 					}
-					
+
 					writer.Write(f.FullName);
 				}
 
@@ -681,7 +681,7 @@ namespace UnrealBuildTool
 					TimeSpan processorTime = process.TotalProcessorTime;
 					TimeSpan executionTime = process.TotalWallTime;
 					List<string> logLines = process.LogLines;
-					logLines.RemoveAll((line) => line.StartsWith("   Creating library ", StringComparison.OrdinalIgnoreCase) && line.EndsWith(".exp", StringComparison.OrdinalIgnoreCase) || line.EndsWith("file(s) copied.", StringComparison.OrdinalIgnoreCase));
+					logLines.RemoveAll((line) => (line.StartsWith("   Creating library ", StringComparison.OrdinalIgnoreCase) && line.EndsWith(".exp", StringComparison.OrdinalIgnoreCase)) || line.EndsWith("file(s) copied.", StringComparison.OrdinalIgnoreCase));
 
 					string? additionalDescription = !enableDetour ? "(UBA disabled)" : null;
 					ActionFinished(queue, new ExecuteResults(logLines, process.ExitCode, executionTime, processorTime, additionalDescription), action, pchItem, process);
diff --git a/Engine/Source/Programs/UnrealBuildTool/Executors/XGE.cs b/Engine/Source/Programs/UnrealBuildTool/Executors/XGE.cs
--- a/Engine/Source/Programs/UnrealBuildTool/Executors/XGE.cs
+++ b/Engine/Source/Programs/UnrealBuildTool/Executors/XGE.cs
@@ -781,7 +781,7 @@ namespace UnrealBuildTool
 				// Run through the standard XGE executor
 				bool result = ExecuteTaskFile(TaskFilePath, EventHandlerWrapper, NumActions, Logger);
 
-				telemetryEvent = new TelemetryExecutorEvent(Name, startTimeUTC, result, NumActions, -1, -1, DateTime.UtcNow);
+				telemetryEvent = new TelemetryExecutorEvent(Name, startTimeUTC, result, NumActions, -1, -1, 0, 0, DateTime.UtcNow);
 
 				return result;
 			}
diff --git a/Engine/Source/Programs/UnrealBuildTool/System/Telemetry.cs b/Engine/Source/Programs/UnrealBuildTool/System/Telemetry.cs
--- a/Engine/Source/Programs/UnrealBuildTool/System/Telemetry.cs
+++ b/Engine/Source/Programs/UnrealBuildTool/System/Telemetry.cs
@@ -119,13 +119,16 @@ namespace UnrealBuildTool
 		public int TotalActions { get; }
 		public int SucceededActions { get; }
 		public int FailedActions { get; }
-		public int SkippedActions { get; }
-
+		public int SkippedActions => TotalActions - SucceededActions - FailedActions;
+		public int CacheCheckActions => CacheHitActions + CacheMissActions;
+		public int CacheHitActions { get; }
+		public int CacheMissActions { get; }
 		public double SuccessRate => SucceededActions / (double)TotalActions;
 		public double FailureRate => FailedActions / (double)TotalActions;
 		public double SkipRate => SkippedActions / (double)TotalActions;
+		public double CacheRate => CacheHitActions + CacheMissActions > 0 ? CacheHitActions / (double)CacheCheckActions : 0.0;
 
-		public TelemetryExecutorEvent(string executor, DateTime startUTC, bool result, int totalActions, int succeededActions, int failedActions, DateTime timestamp)
+		public TelemetryExecutorEvent(string executor, DateTime startUTC, bool result, int totalActions, int succeededActions, int failedActions, int cacheHitActions, int cacheMissActions, DateTime timestamp)
 			: base(timestamp)
 		{
 			Executor = executor;
@@ -135,7 +138,8 @@ namespace UnrealBuildTool
 			TotalActions = totalActions;
 			SucceededActions = succeededActions;
 			FailedActions = failedActions;
-			SkippedActions = TotalActions - SucceededActions - FailedActions;
+			CacheHitActions = cacheHitActions;
+			CacheMissActions = cacheMissActions;
 		}
 	}
 
@@ -153,8 +157,8 @@ namespace UnrealBuildTool
 		public double LocalUsage => LocalActions / (double)(TotalActions + RetriedLocalActions + RetriedDisabledActions - SkippedActions);
 		public double RemoteUsage => RemoteActions / (double)(TotalActions + RetriedLocalActions + RetriedDisabledActions - SkippedActions);
 
-		public TelemetryExecutorUBAEvent(string executor, DateTime startUTC, bool result, int totalActions, int succeededActions, int failedActions, int localActions, int remoteActions, int retriedLocalActions, int retriedDisabledActions, DateTime timestamp)
-			: base(executor, startUTC, result, totalActions, succeededActions, failedActions, timestamp)
+		public TelemetryExecutorUBAEvent(string executor, DateTime startUTC, bool result, int totalActions, int succeededActions, int failedActions, int cacheHitActions, int cacheMissActions, int localActions, int remoteActions, int retriedLocalActions, int retriedDisabledActions, DateTime timestamp)
+			: base(executor, startUTC, result, totalActions, succeededActions, failedActions, cacheHitActions, cacheMissActions, timestamp)
 		{
 			LocalActions = localActions;
 			RemoteActions = remoteActions;
