Fixing header compiles in clang by disabling the "pragma once in source file" warning for .h files.

diff --git a/Engine/Source/Programs/UnrealBuildTool/ToolChain/ClangToolChain.cs b/Engine/Source/Programs/UnrealBuildTool/ToolChain/ClangToolChain.cs
--- a/Engine/Source/Programs/UnrealBuildTool/ToolChain/ClangToolChain.cs
+++ b/Engine/Source/Programs/UnrealBuildTool/ToolChain/ClangToolChain.cs
@@ -908,6 +908,15 @@ namespace UnrealBuildTool
 				// Compile the file as Objective-C code.
 				GetCompileArguments_M(CompileEnvironment, Arguments);
 			}
+			else if (Extension == ".H")
+			{
+				// Compile the file as C++ code.
+				GetCompileArguments_CPP(CompileEnvironment, Arguments);
+
+				// This warning was to catch #pragma once inside a source file.
+				// If we're compiling a header directly, we should always have the pragma once, so we need to ignore this warning.
+				Arguments.Add("-Wno-pragma-once-outside-header");
+			}
 			else
 			{
 				// Compile the file as C++ code.
