[UBA]
* Added -UbaReportCacheMissReason which can be set to output why a cache miss happened. This is very useful for debugging purposes to identify inputs that are not deterministic/portable

diff --git a/Engine/Source/Programs/Shared/EpicGames.UBA/ICacheClient.cs b/Engine/Source/Programs/Shared/EpicGames.UBA/ICacheClient.cs
--- a/Engine/Source/Programs/Shared/EpicGames.UBA/ICacheClient.cs
+++ b/Engine/Source/Programs/Shared/EpicGames.UBA/ICacheClient.cs
@@ -65,10 +65,11 @@ namespace EpicGames.UBA
 		/// Create a ICacheClient object
 		/// </summary>
 		/// <param name="session">The session</param>
+		/// <param name="reportMissReason">Output reason for cache miss to log.</param>
 		/// <returns>The ICacheClient</returns>
-		public static ICacheClient CreateCacheClient(ISessionServer session)
+		public static ICacheClient CreateCacheClient(ISessionServer session, bool reportMissReason)
 		{
-			return new CacheClientImpl(session);
+			return new CacheClientImpl(session, reportMissReason);
 		}
 	}
 }
diff --git a/Engine/Source/Programs/Shared/EpicGames.UBA/Impl/CacheClientImpl.cs b/Engine/Source/Programs/Shared/EpicGames.UBA/Impl/CacheClientImpl.cs
--- a/Engine/Source/Programs/Shared/EpicGames.UBA/Impl/CacheClientImpl.cs
+++ b/Engine/Source/Programs/Shared/EpicGames.UBA/Impl/CacheClientImpl.cs
@@ -77,7 +77,7 @@ namespace EpicGames.UBA
 		static extern void ProcessStartInfo_Destroy(IntPtr server);
 
 		[DllImport("UbaHost", CharSet = CharSet.Auto)]
-		static extern IntPtr CacheClient_Create(IntPtr session);
+		static extern IntPtr CacheClient_Create(IntPtr session, bool reportMissReason);
 
 		[DllImport("UbaHost", CharSet = CharSet.Auto)]
 		static extern bool CacheClient_Connect(IntPtr cacheClient, string host, int port);
@@ -115,9 +115,9 @@ namespace EpicGames.UBA
 		}
 		#endregion
 
-		public CacheClientImpl(ISessionServer session)
+		public CacheClientImpl(ISessionServer session, bool reportMissReason)
 		{
-			_handle = CacheClient_Create(session.GetHandle());
+			_handle = CacheClient_Create(session.GetHandle(), reportMissReason);
 		}
 
 		public bool Connect(string host, int port)
diff --git a/Engine/Source/Programs/UnrealBuildTool/Executors/UnrealBuildAccelerator/UBAExecutor.cs b/Engine/Source/Programs/UnrealBuildTool/Executors/UnrealBuildAccelerator/UBAExecutor.cs
--- a/Engine/Source/Programs/UnrealBuildTool/Executors/UnrealBuildAccelerator/UBAExecutor.cs
+++ b/Engine/Source/Programs/UnrealBuildTool/Executors/UnrealBuildAccelerator/UBAExecutor.cs
@@ -298,7 +298,7 @@ namespace UnrealBuildTool
 					using IStorageServer ubaStorageServer = IStorageServer.CreateStorageServer(Server, ubaLogger, new StorageServerCreateInfo(_rootDirRef.FullName, ((ulong)UBAConfig.StoreCapacityGb) * 1000 * 1000 * 1000, !UBAConfig.bStoreRaw, UBAConfig.Zone));
 					using ISessionServerCreateInfo serverCreateInfo = ISessionServerCreateInfo.CreateSessionServerCreateInfo(ubaStorageServer, Server, ubaLogger, new SessionServerCreateInfo(_rootDirRef.FullName, ubaTraceFile.FullName.Replace('\\', '/'), UBAConfig.bDisableCustomAlloc, false, UBAConfig.bResetCas, UBAConfig.bWriteToDisk, UBAConfig.bDetailedTrace, !UBAConfig.bDisableWaitOnMem, UBAConfig.bAllowKillOnMem, UBAConfig.bStoreObjFilesCompressed));
 					using (_session = ISessionServer.CreateSessionServer(serverCreateInfo))
-					using (_cacheClient = ICacheClient.CreateCacheClient(_session))
+					using (_cacheClient = ICacheClient.CreateCacheClient(_session, UBAConfig.bReportCacheMissReason))
 					{
 						if (!String.IsNullOrEmpty(UBAConfig.CacheServer))
 						{
diff --git a/Engine/Source/Programs/UnrealBuildTool/Executors/UnrealBuildAccelerator/UnrealBuildAcceleratorConfig.cs b/Engine/Source/Programs/UnrealBuildTool/Executors/UnrealBuildAccelerator/UnrealBuildAcceleratorConfig.cs
--- a/Engine/Source/Programs/UnrealBuildTool/Executors/UnrealBuildAccelerator/UnrealBuildAcceleratorConfig.cs
+++ b/Engine/Source/Programs/UnrealBuildTool/Executors/UnrealBuildAccelerator/UnrealBuildAcceleratorConfig.cs
@@ -248,5 +248,12 @@ namespace UnrealBuildTool
 		[XmlConfigFile(Category = "UnrealBuildAccelerator", Name = "CacheMaxWorkers")]
 		[CommandLine("-UBACacheMaxWorkers")]
 		public int CacheMaxWorkers { get; set; } = 32;
+
+		/// <summary>
+		/// Report reason a cache miss happened. Useful when searching for determinism/portability issues
+		/// </summary>
+		[XmlConfigFile(Category = "UnrealBuildAccelerator", Name = "ReportCacheMissReason")]
+		[CommandLine("-UBAReportCacheMissReason")]
+		public bool bReportCacheMissReason { get; set; }
 	}
 }
\ No newline at end of file
