UECommandLines per Configuration for IOS CDN and bulk builds

diff --git a/Engine/Source/Programs/UnrealBuildTool/ProjectFiles/Xcode/XcodeProject.cs b/Engine/Source/Programs/UnrealBuildTool/ProjectFiles/Xcode/XcodeProject.cs
--- a/Engine/Source/Programs/UnrealBuildTool/ProjectFiles/Xcode/XcodeProject.cs
+++ b/Engine/Source/Programs/UnrealBuildTool/ProjectFiles/Xcode/XcodeProject.cs
@@ -1448,12 +1448,6 @@ namespace UnrealBuildTool.XcodeProjectXcconfig
 				// make a string like --exclude=/Info.plist --exclude=/Manifest_* ...
 				string ExcludeString = String.Join(" ", Exclusions.Select(x => (x[0] == '+' ? "--include" : "--exclude") + $"=\\\"{x.Substring(1)}\\\""));
 
-				CopyScript.AddRange(new string[]
-				{
-					"",
-					$"echo \\\"Syncing ${{STAGED_DIR}}{SyncSourceSubdir} to ${{CONFIGURATION_BUILD_DIR}}/${{CONTENTS_FOLDER_PATH}}{SyncDestSubdir}\\\"",
-					$"rsync -a --delete {ExcludeString} \\\"${{STAGED_DIR}}{SyncSourceSubdir}/\\\" \\\"${{CONFIGURATION_BUILD_DIR}}/${{CONTENTS_FOLDER_PATH}}{SyncDestSubdir}\\\"",
-				});
 
 				// copy uecommandline.txt for IOS type platforms
 				if (Platform != UnrealTargetPlatform.Mac)
@@ -1466,6 +1460,23 @@ namespace UnrealBuildTool.XcodeProjectXcconfig
 						"fi"
 					});
 				}
+
+				CopyScript.AddRange(new string[]
+				{
+					"",
+					$"echo \\\"Syncing ${{STAGED_DIR}}{SyncSourceSubdir} to ${{CONFIGURATION_BUILD_DIR}}/${{CONTENTS_FOLDER_PATH}}{SyncDestSubdir}\\\"",
+					$"if [[ -e  \\\"${{STAGED_DIR}}{SyncSourceSubdir}\\\" ]]; then",
+					$"rsync -a --delete {ExcludeString} \\\"${{STAGED_DIR}}{SyncSourceSubdir}/\\\" \\\"${{CONFIGURATION_BUILD_DIR}}/${{CONTENTS_FOLDER_PATH}}{SyncDestSubdir}\\\"",
+					"else",
+					"  echo =========================================================================================",
+					"  echo \\\"WARNING: To run, you must have a valid staged sync source directory. The Staged SyncSource location is:\\\"",
+					$"  echo \\\"${{STAGED_DIR}}{SyncSourceSubdir}\\\"",
+					"  echo \\\"Use the editor's Platforms menu, or run a command like::\\\"",
+					$"  echo \\\"./RunUAT.sh BuildCookRun -platform={Platform} -project=<project> -build -cook -stage -pak\\\"",
+					"  echo =========================================================================================",
+					"  exit -0 ",
+					"fi",
+				});
 			}
 
 			// run this script every time, but xcode will show a warning if there isn't _some_ output
