UnrealBuildTool: Fix description of -SingleFile compile marked as Analyze when it should be Compiling (visual only bug)

#rnx

diff --git a/Engine/Source/Programs/UnrealBuildTool/Platform/Windows/VCToolChain.cs b/Engine/Source/Programs/UnrealBuildTool/Platform/Windows/VCToolChain.cs
--- a/Engine/Source/Programs/UnrealBuildTool/Platform/Windows/VCToolChain.cs
+++ b/Engine/Source/Programs/UnrealBuildTool/Platform/Windows/VCToolChain.cs
@@ -1565,7 +1565,7 @@ namespace UnrealBuildTool
 
 			AppendCLArguments_Global(CompileEnvironment, BaseCompileAction.Arguments);
 
-			BaseCompileAction.bIsAnalyzing = Target.StaticAnalyzer != StaticAnalyzer.Default && !CompileEnvironment.bDisableStaticAnalysis && !(Target.WindowsPlatform.Compiler.IsClang() && CompileEnvironment.PrecompiledHeaderAction == PrecompiledHeaderAction.Create);
+			BaseCompileAction.bIsAnalyzing = Target.StaticAnalyzer != StaticAnalyzer.None && !CompileEnvironment.bDisableStaticAnalysis && !(Target.WindowsPlatform.Compiler.IsClang() && CompileEnvironment.PrecompiledHeaderAction == PrecompiledHeaderAction.Create);
 
 			// Add include paths to the argument list.
 			BaseCompileAction.IncludePaths.AddRange(CompileEnvironment.UserIncludePaths);
@@ -1930,7 +1930,7 @@ namespace UnrealBuildTool
 					CompileAction.WriteResponseFile(Graph, Logger);
 				}
 
-				CompileAction.bIsAnalyzing = Target.StaticAnalyzer != StaticAnalyzer.None && !(Target.WindowsPlatform.Compiler.IsClang() && CompileEnvironment.PrecompiledHeaderAction == PrecompiledHeaderAction.Create);
+				CompileAction.bIsAnalyzing = Target.StaticAnalyzer != StaticAnalyzer.None && !CompileEnvironment.bDisableStaticAnalysis && !(Target.WindowsPlatform.Compiler.IsClang() && CompileEnvironment.PrecompiledHeaderAction == PrecompiledHeaderAction.Create);
 
 				// Update the output
 				Graph.AddAction(CompileAction);
