Fix order of arguments when using AFS to push libUnreal.so outside of .apk

diff --git a/Engine/Source/Programs/UnrealBuildTool/Platform/Android/UEDeployAndroid.cs b/Engine/Source/Programs/UnrealBuildTool/Platform/Android/UEDeployAndroid.cs
--- a/Engine/Source/Programs/UnrealBuildTool/Platform/Android/UEDeployAndroid.cs
+++ b/Engine/Source/Programs/UnrealBuildTool/Platform/Android/UEDeployAndroid.cs
@@ -5033,7 +5033,7 @@ namespace UnrealBuildTool
 					AFSToken = String.IsNullOrEmpty(AFSToken) ? "" : " -k " + AFSToken;
 
 					string AFSExecutable = Path.Combine(Unreal.EngineDirectory.ToString(), @"Binaries/DotNET/Android/UnrealAndroidFileTool", GetAFSExecutable(UnrealTargetPlatform.Win64, Logger));
-					string AFS = $"{AFSExecutable} -p {PackageName}{AFSToken}";
+					string AFSArguments = $"-p {PackageName}{AFSToken}";
 
 					string? SOPushScriptLocation = Path.GetDirectoryName(FinalSOName)!;
 
@@ -5042,7 +5042,8 @@ namespace UnrealBuildTool
 					// MakeApk will not be called in bDontBundleLibrariesInAPK mode, so we need to run stripping outside of it
 					string SOPushScript = @$"
 set ADB=adb
-set AFS={AFS.Replace("/", "\\")}
+set AFS={AFSExecutable.Replace("/", "\\")}
+set AFSARGS={AFSArguments}
 set DEVICE=
 if not \""%1\""==\""\"" set DEVICE=-s %1
 pushd %~dp0
@@ -5050,7 +5051,7 @@ pushd %~dp0
 					if (bUseAFS)
 					{
 						SOPushScript += @$"
-%AFS% %DEVICE% push {FinalSONameStrippedRelative} ""^int/libUnreal.so""
+%AFS% %DEVICE% %AFSARGS% push {FinalSONameStrippedRelative} ""^int/libUnreal.so""
 if ""%ERRORLEVEL%"" NEQ ""0"" (exit /b %ERRORLEVEL%)
 ";
 					}
