UnrealBuildTool: -IncludeHeader fixes
* Headers can have duplicate filenames which results in intermediate file conflicts, generate an intermediate subdirectory using the header relative path from the module root
* Don't use PCH for headers (also fix the original unmodified module compile environment being overwitten due to not making a copy !!!)
* Set SUPPRESS_MONOLITHIC_HEADER_WARNINGS=0 for header compiles
* Remove the need to generate a separate file to include the header with -SingleFile
* Fix /TC not being passed to clang-cl
* Disable a few clang #pragma related warnings

#rnx

diff --git a/Engine/Source/Programs/UnrealBuildTool/Configuration/UEBuildModuleCPP.cs b/Engine/Source/Programs/UnrealBuildTool/Configuration/UEBuildModuleCPP.cs
--- a/Engine/Source/Programs/UnrealBuildTool/Configuration/UEBuildModuleCPP.cs
+++ b/Engine/Source/Programs/UnrealBuildTool/Configuration/UEBuildModuleCPP.cs
@@ -429,19 +429,19 @@ namespace UnrealBuildTool
 				MinFilesUsingPrecompiledHeader = 1;
 			}
 
-			// Set up the environment with which to compile the CPP files
-			CppCompileEnvironment CompileEnvironment = ModuleCompileEnvironment;
-
 			// Generate ISPC headers first so C++ can consume them
 			if (InputFiles.ISPCFiles.Count > 0)
 			{
 				CreateHeadersForISPC(ToolChain, ModuleCompileEnvironment, InputFiles.ISPCFiles, IntermediateDirectory, Graph);
 			}
 
+			// Set up the environment with which to compile the CPP files
+			CppCompileEnvironment CompileEnvironment = new(ModuleCompileEnvironment);
+
 			// Compile any module interfaces
 			if (InputFiles.IXXFiles.Count > 0 && Target.bEnableCppModules)
 			{
-				CppCompileEnvironment IxxCompileEnvironment = CompileEnvironment;
+				CppCompileEnvironment IxxCompileEnvironment = new(ModuleCompileEnvironment);
 
 				// Write all the definitions to a separate file for the ixx compile
 				CreateHeaderForDefinitions(IxxCompileEnvironment, IntermediateDirectory, "ixx", Graph);
@@ -817,8 +817,19 @@ namespace UnrealBuildTool
 				List<FileItem> HeaderFileItems = GetCompilableHeaders(InputFiles, CompileEnvironment);
 				if (HeaderFileItems.Count > 0)
 				{
-					// Add the compile actions
-					LinkInputFiles.AddRange(ToolChain.CompileAllCPPFiles(CompileEnvironment, HeaderFileItems, IntermediateDirectory, Name, Graph).ObjectFiles);
+					CppCompileEnvironment HeaderCompileEnvironment = new CppCompileEnvironment(ModuleCompileEnvironment);
+					DirectoryReference HeaderIntermediateDirectory = DirectoryReference.Combine(IntermediateDirectory, "H");
+					HeaderCompileEnvironment.Definitions.RemoveAll(x => x.StartsWith("SUPPRESS_MONOLITHIC_HEADER_WARNINGS"));
+					HeaderCompileEnvironment.Definitions.Add("SUPPRESS_MONOLITHIC_HEADER_WARNINGS=0");
+					CreateHeaderForDefinitions(HeaderCompileEnvironment, HeaderIntermediateDirectory, "h", Graph);
+
+					// Duplicate named headers are allowed, so adjust the IntermediateDirectory to compensate for this
+					foreach (DirectoryItem ParentDir in HeaderFileItems.Select(x => x.Directory!).Distinct())
+					{
+						DirectoryReference HeaderSubIntermediateDirectory = DirectoryReference.Combine(HeaderIntermediateDirectory, ParentDir.Location.MakeRelativeTo(ModuleDirectory));
+						// Add the compile actions
+						LinkInputFiles.AddRange(ToolChain.CompileAllCPPFiles(HeaderCompileEnvironment, HeaderFileItems.Where(x => x.Directory == ParentDir), HeaderSubIntermediateDirectory, Name, Graph).ObjectFiles);
+					}
 				}
 			}
 
diff --git a/Engine/Source/Programs/UnrealBuildTool/Platform/Windows/VCSpecificFileAction.cs b/Engine/Source/Programs/UnrealBuildTool/Platform/Windows/VCSpecificFileAction.cs
--- a/Engine/Source/Programs/UnrealBuildTool/Platform/Windows/VCSpecificFileAction.cs
+++ b/Engine/Source/Programs/UnrealBuildTool/Platform/Windows/VCSpecificFileAction.cs
@@ -64,17 +64,7 @@ namespace UnrealBuildTool
 			}
 			SingleFiles[SourceFile.Name].Add(SourceFile);
 
-			// If it is a header file we need to wrap it from another file.. because otherwise it will fail when there are circular dependencies
-			// (There are a lot of those, very few of the .h files in core can be compiled without errors)
-			if (SourceFile.HasExtension(".h"))
-			{
-				DirectoryReference.CreateDirectory(IntermediateDirectory);
-				FileItem DummyFile = FileItem.GetItemByFileReference(FileReference.Combine(IntermediateDirectory, Filename));
-				Directory.CreateDirectory(DummyFile.Directory.FullName);
-				File.WriteAllText(DummyFile.FullName, $"#include \"{SourceFile.FullName.Replace('\\', '/')}\"");
-				SourceFile = DummyFile;
-			}
-			else if (!SourceFile.HasExtension(".cpp"))
+			if (!SourceFile.HasExtension(".cpp") && !SourceFile.HasExtension(".h"))
 			{
 				return null;
 			}
@@ -82,6 +72,15 @@ namespace UnrealBuildTool
 			DirectoryReference.CreateDirectory(IntermediateDirectory);
 			VCCompileAction Action = new VCCompileAction(BaseAction);
 			Action.SourceFile = SourceFile;
+			if (SourceFile.HasExtension(".h"))
+			{
+				Action.Definitions.Add("SUPPRESS_MONOLITHIC_HEADER_WARNINGS=0");
+				if (BaseAction.CompilerType.IsClang())
+				{
+					ClangWarnings.GetHeaderDisabledWarnings(Action.Arguments);
+				}
+			}
+			
 			if (bPreprocessOnly)
 			{
 				Action.PreprocessedFile = FileItem.GetItemByFileReference(FileReference.Combine(IntermediateDirectory, $"{Filename}.i"));
diff --git a/Engine/Source/Programs/UnrealBuildTool/Platform/Windows/VCToolChain.cs b/Engine/Source/Programs/UnrealBuildTool/Platform/Windows/VCToolChain.cs
--- a/Engine/Source/Programs/UnrealBuildTool/Platform/Windows/VCToolChain.cs
+++ b/Engine/Source/Programs/UnrealBuildTool/Platform/Windows/VCToolChain.cs
@@ -25,11 +25,6 @@ namespace UnrealBuildTool
 		/// </summary>
 		protected VCEnvironment EnvVars;
 
-		/// <summary>
-		/// Length of a path string that will trigger a build warning, as long paths may cause unexpected errors with the MSVC toolchain.
-		/// </summary>
-		private static int MaxPathWarningLength = 260;
-
 		public VCToolChain(ReadOnlyTargetRules Target, ILogger Logger)
 			: base(Logger)
 		{
@@ -162,33 +157,14 @@ namespace UnrealBuildTool
 			}
 		}
 
-		private static void CheckCommandLinePathLength(string PathString)
-		{
-			if (!Path.IsPathRooted(PathString))
-			{
-				string ResolvedPath = Path.Combine(Unreal.EngineSourceDirectory.FullName, PathString);
-				if (ResolvedPath.Length > MaxPathWarningLength)
-				{
-					Log.TraceWarningOnce($"Relative path '{PathString}' when resolved will have length '{ResolvedPath.Length}' which is greater than MAX_PATH (260) and may cause unexpected errors with the MSVC toolchain.");
-				}
-			}
-			else if (PathString.Length > MaxPathWarningLength)
-			{
-				Log.TraceWarningOnce($"Absolute path '{PathString}' has length '{PathString.Length}' which is greater than MAX_PATH (260) and may cause unexpected errors with the MSVC toolchain.");
-			}
-		}
-
 		public static new string NormalizeCommandLinePath(FileSystemReference Reference)
 		{
 			// Try to use a relative path to shorten command line length and to enable remote distribution where absolute paths are not desired
 			if (Reference.IsUnderDirectory(Unreal.EngineDirectory))
 			{
-				string RelativePath = Reference.MakeRelativeTo(Unreal.EngineSourceDirectory);
-				CheckCommandLinePathLength(RelativePath);
-				return RelativePath;
+				return Reference.MakeRelativeTo(Unreal.EngineSourceDirectory);
 			}
 
-			CheckCommandLinePathLength(Reference.FullName);
 			return Reference.FullName;
 		}
 
@@ -359,13 +335,6 @@ namespace UnrealBuildTool
 					// Enable the static analyzer with default checks.
 					Arguments.Add("--analyze");
 
-					// Deprecated in LLVM 15
-					if (EnvVars.CompilerVersion <= new VersionNumber(14))
-					{
-						// Make sure we check inside nested blocks (e.g. 'if ((foo = getchar()) == 0) {}')
-						Arguments.Add("-Xclang -analyzer-opt-analyze-nested-blocks");
-					}
-
 					// Write out a pretty web page with navigation to understand how the analysis was derived if HTML is enabled.
 					Arguments.Add($"-Xclang -analyzer-output={Target.StaticAnalyzerOutputType.ToString().ToLowerInvariant()}");
 
@@ -1045,14 +1014,22 @@ namespace UnrealBuildTool
 			}
 		}
 
-		protected virtual void AppendCLArguments_CPP(CppCompileEnvironment CompileEnvironment, List<string> Arguments)
+		protected virtual void AppendCLArguments_H(CppCompileEnvironment CompileEnvironment, List<string> Arguments)
 		{
-			if (Target.WindowsPlatform.Compiler.IsMSVC())
+			AppendCLArguments_CPP(CompileEnvironment, Arguments);
+
+			if (Target.WindowsPlatform.Compiler.IsClang())
 			{
-				// Explicitly compile the file as C++.
-				Arguments.Add("/TP");
+				ClangWarnings.GetHeaderDisabledWarnings(Arguments);
 			}
-			else
+		}
+
+		protected virtual void AppendCLArguments_CPP(CppCompileEnvironment CompileEnvironment, List<string> Arguments)
+		{
+			// Explicitly compile the file as C++.
+			Arguments.Add("/TP");
+
+			if (Target.WindowsPlatform.Compiler.IsClang())
 			{
 				string FileSpecifier = "c++";
 				if (CompileEnvironment.PrecompiledHeaderAction == PrecompiledHeaderAction.Create)
@@ -1693,6 +1670,7 @@ namespace UnrealBuildTool
 				CompileAction.SourceFile = SourceFile;
 
 				bool bIsPlainCFile = Path.GetExtension(SourceFile.AbsolutePath).ToUpperInvariant() == ".C";
+				bool bIsHeaderFile = Path.GetExtension(SourceFile.AbsolutePath).ToUpperInvariant() == ".H";
 
 				if (CompileEnvironment.PrecompiledHeaderAction == PrecompiledHeaderAction.Create)
 				{
@@ -1829,6 +1807,10 @@ namespace UnrealBuildTool
 				{
 					AppendCLArguments_C(CompileEnvironment, CompileAction.Arguments);
 				}
+				else if (bIsHeaderFile)
+				{
+					AppendCLArguments_H(CompileEnvironment, CompileAction.Arguments);
+				}
 				else
 				{
 					AppendCLArguments_CPP(CompileEnvironment, CompileAction.Arguments);
diff --git a/Engine/Source/Programs/UnrealBuildTool/System/ActionGraph.cs b/Engine/Source/Programs/UnrealBuildTool/System/ActionGraph.cs
--- a/Engine/Source/Programs/UnrealBuildTool/System/ActionGraph.cs
+++ b/Engine/Source/Programs/UnrealBuildTool/System/ActionGraph.cs
@@ -289,22 +289,33 @@ namespace UnrealBuildTool
 			{
 				const int MAX_PATH = 260;
 
-				List<FileReference> FailPaths = new List<FileReference>();
-				List<FileReference> WarnPaths = new List<FileReference>();
+				bool ShouldFail(FileItem item)
+				{
+					return item.Location.FullName.Length >= MAX_PATH &&
+						!(item.Location.ContainsName("Intermediate", 0) && item.Location.ContainsName("H", 0)); // Ignore -IncludeHeader items
+				}
+
+				bool ShouldWarn(FileItem item)
+				{
+					return item.Location.FullName.Length > Unreal.RootDirectory.FullName.Length + BuildConfiguration.MaxNestedPathLength &&
+						item.Location.IsUnderDirectory(Unreal.RootDirectory) &&
+						!(item.Location.ContainsName("Intermediate", 0) && item.Location.ContainsName("H", 0)) && // Ignore -IncludeHeader items
+						(item.Location.ContainsName("Restricted", 0) == false) && //Be more relaxed for internal only code
+						item.Location.ContainsName("NotForLicensees", 0) == false;
+				}
+
+				HashSet<FileReference> FailPaths = new();
+				HashSet<FileReference> WarnPaths = new();
 				foreach (IExternalAction Action in Actions)
 				{
 					foreach (FileItem PrerequisiteItem in Action.PrerequisiteItems)
 					{
-						if (PrerequisiteItem.Location.FullName.Length >= MAX_PATH)
+						if (ShouldFail(PrerequisiteItem))
 						{
 							FailPaths.Add(PrerequisiteItem.Location);
 						}
 
-						if (PrerequisiteItem.Location.FullName.Length > Unreal.RootDirectory.FullName.Length + BuildConfiguration.MaxNestedPathLength &&
-							PrerequisiteItem.Location.IsUnderDirectory(Unreal.RootDirectory) &&
-							(PrerequisiteItem.Location.ContainsName("Restricted", 0) == false) && //Be more relaxed for internal only code
-							(PrerequisiteItem.Location.ContainsName("NotForLicensees", 0) == false)
-							)
+						if (ShouldWarn(PrerequisiteItem))
 						{
 							WarnPaths.Add(PrerequisiteItem.Location);
 						}
@@ -312,16 +323,12 @@ namespace UnrealBuildTool
 
 					foreach (FileItem ProducedItem in Action.ProducedItems)
 					{
-						if (ProducedItem.Location.FullName.Length >= MAX_PATH)
+						if (ShouldFail(ProducedItem))
 						{
 							FailPaths.Add(ProducedItem.Location);
 						}
 
-						if (ProducedItem.Location.FullName.Length > Unreal.RootDirectory.FullName.Length + BuildConfiguration.MaxNestedPathLength &&
-							ProducedItem.Location.IsUnderDirectory(Unreal.RootDirectory) &&
-							(ProducedItem.Location.ContainsName("Restricted", 0) == false) && //Be more relaxed for internal only code
-							(ProducedItem.Location.ContainsName("NotForLicensees", 0) == false)
-							)
+						if (ShouldWarn(ProducedItem))
 						{
 							WarnPaths.Add(ProducedItem.Location);
 						}
@@ -331,8 +338,8 @@ namespace UnrealBuildTool
 				if (FailPaths.Count > 0)
 				{
 					StringBuilder Message = new StringBuilder();
-					Message.Append($"The following output paths are longer than {MAX_PATH} characters. Please move the engine to a directory with a shorter path.");
-					foreach (FileReference Path in FailPaths)
+					Message.Append($"The following action paths are longer than {MAX_PATH} characters. Please move the engine to a directory with a shorter path.");
+					foreach (FileReference Path in FailPaths.OrderBy(x => x.FullName))
 					{
 						Message.Append($"\n[{Path.FullName.Length.ToString()} characters] {Path}");
 					}
@@ -343,7 +350,7 @@ namespace UnrealBuildTool
 				{
 					StringBuilder Message = new StringBuilder();
 					Message.Append($"Detected paths more than {BuildConfiguration.MaxNestedPathLength.ToString()} characters below UE root directory. This may cause portability issues due to the {MAX_PATH.ToString()} character maximum path length on Windows:\n");
-					foreach (FileReference Path in WarnPaths)
+					foreach (FileReference Path in WarnPaths.OrderBy(x => x.FullName))
 					{
 						string RelativePath = Path.MakeRelativeTo(Unreal.RootDirectory);
 						Message.Append($"\n[{RelativePath.Length.ToString()} characters] {RelativePath}");
diff --git a/Engine/Source/Programs/UnrealBuildTool/ToolChain/ClangSpecificFileAction.cs b/Engine/Source/Programs/UnrealBuildTool/ToolChain/ClangSpecificFileAction.cs
--- a/Engine/Source/Programs/UnrealBuildTool/ToolChain/ClangSpecificFileAction.cs
+++ b/Engine/Source/Programs/UnrealBuildTool/ToolChain/ClangSpecificFileAction.cs
@@ -49,20 +49,12 @@ namespace UnrealBuildTool
 		public IExternalAction? CreateAction(FileItem SourceFile, ILogger Logger)
 		{
 			string DummyName = "SingleFile.cpp";
-			string UniqueDummyName = $"SingleFile{SingleFileCounter}.cpp";
+			string UniqueDummyName = $"SingleFile{SingleFileCounter}{SourceFile.Location.GetExtension()}";
 			++SingleFileCounter;
 
 			int FileNameIndex = CommandArguments.IndexOf(DummyName);
 			string DummyPath = CommandArguments.Substring(2, FileNameIndex + DummyName.Length - 2);
 
-			if (SourceFile.HasExtension(".h"))
-			{
-				FileItem DummyFile = FileItem.GetItemByFileReference(FileReference.Combine(OutputDir, "SingleFile", SourceFile.Name));
-				Directory.CreateDirectory(DummyFile.Directory.FullName);
-				File.WriteAllText(DummyFile.FullName, $"#include \"{SourceFile.FullName.Replace('\\', '/')}\"");
-				SourceFile = DummyFile;
-			}
-
 			List<string> NewRspLines = new();
 			foreach (string L in RspLines)
 			{
@@ -82,6 +74,12 @@ namespace UnrealBuildTool
 				NewRspLines.Add(Line);
 			}
 
+			if (SourceFile.HasExtension(".h"))
+			{
+				NewRspLines.Add("-DSUPPRESS_MONOLITHIC_HEADER_WARNINGS=0");
+				ClangWarnings.GetHeaderDisabledWarnings(NewRspLines);
+			}
+
 			Action Action = new Action(this);
 			Action.CommandArguments = CommandArguments.Replace(DummyName, UniqueDummyName);
 			Action.DependencyListFile = null;
diff --git a/Engine/Source/Programs/UnrealBuildTool/ToolChain/ClangToolChain.cs b/Engine/Source/Programs/UnrealBuildTool/ToolChain/ClangToolChain.cs
--- a/Engine/Source/Programs/UnrealBuildTool/ToolChain/ClangToolChain.cs
+++ b/Engine/Source/Programs/UnrealBuildTool/ToolChain/ClangToolChain.cs
@@ -426,6 +426,12 @@ namespace UnrealBuildTool
 			}
 		}
 
+		protected virtual void GetCompileArguments_H(CppCompileEnvironment CompileEnvironment, List<string> Arguments)
+		{
+			GetCompileArguments_CPP(CompileEnvironment, Arguments);
+			ClangWarnings.GetHeaderDisabledWarnings(Arguments);
+		}
+
 		protected virtual void GetCompileArguments_CPP(CppCompileEnvironment CompileEnvironment, List<string> Arguments)
 		{
 			Arguments.Add("-x c++");
@@ -910,12 +916,8 @@ namespace UnrealBuildTool
 			}
 			else if (Extension == ".H")
 			{
-				// Compile the file as C++ code.
-				GetCompileArguments_CPP(CompileEnvironment, Arguments);
-
-				// This warning was to catch #pragma once inside a source file.
-				// If we're compiling a header directly, we should always have the pragma once, so we need to ignore this warning.
-				Arguments.Add("-Wno-pragma-once-outside-header");
+				// Compile the file as C++ code with some additional arguments
+				GetCompileArguments_H(CompileEnvironment, Arguments);
 			}
 			else
 			{
diff --git a/Engine/Source/Programs/UnrealBuildTool/ToolChain/ClangWarnings.cs b/Engine/Source/Programs/UnrealBuildTool/ToolChain/ClangWarnings.cs
--- a/Engine/Source/Programs/UnrealBuildTool/ToolChain/ClangWarnings.cs
+++ b/Engine/Source/Programs/UnrealBuildTool/ToolChain/ClangWarnings.cs
@@ -160,6 +160,14 @@ namespace UnrealBuildTool
 			}
 		}
 
+		internal static void GetHeaderDisabledWarnings(List<string> Arguments)
+		{
+			// This warning was to catch #pragma once inside a source file.
+			// If we're compiling a header directly, we should always have the pragma once, so we need to ignore this warning.
+			Arguments.Add("-Wno-pragma-once-outside-header");
+			Arguments.Add("-Wno-#pragma-messages");
+		}
+
 		// Additional disabled warnings for msvc. Everything below should be checked if it is necessary
 		internal static void GetVCDisabledWarnings(List<string> Arguments)
 		{
diff --git a/Engine/Source/Programs/UnrealBuildTool/ToolChain/IWYUToolChain.cs b/Engine/Source/Programs/UnrealBuildTool/ToolChain/IWYUToolChain.cs
--- a/Engine/Source/Programs/UnrealBuildTool/ToolChain/IWYUToolChain.cs
+++ b/Engine/Source/Programs/UnrealBuildTool/ToolChain/IWYUToolChain.cs
@@ -195,7 +195,7 @@ namespace UnrealBuildTool
 			Arguments.Add("-Wno-undefined-bool-conversion");
 			Arguments.Add("-Wno-deprecated-anon-enum-enum-conversion");
 			Arguments.Add("-Wno-ambiguous-reversed-operator");
-			Arguments.Add("-Wno-pragma-once-outside-header");
+			ClangWarnings.GetHeaderDisabledWarnings(Arguments);
 		}
 
 		// Skip ISPC headers
