UnrealBuildTool: Shorten -IncludeHeader path

#rnx

diff --git a/Engine/Source/Programs/UnrealBuildTool/Configuration/UEBuildModuleCPP.cs b/Engine/Source/Programs/UnrealBuildTool/Configuration/UEBuildModuleCPP.cs
--- a/Engine/Source/Programs/UnrealBuildTool/Configuration/UEBuildModuleCPP.cs
+++ b/Engine/Source/Programs/UnrealBuildTool/Configuration/UEBuildModuleCPP.cs
@@ -823,10 +823,11 @@ namespace UnrealBuildTool
 					HeaderCompileEnvironment.Definitions.Add("SUPPRESS_MONOLITHIC_HEADER_WARNINGS=0");
 					CreateHeaderForDefinitions(HeaderCompileEnvironment, HeaderIntermediateDirectory, "h", Graph);
 
-					// Duplicate named headers are allowed, so adjust the IntermediateDirectory to compensate for this
-					foreach (DirectoryItem ParentDir in HeaderFileItems.Select(x => x.Directory!).Distinct())
+					// Duplicate named headers are allowed, so adjust the IntermediateDirectory to compensate for this by using an index
+					int DirectoryIndex = 0;
+					foreach (DirectoryItem ParentDir in HeaderFileItems.Select(x => x.Directory!).Distinct().OrderBy(x => x.FullName))
 					{
-						DirectoryReference HeaderSubIntermediateDirectory = DirectoryReference.Combine(HeaderIntermediateDirectory, ParentDir.Location.MakeRelativeTo(ModuleDirectory));
+						DirectoryReference HeaderSubIntermediateDirectory = DirectoryReference.Combine(HeaderIntermediateDirectory, (DirectoryIndex++).ToString());
 						// Add the compile actions
 						LinkInputFiles.AddRange(ToolChain.CompileAllCPPFiles(HeaderCompileEnvironment, HeaderFileItems.Where(x => x.Directory == ParentDir), HeaderSubIntermediateDirectory, Name, Graph).ObjectFiles);
 					}
