UnrealBuildTool: Always write a .cpp file when compiling stand-alone headers

#rnx

diff --git a/Engine/Source/Programs/UnrealBuildTool/Configuration/UEBuildModuleCPP.cs b/Engine/Source/Programs/UnrealBuildTool/Configuration/UEBuildModuleCPP.cs
--- a/Engine/Source/Programs/UnrealBuildTool/Configuration/UEBuildModuleCPP.cs
+++ b/Engine/Source/Programs/UnrealBuildTool/Configuration/UEBuildModuleCPP.cs
@@ -828,8 +828,26 @@ namespace UnrealBuildTool
 					foreach (DirectoryItem ParentDir in HeaderFileItems.Select(x => x.Directory!).Distinct().OrderBy(x => x.FullName))
 					{
 						DirectoryReference HeaderSubIntermediateDirectory = DirectoryReference.Combine(HeaderIntermediateDirectory, (DirectoryIndex++).ToString());
+						List<FileItem> HeadersToCompile = new();
+						foreach (FileItem Header in HeaderFileItems.Where(x => x.Directory == ParentDir))
+						{
+							string IncludeFileString = Header.AbsolutePath;
+							if (Header.Location.IsUnderDirectory(Unreal.RootDirectory))
+							{
+								IncludeFileString = Header.Location.MakeRelativeTo(Unreal.EngineSourceDirectory);
+							}
+
+							List<string> GeneratedHeaderCppContents = new();
+							GeneratedHeaderCppContents.Add("// This file is automatically generated at compile-time to include a user created header file.");
+							GeneratedHeaderCppContents.Add($"#include \"{IncludeFileString.Replace('\\', '/')}\"");
+
+							FileItem GeneratedHeaderCpp = FileItem.GetItemByFileReference(FileReference.Combine(HeaderSubIntermediateDirectory, $"{Header.Name}.cpp"));
+							Graph.CreateIntermediateTextFile(GeneratedHeaderCpp, GeneratedHeaderCppContents);
+							HeadersToCompile.Add(GeneratedHeaderCpp);
+						}
+
 						// Add the compile actions
-						LinkInputFiles.AddRange(ToolChain.CompileAllCPPFiles(HeaderCompileEnvironment, HeaderFileItems.Where(x => x.Directory == ParentDir), HeaderSubIntermediateDirectory, Name, Graph).ObjectFiles);
+						LinkInputFiles.AddRange(ToolChain.CompileAllCPPFiles(HeaderCompileEnvironment, HeadersToCompile, HeaderSubIntermediateDirectory, Name, Graph).ObjectFiles);
 					}
 				}
 			}
diff --git a/Engine/Source/Programs/UnrealBuildTool/Platform/Windows/VCSpecificFileAction.cs b/Engine/Source/Programs/UnrealBuildTool/Platform/Windows/VCSpecificFileAction.cs
--- a/Engine/Source/Programs/UnrealBuildTool/Platform/Windows/VCSpecificFileAction.cs
+++ b/Engine/Source/Programs/UnrealBuildTool/Platform/Windows/VCSpecificFileAction.cs
@@ -2,6 +2,7 @@
 
 using System.Collections.Generic;
 using System.IO;
+using System.Reflection.PortableExecutable;
 using EpicGames.Core;
 using Microsoft.Extensions.Logging;
 using UnrealBuildBase;
@@ -79,6 +80,19 @@ namespace UnrealBuildTool
 				{
 					ClangWarnings.GetHeaderDisabledWarnings(Action.Arguments);
 				}
+
+				string IncludeFileString = SourceFile.AbsolutePath;
+				if (SourceFile.Location.IsUnderDirectory(Unreal.RootDirectory))
+				{
+					IncludeFileString = SourceFile.Location.MakeRelativeTo(Unreal.EngineSourceDirectory);
+				}
+
+				List<string> GeneratedHeaderCppContents = new();
+				GeneratedHeaderCppContents.Add("// This file is automatically generated at compile-time to include a user created header file.");
+				GeneratedHeaderCppContents.Add($"#include \"{IncludeFileString.Replace('\\', '/')}\"");
+
+				Action.SourceFile = FileItem.GetItemByFileReference(FileReference.Combine(IntermediateDirectory, $"{SourceFile.Name}.cpp"));
+				File.WriteAllLines(Action.SourceFile.FullName, GeneratedHeaderCppContents);
 			}
 			
 			if (bPreprocessOnly)
diff --git a/Engine/Source/Programs/UnrealBuildTool/ToolChain/ClangSpecificFileAction.cs b/Engine/Source/Programs/UnrealBuildTool/ToolChain/ClangSpecificFileAction.cs
--- a/Engine/Source/Programs/UnrealBuildTool/ToolChain/ClangSpecificFileAction.cs
+++ b/Engine/Source/Programs/UnrealBuildTool/ToolChain/ClangSpecificFileAction.cs
@@ -17,7 +17,8 @@ namespace UnrealBuildTool
 		DirectoryReference SourceDir;
 		DirectoryReference OutputDir;
 		IEnumerable<string> RspLines;
-		int SingleFileCounter;
+
+		Dictionary<string, List<FileItem>> SingleFiles = new();
 
 		internal ClangSpecificFileAction(DirectoryReference Source, DirectoryReference Output, Action Action, IEnumerable<string> ContentLines) : base(Action)
 		{
@@ -48,14 +49,47 @@ namespace UnrealBuildTool
 
 		public IExternalAction? CreateAction(FileItem SourceFile, ILogger Logger)
 		{
+			DirectoryReference.CreateDirectory(DirectoryReference.Combine(OutputDir, "SingleFile"));
+
+			// Keep track of all specific files, so the output file can be renamed if there's a naming conflict
+			string Filename = $"SingleFile/{SourceFile.Name}";
+			if (!SingleFiles.ContainsKey(SourceFile.Name))
+			{
+				SingleFiles[SourceFile.Name] = new();
+			}
+			else
+			{
+				Filename = $"SingleFile/{Path.GetFileNameWithoutExtension(SourceFile.Name)}{SingleFiles[Filename].Count}{Path.GetExtension(SourceFile.Name)}";
+			}
+			SingleFiles[SourceFile.Name].Add(SourceFile);
+
 			string DummyName = "SingleFile.cpp";
-			string UniqueDummyName = $"SingleFile{SingleFileCounter}{SourceFile.Location.GetExtension()}";
-			++SingleFileCounter;
+			string UniqueDummyName = Filename;
 
 			int FileNameIndex = CommandArguments.IndexOf(DummyName);
 			string DummyPath = CommandArguments.Substring(2, FileNameIndex + DummyName.Length - 2);
 
 			List<string> NewRspLines = new();
+
+			if (SourceFile.HasExtension(".h"))
+			{
+				NewRspLines.Add("-DSUPPRESS_MONOLITHIC_HEADER_WARNINGS=0");
+				ClangWarnings.GetHeaderDisabledWarnings(NewRspLines);
+
+				string IncludeFileString = SourceFile.AbsolutePath;
+				if (SourceFile.Location.IsUnderDirectory(Unreal.RootDirectory))
+				{
+					IncludeFileString = SourceFile.Location.MakeRelativeTo(Unreal.EngineSourceDirectory);
+				}
+
+				List<string> GeneratedHeaderCppContents = new();
+				GeneratedHeaderCppContents.Add("// This file is automatically generated at compile-time to include a user created header file.");
+				GeneratedHeaderCppContents.Add($"#include \"{IncludeFileString.Replace('\\', '/')}\"");
+
+				SourceFile = FileItem.GetItemByFileReference(FileReference.Combine(OutputDir, $"{UniqueDummyName}.cpp"));
+				File.WriteAllLines(SourceFile.FullName, GeneratedHeaderCppContents);
+			}
+
 			foreach (string L in RspLines)
 			{
 				string Line = L;
@@ -74,12 +108,6 @@ namespace UnrealBuildTool
 				NewRspLines.Add(Line);
 			}
 
-			if (SourceFile.HasExtension(".h"))
-			{
-				NewRspLines.Add("-DSUPPRESS_MONOLITHIC_HEADER_WARNINGS=0");
-				ClangWarnings.GetHeaderDisabledWarnings(NewRspLines);
-			}
-
 			Action Action = new Action(this);
 			Action.CommandArguments = CommandArguments.Replace(DummyName, UniqueDummyName);
 			Action.DependencyListFile = null;
