Fix Windows ASan issues when using the AutoRTFM compiler.

diff --git a/Engine/Source/Programs/UnrealBuildTool/Platform/Windows/VCToolChain.cs b/Engine/Source/Programs/UnrealBuildTool/Platform/Windows/VCToolChain.cs
--- a/Engine/Source/Programs/UnrealBuildTool/Platform/Windows/VCToolChain.cs
+++ b/Engine/Source/Programs/UnrealBuildTool/Platform/Windows/VCToolChain.cs
@@ -611,13 +611,13 @@ namespace UnrealBuildTool
 				{
 					// MSVC has no support for __has_feature(address_sanitizer)
 					AddDefinition(Arguments, "USING_ADDRESS_SANITISER=1");
-
-					// Disabling a couple annotations to workaround an issue related to building third party libraries with different options than the main binary.
-					// This fixes an error that looks like this: error LNK2038: mismatch detected for 'annotate_string': value '0' doesn't match value '1' in Module.Core.XX_of_YY.cpp.obj
-					AddDefinition(Arguments, "_DISABLE_STRING_ANNOTATION=1");
-					AddDefinition(Arguments, "_DISABLE_VECTOR_ANNOTATION=1");
 				}
 
+				// Disabling a couple annotations to workaround an issue related to building third party libraries with different options than the main binary.
+				// This fixes an error that looks like this: error LNK2038: mismatch detected for 'annotate_string': value '0' doesn't match value '1' in Module.Core.XX_of_YY.cpp.obj
+				AddDefinition(Arguments, "_DISABLE_STRING_ANNOTATION=1");
+				AddDefinition(Arguments, "_DISABLE_VECTOR_ANNOTATION=1");
+
 				// Currently the ASan headers are not default around. They can be found at this location so lets use this until this is resolved in the toolchain
 				// Jira with some more info and the MSVC bug at UE-144727
 				AddSystemIncludePath(Arguments, DirectoryReference.Combine(EnvVars.CompilerDir, "crt", "src"), Target.WindowsPlatform.Compiler);
@@ -1439,6 +1439,39 @@ namespace UnrealBuildTool
 			// [REDACTED_PROFILE_URL]/STL/issues/2655
 			Arguments.Add("/ALTERNATENAME:__imp___std_init_once_begin_initialize=__imp_InitOnceBeginInitialize");
 			Arguments.Add("/ALTERNATENAME:__imp___std_init_once_complete=__imp_InitOnceComplete");
+
+			if (Target.WindowsPlatform.bEnableAddressSanitizer)
+			{
+				// With clang we seemingly need to explicitly pass the .lib's to link against for ASan.
+				if (Target.WindowsPlatform.Compiler.IsClang())
+				{
+					DirectoryReference ASanRuntimeDir;
+					string ASanArchSuffix;
+					if (EnvVars.Architecture == UnrealArch.X64)
+					{
+						ASanRuntimeDir = DirectoryReference.Combine(EnvVars.ToolChainDir, "lib", "x64");
+						ASanArchSuffix = "x86_64";
+					}
+					else
+					{
+						throw new BuildException("Unsupported build architecture for Address Sanitizer");
+					}
+
+					string ASanRuntimeLib = $"clang_rt.asan_dynamic-{ASanArchSuffix}.lib";
+					string ASanDebugRuntimeLib = $"clang_rt.asan_dbg_dynamic-{ASanArchSuffix}.lib";
+
+					if (Target.bDebugBuildsActuallyUseDebugCRT)
+					{
+						LinkEnvironment.Libraries.Add(FileReference.Combine(ASanRuntimeDir, ASanDebugRuntimeLib));
+					}
+					else
+					{
+						LinkEnvironment.Libraries.Add(FileReference.Combine(ASanRuntimeDir, ASanRuntimeLib));
+					}
+
+					LinkEnvironment.Libraries.Add(FileReference.Combine(ASanRuntimeDir, $"clang_rt.asan_dynamic_runtime_thunk-{ASanArchSuffix}.lib"));
+				}
+			}
 		}
 
 		protected virtual void AppendLibArguments(LinkEnvironment LinkEnvironment, List<string> Arguments)
