Allow ShowIncludes to work for clang as well as windows

diff --git a/Engine/Source/Programs/UnrealBuildTool/Configuration/ReadOnlyTargetRules.cs b/Engine/Source/Programs/UnrealBuildTool/Configuration/ReadOnlyTargetRules.cs
--- a/Engine/Source/Programs/UnrealBuildTool/Configuration/ReadOnlyTargetRules.cs
+++ b/Engine/Source/Programs/UnrealBuildTool/Configuration/ReadOnlyTargetRules.cs
@@ -360,6 +360,8 @@ namespace UnrealBuildTool
 
 		public WarningLevel ModuleIncludeSubdirectoryWarningLevel => Inner.ModuleIncludeSubdirectoryWarningLevel;
 
+		public bool bShowIncludes => Inner.bShowIncludes;
+
 		public bool bWarningsAsErrors => Inner.bWarningsAsErrors;
 
 		public bool bUseFastMonoCalls => Inner.bUseFastMonoCalls;
diff --git a/Engine/Source/Programs/UnrealBuildTool/Configuration/TargetRules.cs b/Engine/Source/Programs/UnrealBuildTool/Configuration/TargetRules.cs
--- a/Engine/Source/Programs/UnrealBuildTool/Configuration/TargetRules.cs
+++ b/Engine/Source/Programs/UnrealBuildTool/Configuration/TargetRules.cs
@@ -1790,6 +1790,13 @@ namespace UnrealBuildTool
 		[CommandLine("-ModuleIncludeSubdirectoryWarningLevel=")]
 		public WarningLevel ModuleIncludeSubdirectoryWarningLevel { get; set; } = WarningLevel.Off;
 
+		/// <summary>
+		/// Print out files that are included by each source file
+		/// </summary>
+		[XmlConfigFile(Category = "BuildConfiguration")]
+		[CommandLine("-ShowIncludes")]
+		public bool bShowIncludes = false;
+
 		/// <summary>
 		/// Forces frame pointers to be retained this is usually required when you want reliable callstacks e.g. mallocframeprofiler
 		/// </summary>
diff --git a/Engine/Source/Programs/UnrealBuildTool/Platform/Windows/UEBuildWindows.cs b/Engine/Source/Programs/UnrealBuildTool/Platform/Windows/UEBuildWindows.cs
--- a/Engine/Source/Programs/UnrealBuildTool/Platform/Windows/UEBuildWindows.cs
+++ b/Engine/Source/Programs/UnrealBuildTool/Platform/Windows/UEBuildWindows.cs
@@ -661,8 +661,7 @@ namespace UnrealBuildTool
 		/// <summary>
 		/// Print out files that are included by each source file
 		/// </summary>
-		[CommandLine("-ShowIncludes")]
-		[XmlConfigFile(Category = "WindowsPlatform")]
+		[Obsolete("Deprecated in UE5.5 - Use bShowIncludes on TargetRules instead.")]
 		public bool bShowIncludes = false;
 
 		/// <summary>
@@ -931,6 +930,7 @@ namespace UnrealBuildTool
 
 		public bool bCompilerTrace => Inner.bCompilerTrace;
 
+		[Obsolete("Deprecated in UE5.5 - Use bShowIncludes on TargetRules instead.")]
 		public bool bShowIncludes => Inner.bShowIncludes;
 
 		public string GetVisualStudioCompilerVersionName()
diff --git a/Engine/Source/Programs/UnrealBuildTool/Platform/Windows/VCToolChain.cs b/Engine/Source/Programs/UnrealBuildTool/Platform/Windows/VCToolChain.cs
--- a/Engine/Source/Programs/UnrealBuildTool/Platform/Windows/VCToolChain.cs
+++ b/Engine/Source/Programs/UnrealBuildTool/Platform/Windows/VCToolChain.cs
@@ -421,7 +421,7 @@ namespace UnrealBuildTool
 			Arguments.Add("/Oi");
 
 			// Trace includes
-			if (Target.WindowsPlatform.bShowIncludes)
+			if (Target.bShowIncludes)
 			{
 				if (Target.WindowsPlatform.Compiler.IsClang())
 				{
@@ -1923,7 +1923,7 @@ namespace UnrealBuildTool
 				else
 				{
 					CompileAction.DependencyListFile = FileItem.GetItemByFileReference(FileReference.Combine(OutputDir, $"{FileName}.txt"));
-					CompileAction.bShowIncludes = Target.WindowsPlatform.bShowIncludes;
+					CompileAction.bShowIncludes = Target.bShowIncludes;
 				}
 
 				// Write cl errors and warnings to a file if supported
diff --git a/Engine/Source/Programs/UnrealBuildTool/ToolChain/ClangToolChain.cs b/Engine/Source/Programs/UnrealBuildTool/ToolChain/ClangToolChain.cs
--- a/Engine/Source/Programs/UnrealBuildTool/ToolChain/ClangToolChain.cs
+++ b/Engine/Source/Programs/UnrealBuildTool/ToolChain/ClangToolChain.cs
@@ -200,6 +200,7 @@ namespace UnrealBuildTool
 
 		// Target settings
 		protected bool PreprocessDepends = false;
+		protected bool ShowIncludes = false;
 		protected StaticAnalyzer StaticAnalyzer = StaticAnalyzer.None;
 		protected StaticAnalyzerMode StaticAnalyzerMode = StaticAnalyzerMode.Deep;
 		protected StaticAnalyzerOutputType StaticAnalyzerOutputType = StaticAnalyzerOutputType.Text;
@@ -232,6 +233,7 @@ namespace UnrealBuildTool
 			base.SetUpGlobalEnvironment(Target);
 
 			PreprocessDepends = Target.bPreprocessDepends;
+			ShowIncludes = Target.bShowIncludes;
 			StaticAnalyzer = Target.StaticAnalyzer;
 			StaticAnalyzerMode = Target.StaticAnalyzerMode;
 			StaticAnalyzerOutputType = Target.StaticAnalyzerOutputType;
@@ -488,6 +490,11 @@ namespace UnrealBuildTool
 		{
 			Arguments.AddRange(CompileEnvironment.UserIncludePaths.Select(IncludePath => GetUserIncludePathArgument(IncludePath)));
 			Arguments.AddRange(CompileEnvironment.SystemIncludePaths.Select(IncludePath => GetSystemIncludePathArgument(IncludePath)));
+
+			if (ShowIncludes)
+			{
+				Arguments.Add("-H");
+			}
 		}
 
 		protected virtual string GetPreprocessorDefinitionArgument(string Definition)
