Minor clean up of SUPPRESS_MONOLITHIC_HEADER_WARNINGS
* Remove pointless define/undef of it
* Set it to 1 consistently in all cases (even though we only check if it is defined)

diff --git a/Engine/Source/Programs/UnrealBuildTool/Platform/Windows/VCSpecificFileAction.cs b/Engine/Source/Programs/UnrealBuildTool/Platform/Windows/VCSpecificFileAction.cs
--- a/Engine/Source/Programs/UnrealBuildTool/Platform/Windows/VCSpecificFileAction.cs
+++ b/Engine/Source/Programs/UnrealBuildTool/Platform/Windows/VCSpecificFileAction.cs
@@ -74,7 +74,7 @@ namespace UnrealBuildTool
 			Action.SourceFile = SourceFile;
 			if (SourceFile.HasExtension(".h"))
 			{
-				Action.Definitions.Add("SUPPRESS_MONOLITHIC_HEADER_WARNINGS=0");
+				Action.Definitions.Add("SUPPRESS_MONOLITHIC_HEADER_WARNINGS=1");
 				if (BaseAction.CompilerType.IsClang())
 				{
 					ClangWarnings.GetHeaderDisabledWarnings(Action.Arguments);
diff --git a/Engine/Source/Programs/UnrealBuildTool/ToolChain/ClangSpecificFileAction.cs b/Engine/Source/Programs/UnrealBuildTool/ToolChain/ClangSpecificFileAction.cs
--- a/Engine/Source/Programs/UnrealBuildTool/ToolChain/ClangSpecificFileAction.cs
+++ b/Engine/Source/Programs/UnrealBuildTool/ToolChain/ClangSpecificFileAction.cs
@@ -73,7 +73,7 @@ namespace UnrealBuildTool
 
 			if (SourceFile.HasExtension(".h"))
 			{
-				NewRspLines.Add("-DSUPPRESS_MONOLITHIC_HEADER_WARNINGS=0");
+				NewRspLines.Add("-DSUPPRESS_MONOLITHIC_HEADER_WARNINGS=1");
 				ClangWarnings.GetHeaderDisabledWarnings(NewRspLines);
 
 				string IncludeFileString = SourceFile.AbsolutePath;
