Add UE_DIRECT_HEADER_COMPILE to allow avoiding warnings/errors in single-header compilation.
Replace setting SUPPRESS_MONOLITHIC_HEADER_WARNINGS with UE_DIRECT_HEADER_COMPILE for header compilation situations
#rnx

diff --git a/Engine/Source/Programs/UnrealBuildTool/Configuration/UEBuildModuleCPP.cs b/Engine/Source/Programs/UnrealBuildTool/Configuration/UEBuildModuleCPP.cs
--- a/Engine/Source/Programs/UnrealBuildTool/Configuration/UEBuildModuleCPP.cs
+++ b/Engine/Source/Programs/UnrealBuildTool/Configuration/UEBuildModuleCPP.cs
@@ -819,8 +819,10 @@ namespace UnrealBuildTool
 				{
 					CppCompileEnvironment HeaderCompileEnvironment = new CppCompileEnvironment(ModuleCompileEnvironment);
 					DirectoryReference HeaderIntermediateDirectory = DirectoryReference.Combine(IntermediateDirectory, "H");
-					HeaderCompileEnvironment.Definitions.RemoveAll(x => x.StartsWith("SUPPRESS_MONOLITHIC_HEADER_WARNINGS"));
-					HeaderCompileEnvironment.Definitions.Add("SUPPRESS_MONOLITHIC_HEADER_WARNINGS=0");
+
+					HeaderCompileEnvironment.Definitions.RemoveAll(x => x.StartsWith("UE_DIRECT_HEADER_COMPILE"));
+					HeaderCompileEnvironment.Definitions.Add("UE_DIRECT_HEADER_COMPILE=1");
+
 					CreateHeaderForDefinitions(HeaderCompileEnvironment, HeaderIntermediateDirectory, "h", Graph);
 
 					// Duplicate named headers are allowed, so adjust the IntermediateDirectory to compensate for this by using an index
diff --git a/Engine/Source/Programs/UnrealBuildTool/Platform/Windows/VCSpecificFileAction.cs b/Engine/Source/Programs/UnrealBuildTool/Platform/Windows/VCSpecificFileAction.cs
--- a/Engine/Source/Programs/UnrealBuildTool/Platform/Windows/VCSpecificFileAction.cs
+++ b/Engine/Source/Programs/UnrealBuildTool/Platform/Windows/VCSpecificFileAction.cs
@@ -74,7 +74,7 @@ namespace UnrealBuildTool
 			Action.SourceFile = SourceFile;
 			if (SourceFile.HasExtension(".h"))
 			{
-				Action.Definitions.Add("SUPPRESS_MONOLITHIC_HEADER_WARNINGS=1");
+				Action.Definitions.Add("UE_DIRECT_HEADER_COMPILE=1");
 				if (BaseAction.CompilerType.IsClang())
 				{
 					ClangWarnings.GetHeaderDisabledWarnings(Action.Arguments);
diff --git a/Engine/Source/Programs/UnrealBuildTool/ToolChain/ClangSpecificFileAction.cs b/Engine/Source/Programs/UnrealBuildTool/ToolChain/ClangSpecificFileAction.cs
--- a/Engine/Source/Programs/UnrealBuildTool/ToolChain/ClangSpecificFileAction.cs
+++ b/Engine/Source/Programs/UnrealBuildTool/ToolChain/ClangSpecificFileAction.cs
@@ -73,7 +73,7 @@ namespace UnrealBuildTool
 
 			if (SourceFile.HasExtension(".h"))
 			{
-				NewRspLines.Add("-DSUPPRESS_MONOLITHIC_HEADER_WARNINGS=1");
+				NewRspLines.Add("-DUE_DIRECT_HEADER_COMPILE=1");
 				ClangWarnings.GetHeaderDisabledWarnings(NewRspLines);
 
 				string IncludeFileString = SourceFile.AbsolutePath;
diff --git a/Engine/Source/Programs/UnrealBuildTool/ToolChain/IWYUToolChain.cs b/Engine/Source/Programs/UnrealBuildTool/ToolChain/IWYUToolChain.cs
--- a/Engine/Source/Programs/UnrealBuildTool/ToolChain/IWYUToolChain.cs
+++ b/Engine/Source/Programs/UnrealBuildTool/ToolChain/IWYUToolChain.cs
@@ -159,7 +159,7 @@ namespace UnrealBuildTool
 
 		protected override void GetCompileArguments_Global(CppCompileEnvironment CompileEnvironment, List<string> Arguments)
 		{
-			Arguments.Add(GetPreprocessorDefinitionArgument("SUPPRESS_MONOLITHIC_HEADER_WARNINGS=1"));
+			Arguments.Add(GetPreprocessorDefinitionArgument("UE_DIRECT_HEADER_COMPILE=1"));
 			Arguments.Add(GetPreprocessorDefinitionArgument("PLATFORM_COMPILER_IWYU=1"));
 
 			base.GetCompileArguments_Global(CompileEnvironment, Arguments);
diff --git a/Engine/Source/Runtime/Core/Public/Misc/CoreMiscDefines.h b/Engine/Source/Runtime/Core/Public/Misc/CoreMiscDefines.h
--- a/Engine/Source/Runtime/Core/Public/Misc/CoreMiscDefines.h
+++ b/Engine/Source/Runtime/Core/Public/Misc/CoreMiscDefines.h
@@ -291,7 +291,10 @@ enum EPerElement {PerElement};
 /**
  * Macro which can be placed in a header to throw a deprecation warning when it is included.
  */
-#ifdef _MSC_VER
+#if defined(UE_DIRECT_HEADER_COMPILE)
+	// Don't warn about the deprecated header when we are directly compiling the header
+	#define UE_DEPRECATED_HEADER(Version, Message)
+#elif defined(_MSC_VER)
     #if UE_WARNINGS_AS_ERRORS
     	#define UE_DEPRECATED_HEADER(Version, Message) __pragma(message(__FILE__ "(" PREPROCESSOR_TO_STRING(__LINE__) "): error C4996: " Message " Please update your code to the new API before upgrading to the next release, otherwise your project will no longer compile."))
     #else
diff --git a/Engine/Source/Runtime/Core/Public/Misc/MonolithicHeaderBoilerplate.h b/Engine/Source/Runtime/Core/Public/Misc/MonolithicHeaderBoilerplate.h
--- a/Engine/Source/Runtime/Core/Public/Misc/MonolithicHeaderBoilerplate.h
+++ b/Engine/Source/Runtime/Core/Public/Misc/MonolithicHeaderBoilerplate.h
@@ -2,7 +2,7 @@
 
 #include "CoreTypes.h"
 
-#if UE_IS_ENGINE_MODULE && !defined(SUPPRESS_MONOLITHIC_HEADER_WARNINGS)
+#if UE_IS_ENGINE_MODULE && !defined(SUPPRESS_MONOLITHIC_HEADER_WARNINGS) && !defined(UE_DIRECT_HEADER_COMPILE)
 	#define MONOLITHIC_HEADER_BOILERPLATE() COMPILE_WARNING("Monolithic headers should not be used by this module. Please change it to explicitly include the headers it needs.")
 	#define SUPPRESS_MONOLITHIC_HEADER_WARNINGS 1
 #else
