UnrealBuildTool: Preprocess SharedPCH when generating project files to prevent writing to the same file multiple times

#rnx

diff --git a/Engine/Source/Programs/UnrealBuildTool/Configuration/UEBuildModuleCPP.cs b/Engine/Source/Programs/UnrealBuildTool/Configuration/UEBuildModuleCPP.cs
--- a/Engine/Source/Programs/UnrealBuildTool/Configuration/UEBuildModuleCPP.cs
+++ b/Engine/Source/Programs/UnrealBuildTool/Configuration/UEBuildModuleCPP.cs
@@ -2032,9 +2032,6 @@ namespace UnrealBuildTool
 
 		public CppCompileEnvironment CreateCompileEnvironmentForIntellisense(ReadOnlyTargetRules Target, CppCompileEnvironment BaseCompileEnvironment, ILogger Logger)
 		{
-			string MutexName = SingleInstanceMutex.GetUniqueMutexForPath("UnrealBuildTool_CreateCompileEnvironmentForIntellisense", IntermediateDirectory.FullName);
-			using SingleInstanceMutex mutex = new SingleInstanceMutex(MutexName, true);
-
 			CppCompileEnvironment CompileEnvironment = CreateModuleCompileEnvironment(Target, BaseCompileEnvironment, Logger);
 			CompileEnvironment = SetupPrecompiledHeaders(Target, null, CompileEnvironment, new List<FileItem>(), new NullActionGraphBuilder(Logger), Logger);
 			CreateHeaderForDefinitions(CompileEnvironment, IntermediateDirectory, null, new NullActionGraphBuilder(Logger));
diff --git a/Engine/Source/Programs/UnrealBuildTool/Configuration/UEBuildTarget.cs b/Engine/Source/Programs/UnrealBuildTool/Configuration/UEBuildTarget.cs
--- a/Engine/Source/Programs/UnrealBuildTool/Configuration/UEBuildTarget.cs
+++ b/Engine/Source/Programs/UnrealBuildTool/Configuration/UEBuildTarget.cs
@@ -2444,6 +2444,16 @@ namespace UnrealBuildTool
 
 			FindSharedPCHs(Binaries, GlobalCompileEnvironment, Logger);
 
+			SingleInstanceMutex? mutex = null;
+			if (Rules.BuildEnvironment == TargetBuildEnvironment.Shared)
+			{
+				string MutexName = SingleInstanceMutex.GetUniqueMutexForPath("UnrealBuildTool_CreateCompileEnvironmentForProjectFiles", Unreal.EngineDirectory.FullName);
+				mutex = new SingleInstanceMutex(MutexName, true);
+			}
+
+			CreateSharedPCHInstances(Rules, TargetToolChain, Binaries, GlobalCompileEnvironment, new NullActionGraphBuilder(Logger), Logger);
+			mutex?.Dispose();
+
 			return GlobalCompileEnvironment;
 		}
 
