Fix for iOS deploy fails

diff --git a/Engine/Source/Programs/UnrealBuildTool/ProjectFiles/Xcode/XcodeProject.Scheme.cs b/Engine/Source/Programs/UnrealBuildTool/ProjectFiles/Xcode/XcodeProject.Scheme.cs
--- a/Engine/Source/Programs/UnrealBuildTool/ProjectFiles/Xcode/XcodeProject.Scheme.cs
+++ b/Engine/Source/Programs/UnrealBuildTool/ProjectFiles/Xcode/XcodeProject.Scheme.cs
@@ -67,6 +67,13 @@ namespace UnrealBuildTool.XcodeProjectXcconfig
 					"  touch &quot;${PROJECT_DIR}/${INFOPLIST_FILE}&quot;",
 					"fi",
 					"",
+					"# Sometimes Xcode won't redo code sign iOS.app, leads to deployment fails",
+					"# Force Xcode to redo code sign during each build",
+					"if [ &quot;${UE_PLATFORM_NAME}&quot; != &quot;Mac&quot; ]",
+					"then",
+					"  rm -f &quot;${UE_BINARIES_DIR}/${PRODUCT_NAME}.app/Info.plist&quot;",
+					"fi",
+					"",
 					"# if we had previously archived the build, then Xcode would have made a soft-link, but it will fail when building normally, so delete a soft-link",
 					"if [ -L &quot;${CONFIGURATION_BUILD_DIR}/${PRODUCT_NAME}.app&quot; ]",
 					"then",
