Changed so operator new/delete overloads are pulled in through a special include compiled in a generated cpp file that goes in its own translation unit. This will make it possible to support merging of modules and only include one of the .obj files for these overloads

If a module depends on Core and does not want the overloads it needs to add the private define SUPPRESS_PER_MODULE_INLINE_FILE or set PER_MODULE_INLINE_FILE to point to a different file

diff --git a/Engine/Plugins/Experimental/ChaosCachingUSD/Source/ChaosCachingUSD/ChaosCachingUSD.build.cs b/Engine/Plugins/Experimental/ChaosCachingUSD/Source/ChaosCachingUSD/ChaosCachingUSD.build.cs
--- a/Engine/Plugins/Experimental/ChaosCachingUSD/Source/ChaosCachingUSD/ChaosCachingUSD.build.cs
+++ b/Engine/Plugins/Experimental/ChaosCachingUSD/Source/ChaosCachingUSD/ChaosCachingUSD.build.cs
@@ -39,6 +39,8 @@ namespace UnrealBuildTool.Rules
 				"USDUtilities",
 				}
 			);
+
+			PrivateDefinitions.Add("SUPPRESS_PER_MODULE_INLINE_FILE"); // This module does not use core's standard operator new/delete overloads
 		}
 	}
 }
\ No newline at end of file
diff --git a/Engine/Plugins/Importers/USDImporter/Source/USDExporter/USDExporter.Build.cs b/Engine/Plugins/Importers/USDImporter/Source/USDExporter/USDExporter.Build.cs
--- a/Engine/Plugins/Importers/USDImporter/Source/USDExporter/USDExporter.Build.cs
+++ b/Engine/Plugins/Importers/USDImporter/Source/USDExporter/USDExporter.Build.cs
@@ -59,6 +59,8 @@ namespace UnrealBuildTool.Rules
 					"USDUtilities",
 				}
 			);
+
+			PrivateDefinitions.Add("SUPPRESS_PER_MODULE_INLINE_FILE"); // This module does not use core's standard operator new/delete overloads
 		}
 	}
 }
diff --git a/Engine/Plugins/Importers/USDImporter/Source/USDSchemas/USDSchemas.Build.cs b/Engine/Plugins/Importers/USDImporter/Source/USDSchemas/USDSchemas.Build.cs
--- a/Engine/Plugins/Importers/USDImporter/Source/USDSchemas/USDSchemas.Build.cs
+++ b/Engine/Plugins/Importers/USDImporter/Source/USDSchemas/USDSchemas.Build.cs
@@ -78,6 +78,8 @@ namespace UnrealBuildTool.Rules
 					);
 				}
 			}
+
+			PrivateDefinitions.Add("SUPPRESS_PER_MODULE_INLINE_FILE"); // This module does not use core's standard operator new/delete overloads
 		}
 	}
 }
diff --git a/Engine/Plugins/Importers/USDImporter/Source/USDStage/USDStage.Build.cs b/Engine/Plugins/Importers/USDImporter/Source/USDStage/USDStage.Build.cs
--- a/Engine/Plugins/Importers/USDImporter/Source/USDStage/USDStage.Build.cs
+++ b/Engine/Plugins/Importers/USDImporter/Source/USDStage/USDStage.Build.cs
@@ -67,6 +67,8 @@ namespace UnrealBuildTool.Rules
 					}
 				);
 			}
+
+			PrivateDefinitions.Add("SUPPRESS_PER_MODULE_INLINE_FILE"); // This module does not use core's standard operator new/delete overloads
 		}
 	}
 }
diff --git a/Engine/Plugins/Importers/USDImporter/Source/USDStageEditor/USDStageEditor.Build.cs b/Engine/Plugins/Importers/USDImporter/Source/USDStageEditor/USDStageEditor.Build.cs
--- a/Engine/Plugins/Importers/USDImporter/Source/USDStageEditor/USDStageEditor.Build.cs
+++ b/Engine/Plugins/Importers/USDImporter/Source/USDStageEditor/USDStageEditor.Build.cs
@@ -37,6 +37,8 @@ namespace UnrealBuildTool.Rules
 					"WorkspaceMenuStructure",
 				}
 			);
+
+			PrivateDefinitions.Add("SUPPRESS_PER_MODULE_INLINE_FILE"); // This module does not use core's standard operator new/delete overloads
 		}
 	}
 }
diff --git a/Engine/Plugins/Importers/USDImporter/Source/USDStageEditorViewModels/USDStageEditorViewModels.Build.cs b/Engine/Plugins/Importers/USDImporter/Source/USDStageEditorViewModels/USDStageEditorViewModels.Build.cs
--- a/Engine/Plugins/Importers/USDImporter/Source/USDStageEditorViewModels/USDStageEditorViewModels.Build.cs
+++ b/Engine/Plugins/Importers/USDImporter/Source/USDStageEditorViewModels/USDStageEditorViewModels.Build.cs
@@ -33,6 +33,8 @@ namespace UnrealBuildTool.Rules
 					"WorkspaceMenuStructure",
 				}
 			);
+
+			PrivateDefinitions.Add("SUPPRESS_PER_MODULE_INLINE_FILE"); // This module does not use core's standard operator new/delete overloads
 		}
 	}
 }
diff --git a/Engine/Plugins/Importers/USDImporter/Source/USDStageImporter/USDStageImporter.Build.cs b/Engine/Plugins/Importers/USDImporter/Source/USDStageImporter/USDStageImporter.Build.cs
--- a/Engine/Plugins/Importers/USDImporter/Source/USDStageImporter/USDStageImporter.Build.cs
+++ b/Engine/Plugins/Importers/USDImporter/Source/USDStageImporter/USDStageImporter.Build.cs
@@ -44,6 +44,8 @@ namespace UnrealBuildTool.Rules
 					"HairStrandsCore",
 				}
 			);
+
+			PrivateDefinitions.Add("SUPPRESS_PER_MODULE_INLINE_FILE"); // This module does not use core's standard operator new/delete overloads
 		}
 	}
 }
diff --git a/Engine/Plugins/Importers/USDImporter/Source/USDTests/USDTests.Build.cs b/Engine/Plugins/Importers/USDImporter/Source/USDTests/USDTests.Build.cs
--- a/Engine/Plugins/Importers/USDImporter/Source/USDTests/USDTests.Build.cs
+++ b/Engine/Plugins/Importers/USDImporter/Source/USDTests/USDTests.Build.cs
@@ -33,6 +33,8 @@ namespace UnrealBuildTool.Rules
 					}
 				);
 			}
+
+			PrivateDefinitions.Add("SUPPRESS_PER_MODULE_INLINE_FILE"); // This module does not use core's standard operator new/delete overloads
 		}
 	}
 }
diff --git a/Engine/Plugins/Runtime/USDCore/Source/USDUtilities/USDUtilities.Build.cs b/Engine/Plugins/Runtime/USDCore/Source/USDUtilities/USDUtilities.Build.cs
--- a/Engine/Plugins/Runtime/USDCore/Source/USDUtilities/USDUtilities.Build.cs
+++ b/Engine/Plugins/Runtime/USDCore/Source/USDUtilities/USDUtilities.Build.cs
@@ -61,6 +61,8 @@ namespace UnrealBuildTool.Rules
 					}
 				);
 			}
+
+			PrivateDefinitions.Add("SUPPRESS_PER_MODULE_INLINE_FILE"); // This module does not use core's standard operator new/delete overloads
 		}
 	}
 }
diff --git a/Engine/Plugins/Runtime/USDCore/Source/UnrealUSDWrapper/UnrealUSDWrapper.Build.cs b/Engine/Plugins/Runtime/USDCore/Source/UnrealUSDWrapper/UnrealUSDWrapper.Build.cs
--- a/Engine/Plugins/Runtime/USDCore/Source/UnrealUSDWrapper/UnrealUSDWrapper.Build.cs
+++ b/Engine/Plugins/Runtime/USDCore/Source/UnrealUSDWrapper/UnrealUSDWrapper.Build.cs
@@ -168,6 +168,8 @@ namespace UnrealBuildTool.Rules
 				PublicDefinitions.Add("USE_USD_SDK=0");
 				PublicDefinitions.Add("USD_USES_SYSTEM_MALLOC=0");
 			}
+
+			PrivateDefinitions.Add("SUPPRESS_PER_MODULE_INLINE_FILE"); // This module does not use core's standard operator new/delete overloads
 		}
 
 		bool EnableUsdSdk(ReadOnlyTargetRules Target)
diff --git a/Engine/Source/Programs/UnrealBuildTool/Configuration/UEBuildModuleCPP.cs b/Engine/Source/Programs/UnrealBuildTool/Configuration/UEBuildModuleCPP.cs
--- a/Engine/Source/Programs/UnrealBuildTool/Configuration/UEBuildModuleCPP.cs
+++ b/Engine/Source/Programs/UnrealBuildTool/Configuration/UEBuildModuleCPP.cs
@@ -590,6 +590,31 @@ namespace UnrealBuildTool
 				}
 			}
 
+			// The operator new/delete overrides need to exist in all modules.
+			// We put these overrides in a separate file in order to be able to filter out so only one file is added when modules are merged
+			if (Target.LinkType == TargetLinkType.Modular)
+			{
+				FileReference NewDeleteOverrides = FileReference.Combine(IntermediateDirectory, "PerModuleInline.gen.cpp");
+				string Content =
+					"#if !defined(PER_MODULE_INLINE_FILE) && defined(CORE_API)\r\n" +
+					"#define PER_MODULE_INLINE_FILE \"HAL/PerModuleInline.inl\"\r\n" +
+					"#endif\r\n" +
+					"#if defined(PER_MODULE_INLINE_FILE) && !defined(SUPPRESS_PER_MODULE_INLINE_FILE)\r\n" +
+					"#include PER_MODULE_INLINE_FILE\r\n" +
+					"#endif";
+
+				bool bMergeModules = false;
+				FileItem PerModuleFile = Graph.CreateIntermediateTextFile(NewDeleteOverrides, Content, bMergeModules);
+				if (bMergeModules)
+				{
+					LinkInputFiles.AddRange(ToolChain.CompileAllCPPFiles(CompileEnvironment, new[] { PerModuleFile }, IntermediateDirectory, Name, Graph).ObjectFiles);
+				}
+				else
+				{
+					CPPFiles.Add(PerModuleFile);
+				}
+			}
+
 			// Engine modules will always use unity build mode unless MinSourceFilesForUnityBuildOverride is specified in
 			// the module rules file.  By default, game modules only use unity of they have enough source files for that
 			// to be worthwhile.  If you have a lot of small game modules, consider specifying MinSourceFilesForUnityBuildOverride=0
@@ -615,7 +640,7 @@ namespace UnrealBuildTool
 			{
 				int FileCount = CPPFiles.Count;
 
-				// if we are merging the genearted cpp files then that needs to be part of the count
+				// if we are merging the generated cpp files then that needs to be part of the count
 				if (bMergeUnityFiles)
 				{
 					FileCount += GeneratedFileItems.Count;
diff --git a/Engine/Source/Runtime/BuildSettings/BuildSettings.Build.cs b/Engine/Source/Runtime/BuildSettings/BuildSettings.Build.cs
--- a/Engine/Source/Runtime/BuildSettings/BuildSettings.Build.cs
+++ b/Engine/Source/Runtime/BuildSettings/BuildSettings.Build.cs
@@ -44,5 +44,7 @@ public class BuildSettings : ModuleRules
 		PrivateDefinitions.Add($"BUILD_USER=\"{userName}\"");
 		PrivateDefinitions.Add($"BUILD_USERDOMAINNAME=\"{userDomainName}\"");
 		PrivateDefinitions.Add($"BUILD_MACHINENAME=\"{machineName}\"");
+
+		PrivateDefinitions.Add("SUPPRESS_PER_MODULE_INLINE_FILE"); // This module does not use core's standard operator new/delete overloads
 	}
 }
diff --git a/Engine/Source/Runtime/Core/Private/Unix/UnixPlatformProcess.cpp b/Engine/Source/Runtime/Core/Private/Unix/UnixPlatformProcess.cpp
--- a/Engine/Source/Runtime/Core/Private/Unix/UnixPlatformProcess.cpp
+++ b/Engine/Source/Runtime/Core/Private/Unix/UnixPlatformProcess.cpp
@@ -85,7 +85,7 @@ void* FUnixPlatformProcess::GetDllHandle( const TCHAR* Filename )
 		else
 		{
 			// is it ia ue4 module ? if not, move it to GLOBAL
-			void *IsUE4Module = dlsym(Handle, "InitializeModule");
+			void *IsUE4Module = dlsym(Handle, "ThisIsAnUnrealEngineModule");
 			if (!IsUE4Module)
 			{
 				UpgradeToGlobal = true;
diff --git a/Engine/Source/Runtime/Core/Public/HAL/PerModuleInline.inl b/Engine/Source/Runtime/Core/Public/HAL/PerModuleInline.inl
new file mode 100644
--- /dev/null
+++ b/Engine/Source/Runtime/Core/Public/HAL/PerModuleInline.inl
@@ -0,0 +1,14 @@
+// Copyright Epic Games, Inc. All Rights Reserved.
+
+#pragma once
+
+#include "CoreGlobals.h"
+#include "Modules/Boilerplate/ModuleBoilerplate.h"
+
+UE4_VISUALIZERS_HELPERS
+REPLACEMENT_OPERATOR_NEW_AND_DELETE
+UE_DEFINE_FMEMORY_WRAPPERS
+
+#if !IS_MONOLITHIC
+extern "C" DLLEXPORT void ThisIsAnUnrealEngineModule() {}
+#endif
diff --git a/Engine/Source/Runtime/Core/Public/Modules/ModuleManager.h b/Engine/Source/Runtime/Core/Public/Modules/ModuleManager.h
--- a/Engine/Source/Runtime/Core/Public/Modules/ModuleManager.h
+++ b/Engine/Source/Runtime/Core/Public/Modules/ModuleManager.h
@@ -821,7 +821,6 @@ class FDefaultGameModuleImpl
 		} \
 		/* Forced reference to this function is added by the linker to check that each module uses IMPLEMENT_MODULE */ \
 		extern "C" void IMPLEMENT_MODULE_##ModuleName() { } \
-		PER_MODULE_BOILERPLATE \
 		PER_MODULE_BOILERPLATE_ANYLINK(ModuleImplClass, ModuleName)
 
 #endif //IS_MONOLITHIC
@@ -941,7 +940,6 @@ class FDefaultGameModuleImpl
 					FCString::Strncpy(GInternalProjectName, TEXT( GameName ), UE_ARRAY_COUNT(GInternalProjectName)); \
 				} \
 			} AutoSet##ModuleName; \
-			PER_MODULE_BOILERPLATE \
 			PER_MODULE_BOILERPLATE_ANYLINK(FDefaultGameModuleImpl, ModuleName) \
 			FEngineLoop GEngineLoop;
 	#endif
diff --git a/Engine/Source/Runtime/Experimental/IoStore/HttpClient/Private/Module.cpp b/Engine/Source/Runtime/Experimental/IoStore/HttpClient/Private/Module.cpp
deleted file mode 100644
diff --git a/Engine/Source/Runtime/TraceLog/TraceLog.Build.cs b/Engine/Source/Runtime/TraceLog/TraceLog.Build.cs
--- a/Engine/Source/Runtime/TraceLog/TraceLog.Build.cs
+++ b/Engine/Source/Runtime/TraceLog/TraceLog.Build.cs
@@ -14,5 +14,7 @@ public class TraceLog : ModuleRules
 		PublicIncludePathModuleNames.Add("Core");
 
 		bAllowAutoRTFMInstrumentation = true;
+
+		PrivateDefinitions.Add("SUPPRESS_PER_MODULE_INLINE_FILE"); // This module does not use core's standard operator new/delete overloads
 	}
 }
