UnrealBuildTool: Initial support for merging modular builds into fewer libraries, highly experimental

#rnx

diff --git a/Engine/Source/Programs/UnrealBuildTool/Configuration/ReadOnlyTargetRules.cs b/Engine/Source/Programs/UnrealBuildTool/Configuration/ReadOnlyTargetRules.cs
--- a/Engine/Source/Programs/UnrealBuildTool/Configuration/ReadOnlyTargetRules.cs
+++ b/Engine/Source/Programs/UnrealBuildTool/Configuration/ReadOnlyTargetRules.cs
@@ -529,6 +529,14 @@ namespace UnrealBuildTool
 
 		public TargetLinkType LinkType => Inner.LinkType;
 
+		public bool bStripExports => Inner.bStripExports;
+
+		public bool bMergeModules => Inner.bMergeModules;
+
+		public IReadOnlyList<string> MergePlugins => Inner.MergePlugins;
+
+		public IReadOnlyDictionary<string, IEnumerable<string>> MergePluginsShared => Inner.MergePluginsShared.ToDictionary(k => k.Key, v => v.Value.AsEnumerable());
+
 		public IReadOnlyList<string> GlobalDefinitions => Inner.GlobalDefinitions.AsReadOnly();
 
 		public IReadOnlyList<string> ProjectDefinitions => Inner.ProjectDefinitions.AsReadOnly();
diff --git a/Engine/Source/Programs/UnrealBuildTool/Configuration/TargetRules.cs b/Engine/Source/Programs/UnrealBuildTool/Configuration/TargetRules.cs
--- a/Engine/Source/Programs/UnrealBuildTool/Configuration/TargetRules.cs
+++ b/Engine/Source/Programs/UnrealBuildTool/Configuration/TargetRules.cs
@@ -2434,6 +2434,32 @@ namespace UnrealBuildTool
 			set => LinkTypePrivate = value;
 		}
 
+		/// <summary>
+		/// Experimental: Strip unused exports from libraries. Only applies when LinkType is Modular
+		/// </summary>
+		[RequiresUniqueBuildEnvironment]
+		[CommandLine("-StripExports")]
+		public bool bStripExports { get; set; }
+
+		/// <summary>
+		/// Experimental: Merge modular modules into combined libraries. Sets LinkType to Modular and enables bStripExports
+		/// </summary>
+		[RequiresUniqueBuildEnvironment]
+		[CommandLine("-MergeModules")]
+		public bool bMergeModules { get; set; }
+
+		/// <summary>
+		/// Experimental: List of plugins (and their dependencies) to each merge into separate libraries. Requires bMergeModules to be enabled
+		/// </summary>
+		[RequiresUniqueBuildEnvironment]
+		[CommandLine("-MergePlugins=", ListSeparator = '+')]
+		public List<string> MergePlugins { get; } = new();
+
+		/// <summary>
+		/// Experimental: List of merged plugins to move common dependenceis to a shared library, requires MergePlugins. Can be chained
+		/// </summary>
+		public Dictionary<string, HashSet<string>> MergePluginsShared { get; } = new();
+
 		/// <summary>
 		/// Backing storage for the LinkType property.
 		/// </summary>
diff --git a/Engine/Source/Programs/UnrealBuildTool/Configuration/UEBuildBinary.cs b/Engine/Source/Programs/UnrealBuildTool/Configuration/UEBuildBinary.cs
--- a/Engine/Source/Programs/UnrealBuildTool/Configuration/UEBuildBinary.cs
+++ b/Engine/Source/Programs/UnrealBuildTool/Configuration/UEBuildBinary.cs
@@ -128,6 +128,16 @@ namespace UnrealBuildTool
 		/// </summary>
 		public readonly List<ModuleRules.RuntimeDependency> RuntimeDependencies = new List<ModuleRules.RuntimeDependency>();
 
+		/// <summary>
+		/// If compiled objects will have unused exports stripped, only affects objects compiled into modular libraries.
+		/// </summary>
+		public bool bStripUnusedExports = false;
+
+		/// <summary>
+		/// List of all objects that would be linked.
+		/// </summary>
+		public readonly HashSet<FileItem> InputObjects = new();
+
 		/// <summary>
 		/// Create an instance of the class with the given configuration data
 		/// </summary>
@@ -775,6 +785,7 @@ namespace UnrealBuildTool
 					// Compile each module.
 					Logger.LogDebug("Compile module: {ModuleName}", Module.Name);
 					List<FileItem> LinkInputFiles = Module.Compile(Target, ToolChain, BinaryCompileEnvironment, WorkingSet, Graph, Logger);
+					InputObjects.UnionWith(LinkInputFiles.Where(x => x.HasExtension(".obj") || x.HasExtension(".o")));
 
 					// Save the module outputs. In monolithic builds, this is just the object files.
 					if (Target.LinkType == TargetLinkType.Monolithic)
@@ -837,6 +848,25 @@ namespace UnrealBuildTool
 			BinaryLinkEnvironment.bIsBuildingDLL = IsBuildingDll(Type);
 			BinaryLinkEnvironment.bIsBuildingLibrary = IsBuildingLibrary(Type);
 
+			// Setup object export stripping
+			if ((BinaryLinkEnvironment.bIsBuildingDLL || BinaryLinkEnvironment.bIsBuildingLibrary) && Target.bStripExports)
+			{
+				bStripUnusedExports = true;
+				List<FileItem> inputFiles = new();
+				foreach (FileItem item in BinaryLinkEnvironment.InputFiles)
+				{
+					if (InputObjects.Contains(item))
+					{
+						inputFiles.Add(FileItem.GetItemByFileReference(item.Location.ChangeExtension($".strip{item.Location.GetExtension()}")));
+					}
+					else
+					{
+						inputFiles.Add(item);
+					}
+				}
+				BinaryLinkEnvironment.InputFiles = inputFiles;
+			}
+
 			// Code coverage inherited from compile environment
 			BinaryLinkEnvironment.bCodeCoverage = CompileEnvironment.bCodeCoverage;
 
diff --git a/Engine/Source/Programs/UnrealBuildTool/Configuration/UEBuildModule.cs b/Engine/Source/Programs/UnrealBuildTool/Configuration/UEBuildModule.cs
--- a/Engine/Source/Programs/UnrealBuildTool/Configuration/UEBuildModule.cs
+++ b/Engine/Source/Programs/UnrealBuildTool/Configuration/UEBuildModule.cs
@@ -660,7 +660,7 @@ namespace UnrealBuildTool
 						Definitions.Add(ModuleApiDefine + "=");
 					}
 				}
-				else if (Binary == null || SourceBinary != Binary)
+				else if (Binary == null || (Rules.Target.bMergeModules ? (SourceBinary == null || SourceModule != this) : SourceBinary != Binary))
 				{
 					Definitions.Add(ModuleApiDefine + "=DLLIMPORT");
 				}
diff --git a/Engine/Source/Programs/UnrealBuildTool/Configuration/UEBuildTarget.cs b/Engine/Source/Programs/UnrealBuildTool/Configuration/UEBuildTarget.cs
--- a/Engine/Source/Programs/UnrealBuildTool/Configuration/UEBuildTarget.cs
+++ b/Engine/Source/Programs/UnrealBuildTool/Configuration/UEBuildTarget.cs
@@ -8,6 +8,7 @@ using System.Globalization;
 using System.IO;
 using System.Linq;
 using System.Reflection;
+using System.Runtime.InteropServices;
 using System.Runtime.Serialization;
 using System.Text;
 using System.Text.Json.Serialization;
@@ -2644,6 +2645,11 @@ namespace UnrealBuildTool
 				}
 			}
 
+			if (Rules.bStripExports)
+			{
+				CreateStripExportsActions(MakefileBuilder, Makefile);
+			}
+
 			// Cache inline gen cpp data
 			using (GlobalTracer.Instance.BuildSpan("CacheInlineGenCppData").StartActive())
 			{
@@ -2980,6 +2986,73 @@ namespace UnrealBuildTool
 			return Makefile;
 		}
 
+		void CreateStripExportsActions(TargetMakefileBuilder MakefileBuilder, TargetMakefile Makefile)
+		{
+			FileReference commandPath;
+			if (OperatingSystem.IsWindows())
+			{
+#pragma warning disable CA1308 // Normalize strings to uppercase
+				commandPath = FileReference.Combine(Unreal.EngineDirectory, "Binaries", "Win64", "UnrealBuildAccelerator", RuntimeInformation.ProcessArchitecture.ToString().ToLowerInvariant(), "UbaObjTool.exe");
+#pragma warning restore CA1308 // Normalize strings to uppercase
+			}
+			else if (OperatingSystem.IsLinux())
+			{
+				if (RuntimeInformation.ProcessArchitecture == Architecture.X64)
+				{
+					commandPath = FileReference.Combine(Unreal.EngineDirectory, "Binaries", "Linux", "UnrealBuildAccelerator", "UbaObjTool");
+				}
+				else if (RuntimeInformation.ProcessArchitecture == Architecture.Arm64)
+				{
+					commandPath = FileReference.Combine(Unreal.EngineDirectory, "Binaries", "LinuxArm64", "UnrealBuildAccelerator", "UbaObjTool");
+				}
+				else
+				{
+					throw new PlatformNotSupportedException();
+				}
+			}
+			else if (OperatingSystem.IsMacOS())
+			{
+				commandPath = FileReference.Combine(Unreal.EngineDirectory, "Binaries", "Mac", "UnrealBuildAccelerator", "UbaObjTool");
+			}
+			else
+			{
+				throw new PlatformNotSupportedException();
+			}
+
+			HashSet<FileItem> allObjects = new(Binaries.SelectMany(x => x.InputObjects));
+
+			void CreateStripAction(string Name, IEnumerable<FileItem> toStrip)
+			{
+				HashSet<FileItem> otherObjects = new(allObjects.Except(toStrip));
+				FileItem stripRsp = FileItem.GetItemByFileReference(FileReference.Combine(ProjectIntermediateDirectory!, $"{Name}.strip.rsp"));
+				List<string> arguments = new();
+				arguments.AddRange(toStrip.Select(path => $"/S:{path.FullName}").OrderBy(x => x));
+				arguments.AddRange(otherObjects.Select(path => $"/D:{path.FullName}").OrderBy(x => x));
+
+				MakefileBuilder.CreateIntermediateTextFile(stripRsp, arguments);
+
+				Action stripAction = MakefileBuilder.CreateAction(ActionType.Link);
+
+				stripAction.PrerequisiteItems.Add(stripRsp);
+				stripAction.PrerequisiteItems.UnionWith(allObjects);
+				stripAction.ProducedItems.UnionWith(toStrip.Select(x => FileItem.GetItemByFileReference(x.Location.ChangeExtension($".strip{x.Location.GetExtension()}"))));
+				stripAction.CommandPath = commandPath;
+				stripAction.CommandArguments = "@" + stripRsp.FullName;
+				stripAction.CommandDescription = "Strip Objects";
+				stripAction.WorkingDirectory = Unreal.EngineSourceDirectory;
+				stripAction.StatusDescription = Name;
+				stripAction.bCanExecuteRemotely = false;
+				Makefile.OutputItems.AddRange(stripAction.ProducedItems);
+			}
+
+			HashSet<FileItem> pchObjects = new(allObjects.Where(x => x.Name.Contains("SharedPCH", StringComparison.Ordinal)));
+			CreateStripAction($"{AppName}-SharedPCH", pchObjects);
+			foreach (UEBuildBinary Binary in Binaries.Where(x => x.bStripUnusedExports))
+			{
+				CreateStripAction(Binary.OutputFilePaths.First().GetFileNameWithoutExtension(), Binary.InputObjects.Where(x => !pchObjects.Contains(x)));
+			}
+		}
+
 		void CreateWriteMetadataAction(TargetMakefileBuilder MakefileBuilder, string StatusDescription, FileReference InfoFile, WriteMetadataTargetInfo Info, IEnumerable<FileItem> PrerequisiteItems)
 		{
 			List<FileReference> ProducedItems = new List<FileReference>(Info.FileToManifest.Keys);
@@ -3192,7 +3265,8 @@ namespace UnrealBuildTool
 		/// Export the definition of this target to a JSON file
 		/// </summary>
 		/// <param name="OutputFile">File to write to</param>
-		public void ExportJson(FileReference OutputFile)
+		/// <param name="Actions">List of actions to export</param>
+		public void ExportJson(FileReference OutputFile, List<LinkedAction>? Actions = null)
 		{
 			DirectoryReference.CreateDirectory(OutputFile.Directory);
 			using (JsonWriter Writer = new JsonWriter(OutputFile))
@@ -3225,6 +3299,27 @@ namespace UnrealBuildTool
 				}
 				Writer.WriteObjectEnd();
 
+				if (Actions != null)
+				{
+					Writer.WriteArrayStart("Actions");
+					Dictionary<LinkedAction, int> LinkedActionToId = new();
+					int Counter = 0;
+					foreach (LinkedAction Action in Actions)
+					{
+						LinkedActionToId[Action] = Counter++;
+					}
+					foreach (LinkedAction Action in Actions)
+					{
+						if (Action.StatusDescription != "...")
+						{
+							Writer.WriteObjectStart();
+							Action.ExportJson(LinkedActionToId, Writer);
+							Writer.WriteObjectEnd();
+						}
+					}
+					Writer.WriteArrayEnd();
+				}
+
 				Writer.WriteObjectEnd();
 			}
 		}
@@ -3345,6 +3440,116 @@ namespace UnrealBuildTool
 			return true;
 		}
 
+		/// <summary>
+		/// Dynamically merge modules into libraries based on target settings
+		/// </summary>
+		/// <param name="Logger"></param>
+		protected void SetupMergeModules(ILogger Logger)
+		{
+			void MergeModules(string name, IEnumerable<UEBuildModuleCPP> modules)
+			{
+				if (!modules.Any())
+				{
+					return;
+				}
+
+				UEBuildBinary mergedBinary = CreateDynamicLibraryForModules(name, modules.OrderBy(x => x.Name));
+				foreach (UEBuildModuleCPP module in modules)
+				{
+					module.Binary = mergedBinary;
+				}
+
+				Binaries.Add(mergedBinary);
+			}
+
+			UEBuildModuleCPP launchModule = Binaries[0].PrimaryModule;
+			HashSet<UEBuildModuleCPP> allModules = new(Binaries.SelectMany(x => x.Modules.OfType<UEBuildModuleCPP>().Where(x => x != launchModule)));
+
+			if (Rules.MergePlugins.Any() && BuildPlugins != null)
+			{
+				// Find the plugin to merge and gather all dependencies
+				Dictionary<string, HashSet<UEBuildPlugin>> pluginGroups = new();
+				foreach (string name in Rules.MergePlugins)
+				{
+					UEBuildPlugin? plugin = BuildPlugins.FirstOrDefault(x => x.Name == name);
+					if (plugin == null)
+					{
+						Logger.LogWarning("Plugin {Name} not found in BuildPlugin and will not be merged", name);
+						continue;
+					}
+					HashSet<UEBuildPlugin> needed = new();
+					needed.Add(plugin);
+					needed.UnionWith(plugin.Dependencies ?? new());
+					pluginGroups.Add(plugin.Name, needed);
+				}
+
+				// Move common plugins to a new group
+				foreach (KeyValuePair<string, IEnumerable<string>> group in Rules.MergePluginsShared)
+				{
+					// Get the existing groups that will be shared
+					Dictionary<string, HashSet<UEBuildPlugin>> shared = new();
+					foreach (string dependecy in group.Value)
+					{
+						if (pluginGroups.TryGetValue(dependecy, out HashSet<UEBuildPlugin>? value))
+						{
+							shared.Add(dependecy, value);
+						}
+					}
+
+					// Find all plugins that are referenced by more than one group
+					HashSet<UEBuildPlugin> common = new(shared.SelectMany(x => x.Value).Where(x => shared.Count(y => y.Value.Contains(x)) > 1));
+					foreach (KeyValuePair<string, HashSet<UEBuildPlugin>> item in shared)
+					{
+						item.Value.ExceptWith(common);
+						pluginGroups.Remove(item.Key);
+						pluginGroups.Add(item.Key, item.Value);
+					}
+					pluginGroups.Add(group.Key, common);
+				}
+
+				// Remove any plugins that are shared between multiple merge groups
+				if (pluginGroups.Any())
+				{
+					Dictionary<string, HashSet<UEBuildPlugin>> filtered = new();
+					foreach (KeyValuePair<string, HashSet<UEBuildPlugin>> item in pluginGroups)
+					{
+						IEnumerable<UEBuildPlugin> otherPluginDependencies = pluginGroups.Where(x => x.Key != item.Key).SelectMany(x => x.Value).Distinct();
+						filtered.Add(item.Key, new(item.Value.Except(otherPluginDependencies)));
+					}
+					pluginGroups = filtered;
+				}
+
+				// Find remaining unmerged plugins and gather their dependencies
+				HashSet<UEBuildPlugin> unmergablePlugins = new(BuildPlugins);
+				unmergablePlugins.ExceptWith(pluginGroups.SelectMany(x => x.Value));
+				unmergablePlugins = new(unmergablePlugins.Union(unmergablePlugins.SelectMany(x => x.Dependencies ?? new())));
+
+				// Create the merged plugin modules
+				foreach (KeyValuePair<string, HashSet<UEBuildPlugin>> item in pluginGroups)
+				{
+					// Find all modules for the unique plugins to be merged, except those that were unmergable due to being a dependency of an unmerged plugin
+					IEnumerable<UEBuildModuleCPP> pluginModules = item.Value.Except(unmergablePlugins).SelectMany(x => x.Modules).Distinct();
+
+					// Merge the plugin modules and remove those modules from remaining set of modules to merge 
+					MergeModules(item.Key, pluginModules);
+					allModules.ExceptWith(pluginModules);
+				}
+			}
+
+			IEnumerable<UEBuildModuleCPP> engineModules = allModules.Where(x => x.RulesFile.IsUnderDirectory(Unreal.EngineDirectory));
+			IEnumerable<UEBuildModuleCPP> projectModules = allModules.Except(engineModules);
+
+			// Merge remaining engine modules
+			MergeModules("CommonEngine", engineModules);
+
+			// Merge remaining project modules
+			MergeModules("Common", projectModules);
+
+			// Remove any binaries that no longer have modules due to merging
+			Binaries.ForEach(x => x.Modules.RemoveAll(y => y.Binary != x));
+			Binaries.RemoveAll(x => !x.Modules.Any());
+		}
+
 		/// <summary>
 		/// Setup target before build. This method finds dependencies, sets up global environment etc.
 		/// </summary>
@@ -3401,6 +3606,12 @@ namespace UnrealBuildTool
 				AddAllValidModulesToTarget(Logger);
 			}
 
+			// Merge modules into fewer libraries if requested.
+			if (Rules.bMergeModules)
+			{
+				SetupMergeModules(Logger);
+			}
+
 			// Add the external and non-C++ referenced modules to the binaries that reference them.
 			foreach (UEBuildModuleCPP Module in Modules.Values.OfType<UEBuildModuleCPP>())
 			{
@@ -4171,6 +4382,52 @@ namespace UnrealBuildTool
 			);
 		}
 
+		/// <summary>
+		/// Adds a dynamic library for the given set of modules. Does not check whether a binary already exists, or whether a binary should be created for this build configuration.
+		/// </summary>
+		/// <param name="MergedName">The name of the merged binary</param>
+		/// <param name="Modules">The modules to create a binary for</param>
+		/// <returns>The new binary. This has not been added to the target.</returns>
+		private UEBuildBinary CreateDynamicLibraryForModules(string MergedName, IEnumerable<UEBuildModuleCPP> Modules)
+		{
+			DirectoryReference IntermediateDirectory = DirectoryReference.Combine(ProjectIntermediateDirectory, MergedName);
+
+			// Get the root output directory and base name (target name/app name) for this binary
+			DirectoryReference OutputDirectory;
+			if (ProjectFile != null && Rules.File.IsUnderDirectory(ProjectDirectory) || Modules.Any(x => x.Rules.bUsePrecompiled))
+			{
+				OutputDirectory = GetOutputDirectoryForExecutable(ProjectDirectory, Rules.File);
+			}
+			else
+			{
+				OutputDirectory = GetOutputDirectoryForExecutable(Unreal.EngineDirectory, Rules.File);
+			}
+
+			// Get the configuration that this module will be built in. Engine modules compiled in DebugGame will use Development.
+			UnrealTargetConfiguration ModuleConfiguration = Configuration;
+			if (Configuration == UnrealTargetConfiguration.DebugGame && !Modules.All(x => x.Rules.Context.bCanBuildDebugGame))
+			{
+				ModuleConfiguration = UnrealTargetConfiguration.Development;
+			}
+
+			List<FileReference> OutputPaths = MakeBinaryPaths(OutputDirectory, AppName + "-" + MergedName, Platform, Configuration, UEBuildBinaryType.DynamicLinkLibrary, Rules.Architectures, Rules.UndecoratedConfiguration, bCompileMonolithic && ProjectFile != null, Rules.ExeBinariesSubFolder, ProjectFile, Rules);
+
+			// Create the binary
+			UEBuildBinary Binary = new UEBuildBinary(
+				Type: UEBuildBinaryType.DynamicLinkLibrary,
+				OutputFilePaths: OutputPaths,
+				IntermediateDirectory: ProjectIntermediateDirectory,
+				bAllowExports: true,
+				bBuildAdditionalConsoleApp: false,
+				bBuildConsoleAppOnly: false,
+				PrimaryModule: Modules.First(),
+				bUsePrecompiled: false
+			);
+
+			Binary.Modules.AddRange(Modules.Skip(1));
+			return Binary;
+		}
+
 		/// <summary>
 		/// Makes a filename (without path) for a compiled binary (e.g. "Core-Win64-Debug.lib") */
 		/// </summary>
diff --git a/Engine/Source/Programs/UnrealBuildTool/System/RulesAssembly.cs b/Engine/Source/Programs/UnrealBuildTool/System/RulesAssembly.cs
--- a/Engine/Source/Programs/UnrealBuildTool/System/RulesAssembly.cs
+++ b/Engine/Source/Programs/UnrealBuildTool/System/RulesAssembly.cs
@@ -696,6 +696,14 @@ namespace UnrealBuildTool
 				Rules.bOmitPCDebugInfoInDevelopment = false;
 			}
 
+			// If merging modules, force modular and strip exports
+			if (Rules.bMergeModules)
+			{
+				Rules.BuildEnvironment = TargetBuildEnvironment.Unique;
+				Rules.LinkType = TargetLinkType.Modular;
+				Rules.bStripExports = true;
+			}
+
 			// Setup utrace for Shader Compiler Worker
 			if (Rules.bShaderCompilerWorkerTrace)
 			{
