UBT: SNDBS: Fixed possible distribution problems when cl-filter is used during the build.

cl-filter is invoked as follows:

cl-filter -dependencies=<path> -timing=<path> -compiler=<cl_path> -- <cl_path> [at]<rsp_path> <params>

when the command line is too long, SNDBS moves the parameters into a new response file and modifies the call to be:

cl-filter [at]<sndbs_rsp_path>

However, cl-filter doesn't support response files and fails.

Resolved by modifying the SNDBS configuration to only allow the part after "-- <cl_path>" to be moved to the response file.

diff --git a/Engine/Source/Programs/UnrealBuildTool/Executors/SNDBS.cs b/Engine/Source/Programs/UnrealBuildTool/Executors/SNDBS.cs
--- a/Engine/Source/Programs/UnrealBuildTool/Executors/SNDBS.cs
+++ b/Engine/Source/Programs/UnrealBuildTool/Executors/SNDBS.cs
@@ -459,6 +459,7 @@ vc_major_version=14
 use_surrogate=true
 force_synchronous_pdb_writes=true
 error_report_mode=prompt
+response_file_content_pattern=\s--\s"".*?cl\.exe""\s(.*)
 
 [group]
 server={VC_COMPILER_DIR}\mspdbsrv.exe
