[UBT]
* Fixed so PerModule cpp file creates its own translation unit if module merging is enabled
* Fixed so only one translation unit of PerModule.cpp is added per binary if module merging is enabled

diff --git a/Engine/Source/Programs/UnrealBuildTool/Configuration/UEBuildBinary.cs b/Engine/Source/Programs/UnrealBuildTool/Configuration/UEBuildBinary.cs
--- a/Engine/Source/Programs/UnrealBuildTool/Configuration/UEBuildBinary.cs
+++ b/Engine/Source/Programs/UnrealBuildTool/Configuration/UEBuildBinary.cs
@@ -778,6 +778,8 @@ namespace UnrealBuildTool
 
 			HashSet<FileItem> InputFilesLookup = new();
 
+			FileItem? PerModuleItemToKeep = null;
+
 			foreach (UEBuildModule Module in UEBuildModule.StableTopologicalSort(Modules))
 			{
 				if (Module.Binary == null || Module.Binary == this)
@@ -785,6 +787,23 @@ namespace UnrealBuildTool
 					// Compile each module.
 					Logger.LogDebug("Compile module: {ModuleName}", Module.Name);
 					List<FileItem> LinkInputFiles = Module.Compile(Target, ToolChain, BinaryCompileEnvironment, WorkingSet, Graph, Logger);
+
+
+					// If modules are merged we will have multiple PerModuleInline which will cause duplicated symbols. So let's only keep one of those
+					// Also, some modules might have empty PerModuleInline.gen.cpp therefore it is important that we use Core if it exists
+					if (Target.bMergeModules)
+					{
+						FileItem PerModuleItem = LinkInputFiles.FirstOrDefault(x => x!.Name.Contains("PerModuleInline.gen.cpp"), null);
+						if (PerModuleItem != null)
+						{
+							if (PerModuleItemToKeep == null || Module.Name == "Core")
+							{
+								PerModuleItemToKeep = PerModuleItem;
+							}
+							LinkInputFiles.Remove(PerModuleItem);
+						}
+					}
+
 					InputObjects.UnionWith(LinkInputFiles.Where(x => x.HasExtension(".obj") || x.HasExtension(".o")));
 
 					// Save the module outputs. In monolithic builds, this is just the object files.
@@ -817,6 +836,12 @@ namespace UnrealBuildTool
 				Module.SetupPrivateLinkEnvironment(this, BinaryLinkEnvironment, BinaryDependencies, LinkEnvironmentVisitedModules, ExeDir);
 			}
 
+			if (PerModuleItemToKeep != null)
+			{
+				BinaryLinkEnvironment.InputFiles.Add(PerModuleItemToKeep);
+				InputObjects.Add(PerModuleItemToKeep);
+			}
+
 			// Allow the binary dependencies to modify the link environment.
 			foreach (UEBuildBinary BinaryDependency in BinaryDependencies)
 			{
diff --git a/Engine/Source/Programs/UnrealBuildTool/Configuration/UEBuildModuleCPP.cs b/Engine/Source/Programs/UnrealBuildTool/Configuration/UEBuildModuleCPP.cs
--- a/Engine/Source/Programs/UnrealBuildTool/Configuration/UEBuildModuleCPP.cs
+++ b/Engine/Source/Programs/UnrealBuildTool/Configuration/UEBuildModuleCPP.cs
@@ -603,7 +603,7 @@ namespace UnrealBuildTool
 					"#include PER_MODULE_INLINE_FILE\r\n" +
 					"#endif";
 
-				bool bMergeModules = false;
+				bool bMergeModules = Target.bMergeModules;
 				FileItem PerModuleFile = Graph.CreateIntermediateTextFile(NewDeleteOverrides, Content, bMergeModules);
 				if (bMergeModules)
 				{
