[UBT]
* Fixed compile warning

diff --git a/Engine/Source/Programs/UnrealBuildTool/Configuration/UEBuildBinary.cs b/Engine/Source/Programs/UnrealBuildTool/Configuration/UEBuildBinary.cs
--- a/Engine/Source/Programs/UnrealBuildTool/Configuration/UEBuildBinary.cs
+++ b/Engine/Source/Programs/UnrealBuildTool/Configuration/UEBuildBinary.cs
@@ -793,7 +793,7 @@ namespace UnrealBuildTool
 					// Also, some modules might have empty PerModuleInline.gen.cpp therefore it is important that we use Core if it exists
 					if (Target.bMergeModules)
 					{
-						FileItem PerModuleItem = LinkInputFiles.FirstOrDefault(x => x!.Name.Contains("PerModuleInline.gen.cpp"), null);
+						FileItem? PerModuleItem = LinkInputFiles.FirstOrDefault(x => x!.Name.Contains("PerModuleInline.gen.cpp"), null);
 						if (PerModuleItem != null)
 						{
 							if (PerModuleItemToKeep == null || Module.Name == "Core")
