UnrealBuildTool: Skip drive letter when checking for output file paths greater than MAX_PATH

#rnx

diff --git a/Engine/Source/Programs/UnrealBuildTool/System/ActionGraph.cs b/Engine/Source/Programs/UnrealBuildTool/System/ActionGraph.cs
--- a/Engine/Source/Programs/UnrealBuildTool/System/ActionGraph.cs
+++ b/Engine/Source/Programs/UnrealBuildTool/System/ActionGraph.cs
@@ -292,14 +292,14 @@ namespace UnrealBuildTool
 				bool ShouldFail(FileItem item)
 				{
 					return item.Location.FullName.Length >= MAX_PATH &&
-						!(item.Location.ContainsName("Intermediate", 0) && item.Location.ContainsName("H", 0)); // Ignore -IncludeHeader items
+						!(item.Location.ContainsName("Intermediate", 0) && item.Location.ContainsName("H", 1)); // Ignore -IncludeHeader items
 				}
 
 				bool ShouldWarn(FileItem item)
 				{
 					return item.Location.FullName.Length > Unreal.RootDirectory.FullName.Length + BuildConfiguration.MaxNestedPathLength &&
 						item.Location.IsUnderDirectory(Unreal.RootDirectory) &&
-						!(item.Location.ContainsName("Intermediate", 0) && item.Location.ContainsName("H", 0)) && // Ignore -IncludeHeader items
+						!(item.Location.ContainsName("Intermediate", 0) && item.Location.ContainsName("H", 1)) && // Ignore -IncludeHeader items
 						(item.Location.ContainsName("Restricted", 0) == false) && //Be more relaxed for internal only code
 						item.Location.ContainsName("NotForLicensees", 0) == false;
 				}
