direct header compiling tools

diff --git a/Engine/Source/Programs/UnrealBuildTool/Configuration/UEBuildModuleCPP.cs b/Engine/Source/Programs/UnrealBuildTool/Configuration/UEBuildModuleCPP.cs
--- a/Engine/Source/Programs/UnrealBuildTool/Configuration/UEBuildModuleCPP.cs
+++ b/Engine/Source/Programs/UnrealBuildTool/Configuration/UEBuildModuleCPP.cs
@@ -845,9 +845,6 @@ namespace UnrealBuildTool
 					CppCompileEnvironment HeaderCompileEnvironment = new CppCompileEnvironment(ModuleCompileEnvironment);
 					DirectoryReference HeaderIntermediateDirectory = DirectoryReference.Combine(IntermediateDirectory, "H");
 
-					HeaderCompileEnvironment.Definitions.RemoveAll(x => x.StartsWith("UE_DIRECT_HEADER_COMPILE"));
-					HeaderCompileEnvironment.Definitions.Add("UE_DIRECT_HEADER_COMPILE=1");
-
 					CreateHeaderForDefinitions(HeaderCompileEnvironment, HeaderIntermediateDirectory, "h", Graph);
 
 					// Duplicate named headers are allowed, so adjust the IntermediateDirectory to compensate for this by using an index
@@ -864,10 +861,7 @@ namespace UnrealBuildTool
 								IncludeFileString = Header.Location.MakeRelativeTo(Unreal.EngineSourceDirectory);
 							}
 
-							List<string> GeneratedHeaderCppContents = new();
-							GeneratedHeaderCppContents.Add("// This file is automatically generated at compile-time to include a user created header file.");
-							GeneratedHeaderCppContents.Add($"#include \"{IncludeFileString.Replace('\\', '/')}\"");
-
+							List<string> GeneratedHeaderCppContents = UEBuildModuleCPP.GenerateHeaderCpp(Header.Name, IncludeFileString);
 							FileItem GeneratedHeaderCpp = FileItem.GetItemByFileReference(FileReference.Combine(HeaderSubIntermediateDirectory, $"{Header.Name}.cpp"));
 							Graph.CreateIntermediateTextFile(GeneratedHeaderCpp, GeneratedHeaderCppContents);
 							HeadersToCompile.Add(GeneratedHeaderCpp);
@@ -2430,6 +2424,18 @@ namespace UnrealBuildTool
 			return SourceFiles.ToArray();
 		}
 
+		public static List<string> GenerateHeaderCpp(string HeaderName, string IncludeFileString)
+		{
+			HeaderName = HeaderName.Substring(0, HeaderName.Length - 2).Replace(".", "_").Replace("-", "_");
+
+			List<string> GeneratedHeaderCppContents = new();
+			GeneratedHeaderCppContents.Add("// This file is automatically generated at compile-time to include a user created header file.");
+			GeneratedHeaderCppContents.Add("#define UE_DIRECT_HEADER_COMPILE 1");
+			GeneratedHeaderCppContents.Add($"#define __COMPILING_{HeaderName}");
+			GeneratedHeaderCppContents.Add($"#include \"{IncludeFileString.Replace('\\', '/')}\"");
+			return GeneratedHeaderCppContents;
+		}
+
 		/// <summary>
 		/// Gets a set of source files for the given directory. Used to detect when the makefile is out of date.
 		/// </summary>
diff --git a/Engine/Source/Programs/UnrealBuildTool/Platform/Windows/VCSpecificFileAction.cs b/Engine/Source/Programs/UnrealBuildTool/Platform/Windows/VCSpecificFileAction.cs
--- a/Engine/Source/Programs/UnrealBuildTool/Platform/Windows/VCSpecificFileAction.cs
+++ b/Engine/Source/Programs/UnrealBuildTool/Platform/Windows/VCSpecificFileAction.cs
@@ -74,7 +74,6 @@ namespace UnrealBuildTool
 			Action.SourceFile = SourceFile;
 			if (SourceFile.HasExtension(".h"))
 			{
-				Action.Definitions.Add("UE_DIRECT_HEADER_COMPILE=1");
 				if (BaseAction.CompilerType.IsClang())
 				{
 					ClangWarnings.GetHeaderDisabledWarnings(Action.Arguments);
@@ -86,10 +85,7 @@ namespace UnrealBuildTool
 					IncludeFileString = SourceFile.Location.MakeRelativeTo(Unreal.EngineSourceDirectory);
 				}
 
-				List<string> GeneratedHeaderCppContents = new();
-				GeneratedHeaderCppContents.Add("// This file is automatically generated at compile-time to include a user created header file.");
-				GeneratedHeaderCppContents.Add($"#include \"{IncludeFileString.Replace('\\', '/')}\"");
-
+				List<string> GeneratedHeaderCppContents = UEBuildModuleCPP.GenerateHeaderCpp(SourceFile.Name, IncludeFileString);
 				Action.SourceFile = FileItem.GetItemByFileReference(FileReference.Combine(IntermediateDirectory, $"{SourceFile.Name}.cpp"));
 				File.WriteAllLines(Action.SourceFile.FullName, GeneratedHeaderCppContents);
 			}
diff --git a/Engine/Source/Programs/UnrealBuildTool/ToolChain/ClangSpecificFileAction.cs b/Engine/Source/Programs/UnrealBuildTool/ToolChain/ClangSpecificFileAction.cs
--- a/Engine/Source/Programs/UnrealBuildTool/ToolChain/ClangSpecificFileAction.cs
+++ b/Engine/Source/Programs/UnrealBuildTool/ToolChain/ClangSpecificFileAction.cs
@@ -73,7 +73,6 @@ namespace UnrealBuildTool
 
 			if (SourceFile.HasExtension(".h"))
 			{
-				NewRspLines.Add("-DUE_DIRECT_HEADER_COMPILE=1");
 				ClangWarnings.GetHeaderDisabledWarnings(NewRspLines);
 
 				string IncludeFileString = SourceFile.AbsolutePath;
@@ -82,10 +81,7 @@ namespace UnrealBuildTool
 					IncludeFileString = SourceFile.Location.MakeRelativeTo(Unreal.EngineSourceDirectory);
 				}
 
-				List<string> GeneratedHeaderCppContents = new();
-				GeneratedHeaderCppContents.Add("// This file is automatically generated at compile-time to include a user created header file.");
-				GeneratedHeaderCppContents.Add($"#include \"{IncludeFileString.Replace('\\', '/')}\"");
-
+				List<string> GeneratedHeaderCppContents = UEBuildModuleCPP.GenerateHeaderCpp(SourceFile.Name, IncludeFileString);
 				SourceFile = FileItem.GetItemByFileReference(FileReference.Combine(OutputDir, $"{UniqueDummyName}.cpp"));
 				File.WriteAllLines(SourceFile.FullName, GeneratedHeaderCppContents);
 			}
diff --git a/Engine/Source/Runtime/Core/Public/Misc/DirectHeaderCompiling.h b/Engine/Source/Runtime/Core/Public/Misc/DirectHeaderCompiling.h
new file mode 100644
--- /dev/null
+++ b/Engine/Source/Runtime/Core/Public/Misc/DirectHeaderCompiling.h
@@ -0,0 +1,12 @@
+// Copyright Epic Games, Inc. All Rights Reserved.
+
+#pragma once
+
+// Used to allow control flow when direct header compiling to avoid issues with headers
+// that typically expect to be included after another header
+#ifndef UE_DIRECT_HEADER_COMPILING
+	// C5105: macro expansion producing 'defined' has undefined behavior
+	// however, in all compilers we use, the behavior is as desired
+	#pragma warning(disable : 5105) 
+	#define UE_DIRECT_HEADER_COMPILING(id) defined __COMPILING_##id
+#endif
\ No newline at end of file
