Remove the Platform.h experiment as it fails if files are in modules that don't include Core and put back the one off pragma disable for the single file cases in to DirectHeaderCompiling.h, leaving the global case in place for most use cases
#rnx

diff --git a/Engine/Source/Programs/UnrealBuildTool/Configuration/UEBuildModuleCPP.cs b/Engine/Source/Programs/UnrealBuildTool/Configuration/UEBuildModuleCPP.cs
--- a/Engine/Source/Programs/UnrealBuildTool/Configuration/UEBuildModuleCPP.cs
+++ b/Engine/Source/Programs/UnrealBuildTool/Configuration/UEBuildModuleCPP.cs
@@ -2432,7 +2432,6 @@ namespace UnrealBuildTool
 			GeneratedHeaderCppContents.Add("// This file is automatically generated at compile-time to include a user created header file.");
 			GeneratedHeaderCppContents.Add("#define UE_DIRECT_HEADER_COMPILE 1");
 			GeneratedHeaderCppContents.Add($"#define __COMPILING_{HeaderName}");
-			GeneratedHeaderCppContents.Add($"#include \"HAL/Platform.h\"");
 			GeneratedHeaderCppContents.Add($"#include \"{IncludeFileString.Replace('\\', '/')}\"");
 			return GeneratedHeaderCppContents;
 		}
diff --git a/Engine/Source/Runtime/Core/Public/Misc/DirectHeaderCompiling.h b/Engine/Source/Runtime/Core/Public/Misc/DirectHeaderCompiling.h
--- a/Engine/Source/Runtime/Core/Public/Misc/DirectHeaderCompiling.h
+++ b/Engine/Source/Runtime/Core/Public/Misc/DirectHeaderCompiling.h
@@ -2,6 +2,10 @@
 
 #pragma once
 
+// C5105: macro expansion producing 'defined' has undefined behavior
+// however, in all compilers we use, the behavior is as desired
+#pragma warning(disable : 5105) 
+
 // Used to allow control flow when direct header compiling to avoid issues with headers
 // that typically expect to be included after another header
 #ifndef UE_DIRECT_HEADER_COMPILING
