Fix Empty AppName causing archving fail

diff --git a/Engine/Source/Programs/UnrealBuildTool/ProjectFiles/Xcode/XcodeProject.cs b/Engine/Source/Programs/UnrealBuildTool/ProjectFiles/Xcode/XcodeProject.cs
--- a/Engine/Source/Programs/UnrealBuildTool/ProjectFiles/Xcode/XcodeProject.cs
+++ b/Engine/Source/Programs/UnrealBuildTool/ProjectFiles/Xcode/XcodeProject.cs
@@ -2136,8 +2136,14 @@ namespace UnrealBuildTool.XcodeProjectXcconfig
 
 					ConfigXcconfig.AppendLine($"PRODUCT_NAME_ = {ProductName}");
 					ConfigXcconfig.AppendLine($"PRODUCT_NAME_build = $(PRODUCT_NAME_)");
-					string ArchivedName = ApplicationDisplayName ?? (UnrealData.UProjectFileLocation == null ? ProductName : UnrealData.UProjectFileLocation!.GetFileNameWithoutAnyExtensions());
-					ConfigXcconfig.AppendLine($"PRODUCT_NAME_install = {ArchivedName}");
+					if (String.IsNullOrEmpty(ApplicationDisplayName))
+					{
+						ConfigXcconfig.AppendLine($"PRODUCT_NAME_install = {(UnrealData.UProjectFileLocation == null ? ProductName : UnrealData.UProjectFileLocation!.GetFileNameWithoutAnyExtensions())}");
+					}
+					else
+					{
+						ConfigXcconfig.AppendLine($"PRODUCT_NAME_install = {ApplicationDisplayName}");
+					}
 
 					// this will choose the proper PRODUCT_NAME when archiving vs normal building
 					ConfigXcconfig.AppendLine("PRODUCT_NAME = $(PRODUCT_NAME_$(ACTION))");
