UnrealBuildTool: Add System.Security.Cryptography package to DynamicCompilation reference, required for .NET 8

#rnx

diff --git a/Engine/Source/Programs/UnrealBuildTool/System/DynamicCompilation.cs b/Engine/Source/Programs/UnrealBuildTool/System/DynamicCompilation.cs
--- a/Engine/Source/Programs/UnrealBuildTool/System/DynamicCompilation.cs
+++ b/Engine/Source/Programs/UnrealBuildTool/System/DynamicCompilation.cs
@@ -228,6 +228,14 @@ namespace UnrealBuildTool
 			MetadataReferences.Add(MetadataReference.CreateFromFile(Assembly.Load("Microsoft.Win32.Registry").Location));
 
 			// RNGCryptoServiceProvider, used to generate random hex bytes
+			try
+			{
+				// Does not exist for .NET6
+				MetadataReferences.Add(MetadataReference.CreateFromFile(Assembly.Load("System.Security.Cryptography").Location));
+			}
+			catch (FileNotFoundException)
+			{
+			}
 			MetadataReferences.Add(MetadataReference.CreateFromFile(Assembly.Load("System.Security.Cryptography.Algorithms").Location));
 			MetadataReferences.Add(MetadataReference.CreateFromFile(Assembly.Load("System.Security.Cryptography.Csp").Location));
 
