[UBT]
* Changed so it is possible to have extra command line parameters in combination with rsp file on vc targets
* Changed so we don't create .strip.obj file anymore but instead use UbaObjTool to create an additional obj file that carry exports etc

diff --git a/Engine/Source/Programs/UnrealBuildTool/Configuration/UEBuildBinary.cs b/Engine/Source/Programs/UnrealBuildTool/Configuration/UEBuildBinary.cs
--- a/Engine/Source/Programs/UnrealBuildTool/Configuration/UEBuildBinary.cs
+++ b/Engine/Source/Programs/UnrealBuildTool/Configuration/UEBuildBinary.cs
@@ -877,19 +877,10 @@ namespace UnrealBuildTool
 			if ((BinaryLinkEnvironment.bIsBuildingDLL || BinaryLinkEnvironment.bIsBuildingLibrary) && Target.bStripExports)
 			{
 				bStripUnusedExports = true;
-				List<FileItem> inputFiles = new();
-				foreach (FileItem item in BinaryLinkEnvironment.InputFiles)
-				{
-					if (InputObjects.Contains(item))
-					{
-						inputFiles.Add(FileItem.GetItemByFileReference(item.Location.ChangeExtension($".strip{item.Location.GetExtension()}")));
-					}
-					else
-					{
-						inputFiles.Add(item);
-					}
-				}
-				BinaryLinkEnvironment.InputFiles = inputFiles;
+
+				string Name = OutputFilePaths.First().GetFileNameWithoutExtension();
+				FileReference extraObj = FileReference.Combine(IntermediateDirectory!, $"{Name}.extra.obj");
+				BinaryLinkEnvironment.InputFiles.Add(FileItem.GetItemByFileReference(extraObj));
 			}
 
 			// Code coverage inherited from compile environment
diff --git a/Engine/Source/Programs/UnrealBuildTool/Configuration/UEBuildTarget.cs b/Engine/Source/Programs/UnrealBuildTool/Configuration/UEBuildTarget.cs
--- a/Engine/Source/Programs/UnrealBuildTool/Configuration/UEBuildTarget.cs
+++ b/Engine/Source/Programs/UnrealBuildTool/Configuration/UEBuildTarget.cs
@@ -3037,8 +3037,11 @@ namespace UnrealBuildTool
 				HashSet<FileItem> otherObjects = new(allObjects.Except(toStrip));
 				FileItem stripRsp = FileItem.GetItemByFileReference(FileReference.Combine(ProjectIntermediateDirectory!, $"{Name}.strip.rsp"));
 				List<string> arguments = new();
-				arguments.AddRange(toStrip.Select(path => $"/S:{path.FullName}").OrderBy(x => x));
-				arguments.AddRange(otherObjects.Select(path => $"/D:{path.FullName}").OrderBy(x => x));
+				arguments.AddRange(toStrip.Select(path => $"/S:{path.FullName.Replace(".obj", ".exi")}").OrderBy(x => x));
+				arguments.AddRange(otherObjects.Select(path => $"/D:{path.FullName.Replace(".obj", ".exi")}").OrderBy(x => x));
+
+				FileReference extraObj = FileReference.Combine(ProjectIntermediateDirectory!, $"{Name}.extra.obj");
+				arguments.Add($"/O:{extraObj}");
 
 				MakefileBuilder.CreateIntermediateTextFile(stripRsp, arguments);
 
@@ -3046,10 +3049,10 @@ namespace UnrealBuildTool
 
 				stripAction.PrerequisiteItems.Add(stripRsp);
 				stripAction.PrerequisiteItems.UnionWith(allObjects);
-				stripAction.ProducedItems.UnionWith(toStrip.Select(x => FileItem.GetItemByFileReference(x.Location.ChangeExtension($".strip{x.Location.GetExtension()}"))));
+				stripAction.ProducedItems.Add(FileItem.GetItemByFileReference(extraObj));
 				stripAction.CommandPath = commandPath;
 				stripAction.CommandArguments = "@" + stripRsp.FullName;
-				stripAction.CommandDescription = "Strip Objects";
+				stripAction.CommandDescription = "GenerateExports";
 				stripAction.WorkingDirectory = Unreal.EngineSourceDirectory;
 				stripAction.StatusDescription = Name;
 				stripAction.bCanExecuteRemotely = false;
@@ -3057,7 +3060,11 @@ namespace UnrealBuildTool
 			}
 
 			HashSet<FileItem> pchObjects = new(allObjects.Where(x => x.Name.Contains("SharedPCH", StringComparison.Ordinal)));
-			CreateStripAction($"{AppName}-SharedPCH", pchObjects);
+			if (pchObjects.Count > 0)
+			{
+				CreateStripAction($"{AppName}-SharedPCH", pchObjects);
+			}
+
 			foreach (UEBuildBinary Binary in Binaries.Where(x => x.bStripUnusedExports))
 			{
 				CreateStripAction(Binary.OutputFilePaths.First().GetFileNameWithoutExtension(), Binary.InputObjects.Where(x => !pchObjects.Contains(x)));
diff --git a/Engine/Source/Programs/UnrealBuildTool/Platform/Windows/VCCompileAction.cs b/Engine/Source/Programs/UnrealBuildTool/Platform/Windows/VCCompileAction.cs
--- a/Engine/Source/Programs/UnrealBuildTool/Platform/Windows/VCCompileAction.cs
+++ b/Engine/Source/Programs/UnrealBuildTool/Platform/Windows/VCCompileAction.cs
@@ -588,7 +588,7 @@ namespace UnrealBuildTool
 				{
 					ResponseFileString = ResponseFile.FullName;
 				}
-				return String.Format("@{0}", Utils.MakePathSafeToUseWithCommandLine(ResponseFileString));
+				return String.Format("@{0} {1}", Utils.MakePathSafeToUseWithCommandLine(ResponseFileString), String.Join(" ", Arguments));
 			}
 		}
 
diff --git a/Engine/Source/Programs/UnrealBuildTool/Platform/Windows/VCToolChain.cs b/Engine/Source/Programs/UnrealBuildTool/Platform/Windows/VCToolChain.cs
--- a/Engine/Source/Programs/UnrealBuildTool/Platform/Windows/VCToolChain.cs
+++ b/Engine/Source/Programs/UnrealBuildTool/Platform/Windows/VCToolChain.cs
@@ -1897,6 +1897,7 @@ namespace UnrealBuildTool
 					if (!ProjectFileGenerator.bGenerateProjectFiles)
 					{
 						CompileDepsAction.WriteResponseFile(Graph, Logger);
+						CompileDepsAction.Arguments.Clear();
 					}
 
 					CompileAction.ActionType = ActionType.CompileModuleInterface;
@@ -1943,6 +1944,16 @@ namespace UnrealBuildTool
 				if (!ProjectFileGenerator.bGenerateProjectFiles)
 				{
 					CompileAction.WriteResponseFile(Graph, Logger);
+					CompileAction.Arguments.Clear();
+				}
+
+				// Must be added after response file is created just to make sure it ends up on the command line and not in the response file
+				if (Target.bMergeModules)
+				{
+					// EXTRACTEXPORTS can only be interpreted by UBA.. so this action won't build outside uba
+					CompileAction.Arguments.Add("/EXTRACTEXPORTS");
+					FileItem SymFile = FileItem.GetItemByFileReference(FileReference.Combine(OutputDir, FileName + ".exi"));
+					CompileAction.AdditionalProducedItems.Add(SymFile);
 				}
 
 				CompileAction.bIsAnalyzing = Target.StaticAnalyzer != StaticAnalyzer.None && !CompileEnvironment.bDisableStaticAnalysis && !(Target.WindowsPlatform.Compiler.IsClang() && CompileEnvironment.PrecompiledHeaderAction == PrecompiledHeaderAction.Create);
diff --git a/Engine/Source/Programs/UnrealBuildTool/ToolChain/ISPCToolChain.cs b/Engine/Source/Programs/UnrealBuildTool/ToolChain/ISPCToolChain.cs
--- a/Engine/Source/Programs/UnrealBuildTool/ToolChain/ISPCToolChain.cs
+++ b/Engine/Source/Programs/UnrealBuildTool/ToolChain/ISPCToolChain.cs
@@ -18,11 +18,13 @@ namespace UnrealBuildTool
 		}
 
 		protected FileReference? ProjectFile = null;
+		protected bool bMergeModules = false;
 
 		public override void SetUpGlobalEnvironment(ReadOnlyTargetRules Target)
 		{
 			base.SetUpGlobalEnvironment(Target);
 			ProjectFile = Target.ProjectFile;
+			bMergeModules = Target.bMergeModules;
 		}
 
 		/// <summary>
@@ -675,6 +677,8 @@ namespace UnrealBuildTool
 
 				List<FileItem> CompiledISPCObjFiles = new List<FileItem>();
 
+				string FileName = Path.GetFileName(ISPCFile.AbsolutePath);
+
 				string CompiledISPCObjFileSuffix = bByteCodeOutput ? ".bc" : GetISPCObjectFileSuffix(CompileEnvironment.Platform);
 				foreach (string Target in CompileTargets)
 				{
@@ -693,7 +697,7 @@ namespace UnrealBuildTool
 						CompiledISPCObjFile = FileItem.GetItemByFileReference(
 						FileReference.Combine(
 							OutputDir,
-							Path.GetFileName(ISPCFile.AbsolutePath) + "_" + ObjTarget + CompiledISPCObjFileSuffix
+							FileName + "_" + ObjTarget + CompiledISPCObjFileSuffix
 							)
 						);
 					}
@@ -702,7 +706,7 @@ namespace UnrealBuildTool
 						CompiledISPCObjFile = FileItem.GetItemByFileReference(
 						FileReference.Combine(
 							OutputDir,
-							Path.GetFileName(ISPCFile.AbsolutePath) + CompiledISPCObjFileSuffix
+							FileName + CompiledISPCObjFileSuffix
 							)
 						);
 					}
@@ -715,7 +719,7 @@ namespace UnrealBuildTool
 				FileItem CompiledISPCObjFileNoISA = FileItem.GetItemByFileReference(
 					FileReference.Combine(
 						OutputDir,
-						Path.GetFileName(ISPCFile.AbsolutePath) + CompiledISPCObjFileSuffix
+						FileName + CompiledISPCObjFileSuffix
 						)
 					);
 
@@ -738,7 +742,7 @@ namespace UnrealBuildTool
 				Arguments.AddRange(GlobalArguments);
 
 				// Consume the included header dependency list
-				FileItem DependencyListFile = FileItem.GetItemByFileReference(FileReference.Combine(OutputDir, Path.GetFileName(ISPCFile.AbsolutePath) + ".txt"));
+				FileItem DependencyListFile = FileItem.GetItemByFileReference(FileReference.Combine(OutputDir, FileName + ".txt"));
 				CompileAction.DependencyListFile = DependencyListFile;
 				CompileAction.PrerequisiteItems.Add(DependencyListFile);
 
@@ -746,7 +750,18 @@ namespace UnrealBuildTool
 
 				FileReference ResponseFileName = GetResponseFileName(CompileEnvironment, CompiledISPCObjFileNoISA);
 				FileItem ResponseFileItem = Graph.CreateIntermediateTextFile(ResponseFileName, Arguments.Select(x => Utils.ExpandVariables(x)));
-				CompileAction.CommandArguments = $"@\"{ResponseFileName}\"";
+
+				string AdditionalArguments = "";
+				// Must be added after response file is created just to make sure it ends up on the command line and not in the response file
+				if (bMergeModules)
+				{
+					// EXTRACTEXPORTS can only be interpreted by UBA.. so this action won't build outside uba
+					AdditionalArguments = " /EXTRACTEXPORTS";
+					FileItem SymFile = FileItem.GetItemByFileReference(FileReference.Combine(OutputDir, FileName + ".exi"));
+					CompileAction.ProducedItems.Add(SymFile);
+				}
+
+				CompileAction.CommandArguments = $"@\"{ResponseFileName}\"{AdditionalArguments}";
 				CompileAction.PrerequisiteItems.Add(ResponseFileItem);
 
 				// Add the source file and its included files to the prerequisite item list.
