LLT job - selecting a list of platforms must select appropriate tests

diff --git a/Engine/Source/Programs/UnrealBuildTool/Configuration/TestModuleRules.cs b/Engine/Source/Programs/UnrealBuildTool/Configuration/TestModuleRules.cs
--- a/Engine/Source/Programs/UnrealBuildTool/Configuration/TestModuleRules.cs
+++ b/Engine/Source/Programs/UnrealBuildTool/Configuration/TestModuleRules.cs
@@ -334,6 +334,23 @@ namespace UnrealBuildTool
 				}
 
 				AppendOrUpdateRunAllTestsNode(Root, "DeployAndTest", ValidPlatform.ToString(), ExpandArguments);
+
+				if (IsRestrictedPlatformName)
+				{
+					// Create a General.xml file and add a TestPlatform* option
+					string RestrictedPlatformFolderPath = Path.GetDirectoryName(GeneratedPropertiesPlatformFile)!;
+					string RestrictedPlatformGeneral = Path.Combine(RestrictedPlatformFolderPath, "General.xml");
+					if (!System.IO.File.Exists(RestrictedPlatformGeneral))
+					{
+						using (FileStream FileStream = System.IO.File.Create(RestrictedPlatformGeneral))
+						{
+							Log.LogInformation("Saving general metadata to {File}", RestrictedPlatformGeneral);
+							XDocument GeneralProps = new XDocument(new XElement(BuildGraphNamespace + "BuildGraph", new XAttribute(XNamespace.Xmlns + "xsi", SchemaInstance), new XAttribute(SchemaInstance + "schemaLocation", SchemaLocation)));
+							InsertOrUpdateTestOption(GeneralProps.Root, $"TestPlatform{ValidPlatform}", $"Run tests on {ValidPlatform}", false.ToString());
+							GeneralProps.Save(FileStream);
+						}
+					}
+				}
 			}
 
 			foreach (KeyValuePair<string, XDocument> KVP in SaveAtEnd)
