Changed so rsp files for clang targets on windows hosts uses backslash instead of forward slash. This makes it possible to use uba cache on these targets

diff --git a/Engine/Source/Programs/UnrealBuildTool/ToolChain/ISPCToolChain.cs b/Engine/Source/Programs/UnrealBuildTool/ToolChain/ISPCToolChain.cs
--- a/Engine/Source/Programs/UnrealBuildTool/ToolChain/ISPCToolChain.cs
+++ b/Engine/Source/Programs/UnrealBuildTool/ToolChain/ISPCToolChain.cs
@@ -375,13 +375,21 @@ namespace UnrealBuildTool
 		/// <returns>Normalized path as a string</returns>
 		protected virtual string NormalizeCommandLinePath(FileSystemReference Reference)
 		{
+			string path = Reference.FullName;
 			// Try to use a relative path to shorten command line length.
 			if (Reference.IsUnderDirectory(Unreal.RootDirectory))
 			{
-				return Reference.MakeRelativeTo(Unreal.EngineSourceDirectory).Replace("\\", "/");
+				path = Reference.MakeRelativeTo(Unreal.EngineSourceDirectory);
 			}
-
-			return Reference.FullName.Replace("\\", "/");
+			if (Path.DirectorySeparatorChar == '/')
+			{
+				path = path.Replace("\\", "/");
+			}
+			else
+			{
+				path = path.Replace("\\", "\\\\");
+			}
+			return path;
 		}
 
 		/// <summary>
