[UBT]
* Fixed so executable also run strip action to make sure exports that must exist will exist
* Fixed so clang-based platforms also work with merged modules.
* Removed pch stripping action now when we don't create stripped obj files anymore

diff --git a/Engine/Source/Programs/UnrealBuildTool/Configuration/UEBuildBinary.cs b/Engine/Source/Programs/UnrealBuildTool/Configuration/UEBuildBinary.cs
--- a/Engine/Source/Programs/UnrealBuildTool/Configuration/UEBuildBinary.cs
+++ b/Engine/Source/Programs/UnrealBuildTool/Configuration/UEBuildBinary.cs
@@ -874,7 +874,7 @@ namespace UnrealBuildTool
 			BinaryLinkEnvironment.bIsBuildingLibrary = IsBuildingLibrary(Type);
 
 			// Setup object export stripping
-			if ((BinaryLinkEnvironment.bIsBuildingDLL || BinaryLinkEnvironment.bIsBuildingLibrary) && Target.bStripExports)
+			if (Target.bStripExports)
 			{
 				bStripUnusedExports = true;
 
diff --git a/Engine/Source/Programs/UnrealBuildTool/Configuration/UEBuildTarget.cs b/Engine/Source/Programs/UnrealBuildTool/Configuration/UEBuildTarget.cs
--- a/Engine/Source/Programs/UnrealBuildTool/Configuration/UEBuildTarget.cs
+++ b/Engine/Source/Programs/UnrealBuildTool/Configuration/UEBuildTarget.cs
@@ -3032,15 +3032,15 @@ namespace UnrealBuildTool
 
 			HashSet<FileItem> allObjects = new(Binaries.SelectMany(x => x.InputObjects));
 
-			void CreateStripAction(string Name, IEnumerable<FileItem> toStrip)
+			void CreateStripAction(string Name, IEnumerable<FileItem> toStrip, DirectoryReference IntermediateDirectory)
 			{
 				HashSet<FileItem> otherObjects = new(allObjects.Except(toStrip));
-				FileItem stripRsp = FileItem.GetItemByFileReference(FileReference.Combine(ProjectIntermediateDirectory!, $"{Name}.strip.rsp"));
+				FileItem stripRsp = FileItem.GetItemByFileReference(FileReference.Combine(IntermediateDirectory, $"{Name}.strip.rsp"));
 				List<string> arguments = new();
-				arguments.AddRange(toStrip.Select(path => $"/S:{path.FullName.Replace(".obj", ".exi")}").OrderBy(x => x));
-				arguments.AddRange(otherObjects.Select(path => $"/D:{path.FullName.Replace(".obj", ".exi")}").OrderBy(x => x));
+				arguments.AddRange(toStrip.Select(path => $"/S:{path.Location.ChangeExtension(".exi")}").OrderBy(x => x));
+				arguments.AddRange(otherObjects.Select(path => $"/D:{path.Location.ChangeExtension(".exi")}").OrderBy(x => x));
 
-				FileReference extraObj = FileReference.Combine(ProjectIntermediateDirectory!, $"{Name}.extra.obj");
+				FileReference extraObj = FileReference.Combine(IntermediateDirectory, $"{Name}.extra.obj");
 				arguments.Add($"/O:{extraObj}");
 
 				MakefileBuilder.CreateIntermediateTextFile(stripRsp, arguments);
@@ -3059,15 +3059,9 @@ namespace UnrealBuildTool
 				Makefile.OutputItems.AddRange(stripAction.ProducedItems);
 			}
 
-			HashSet<FileItem> pchObjects = new(allObjects.Where(x => x.Name.Contains("SharedPCH", StringComparison.Ordinal)));
-			if (pchObjects.Count > 0)
-			{
-				CreateStripAction($"{AppName}-SharedPCH", pchObjects);
-			}
-
 			foreach (UEBuildBinary Binary in Binaries.Where(x => x.bStripUnusedExports))
 			{
-				CreateStripAction(Binary.OutputFilePaths.First().GetFileNameWithoutExtension(), Binary.InputObjects.Where(x => !pchObjects.Contains(x)));
+				CreateStripAction(Binary.OutputFilePaths.First().GetFileNameWithoutExtension(), Binary.InputObjects, Binary.IntermediateDirectory);
 			}
 		}
 
diff --git a/Engine/Source/Programs/UnrealBuildTool/ToolChain/ClangToolChain.cs b/Engine/Source/Programs/UnrealBuildTool/ToolChain/ClangToolChain.cs
--- a/Engine/Source/Programs/UnrealBuildTool/ToolChain/ClangToolChain.cs
+++ b/Engine/Source/Programs/UnrealBuildTool/ToolChain/ClangToolChain.cs
@@ -1196,7 +1196,17 @@ namespace UnrealBuildTool
 
 			// Adds the response file to the compiler input.
 			FileItem CompilerResponseFileItem = Graph.CreateIntermediateTextFile(ResponseFileName, ResponseFileContents);
-			CompileAction.CommandArguments = GetResponseFileArgument(CompilerResponseFileItem);
+			string CommandArguments = GetResponseFileArgument(CompilerResponseFileItem);
+
+			if (bMergeModules)
+			{
+				// EXTRACTEXPORTS can only be interpreted by UBA.. so this action won't build outside uba
+				CommandArguments += " /EXTRACTEXPORTS";
+				FileItem SymFile = FileItem.GetItemByFileReference(FileReference.Combine(OutputDir, Path.GetFileName(SourceFile.AbsolutePath) + ".exi"));
+				CompileAction.ProducedItems.Add(SymFile);
+			}
+
+			CompileAction.CommandArguments = CommandArguments;
 			CompileAction.PrerequisiteItems.Add(CompilerResponseFileItem);
 
 			CompileAction.WorkingDirectory = Unreal.EngineSourceDirectory;
