UnrealBuildTool: Add a `bGenerateProgramProject` flag to TargetRules to force a target to be treated as a program during project file generation.

diff --git a/Engine/Source/Programs/UnrealBuildTool/Configuration/TargetRules.cs b/Engine/Source/Programs/UnrealBuildTool/Configuration/TargetRules.cs
--- a/Engine/Source/Programs/UnrealBuildTool/Configuration/TargetRules.cs
+++ b/Engine/Source/Programs/UnrealBuildTool/Configuration/TargetRules.cs
@@ -811,6 +811,11 @@ namespace UnrealBuildTool
 		/// </summary>
 		public string SolutionDirectory { get; set; } = String.Empty;
 
+		/// <summary>
+		/// Force a Target to be treated as a Program for the purposes of project file generation.
+		/// </summary>
+		public bool bGenerateProgramProject { get; set; }
+
 		/// <summary>
 		/// Whether the target should be included in the default solution build configuration
 		/// Setting this to false will skip building when running in the IDE
diff --git a/Engine/Source/Programs/UnrealBuildTool/ProjectFiles/ProjectFileGenerator.cs b/Engine/Source/Programs/UnrealBuildTool/ProjectFiles/ProjectFileGenerator.cs
--- a/Engine/Source/Programs/UnrealBuildTool/ProjectFiles/ProjectFileGenerator.cs
+++ b/Engine/Source/Programs/UnrealBuildTool/ProjectFiles/ProjectFileGenerator.cs
@@ -2811,7 +2811,7 @@ namespace UnrealBuildTool
 
 					DirectoryReference? GameFolder = null;
 					string ProjectFileNameBase;
-					if (TargetRulesObject.Type == TargetType.Program)
+					if (TargetRulesObject.Type == TargetType.Program || TargetRulesObject.bGenerateProgramProject)
 					{
 						if (!ForceProgramInProject && !bIncludeEnginePrograms && IsEngineTarget)
 						{
