[UBT]
* Minor cosmetic fix to uba trace file

diff --git a/Engine/Source/Programs/UnrealBuildTool/Executors/UnrealBuildAccelerator/UBAExecutor.cs b/Engine/Source/Programs/UnrealBuildTool/Executors/UnrealBuildAccelerator/UBAExecutor.cs
--- a/Engine/Source/Programs/UnrealBuildTool/Executors/UnrealBuildAccelerator/UBAExecutor.cs
+++ b/Engine/Source/Programs/UnrealBuildTool/Executors/UnrealBuildAccelerator/UBAExecutor.cs
@@ -512,6 +512,10 @@ namespace UnrealBuildTool
 
 		ProcessStartInfo GetActionStartInfo(LinkedAction action, out FileItem? pchItem)
 		{
+			string description = action.StatusDescription;
+			if (!String.IsNullOrEmpty(action.CommandDescription))
+				description = $"{action.StatusDescription} ({action.CommandDescription})";
+
 			ProcessStartInfo startInfo = new()
 			{
 				Application = action.CommandPath.FullName,
@@ -520,7 +524,7 @@ namespace UnrealBuildTool
 				Priority = ProcessPriority,
 				OutputStatsThresholdMs = (uint)UBAConfig.OutputStatsThresholdMs,
 				UserData = action,
-				Description = $"{action.StatusDescription} ({action.CommandDescription})",
+				Description = description,
 				Configuration = action.bIsGCCCompiler ? EpicGames.UBA.ProcessStartInfo.CommonProcessConfigs.CompileClang : EpicGames.UBA.ProcessStartInfo.CommonProcessConfigs.CompileMsvc,
 				LogFile = UBAConfig.bLogEnabled ? action.Inner.ProducedItems.First().Location.GetFileName() : null,
 			};
