[UBT]
* Added artifact mode and rootpaths to CompileByteCode action. also enabled it to run remotely in uba

diff --git a/Engine/Source/Programs/UnrealBuildTool/ToolChain/ISPCToolChain.cs b/Engine/Source/Programs/UnrealBuildTool/ToolChain/ISPCToolChain.cs
--- a/Engine/Source/Programs/UnrealBuildTool/ToolChain/ISPCToolChain.cs
+++ b/Engine/Source/Programs/UnrealBuildTool/ToolChain/ISPCToolChain.cs
@@ -814,8 +814,11 @@ namespace UnrealBuildTool
 							PostCompileAction.CommandPath = new FileReference(ByteCodeCompilerPath);
 							PostCompileAction.StatusDescription = Path.GetFileName(ISPCFile.AbsolutePath);
 
+							PostCompileAction.RootPaths.AddRange(RootPaths);
+							PostCompileAction.ArtifactMode = ArtifactMode.Enabled;
+
 							// Disable remote execution to workaround mismatched case on XGE
-							PostCompileAction.bCanExecuteRemotely = false;
+							PostCompileAction.bCanExecuteRemotelyWithXGE = false;
 
 							FinalObjectFiles.Add(FinalCompiledISPCObjFile);
 							Logger.LogDebug("   ISPC Compiling bytecode {StatusDescription}: \"{CommandPath}\" {CommandArguments} {ProducedItems}", PostCompileAction.StatusDescription, PostCompileAction.CommandPath, PostCompileAction.CommandArguments, PostCompileAction.ProducedItems);
