UnrealBuildTool: Add package reference to System.CodeDom needed to compile with .NET 8

#rnx

diff --git a/Engine/Source/Programs/Shared/EpicGames.Build/EpicGames.Build.csproj b/Engine/Source/Programs/Shared/EpicGames.Build/EpicGames.Build.csproj
--- a/Engine/Source/Programs/Shared/EpicGames.Build/EpicGames.Build.csproj
+++ b/Engine/Source/Programs/Shared/EpicGames.Build/EpicGames.Build.csproj
@@ -47,6 +47,7 @@
     <PackageReference Include="Microsoft.CSharp" Version="4.7.0" />
     <PackageReference Include="Microsoft.Extensions.FileSystemGlobbing" Version="6.0.0" />
     <PackageReference Include="Microsoft.Win32.Registry" Version="5.0.0" />
+    <PackageReference Include="System.CodeDom" Version="8.0.0" />
     <PackageReference Include="System.Data.DataSetExtensions" Version="4.5.0" />
     <PackageReference Include="System.Security.Permissions" Version="4.7.0" />
   </ItemGroup>
diff --git a/Engine/Source/Programs/UnrealBuildTool/System/DynamicCompilation.cs b/Engine/Source/Programs/UnrealBuildTool/System/DynamicCompilation.cs
--- a/Engine/Source/Programs/UnrealBuildTool/System/DynamicCompilation.cs
+++ b/Engine/Source/Programs/UnrealBuildTool/System/DynamicCompilation.cs
@@ -208,6 +208,7 @@ namespace UnrealBuildTool
 
 			MetadataReferences.Add(MetadataReference.CreateFromFile(typeof(object).Assembly.Location));
 			MetadataReferences.Add(MetadataReference.CreateFromFile(Assembly.Load("System.Runtime").Location));
+			MetadataReferences.Add(MetadataReference.CreateFromFile(Assembly.Load("System.CodeDom").Location));
 			MetadataReferences.Add(MetadataReference.CreateFromFile(Assembly.Load("System.Collections").Location));
 			MetadataReferences.Add(MetadataReference.CreateFromFile(Assembly.Load("System.IO").Location));
 			MetadataReferences.Add(MetadataReference.CreateFromFile(Assembly.Load("System.IO.FileSystem").Location));
diff --git a/Engine/Source/Programs/UnrealBuildTool/UnrealBuildTool.csproj b/Engine/Source/Programs/UnrealBuildTool/UnrealBuildTool.csproj
--- a/Engine/Source/Programs/UnrealBuildTool/UnrealBuildTool.csproj
+++ b/Engine/Source/Programs/UnrealBuildTool/UnrealBuildTool.csproj
@@ -127,6 +127,7 @@
     <PackageReference Include="Microsoft.CodeAnalysis.CSharp" Version="4.2.0" />
     <PackageReference Include="Microsoft.Win32.Registry" Version="5.0.0" />
     <PackageReference Include="OpenTracing" Version="0.12.1" />
+    <PackageReference Include="System.CodeDom" Version="8.0.0" />
     <PackageReference Include="System.Management" Version="4.7.0" />
     <PackageReference Include="System.Security.Cryptography.Csp" Version="4.3.0" />
     <PackageReference Include="System.ServiceProcess.ServiceController" Version="4.7.0" />
