Replace -Filter command line argument that did nothing with -Include and -Exclude arguments so that users can pass -Exclude only to include everything except certain files.
Add -OutputFilename and -OutputDir arguments.

diff --git a/Engine/Source/Programs/UnrealBuildTool/Modes/GenerateClangDatabase.cs b/Engine/Source/Programs/UnrealBuildTool/Modes/GenerateClangDatabase.cs
--- a/Engine/Source/Programs/UnrealBuildTool/Modes/GenerateClangDatabase.cs
+++ b/Engine/Source/Programs/UnrealBuildTool/Modes/GenerateClangDatabase.cs
@@ -2,6 +2,7 @@
 
 using System;
 using System.Collections.Generic;
+using System.IO;
 using System.Linq;
 using System.Text;
 using System.Threading.Tasks;
@@ -18,10 +19,18 @@ namespace UnrealBuildTool
 	class GenerateClangDatabase : ToolMode
 	{
 		/// <summary>
-		/// Set of filters for files to include in the database. Relative to the root directory, or to the project file.
+		/// Optional set of source files to include in the compile database. If this is empty, all files will be included by default and -Exclude can be used to exclude some.
+		/// Relative to the root directory, or to the project file.
 		/// </summary>
-		[CommandLine("-Filter=")]
-		List<string> FilterRules = new List<string>();
+		[CommandLine("-Include=")]
+		List<string> IncludeRules = new List<string>();
+
+		/// <summary>
+		/// Optional set of source files to exclude from the compile database. 
+		/// Relative to the root directory, or to the project file.
+		/// </summary>
+		[CommandLine("-Exclude=")]
+		List<string> ExcludeRules = new List<string>();
 
 		/// <summary>
 		/// Execute any actions which result in code generation (eg. ISPC compilation)
@@ -30,6 +39,18 @@ namespace UnrealBuildTool
 		[CommandLine("-NoExecCodeGenActions", Value = "false")]
 		public bool bExecCodeGenActions = true;
 
+		/// <summary>
+		/// Optionally override the output filename for handling multiple targets
+		/// </summary>
+		[CommandLine("-OutputFilename=")]
+		public string OutputFilename = "compile_commands.json";
+
+		/// <summary>
+		/// Optionally overrite the output directory for the compile_commands file.
+		/// </summary>
+		[CommandLine("-OutputDir=")]
+		public string OutputDir = Unreal.RootDirectory.ToString();
+
 		/// <summary>
 		/// Execute the command
 		/// </summary>
@@ -45,15 +66,16 @@ namespace UnrealBuildTool
 			XmlConfig.ApplyTo(BuildConfiguration);
 			Arguments.ApplyTo(BuildConfiguration);
 
-			// Parse the filter argument
-			FileFilter? FileFilter = null;
-			if (FilterRules.Count > 0)
+			// If we have specific rules on what to include, exclude anything not covered by such a rule.
+			// Otherwise, include everything not covered by explicit Exclude rules
+			FileFilter FileFilter = new FileFilter(IncludeRules.Count == 0 ? FileFilterType.Include : FileFilterType.Exclude);
+			foreach (string Include in IncludeRules)
 			{
-				FileFilter = new FileFilter(FileFilterType.Exclude);
-				foreach (string FilterRule in FilterRules)
-				{
-					FileFilter.AddRules(FilterRule.Split(';'));
-				}
+				FileFilter.Include(Include.Split(';'));
+			}	
+			foreach (string Exclude in ExcludeRules)
+			{
+				FileFilter.Exclude(Exclude.Split(';'));
 			}
 
 			// Force C++ modules to always include their generated code directories
@@ -121,6 +143,12 @@ namespace UnrealBuildTool
 							{
 								continue;
 							}
+
+							if (!FileFilter.Matches(SourceFile.Location.FullName))
+							{
+								continue;
+							}
+
 							// Create the command
 							StringBuilder CommandBuilder = new StringBuilder();
 							string CommandPath = Action.CommandPath.FullName.Contains(' ') ? Utils.MakePathSafeToUseWithCommandLine(Action.CommandPath) : Action.CommandPath.FullName;
@@ -139,8 +167,8 @@ namespace UnrealBuildTool
 				Logger.LogInformation("Writing database...");
 
 				// Write the compile database
-				DirectoryReference DatabaseDirectory = Arguments.GetDirectoryReferenceOrDefault("-OutputDir=", Unreal.RootDirectory);
-				FileReference DatabaseFile = FileReference.Combine(DatabaseDirectory, "compile_commands.json");
+				DirectoryReference DatabaseDirectory = new DirectoryReference(OutputDir);
+				FileReference DatabaseFile = FileReference.Combine(DatabaseDirectory, OutputFilename);
 				using (JsonWriter Writer = new JsonWriter(DatabaseFile))
 				{
 					Writer.WriteArrayStart();
