Fix for occasional iOS deployment fails

diff --git a/Engine/Source/Programs/UnrealBuildTool/ProjectFiles/Xcode/XcodeProject.Scheme.cs b/Engine/Source/Programs/UnrealBuildTool/ProjectFiles/Xcode/XcodeProject.Scheme.cs
--- a/Engine/Source/Programs/UnrealBuildTool/ProjectFiles/Xcode/XcodeProject.Scheme.cs
+++ b/Engine/Source/Programs/UnrealBuildTool/ProjectFiles/Xcode/XcodeProject.Scheme.cs
@@ -58,7 +58,7 @@ namespace UnrealBuildTool.XcodeProjectXcconfig
 					}
 				}
 
-				string[] ScriptLines = new string[]
+				IEnumerable<string> ScriptLines = new string[]
 				{
 					"# if a UBTGenerated plist doesn&apos;t exist, create it now for Xcode&apos;s dependency needs",
 					"if [ ! -f &quot;${PROJECT_DIR}/${INFOPLIST_FILE}&quot; ]",
@@ -80,6 +80,21 @@ namespace UnrealBuildTool.XcodeProjectXcconfig
 					"fi",
 					"",
 				};
+
+				if (!XcodeProjectFileGenerator.bGenerateRunOnlyProject)
+				{
+					ScriptLines = ScriptLines.Concat(new string[]
+					{
+						"# Sometimes Xcode won't redo code sign iOS.app, leads to deployment fails",
+						"# Force Xcode to redo code sign during each build",
+						"if [ &quot;${UE_PLATFORM_NAME}&quot; != &quot;Mac&quot; ]",
+						"then",
+						"  rm -f &quot;${UE_BINARIES_DIR}/${PRODUCT_NAME}.app/Info.plist&quot;",
+						"fi",
+						"",
+					});
+				}
+
 				// convert special characters to be xml happy
 				IEnumerable<string> Lines = ScriptLines.Concat(UnrealData.ExtraPreBuildScriptLines);
 				Lines = Lines.Select(x => x.Replace("\"", "&quot;").Replace("&&", "&amp;&amp;"));
