UnrealBuildTool: Ensure invalidation for ispc and clang objects when uba compression is toggled

#rnx

diff --git a/Engine/Source/Programs/UnrealBuildTool/Platform/Android/AndroidToolChain.cs b/Engine/Source/Programs/UnrealBuildTool/Platform/Android/AndroidToolChain.cs
--- a/Engine/Source/Programs/UnrealBuildTool/Platform/Android/AndroidToolChain.cs
+++ b/Engine/Source/Programs/UnrealBuildTool/Platform/Android/AndroidToolChain.cs
@@ -1093,6 +1093,10 @@ namespace UnrealBuildTool
 				CompileAction.CommandPath = new FileReference(GetISPCHostCompilerPath(BuildHostPlatform.Current.Platform));
 				CompileAction.StatusDescription = Path.GetFileName(ISPCFile.AbsolutePath);
 				CompileAction.CommandVersion = GetISPCHostCompilerVersion(BuildHostPlatform.Current.Platform).ToString();
+				if (bAllowUbaCompression)
+				{
+					CompileAction.CommandVersion = $"{CompileAction.CommandVersion} Compressed";
+				}
 
 				// Disable remote execution to workaround mismatched case on XGE
 				CompileAction.bCanExecuteRemotely = false;
diff --git a/Engine/Source/Programs/UnrealBuildTool/ToolChain/ClangToolChain.cs b/Engine/Source/Programs/UnrealBuildTool/ToolChain/ClangToolChain.cs
--- a/Engine/Source/Programs/UnrealBuildTool/ToolChain/ClangToolChain.cs
+++ b/Engine/Source/Programs/UnrealBuildTool/ToolChain/ClangToolChain.cs
@@ -1212,6 +1212,10 @@ namespace UnrealBuildTool
 			CompileAction.WorkingDirectory = Unreal.EngineSourceDirectory;
 			CompileAction.CommandPath = Info.Clang;
 			CompileAction.CommandVersion = Info.ClangVersionString;
+			if (bAllowUbaCompression)
+			{
+				CompileAction.CommandVersion = $"{CompileAction.CommandVersion} Compressed";
+			}
 			CompileAction.CommandDescription = IsPreprocessing(CompileEnvironment) ? "Preprocess" : IsAnalyzing(CompileEnvironment) ? "Analyze" : "Compile";
 			UnrealArchitectureConfig ArchConfig = UnrealArchitectureConfig.ForPlatform(CompileEnvironment.Platform);
 			if (ArchConfig.Mode != UnrealArchitectureMode.SingleArchitecture)
diff --git a/Engine/Source/Programs/UnrealBuildTool/ToolChain/ISPCToolChain.cs b/Engine/Source/Programs/UnrealBuildTool/ToolChain/ISPCToolChain.cs
--- a/Engine/Source/Programs/UnrealBuildTool/ToolChain/ISPCToolChain.cs
+++ b/Engine/Source/Programs/UnrealBuildTool/ToolChain/ISPCToolChain.cs
@@ -19,12 +19,14 @@ namespace UnrealBuildTool
 
 		protected FileReference? ProjectFile = null;
 		protected bool bMergeModules = false;
+		protected bool bAllowUbaCompression = false;
 
 		public override void SetUpGlobalEnvironment(ReadOnlyTargetRules Target)
 		{
 			base.SetUpGlobalEnvironment(Target);
 			ProjectFile = Target.ProjectFile;
 			bMergeModules = Target.bMergeModules;
+			bAllowUbaCompression = Target.bAllowUbaCompression;
 		}
 
 		/// <summary>
@@ -494,6 +496,7 @@ namespace UnrealBuildTool
 				CompileAction.WorkingDirectory = Unreal.EngineSourceDirectory;
 				CompileAction.CommandPath = new FileReference(GetISPCHostCompilerPath(BuildHostPlatform.Current.Platform));
 				CompileAction.StatusDescription = Path.GetFileName(ISPCFile.AbsolutePath);
+				CompileAction.CommandVersion = GetISPCHostCompilerVersion(BuildHostPlatform.Current.Platform).ToString();
 				CompileAction.ArtifactMode = ArtifactMode.Enabled;
 
 				CompileAction.bCanExecuteRemotely = true;
@@ -667,6 +670,11 @@ namespace UnrealBuildTool
 				CompileAction.WorkingDirectory = Unreal.EngineSourceDirectory;
 				CompileAction.CommandPath = new FileReference(GetISPCHostCompilerPath(BuildHostPlatform.Current.Platform));
 				CompileAction.StatusDescription = Path.GetFileName(ISPCFile.AbsolutePath);
+				CompileAction.CommandVersion = GetISPCHostCompilerVersion(BuildHostPlatform.Current.Platform).ToString();
+				if (bAllowUbaCompression)
+				{
+					CompileAction.CommandVersion = $"{CompileAction.CommandVersion} Compressed";
+				}
 
 				CompileAction.bCanExecuteRemotely = true;
 
@@ -813,6 +821,11 @@ namespace UnrealBuildTool
 							PostCompileAction.WorkingDirectory = Unreal.EngineSourceDirectory;
 							PostCompileAction.CommandPath = new FileReference(ByteCodeCompilerPath);
 							PostCompileAction.StatusDescription = Path.GetFileName(ISPCFile.AbsolutePath);
+							CompileAction.CommandVersion = GetISPCHostCompilerVersion(BuildHostPlatform.Current.Platform).ToString();
+							if (bAllowUbaCompression)
+							{
+								CompileAction.CommandVersion = $"{CompileAction.CommandVersion} Compressed";
+							}
 
 							PostCompileAction.RootPaths.AddRange(RootPaths);
 							PostCompileAction.ArtifactMode = ArtifactMode.Enabled;
