UnrealBuildTool: Fix passing duplicate arguments to -SingleFile actions when using clang for msvc platforms

#rnx

diff --git a/Engine/Source/Programs/UnrealBuildTool/Platform/Windows/VCCompileAction.cs b/Engine/Source/Programs/UnrealBuildTool/Platform/Windows/VCCompileAction.cs
--- a/Engine/Source/Programs/UnrealBuildTool/Platform/Windows/VCCompileAction.cs
+++ b/Engine/Source/Programs/UnrealBuildTool/Platform/Windows/VCCompileAction.cs
@@ -490,6 +490,7 @@ namespace UnrealBuildTool
 			if (ResponseFile != null)
 			{
 				Graph.CreateIntermediateTextFile(ResponseFile, GetCompilerArguments(Logger));
+				Arguments.Clear();
 			}
 		}
 
diff --git a/Engine/Source/Programs/UnrealBuildTool/Platform/Windows/VCSpecificFileAction.cs b/Engine/Source/Programs/UnrealBuildTool/Platform/Windows/VCSpecificFileAction.cs
--- a/Engine/Source/Programs/UnrealBuildTool/Platform/Windows/VCSpecificFileAction.cs
+++ b/Engine/Source/Programs/UnrealBuildTool/Platform/Windows/VCSpecificFileAction.cs
@@ -113,6 +113,7 @@ namespace UnrealBuildTool
 			}
 			Action.ResponseFile = FileItem.GetItemByFileReference(FileReference.Combine(IntermediateDirectory, $"{Filename}.rsp"));
 			File.WriteAllLines(Action.ResponseFile.FullName, Action.GetCompilerArguments(Logger));
+			Action.Arguments.Clear();
 
 			return Action;
 		}
diff --git a/Engine/Source/Programs/UnrealBuildTool/Platform/Windows/VCToolChain.cs b/Engine/Source/Programs/UnrealBuildTool/Platform/Windows/VCToolChain.cs
--- a/Engine/Source/Programs/UnrealBuildTool/Platform/Windows/VCToolChain.cs
+++ b/Engine/Source/Programs/UnrealBuildTool/Platform/Windows/VCToolChain.cs
@@ -1897,7 +1897,6 @@ namespace UnrealBuildTool
 					if (!ProjectFileGenerator.bGenerateProjectFiles)
 					{
 						CompileDepsAction.WriteResponseFile(Graph, Logger);
-						CompileDepsAction.Arguments.Clear();
 					}
 
 					CompileAction.ActionType = ActionType.CompileModuleInterface;
@@ -1944,7 +1943,6 @@ namespace UnrealBuildTool
 				if (!ProjectFileGenerator.bGenerateProjectFiles)
 				{
 					CompileAction.WriteResponseFile(Graph, Logger);
-					CompileAction.Arguments.Clear();
 				}
 
 				// Must be added after response file is created just to make sure it ends up on the command line and not in the response file
