If building a console only application ensure the NMakeOutput file is set correctly (so that VS picks up the correct executable to debug/run)

diff --git a/Engine/Source/Programs/UnrealBuildTool/ProjectFiles/VisualStudio/VCProject.cs b/Engine/Source/Programs/UnrealBuildTool/ProjectFiles/VisualStudio/VCProject.cs
--- a/Engine/Source/Programs/UnrealBuildTool/ProjectFiles/VisualStudio/VCProject.cs
+++ b/Engine/Source/Programs/UnrealBuildTool/ProjectFiles/VisualStudio/VCProject.cs
@@ -2445,7 +2445,14 @@ namespace UnrealBuildTool
 						VCProjectFileContent.AppendLine("    <NMakeReBuildCommandLine>$(RebuildBatchScript) {0}</NMakeReBuildCommandLine>", BuildArguments);
 						VCProjectFileContent.AppendLine("    <NMakeCleanCommandLine>$(CleanBatchScript) {0}</NMakeCleanCommandLine>", BuildArguments);
 					}
-					VCProjectFileContent.AppendLine("    <NMakeOutput>{0}</NMakeOutput>", NormalizeProjectPath(NMakePath.FullName));
+					if (TargetRulesObject.bBuildConsoleAppOnly)
+					{
+						VCProjectFileContent.AppendLine("    <NMakeOutput>{0}</NMakeOutput>", NormalizeProjectPath(UEBuildBinary.GetAdditionalConsoleAppPath(new FileReference(NMakePath.FullName))));
+					} 
+					else 
+					{
+						VCProjectFileContent.AppendLine("    <NMakeOutput>{0}</NMakeOutput>", NormalizeProjectPath(NMakePath.FullName));
+					}
 					if (ProjectFileGenerator.bVisualStudioLinux)
 					{
 						VCProjectFileContent.AppendLine("    <BuildCommandLine>$(NMakeBuildCommandLine)</BuildCommandLine>");
