UnrealBuildTool: Allow disabling specific PVS analyzer errors for modules

#rnx

diff --git a/Engine/Source/Programs/UnrealBuildTool/Configuration/ModuleRules.cs b/Engine/Source/Programs/UnrealBuildTool/Configuration/ModuleRules.cs
--- a/Engine/Source/Programs/UnrealBuildTool/Configuration/ModuleRules.cs
+++ b/Engine/Source/Programs/UnrealBuildTool/Configuration/ModuleRules.cs
@@ -1057,6 +1057,11 @@ namespace UnrealBuildTool
 		/// </summary>
 		public HashSet<string> StaticAnalyzerAdditionalCheckers { get; } = new();
 
+		/// <summary>
+		/// The PVS Studio analysis errors that should be disabled.
+		/// </summary>
+		public HashSet<string> StaticAnalyzerPVSDisabledErrors { get; } = new();
+
 		private bool? bUseUnityOverride;
 		/// <summary>
 		/// If unity builds are enabled this can be used to override if this specific module will build using Unity.
diff --git a/Engine/Source/Programs/UnrealBuildTool/Configuration/UEBuildModuleCPP.cs b/Engine/Source/Programs/UnrealBuildTool/Configuration/UEBuildModuleCPP.cs
--- a/Engine/Source/Programs/UnrealBuildTool/Configuration/UEBuildModuleCPP.cs
+++ b/Engine/Source/Programs/UnrealBuildTool/Configuration/UEBuildModuleCPP.cs
@@ -2090,6 +2090,7 @@ namespace UnrealBuildTool
 			Result.StaticAnalyzerCheckers = Rules.StaticAnalyzerCheckers;
 			Result.StaticAnalyzerDisabledCheckers = Rules.StaticAnalyzerDisabledCheckers;
 			Result.StaticAnalyzerAdditionalCheckers = Rules.StaticAnalyzerAdditionalCheckers;
+			Result.StaticAnalyzerPVSDisabledErrors = Rules.StaticAnalyzerPVSDisabledErrors;
 			Result.bEnableUndefinedIdentifierWarnings = Rules.bEnableUndefinedIdentifierWarnings;
 			Result.bTreatAsEngineModule = Rules.bTreatAsEngineModule;
 			Result.IncludeOrderVersion = Rules.IncludeOrderVersion;
diff --git a/Engine/Source/Programs/UnrealBuildTool/Platform/Windows/PVSToolChain.cs b/Engine/Source/Programs/UnrealBuildTool/Platform/Windows/PVSToolChain.cs
--- a/Engine/Source/Programs/UnrealBuildTool/Platform/Windows/PVSToolChain.cs
+++ b/Engine/Source/Programs/UnrealBuildTool/Platform/Windows/PVSToolChain.cs
@@ -886,11 +886,17 @@ namespace UnrealBuildTool
 					ConfigFileContents.Append("report-disabled-rules=yes\n");
 				}
 
+				// TODO: Investigate into this disabled error
 				if (SourceFileItem.Location.IsUnderDirectory(Unreal.RootDirectory))
 				{
 					ConfigFileContents.AppendFormat("errors-off=V1102\n");
 				}
 
+				foreach (string error in CompileEnvironment.StaticAnalyzerPVSDisabledErrors.OrderBy(x => x))
+				{
+					ConfigFileContents.AppendFormat($"errors-off={error}\n");
+				}
+
 				int Timeout = (int)(Settings.AnalysisTimeoutFlag == AnalysisTimeoutFlags.No_timeout ? 0 : Settings.AnalysisTimeoutFlag);
 				ConfigFileContents.AppendFormat("timeout={0}\n", Timeout);
 
diff --git a/Engine/Source/Programs/UnrealBuildTool/System/CppCompileEnvironment.cs b/Engine/Source/Programs/UnrealBuildTool/System/CppCompileEnvironment.cs
--- a/Engine/Source/Programs/UnrealBuildTool/System/CppCompileEnvironment.cs
+++ b/Engine/Source/Programs/UnrealBuildTool/System/CppCompileEnvironment.cs
@@ -398,6 +398,11 @@ namespace UnrealBuildTool
 		/// </summary>
 		public HashSet<string> StaticAnalyzerAdditionalCheckers = new HashSet<string>();
 
+		/// <summary>
+		/// The PVS Studio analysis warnings that should be disabled.
+		/// </summary>
+		public HashSet<string> StaticAnalyzerPVSDisabledErrors = new();
+
 		/// <summary>
 		/// True if compiler optimizations should be enabled. This setting is distinct from the configuration (see CPPTargetConfiguration).
 		/// </summary>
@@ -765,6 +770,7 @@ namespace UnrealBuildTool
 			StaticAnalyzerCheckers = new HashSet<string>(Other.StaticAnalyzerCheckers);
 			StaticAnalyzerDisabledCheckers = new HashSet<string>(Other.StaticAnalyzerDisabledCheckers);
 			StaticAnalyzerAdditionalCheckers = new HashSet<string>(Other.StaticAnalyzerAdditionalCheckers);
+			StaticAnalyzerPVSDisabledErrors = new HashSet<string>(Other.StaticAnalyzerPVSDisabledErrors);
 			bStaticAnalyzerExtensions = Other.bStaticAnalyzerExtensions;
 			StaticAnalyzerRulesets = new HashSet<FileReference>(Other.StaticAnalyzerRulesets);
 			bOptimizeCode = Other.bOptimizeCode;
