Deduplication of low level tests: remove online-related tests
+ run only selected tests unless none is selected which means run all tests
#rnx

diff --git a/Engine/Build/LowLevelTests.xml b/Engine/Build/LowLevelTests.xml
--- a/Engine/Build/LowLevelTests.xml
+++ b/Engine/Build/LowLevelTests.xml
@@ -77,10 +77,7 @@
 	<!-- If we set one of the Run*Tests option we don't want to run all the tests, just the selected ones. -->
 	<Property Name="RunsAllTests" Value="$(RunAllTests)"/>
 	<ForEach Name="TestName" Values="$(TestNames)">
-		<!-- Don't turn of all tests if only Foundation is in the list. -->
-		<Do If="'$(TestName)' != 'Foundation'">
-			<Property Name="RunsAllTests" Value="false" If="$(Run$(TestName)Tests)"/>
-		</Do>
+		<Property Name="RunsAllTests" Value="false" If="$(Run$(TestName)Tests)"/>
 	</ForEach>
 
 	<!-- Platform Options -->
diff --git a/Engine/Build/LowLevelTests/ChasoUserDataPT.xml b/Engine/Build/LowLevelTests/ChasoUserDataPT.xml
--- a/Engine/Build/LowLevelTests/ChasoUserDataPT.xml
+++ b/Engine/Build/LowLevelTests/ChasoUserDataPT.xml
@@ -1,6 +1,7 @@
 ﻿<?xml version="1.0" encoding="utf-8"?>
 <BuildGraph xmlns:xsi="http://www.w3.org/2001/XMLSchema-instance" xsi:schemaLocation="http://www.epicgames.com/BuildGraph ../../Build/Graph/Schema.xsd" xmlns="http://www.epicgames.com/BuildGraph">
   <Option Name="RunChasoUserDataPTTests" DefaultValue="False" Description="Run Chaos User Data PT Tests" />
+  <Property Name="TestNames" Value="$(TestNames);ChasoUserDataPT" />
   <Extend Name="RunAllTests">
     <Expand Name="DeployAndTest" Platform="Win64" TestName="ChasoUserDataPT" ShortName="Chaos User Data PT" TargetName="ChaosUserDataPTTests" BinaryRelativePath="Engine\Binaries" ReportType="console" />
   </Extend>
diff --git a/Engine/Build/LowLevelTests/EventLoop.xml b/Engine/Build/LowLevelTests/EventLoop.xml
--- a/Engine/Build/LowLevelTests/EventLoop.xml
+++ b/Engine/Build/LowLevelTests/EventLoop.xml
@@ -1,6 +1,7 @@
 ﻿<?xml version="1.0" encoding="utf-8"?>
 <BuildGraph xmlns:xsi="http://www.w3.org/2001/XMLSchema-instance" xsi:schemaLocation="http://www.epicgames.com/BuildGraph ../../Build/Graph/Schema.xsd" xmlns="http://www.epicgames.com/BuildGraph">
   <Option Name="RunEventLoopTests" DefaultValue="False" Description="Run EventLoop Tests" />
+  <Property Name="TestNames" Value="$(TestNames);EventLoop" />
   <Extend Name="RunAllTests">
     <Expand Name="DeployAndTest" Platform="Win64" TestName="EventLoop" ShortName="EventLoop" TargetName="EventLoopUnitTests" BinaryRelativePath="Engine\Binaries" ReportType="console" />
   </Extend>
diff --git a/Engine/Build/LowLevelTests/Foundation.xml b/Engine/Build/LowLevelTests/Foundation.xml
--- a/Engine/Build/LowLevelTests/Foundation.xml
+++ b/Engine/Build/LowLevelTests/Foundation.xml
@@ -1,6 +1,7 @@
 ﻿<?xml version="1.0" encoding="utf-8"?>
 <BuildGraph xmlns:xsi="http://www.w3.org/2001/XMLSchema-instance" xsi:schemaLocation="http://www.epicgames.com/BuildGraph ../../Build/Graph/Schema.xsd" xmlns="http://www.epicgames.com/BuildGraph">
   <Option Name="RunFoundationTests" DefaultValue="False" Description="Run Foundation Tests" />
+  <Property Name="TestNames" Value="$(TestNames);Foundation" />
   <Extend Name="RunAllTests">
     <Expand Name="DeployAndTest" Platform="Win64" TestName="Foundation" ShortName="Foundation" TargetName="FoundationTests" BinaryRelativePath="Engine\Binaries" ReportType="xml" ExtraCompilationArgs="-allmodules" />
     <Expand Name="DeployAndTest" Platform="Linux" TestName="Foundation" ShortName="Foundation" TargetName="FoundationTests" BinaryRelativePath="Engine\Binaries" ReportType="xml" Tags="~[.]~[Slow]" ExtraCompilationArgs="-allmodules" />
diff --git a/Engine/Build/LowLevelTests/GameplayGraph.xml b/Engine/Build/LowLevelTests/GameplayGraph.xml
--- a/Engine/Build/LowLevelTests/GameplayGraph.xml
+++ b/Engine/Build/LowLevelTests/GameplayGraph.xml
@@ -1,6 +1,7 @@
 ﻿<?xml version="1.0" encoding="utf-8"?>
 <BuildGraph xmlns:xsi="http://www.w3.org/2001/XMLSchema-instance" xsi:schemaLocation="http://www.epicgames.com/BuildGraph ../../Build/Graph/Schema.xsd" xmlns="http://www.epicgames.com/BuildGraph">
   <Option Name="RunGameplayGraphTests" DefaultValue="False" Description="Run GameplayGraph Tests" />
+  <Property Name="TestNames" Value="$(TestNames);GameplayGraph" />
   <Extend Name="RunAllTests">
     <Expand Name="DeployAndTest" Platform="Win64" TestName="GameplayGraph" ShortName="GameplayGraph" TargetName="GameplayGraphTests" BinaryRelativePath="Engine\Binaries" ReportType="console" />
   </Extend>
diff --git a/Engine/Build/LowLevelTests/HeadlessChaos.xml b/Engine/Build/LowLevelTests/HeadlessChaos.xml
--- a/Engine/Build/LowLevelTests/HeadlessChaos.xml
+++ b/Engine/Build/LowLevelTests/HeadlessChaos.xml
@@ -1,6 +1,7 @@
 ﻿<?xml version="1.0" encoding="utf-8"?>
 <BuildGraph xmlns:xsi="http://www.w3.org/2001/XMLSchema-instance" xsi:schemaLocation="http://www.epicgames.com/BuildGraph ../../Build/Graph/Schema.xsd" xmlns="http://www.epicgames.com/BuildGraph">
   <Option Name="RunHeadlessChaosTests" DefaultValue="False" Description="Run Headless Chaos Tests" />
+  <Property Name="TestNames" Value="$(TestNames);HeadlessChaos" />
   <Extend Name="RunAllTests">
     <Expand Name="DeployAndTest" Platform="Win64" TestName="HeadlessChaos" ShortName="Headless Chaos" TargetName="HeadlessChaos" BinaryRelativePath="Engine\Binaries" ReportType="console" UsesCatch2="False" />
   </Extend>
diff --git a/Engine/Build/LowLevelTests/MathCore.xml b/Engine/Build/LowLevelTests/MathCore.xml
--- a/Engine/Build/LowLevelTests/MathCore.xml
+++ b/Engine/Build/LowLevelTests/MathCore.xml
@@ -1,6 +1,7 @@
 ﻿<?xml version="1.0" encoding="utf-8"?>
 <BuildGraph xmlns:xsi="http://www.w3.org/2001/XMLSchema-instance" xsi:schemaLocation="http://www.epicgames.com/BuildGraph ../../Build/Graph/Schema.xsd" xmlns="http://www.epicgames.com/BuildGraph">
   <Option Name="RunMathCoreTests" DefaultValue="False" Description="Run MathCore Tests" />
+  <Property Name="TestNames" Value="$(TestNames);MathCore" />
   <Extend Name="RunAllTests">
     <Expand Name="DeployAndTest" Platform="Win64" TestName="MathCore" ShortName="MathCore" TargetName="MathCoreTests" BinaryRelativePath="Engine\Binaries" ReportType="xml" ExtraCompilationArgs="-allmodules" />
     <Expand Name="DeployAndTest" Platform="Linux" TestName="MathCore" ShortName="MathCore" TargetName="MathCoreTests" BinaryRelativePath="Engine\Binaries" ReportType="xml" ExtraCompilationArgs="-allmodules" />
diff --git a/Engine/Build/LowLevelTests/NetworkPredictionPlugin.xml b/Engine/Build/LowLevelTests/NetworkPredictionPlugin.xml
--- a/Engine/Build/LowLevelTests/NetworkPredictionPlugin.xml
+++ b/Engine/Build/LowLevelTests/NetworkPredictionPlugin.xml
@@ -1,6 +1,7 @@
 ﻿<?xml version="1.0" encoding="utf-8"?>
 <BuildGraph xmlns:xsi="http://www.w3.org/2001/XMLSchema-instance" xsi:schemaLocation="http://www.epicgames.com/BuildGraph ../../Build/Graph/Schema.xsd" xmlns="http://www.epicgames.com/BuildGraph">
   <Option Name="RunNetworkPredictionPluginTests" DefaultValue="False" Description="Run Net Prediction Tests" />
+  <Property Name="TestNames" Value="$(TestNames);NetworkPredictionPlugin" />
   <Extend Name="RunAllTests">
     <Expand Name="DeployAndTest" Platform="Win64" TestName="NetworkPredictionPlugin" ShortName="Net Prediction" TargetName="NetworkPredictionTests" BinaryRelativePath="Engine\Binaries" ReportType="console" />
     <Expand Name="DeployAndTest" Platform="Linux" TestName="NetworkPredictionPlugin" ShortName="Net Prediction" TargetName="NetworkPredictionTests" BinaryRelativePath="Engine\Binaries" ReportType="console" />
diff --git a/Engine/Build/LowLevelTests/ReplicationSystem.xml b/Engine/Build/LowLevelTests/ReplicationSystem.xml
--- a/Engine/Build/LowLevelTests/ReplicationSystem.xml
+++ b/Engine/Build/LowLevelTests/ReplicationSystem.xml
@@ -1,6 +1,7 @@
 ﻿<?xml version="1.0" encoding="utf-8"?>
 <BuildGraph xmlns:xsi="http://www.w3.org/2001/XMLSchema-instance" xsi:schemaLocation="http://www.epicgames.com/BuildGraph ../../Build/Graph/Schema.xsd" xmlns="http://www.epicgames.com/BuildGraph">
   <Option Name="RunReplicationSystemTests" DefaultValue="False" Description="Run Replication System Tests" />
+  <Property Name="TestNames" Value="$(TestNames);ReplicationSystem" />
   <Extend Name="RunAllTests">
     <Expand Name="DeployAndTest" Platform="Win64" TestName="ReplicationSystem" ShortName="Replication System" TargetName="ReplicationSystemLowLevelTests" BinaryRelativePath="Engine\Binaries" ReportType="console" />
     <Expand Name="DeployAndTest" Platform="Linux" TestName="ReplicationSystem" ShortName="Replication System" TargetName="ReplicationSystemLowLevelTests" BinaryRelativePath="Engine\Binaries" ReportType="console" />
diff --git a/Engine/Build/LowLevelTests/StateGraph.xml b/Engine/Build/LowLevelTests/StateGraph.xml
--- a/Engine/Build/LowLevelTests/StateGraph.xml
+++ b/Engine/Build/LowLevelTests/StateGraph.xml
@@ -1,6 +1,7 @@
 ﻿<?xml version="1.0" encoding="utf-8"?>
 <BuildGraph xmlns:xsi="http://www.w3.org/2001/XMLSchema-instance" xsi:schemaLocation="http://www.epicgames.com/BuildGraph ../../Build/Graph/Schema.xsd" xmlns="http://www.epicgames.com/BuildGraph">
   <Option Name="RunStateGraphTests" DefaultValue="False" Description="Run State Graph Tests" />
+  <Property Name="TestNames" Value="$(TestNames);StateGraph" />
   <Extend Name="RunAllTests">
     <Expand Name="DeployAndTest" Platform="Win64" TestName="StateGraph" ShortName="State Graph" TargetName="StateGraphTests" BinaryRelativePath="Engine\Binaries" ReportType="console" />
   </Extend>
diff --git a/Engine/Build/LowLevelTests/SwitchboardListener.xml b/Engine/Build/LowLevelTests/SwitchboardListener.xml
--- a/Engine/Build/LowLevelTests/SwitchboardListener.xml
+++ b/Engine/Build/LowLevelTests/SwitchboardListener.xml
@@ -1,6 +1,7 @@
 ﻿<?xml version="1.0" encoding="utf-8"?>
 <BuildGraph xmlns:xsi="http://www.w3.org/2001/XMLSchema-instance" xsi:schemaLocation="http://www.epicgames.com/BuildGraph ../../Build/Graph/Schema.xsd" xmlns="http://www.epicgames.com/BuildGraph">
   <Option Name="RunSwitchboardListenerTests" DefaultValue="False" Description="Run SwitchboardListener Tests" />
+  <Property Name="TestNames" Value="$(TestNames);SwitchboardListener" />
   <Extend Name="RunAllTests">
     <Expand Name="DeployAndTest" Platform="Win64" TestName="SwitchboardListener" ShortName="SwitchboardListener" TargetName="SwitchboardListenerTests" BinaryRelativePath="Engine\Binaries" ReportType="console" />
   </Extend>
diff --git a/Engine/Build/LowLevelTests/ToolMenus.xml b/Engine/Build/LowLevelTests/ToolMenus.xml
--- a/Engine/Build/LowLevelTests/ToolMenus.xml
+++ b/Engine/Build/LowLevelTests/ToolMenus.xml
@@ -1,6 +1,7 @@
 ﻿<?xml version="1.0" encoding="utf-8"?>
 <BuildGraph xmlns:xsi="http://www.w3.org/2001/XMLSchema-instance" xsi:schemaLocation="http://www.epicgames.com/BuildGraph ../../Build/Graph/Schema.xsd" xmlns="http://www.epicgames.com/BuildGraph">
   <Option Name="RunToolMenusTests" DefaultValue="False" Description="Run ToolMenus Tests" />
+  <Property Name="TestNames" Value="$(TestNames);ToolMenus" />
   <Extend Name="RunAllTests">
     <Expand Name="DeployAndTest" Platform="Win64" TestName="ToolMenus" ShortName="ToolMenus" TargetName="ToolMenusTests" BinaryRelativePath="Engine\Binaries" ReportType="xml" ExtraCompilationArgs="-allmodules" />
     <Expand Name="DeployAndTest" Platform="Linux" TestName="ToolMenus" ShortName="ToolMenus" TargetName="ToolMenusTests" BinaryRelativePath="Engine\Binaries" ReportType="xml" ExtraCompilationArgs="-allmodules" />
diff --git a/Engine/Build/LowLevelTests/Web.xml b/Engine/Build/LowLevelTests/Web.xml
--- a/Engine/Build/LowLevelTests/Web.xml
+++ b/Engine/Build/LowLevelTests/Web.xml
@@ -1,6 +1,7 @@
 ﻿<?xml version="1.0" encoding="utf-8"?>
 <BuildGraph xmlns:xsi="http://www.w3.org/2001/XMLSchema-instance" xsi:schemaLocation="http://www.epicgames.com/BuildGraph ../../Build/Graph/Schema.xsd" xmlns="http://www.epicgames.com/BuildGraph">
   <Option Name="RunWebTests" DefaultValue="False" Description="Run Web Tests" />
+  <Property Name="TestNames" Value="$(TestNames);Web" />
   <Extend Name="RunAllTests">
     <Expand Name="DeployAndTest" Platform="Win64" TestName="Web" ShortName="Web" TargetName="WebTests" BinaryRelativePath="Engine\Binaries" ReportType="xml" HasAfterSteps="True" ExtraCompilationArgs="-allmodules" />
     <Expand Name="DeployAndTest" Platform="Linux" TestName="Web" ShortName="Web" TargetName="WebTests" BinaryRelativePath="Engine\Binaries" ReportType="xml" HasAfterSteps="True" ExtraCompilationArgs="-allmodules" />
diff --git a/Engine/Source/Programs/UnrealBuildTool/Configuration/TestModuleRules.cs b/Engine/Source/Programs/UnrealBuildTool/Configuration/TestModuleRules.cs
--- a/Engine/Source/Programs/UnrealBuildTool/Configuration/TestModuleRules.cs
+++ b/Engine/Source/Programs/UnrealBuildTool/Configuration/TestModuleRules.cs
@@ -272,6 +272,7 @@ namespace UnrealBuildTool
 				if (!IsRestrictedPlatformName)
 				{
 					InsertOrUpdateTestOption(Root, $"Run{TestMetadata.TestName}Tests", $"Run {TestMetadata.TestShortName} Tests", false.ToString());
+					InsertOrUpdateTestProperty(Root, $"TestNames", TestMetadata.TestName, true);
 				}
 
 				if (TestMetadata.Deactivated)
@@ -439,6 +440,23 @@ namespace UnrealBuildTool
 			}
 		}
 
+		private static void InsertOrUpdateTestProperty(XElement Root, string PropertyName, string PropertyValue, bool Append)
+		{
+			XElement? PropertyElementWithName = Root.Elements(BuildGraphNamespace + "Property")
+				.Where(prop => prop.Attribute("Name").Value == PropertyName).FirstOrDefault();
+			if (PropertyElementWithName == null)
+			{
+				XElement ElementInsert = new XElement(BuildGraphNamespace + "Property");
+				ElementInsert.SetAttributeValue("Name", PropertyName);
+				ElementInsert.SetAttributeValue("Value", !Append ? PropertyValue : $"$({PropertyName});{PropertyValue}");
+				Root.Add(ElementInsert);
+			}
+			else
+			{
+				PropertyElementWithName.SetAttributeValue("Value", !Append ? PropertyValue : $"$({PropertyName});{PropertyValue}");
+			}
+		}
+
 #pragma warning restore 8604
 #pragma warning restore 8602
 
