UnrealBuildTool: Disallow caching -SingleFile compiles

#rnx

diff --git a/Engine/Source/Programs/UnrealBuildTool/Platform/Windows/VCSpecificFileAction.cs b/Engine/Source/Programs/UnrealBuildTool/Platform/Windows/VCSpecificFileAction.cs
--- a/Engine/Source/Programs/UnrealBuildTool/Platform/Windows/VCSpecificFileAction.cs
+++ b/Engine/Source/Programs/UnrealBuildTool/Platform/Windows/VCSpecificFileAction.cs
@@ -71,6 +71,7 @@ namespace UnrealBuildTool
 
 			DirectoryReference.CreateDirectory(IntermediateDirectory);
 			VCCompileAction Action = new VCCompileAction(BaseAction);
+			Action.ArtifactMode = ArtifactMode.None;
 			Action.SourceFile = SourceFile;
 			if (SourceFile.HasExtension(".h"))
 			{
