Only kill Java tasks started by MakeAPK on build machines

#android

diff --git a/Engine/Source/Programs/UnrealBuildTool/Platform/Android/UEDeployAndroid.cs b/Engine/Source/Programs/UnrealBuildTool/Platform/Android/UEDeployAndroid.cs
--- a/Engine/Source/Programs/UnrealBuildTool/Platform/Android/UEDeployAndroid.cs
+++ b/Engine/Source/Programs/UnrealBuildTool/Platform/Android/UEDeployAndroid.cs
@@ -1630,6 +1630,54 @@ namespace UnrealBuildTool
 			}
 		}
 
+		private static int RunCommandLineProgramAndReturnResultOutput(string WorkingDirectory, string Command, string Params, ILogger Logger, out string Output, string? OverrideDesc = null, bool bUseShellExecute = false)
+		{
+			// Process Arguments follow windows conventions in .NET Core
+			// Which means single quotes ' are not considered quotes.
+			// see [REDACTED_PROFILE_URL]/runtime/issues/29857
+			// also see UE-102580
+			// for rules see https://docs.microsoft.com/en-us/cpp/cpp/main-function-command-line-args
+			if (Params.Contains('\''))
+			{
+				Params = Params.Replace("\"", "\\\"");
+				Params = Params.Replace('\'', '\"');
+			}
+
+			if (OverrideDesc == null)
+			{
+				Logger.LogInformation("\nRunning: {Command} {Params}", Command, Params);
+			}
+			else if (!String.IsNullOrEmpty(OverrideDesc))
+			{
+				Logger.LogInformation("{Message}", OverrideDesc);
+				Logger.LogDebug("\nRunning: {Command} {Params}", Command, Params);
+			}
+
+			ProcessStartInfo StartInfo = new ProcessStartInfo();
+			StartInfo.WorkingDirectory = WorkingDirectory;
+			StartInfo.FileName = Command;
+			StartInfo.Arguments = Params;
+			StartInfo.UseShellExecute = bUseShellExecute;
+			StartInfo.WindowStyle = ProcessWindowStyle.Minimized;
+			StartInfo.RedirectStandardOutput = true;
+			StartInfo.RedirectStandardError = true;
+
+			string ProcOutput = string.Empty;
+
+			Process Proc = new Process();
+			Proc.StartInfo = StartInfo;
+			Proc.OutputDataReceived += (s, e) => ProcOutput += (e.Data != null) ? e.Data + Environment.NewLine : "";
+			Proc.ErrorDataReceived += (s, e) => ProcOutput += (e.Data != null) ? e.Data + Environment.NewLine : "";
+
+			Proc.Start();
+			Proc.BeginOutputReadLine();
+			Proc.BeginErrorReadLine();
+			Proc.WaitForExit();
+
+			Output = ProcOutput;
+			return Proc.ExitCode;
+		}
+
 		private static int RunCommandLineProgramAndReturnResult(string WorkingDirectory, string Command, string Params, ILogger Logger, string? OverrideDesc = null, bool bUseShellExecute = false)
 		{
 			// Process Arguments follow windows conventions in .NET Core
@@ -5280,6 +5328,30 @@ popd
 							// make sure destination exists
 							Directory.CreateDirectory(Path.GetDirectoryName(DestApkName)!);
 
+							// For build machine, get a list of current java.exe tasks already running to ignore for cleanup
+							List<int> PreviousJavaProcesses = new List<int>();
+							if (bIsBuildMachine && RuntimePlatform.IsWindows)
+							{
+								string Output;
+								int ExitCode = RunCommandLineProgramAndReturnResultOutput(UnrealBuildGradlePath, "cmd.exe", "/c tasklist.exe /NH /FI \"imagename eq java.exe\"", Logger, out Output, "Collecting existing java.exe processes");
+								if (ExitCode == 0)
+								{
+									foreach (string Taskline in Output.Split(new[] { "\r\n", "\r", "\n" }, StringSplitOptions.None))
+									{
+										if (Taskline.StartsWith("java.exe"))
+										{
+											string ProcessString = Taskline.Substring(8).TrimStart();
+											ProcessString = ProcessString.Substring(0, ProcessString.IndexOf(' '));
+											int ProcessId;
+											if (Int32.TryParse(ProcessString, out ProcessId))
+											{
+												PreviousJavaProcesses.Add(ProcessId);
+											}
+										}
+									}
+								}
+							}
+
 							// Use gradle to build the .apk file
 							string ShellExecutable = RuntimePlatform.IsWindows ? "cmd.exe" : "/bin/sh";
 							string ShellParametersBegin = RuntimePlatform.IsWindows ? "/c " : "-c '";
@@ -5297,10 +5369,31 @@ popd
 								//GradleOptions = "tasks --all";
 								//RunCommandLineProgramWithException(UnrealBuildGradlePath, ShellExecutable, ShellParametersBegin + "\"" + GradleScriptPath + "\" " + GradleOptions + ShellParametersEnd, "Listing all tasks...");
 
-								// on Windows sometimes minifyReleaseWithR8 is keeping a lock on classes.dex so kill java.exe
+								// on Windows sometimes minifyReleaseWithR8 is keeping a lock on classes.dex so kill java.exe not previously running
 								if (RuntimePlatform.IsWindows)
 								{
-									RunCommandLineProgramAndReturnResult(UnrealBuildGradlePath, ShellExecutable, "/c taskkill.exe /F /IM java.exe /T", Logger, "Terminate java.exe processes");
+									string Output;
+									int ExitCode = RunCommandLineProgramAndReturnResultOutput(UnrealBuildGradlePath, "cmd.exe", "/c tasklist.exe /NH /FI \"imagename eq java.exe\"", Logger, out Output, "Collecting existing java.exe processes");
+									if (ExitCode == 0)
+									{
+										foreach (string Taskline in Output.Split(new[] { "\r\n", "\r", "\n" }, StringSplitOptions.None))
+										{
+											if (Taskline.StartsWith("java.exe"))
+											{
+												string ProcessString = Taskline.Substring(8).TrimStart();
+												ProcessString = ProcessString.Substring(0, ProcessString.IndexOf(' '));
+												int ProcessId;
+												if (Int32.TryParse(ProcessString, out ProcessId))
+												{
+													// if not in previous collected list, kill it
+													if (!PreviousJavaProcesses.Contains(ProcessId))
+													{
+														RunCommandLineProgramAndReturnResultOutput(UnrealBuildGradlePath, "cmd.exe", "/c taskkill.exe /PID " + ProcessString + " /T /F", Logger, out Output, "Terminate java.exe process");
+													}
+												}
+											}
+										}
+									}
 								}
 
 								GradleOptions = "clean";
