UnrealBuildTool: Suppress ISPC loging when parsing version string

#rnx

diff --git a/Engine/Source/Programs/UnrealBuildTool/ToolChain/UEToolChain.cs b/Engine/Source/Programs/UnrealBuildTool/ToolChain/UEToolChain.cs
--- a/Engine/Source/Programs/UnrealBuildTool/ToolChain/UEToolChain.cs
+++ b/Engine/Source/Programs/UnrealBuildTool/ToolChain/UEToolChain.cs
@@ -275,7 +275,7 @@ namespace UnrealBuildTool
 		/// <returns></returns>
 		protected string? RunToolAndCaptureOutput(FileReference Command, string ToolArg, string? Expression = null)
 		{
-			string ProcessOutput = Utils.RunLocalProcessAndReturnStdOut(Command.FullName, ToolArg, Logger);
+			string ProcessOutput = Utils.RunLocalProcessAndReturnStdOut(Command.FullName, ToolArg, null);
 
 			if (String.IsNullOrEmpty(Expression))
 			{
