[UBA]
* Enabled mimalloc for linux targets that are non-dll

diff --git a/Engine/Source/Programs/UnrealBuildAccelerator/Mimalloc/UbaMimalloc.Build.cs b/Engine/Source/Programs/UnrealBuildAccelerator/Mimalloc/UbaMimalloc.Build.cs
--- a/Engine/Source/Programs/UnrealBuildAccelerator/Mimalloc/UbaMimalloc.Build.cs
+++ b/Engine/Source/Programs/UnrealBuildAccelerator/Mimalloc/UbaMimalloc.Build.cs
@@ -11,7 +11,7 @@ public class UbaMimalloc : ModuleRules
 	{
 		PCHUsage = PCHUsageMode.NoPCHs;
 
-		if (Target.Platform.IsInGroup(UnrealPlatformGroup.Windows))// || Target.Platform.IsInGroup(UnrealPlatformGroup.Linux))
+		if (Target.Platform.IsInGroup(UnrealPlatformGroup.Windows) || (!Target.bShouldCompileAsDLL && Target.Platform.IsInGroup(UnrealPlatformGroup.Linux)))
 		{
 			PublicDependencyModuleNames.AddRange(new string[] {
 				"mimalloc212",
diff --git a/Engine/Source/Programs/UnrealBuildAccelerator/UbaAgent.Target.cs b/Engine/Source/Programs/UnrealBuildAccelerator/UbaAgent.Target.cs
--- a/Engine/Source/Programs/UnrealBuildAccelerator/UbaAgent.Target.cs
+++ b/Engine/Source/Programs/UnrealBuildAccelerator/UbaAgent.Target.cs
@@ -105,7 +105,7 @@ public class UbaAgentTarget : TargetRules
 		}
 
 		int useMiMalloc = 0;
-		if (Target.Platform.IsInGroup(UnrealPlatformGroup.Windows))// || Target.Platform.IsInGroup(UnrealPlatformGroup.Linux))
+		if (Target.Platform.IsInGroup(UnrealPlatformGroup.Windows) || (!Rules.bShouldCompileAsDLL && Target.Platform.IsInGroup(UnrealPlatformGroup.Linux)))
 		{
 			if (Rules.bUsePCHFiles)
 			{
diff --git a/Engine/Source/Programs/UnrealBuildTool/Configuration/UEBuildTarget.cs b/Engine/Source/Programs/UnrealBuildTool/Configuration/UEBuildTarget.cs
--- a/Engine/Source/Programs/UnrealBuildTool/Configuration/UEBuildTarget.cs
+++ b/Engine/Source/Programs/UnrealBuildTool/Configuration/UEBuildTarget.cs
@@ -2658,7 +2658,7 @@ namespace UnrealBuildTool
 
 			if (Rules.bStripExports)
 			{
-				CreateStripExportsActions(MakefileBuilder, Makefile);
+				CreateStripExportsActions(MakefileBuilder, Makefile, GlobalCompileEnvironment);
 			}
 
 			// Cache inline gen cpp data
@@ -2997,7 +2997,7 @@ namespace UnrealBuildTool
 			return Makefile;
 		}
 
-		void CreateStripExportsActions(TargetMakefileBuilder MakefileBuilder, TargetMakefile Makefile)
+		void CreateStripExportsActions(TargetMakefileBuilder MakefileBuilder, TargetMakefile Makefile, CppCompileEnvironment CompileEnvironment)
 		{
 			FileReference commandPath;
 			if (OperatingSystem.IsWindows())
@@ -3055,7 +3055,18 @@ namespace UnrealBuildTool
 				stripAction.CommandDescription = "GenerateExports";
 				stripAction.WorkingDirectory = Unreal.EngineSourceDirectory;
 				stripAction.StatusDescription = Name;
-				stripAction.bCanExecuteRemotely = false;
+				stripAction.bCanExecuteRemotely = false; // Can run remotely but is so fast so it would be a waste
+				stripAction.ArtifactMode = ArtifactMode.Enabled;
+
+				List<DirectoryItem> rootPaths = new();
+				rootPaths.Add(DirectoryItem.GetItemByDirectoryReference(Unreal.EngineDirectory));
+				if (ProjectFile != null && (!CompileEnvironment.bUseSharedBuildEnvironment || CompileEnvironment.AllIncludePath.Any(x => x.IsUnderDirectory(ProjectFile.Directory))))
+				{
+					rootPaths.Add(DirectoryItem.GetItemByDirectoryReference(ProjectFile.Directory));
+				}
+				rootPaths.Add(DirectoryItem.GetItemByDirectoryReference(Unreal.RootDirectory));
+				stripAction.RootPaths = rootPaths;
+
 				Makefile.OutputItems.AddRange(stripAction.ProducedItems);
 			}
 
@@ -3548,14 +3559,18 @@ namespace UnrealBuildTool
 				}
 			}
 
-			IEnumerable<UEBuildModuleCPP> engineModules = allModules.Where(x => x.RulesFile.IsUnderDirectory(Unreal.EngineDirectory));
-			IEnumerable<UEBuildModuleCPP> projectModules = allModules.Except(engineModules);
+			IEnumerable<UEBuildModuleCPP> engineModules = allModules.Where(x => x.RulesFile.IsUnderDirectory(Unreal.EngineDirectory) && !x.Rules.IsPlugin);
+			IEnumerable<UEBuildModuleCPP> enginePluginModules = allModules.Where(x => x.RulesFile.IsUnderDirectory(Unreal.EngineDirectory) && x.Rules.IsPlugin);
+			IEnumerable<UEBuildModuleCPP> projectModules = allModules.Except(engineModules).Except(enginePluginModules).Where(x => !x.Rules.IsPlugin);
+			IEnumerable<UEBuildModuleCPP> projectPluginModules = allModules.Except(engineModules).Except(enginePluginModules).Where(x => x.Rules.IsPlugin);
 
 			// Merge remaining engine modules
 			MergeModules("CommonEngine", engineModules);
+			MergeModules("CommonEnginePlugins", enginePluginModules);
 
 			// Merge remaining project modules
 			MergeModules("Common", projectModules);
+			MergeModules("CommonPlugins", projectPluginModules);
 
 			// Remove any binaries that no longer have modules due to merging
 			Binaries.ForEach(x => x.Modules.RemoveAll(y => y.Binary != x));
