UnrealBuildTool: Auto disable SharedPCH if a module is compiling using a CppStandard older than the engine

#rnx

diff --git a/Engine/Plugins/Animation/ACLPlugin/Source/ACLPlugin/ACLPlugin.Build.cs b/Engine/Plugins/Animation/ACLPlugin/Source/ACLPlugin/ACLPlugin.Build.cs
--- a/Engine/Plugins/Animation/ACLPlugin/Source/ACLPlugin/ACLPlugin.Build.cs
+++ b/Engine/Plugins/Animation/ACLPlugin/Source/ACLPlugin/ACLPlugin.Build.cs
@@ -13,9 +13,6 @@ namespace UnrealBuildTool.Rules
 
 			string ACLSDKDir = Path.GetFullPath(Path.Combine(ModuleDirectory, "../ThirdParty"));
 
-			// Replace with PCHUsageMode.UseExplicitOrSharedPCHs when this plugin can compile with cpp20
-			PCHUsage = PCHUsageMode.NoPCHs;
-
 			PublicSystemIncludePaths.Add(Path.Combine(ACLSDKDir, "acl/includes"));
 			PublicSystemIncludePaths.Add(Path.Combine(ACLSDKDir, "acl/external/rtm/includes"));
 
diff --git a/Engine/Plugins/Animation/ACLPlugin/Source/ACLPluginEditor/ACLPluginEditor.Build.cs b/Engine/Plugins/Animation/ACLPlugin/Source/ACLPluginEditor/ACLPluginEditor.Build.cs
--- a/Engine/Plugins/Animation/ACLPlugin/Source/ACLPluginEditor/ACLPluginEditor.Build.cs
+++ b/Engine/Plugins/Animation/ACLPlugin/Source/ACLPluginEditor/ACLPluginEditor.Build.cs
@@ -11,9 +11,6 @@ namespace UnrealBuildTool.Rules
 		{
 			CppStandard = CppStandardVersion.Cpp17;
 
-			// Replace with PCHUsageMode.UseExplicitOrSharedPCHs when this plugin can compile with cpp20
-			PCHUsage = PCHUsageMode.NoPCHs;
-
 			string ACLSDKDir = Path.GetFullPath(Path.Combine(ModuleDirectory, "../ThirdParty"));
 
 			PublicIncludePaths.Add(Path.Combine(ModuleDirectory, "Public"));
diff --git a/Engine/Plugins/Experimental/ChaosCachingUSD/Source/ChaosCachingUSD/ChaosCachingUSD.build.cs b/Engine/Plugins/Experimental/ChaosCachingUSD/Source/ChaosCachingUSD/ChaosCachingUSD.build.cs
--- a/Engine/Plugins/Experimental/ChaosCachingUSD/Source/ChaosCachingUSD/ChaosCachingUSD.build.cs
+++ b/Engine/Plugins/Experimental/ChaosCachingUSD/Source/ChaosCachingUSD/ChaosCachingUSD.build.cs
@@ -15,9 +15,6 @@ namespace UnrealBuildTool.Rules
 			// warning C5103: pasting '"TF_LOG_STACK_TRACE_ON_ERROR"' and '"TF_LOG_STACK_TRACE_ON_WARNING"' does not result in a valid preprocessing token
 			CppStandard = CppStandardVersion.Cpp17;
 
-			// Replace with PCHUsageMode.UseExplicitOrSharedPCHs when this plugin can compile with cpp20
-			PCHUsage = PCHUsageMode.NoPCHs;
-
 			bUseRTTI = true;
 
 			PublicDependencyModuleNames.AddRange(
diff --git a/Engine/Plugins/Experimental/NNERuntimeRDG/Source/NNERuntimeRDG/NNERuntimeRDG.Build.cs b/Engine/Plugins/Experimental/NNERuntimeRDG/Source/NNERuntimeRDG/NNERuntimeRDG.Build.cs
--- a/Engine/Plugins/Experimental/NNERuntimeRDG/Source/NNERuntimeRDG/NNERuntimeRDG.Build.cs
+++ b/Engine/Plugins/Experimental/NNERuntimeRDG/Source/NNERuntimeRDG/NNERuntimeRDG.Build.cs
@@ -10,9 +10,6 @@ public class NNERuntimeRDG : ModuleRules
 	{
 		CppStandard = CppStandardVersion.Cpp17;
 
-		// Replace with PCHUsageMode.UseExplicitOrSharedPCHs when this plugin can compile with cpp20
-		PCHUsage = PCHUsageMode.NoPCHs;
-
 		PublicDependencyModuleNames.AddRange(new string[] 
 		{ 
 			"Core", 
diff --git a/Engine/Plugins/Runtime/Nvidia/Reflex/Source/Reflex.Build.cs b/Engine/Plugins/Runtime/Nvidia/Reflex/Source/Reflex.Build.cs
--- a/Engine/Plugins/Runtime/Nvidia/Reflex/Source/Reflex.Build.cs
+++ b/Engine/Plugins/Runtime/Nvidia/Reflex/Source/Reflex.Build.cs
@@ -14,9 +14,6 @@ namespace UnrealBuildTool.Rules
 			// error C4002: too many arguments for function-like macro invocation '_TLG_FOR_imp0'
 			CppStandard = CppStandardVersion.Cpp17;
 
-			// Replace with PCHUsageMode.UseExplicitOrSharedPCHs when this plugin can compile with cpp20
-			PCHUsage = PCHUsageMode.NoPCHs;
-
 			PublicDependencyModuleNames.AddRange(
 				new string[]
 				{
diff --git a/Engine/Plugins/Runtime/USDCore/Source/UnrealUSDWrapper/UnrealUSDWrapper.Build.cs b/Engine/Plugins/Runtime/USDCore/Source/UnrealUSDWrapper/UnrealUSDWrapper.Build.cs
--- a/Engine/Plugins/Runtime/USDCore/Source/UnrealUSDWrapper/UnrealUSDWrapper.Build.cs
+++ b/Engine/Plugins/Runtime/USDCore/Source/UnrealUSDWrapper/UnrealUSDWrapper.Build.cs
@@ -15,9 +15,6 @@ namespace UnrealBuildTool.Rules
 			// error C2666: 'UE::FDummyRefPtrType::operator ==': overloaded functions have similar conversions
 			CppStandard = CppStandardVersion.Cpp17;
 
-			// Replace with PCHUsageMode.UseExplicitOrSharedPCHs when this plugin can compile with cpp20
-			PCHUsage = PCHUsageMode.NoPCHs;
-
 			bUseRTTI = true;
 
 			PublicDependencyModuleNames.AddRange(
diff --git a/Engine/Source/Developer/TraceAnalysis/TraceAnalysis.Build.cs b/Engine/Source/Developer/TraceAnalysis/TraceAnalysis.Build.cs
--- a/Engine/Source/Developer/TraceAnalysis/TraceAnalysis.Build.cs
+++ b/Engine/Source/Developer/TraceAnalysis/TraceAnalysis.Build.cs
@@ -11,9 +11,6 @@ public class TraceAnalysis : ModuleRules
 		// Asio does not compile with C++20, remove if updated
 		CppStandard = CppStandardVersion.Cpp17;
 
-		// Replace with PCHUsageMode.UseExplicitOrSharedPCHs when this plugin can compile with cpp20
-		PCHUsage = PCHUsageMode.NoPCHs;
-
 		PrivateDependencyModuleNames.AddRange(
 			new string[] {
 				"Asio",
diff --git a/Engine/Source/Developer/Windows/LiveCodingServer/LiveCodingServer.Build.cs b/Engine/Source/Developer/Windows/LiveCodingServer/LiveCodingServer.Build.cs
--- a/Engine/Source/Developer/Windows/LiveCodingServer/LiveCodingServer.Build.cs
+++ b/Engine/Source/Developer/Windows/LiveCodingServer/LiveCodingServer.Build.cs
@@ -9,9 +9,6 @@ public class LiveCodingServer : ModuleRules
 	{
 		CppStandard = CppStandardVersion.Cpp17;
 
-		// Replace with PCHUsageMode.UseExplicitOrSharedPCHs when this plugin can compile with cpp20
-		PCHUsage = PCHUsageMode.NoPCHs;
-
 		PrivateDependencyModuleNames.Add("Core");
         PrivateDependencyModuleNames.Add("Json");
         PrivateDependencyModuleNames.Add("LiveCoding");
diff --git a/Engine/Source/Programs/UnrealBuildTool/Configuration/ModuleRules.cs b/Engine/Source/Programs/UnrealBuildTool/Configuration/ModuleRules.cs
--- a/Engine/Source/Programs/UnrealBuildTool/Configuration/ModuleRules.cs
+++ b/Engine/Source/Programs/UnrealBuildTool/Configuration/ModuleRules.cs
@@ -763,7 +763,13 @@ namespace UnrealBuildTool
 		{
 			get
 			{
-				if (PCHUsagePrivate.HasValue)
+				if (((CppStandard.HasValue && CppStandard < Target.CppStandardEngine) || (Target.CppStandard < Target.CppStandardEngine))
+					&& (!PCHUsagePrivate.HasValue || PCHUsagePrivate == PCHUsageMode.UseSharedPCHs || PCHUsagePrivate == PCHUsageMode.UseExplicitOrSharedPCHs))
+				{
+					// SharedPCH is disallowed for modules that compile against an older CppStandard than the engine
+					return PrivatePCHHeaderFile != null ? PCHUsageMode.NoSharedPCHs : PCHUsageMode.NoPCHs;
+				}
+				else if (PCHUsagePrivate.HasValue)
 				{
 					// Use the override
 					return PCHUsagePrivate.Value;
