UnrealBuildTool: Print upgrade warning if a project is using unsupported C++ standard versions

#rnx

diff --git a/Engine/Source/Programs/UnrealBuildTool/Configuration/TargetRules.cs b/Engine/Source/Programs/UnrealBuildTool/Configuration/TargetRules.cs
--- a/Engine/Source/Programs/UnrealBuildTool/Configuration/TargetRules.cs
+++ b/Engine/Source/Programs/UnrealBuildTool/Configuration/TargetRules.cs
@@ -3271,6 +3271,7 @@ namespace UnrealBuildTool
 				if ((int)value == (int)BuildSettingsVersion.Latest)
 				{
 					latestVersion = value;
+					break;
 				}
 			}
 
@@ -3314,15 +3315,53 @@ namespace UnrealBuildTool
 
 			if (IncludeOrderVersion <= (EngineIncludeOrderVersion)(EngineIncludeOrderVersion.Latest - 1) && ForcedIncludeOrder == null)
 			{
+				// Resolve EngineIncludeOrderVersion.Latest to the version it's assigned to
+				EngineIncludeOrderVersion latestEngineIncludeOrder = EngineIncludeOrderVersion.Latest;
+				foreach (EngineIncludeOrderVersion value in Enum.GetValues(typeof(BuildSettingsVersion)))
+				{
+					if ((int)value == (int)EngineIncludeOrderVersion.Latest)
+					{
+						latestEngineIncludeOrder = value;
+						break;
+					}
+				}
+
 				diagnostics.Add("[Upgrade]");
 				diagnostics.Add("[Upgrade] Using backward-compatible include order. The latest version of UE has changed the order of includes, which may require code changes. The current setting is:");
-				diagnostics.Add(String.Format("[Upgrade]     IncludeOrderVersion = EngineIncludeOrderVersion.{0}", IncludeOrderVersion));
-				diagnostics.Add(String.Format("[Upgrade] Suppress this message by setting 'IncludeOrderVersion = EngineIncludeOrderVersion.{0};' in {1}.", EngineIncludeOrderVersion.Latest, File!.GetFileName()));
+				diagnostics.Add($"[Upgrade]     IncludeOrderVersion = EngineIncludeOrderVersion.{IncludeOrderVersion}");
+				diagnostics.Add($"[Upgrade] Suppress this message by setting 'IncludeOrderVersion = EngineIncludeOrderVersion.{latestEngineIncludeOrder};' in {File!.GetFileName()}.");
 				diagnostics.Add("[Upgrade] Alternatively you can set this to 'EngineIncludeOrderVersion.Latest' to always use the latest include order. This will potentially cause compile errors when integrating new versions of the engine.");
 				diagnostics.Add("[Upgrade]");
 			}
 
 			Logger.LogDebug("Using EngineIncludeOrderVersion.{Version} for target {Target}", IncludeOrderVersion, File!.GetFileName());
+
+			if (CppStandardEngine < CppStandardVersion.EngineDefault)
+			{
+				diagnostics.Add("[Upgrade]");
+				diagnostics.Add($"[Upgrade] The latest version of UE no longer supports CppStandardVersion.{CppStandardEngine} which may require code changes.");
+				diagnostics.Add($"[Upgrade] Suppress this message by removing 'CppStandardEngine = CppStandardVersion.{CppStandardEngine};' in {File!.GetFileName()}.");
+				diagnostics.Add("[Upgrade]");
+			}
+
+			if (CppStandard < CppStandardVersion.Default)
+			{
+				// Resolve CppStandardVersion.Default to the version it's assigned to
+				CppStandardVersion defaultCppStandard = CppStandardVersion.Latest;
+				foreach (CppStandardVersion value in Enum.GetValues(typeof(CppStandardVersion)))
+				{
+					if ((int)value == (int)CppStandardVersion.Default)
+					{
+						defaultCppStandard = value;
+						break;
+					}
+				}
+
+				diagnostics.Add("[Upgrade]");
+				diagnostics.Add($"[Upgrade] The latest version of UE no longer supports CppStandardVersion.{CppStandard} which may require code changes.");
+				diagnostics.Add($"[Upgrade] Suppress this message by setting 'CppStandard = CppStandardVersion.{defaultCppStandard};' in {File!.GetFileName()}.");
+				diagnostics.Add("[Upgrade]");
+			}
 		}
 
 		/// <summary>
diff --git a/Engine/Source/Programs/UnrealBuildTool/System/CppCompileEnvironment.cs b/Engine/Source/Programs/UnrealBuildTool/System/CppCompileEnvironment.cs
--- a/Engine/Source/Programs/UnrealBuildTool/System/CppCompileEnvironment.cs
+++ b/Engine/Source/Programs/UnrealBuildTool/System/CppCompileEnvironment.cs
@@ -25,12 +25,12 @@ namespace UnrealBuildTool
 	public enum CppStandardVersion
 	{
 		/// <summary>
-		/// Supports C++14. No longer maintained, will be removed in 5.5
+		/// Supports C++14. No longer supported, code may not compile with this version.
 		/// </summary>
 		Cpp14,
 
 		/// <summary>
-		/// Supports C++17
+		/// Supports C++17. No longer supported, code may not compile with this version.
 		/// </summary>
 		Cpp17,
 
