[UBT]
* Changed so root paths are included in bucket id when pch is enabled

diff --git a/Engine/Source/Programs/UnrealBuildTool/Executors/UnrealBuildAccelerator/UBAExecutor.cs b/Engine/Source/Programs/UnrealBuildTool/Executors/UnrealBuildAccelerator/UBAExecutor.cs
--- a/Engine/Source/Programs/UnrealBuildTool/Executors/UnrealBuildAccelerator/UBAExecutor.cs
+++ b/Engine/Source/Programs/UnrealBuildTool/Executors/UnrealBuildAccelerator/UBAExecutor.cs
@@ -564,9 +564,24 @@ namespace UnrealBuildTool
 			{
 				return 0;
 			}
-			string platform = action.Target.Platform.ToString();
-			byte[] platformBytes = System.Text.Encoding.ASCII.GetBytes(platform);
-			return (uint)IoHash.Compute(platformBytes).GetHashCode();
+
+			using (Blake3.Hasher hasher = Blake3.Hasher.New())
+			{
+				// Use platform to chose bucket since there will never be any cache hits between platforms
+				hasher.Update(System.Text.Encoding.UTF8.GetBytes(action.Target.Platform.ToString()));
+
+				// Absolute path is set for actions that uses pch.
+				// And since pch contains absolute paths we unfortunately can't share cache data between machines that have different paths
+				if (action.ArtifactMode.HasFlag(ArtifactMode.AbsolutePath))
+				{
+					foreach (DirectoryItem root in action.RootPaths)
+					{
+						hasher.Update(System.Text.Encoding.UTF8.GetBytes(root.FullName));
+					}
+				}
+
+				return (uint)IoHash.FromBlake3(hasher).GetHashCode();
+			}
 		}
 
 		IRootPaths GetActionRootPaths(LinkedAction action)
diff --git a/Engine/Source/Programs/UnrealBuildTool/Platform/Windows/VCToolChain.cs b/Engine/Source/Programs/UnrealBuildTool/Platform/Windows/VCToolChain.cs
--- a/Engine/Source/Programs/UnrealBuildTool/Platform/Windows/VCToolChain.cs
+++ b/Engine/Source/Programs/UnrealBuildTool/Platform/Windows/VCToolChain.cs
@@ -1786,13 +1786,11 @@ namespace UnrealBuildTool
 				}
 
 				{
-					if (CompileEnvironment.PrecompiledHeaderAction == PrecompiledHeaderAction.Create)
-					{
-						CompileAction.ArtifactMode |= ArtifactMode.Enabled | ArtifactMode.PropagateInputs;
-					}
-					else
+					CompileAction.ArtifactMode = ArtifactMode.Enabled;
+
+					if (CompileEnvironment.PrecompiledHeaderAction != PrecompiledHeaderAction.None)
 					{
-						CompileAction.ArtifactMode |= ArtifactMode.Enabled | ArtifactMode.AbsolutePath; // deps output file contains absolute paths
+						CompileAction.ArtifactMode |= ArtifactMode.AbsolutePath; // Unfortunately we require matching absolute paths for pch to be cached
 					}
 				}
 
diff --git a/Engine/Source/Programs/UnrealBuildTool/ToolChain/ClangToolChain.cs b/Engine/Source/Programs/UnrealBuildTool/ToolChain/ClangToolChain.cs
--- a/Engine/Source/Programs/UnrealBuildTool/ToolChain/ClangToolChain.cs
+++ b/Engine/Source/Programs/UnrealBuildTool/ToolChain/ClangToolChain.cs
@@ -1233,6 +1233,11 @@ namespace UnrealBuildTool
 
 			CompileAction.ArtifactMode = ArtifactMode.Enabled;
 
+			if (CompileEnvironment.PrecompiledHeaderAction != PrecompiledHeaderAction.None)
+			{
+				CompileAction.ArtifactMode |= ArtifactMode.AbsolutePath; // Unfortunately we require matching absolute paths for pch to be cached
+			}
+
 			// Two-pass compile where the preprocessor is run first to output the dependency list
 			if (PreprocessDepends)
 			{
