UnrealBuildTool: Don't allow caching clang -SingleFile actions

#rnx

diff --git a/Engine/Source/Programs/UnrealBuildTool/ToolChain/ClangSpecificFileAction.cs b/Engine/Source/Programs/UnrealBuildTool/ToolChain/ClangSpecificFileAction.cs
--- a/Engine/Source/Programs/UnrealBuildTool/ToolChain/ClangSpecificFileAction.cs
+++ b/Engine/Source/Programs/UnrealBuildTool/ToolChain/ClangSpecificFileAction.cs
@@ -28,6 +28,7 @@ namespace UnrealBuildTool
 			SourceDir = Source;
 			OutputDir = Output;
 			RspLines = ContentLines;
+			ArtifactMode = ArtifactMode.None;
 		}
 
 		public ClangSpecificFileAction(BinaryArchiveReader Reader) : base(Reader)
