[UBA]
* Changed so actions with no produced items are never cached
* Moved test if action should be cached to one place

diff --git a/Engine/Source/Programs/UnrealBuildTool/Executors/UnrealBuildAccelerator/UBAExecutor.cs b/Engine/Source/Programs/UnrealBuildTool/Executors/UnrealBuildAccelerator/UBAExecutor.cs
--- a/Engine/Source/Programs/UnrealBuildTool/Executors/UnrealBuildAccelerator/UBAExecutor.cs
+++ b/Engine/Source/Programs/UnrealBuildTool/Executors/UnrealBuildAccelerator/UBAExecutor.cs
@@ -531,6 +531,8 @@ namespace UnrealBuildTool
 
 			pchItem = null;
 
+			// This is not used. Should probably be deleted.
+			/*
 			if (_cacheClient != null && UBAConfig.bWriteCache && action.ArtifactMode.HasFlag(ArtifactMode.Enabled))
 			{
 				startInfo.TrackInputs = true;
@@ -555,6 +557,8 @@ namespace UnrealBuildTool
 					}
 				}
 			}
+			*/
+
 			return startInfo;
 		}
 
@@ -616,6 +620,17 @@ namespace UnrealBuildTool
 
 		bool WriteToCache(LinkedAction action, IProcess process)
 		{
+			if (!UBAConfig.bWriteCache || process.ExitCode != 0 || _cacheClient == null || !action.ArtifactMode.HasFlag(ArtifactMode.Enabled))
+			{
+				return true;
+			}
+
+			// If there are no outputs there is nothing to cache
+			if (!action.ProducedItems.Any())
+			{
+				return true;
+			}
+
 			// Collect all inputs for action
 			// We use prerequisite items plus what we find in dependency list file if it exists.
 
@@ -694,7 +709,7 @@ namespace UnrealBuildTool
 						_session!.RegisterNewFiles(action.ProducedItems.Where(x => FileReference.Exists(x.Location)).Select(x => x.FullName).ToArray());
 					}
 
-					if (startInfo.TrackInputs && UBAConfig.bWriteCache && enableDetour && process.ExitCode == 0 && _cacheClient != null)
+					if (enableDetour)
 					{
 						WriteToCache(action, process);
 					}
@@ -796,10 +811,7 @@ namespace UnrealBuildTool
 
 					IProcess process = (IProcess)s;
 
-					if (startInfo.TrackInputs && UBAConfig.bWriteCache && process.ExitCode == 0 && _cacheClient != null)
-					{
-						WriteToCache(action, process);
-					}
+					WriteToCache(action, process);
 
 					string additionalDescription = $"[RemoteExecutor: {e.ExecutingHost}]";
 					TimeSpan processorTime = e.TotalProcessorTime;
