Support moving to disabling AutoRTFM support only vs explicitly enabling

diff --git a/Engine/Source/Programs/UnrealBuildTool/Configuration/ModuleRules.cs b/Engine/Source/Programs/UnrealBuildTool/Configuration/ModuleRules.cs
--- a/Engine/Source/Programs/UnrealBuildTool/Configuration/ModuleRules.cs
+++ b/Engine/Source/Programs/UnrealBuildTool/Configuration/ModuleRules.cs
@@ -1190,6 +1190,11 @@ namespace UnrealBuildTool
 		/// </summary>
 		public bool bAllowAutoRTFMInstrumentation { get; set; }
 
+		/// <summary>
+		/// Disables AutoRTFM instrumentation to this module only when AutoRTFMCompiler is enabled
+		/// </summary>
+		public bool bDisableAutoRTFMInstrumentation { get; set; }
+
 		/// <summary>
 		/// List of modules names (no path needed) with header files that our module's public headers needs access to, but we don't need to "import" or link against.
 		/// </summary>
diff --git a/Engine/Source/Programs/UnrealBuildTool/Configuration/UEBuildModuleCPP.cs b/Engine/Source/Programs/UnrealBuildTool/Configuration/UEBuildModuleCPP.cs
--- a/Engine/Source/Programs/UnrealBuildTool/Configuration/UEBuildModuleCPP.cs
+++ b/Engine/Source/Programs/UnrealBuildTool/Configuration/UEBuildModuleCPP.cs
@@ -1503,6 +1503,7 @@ namespace UnrealBuildTool
 			CompileEnvironment.IncludeOrderVersion = ModuleCompileEnvironment.IncludeOrderVersion;
 			CompileEnvironment.bUseAutoRTFMCompiler = ModuleCompileEnvironment.bUseAutoRTFMCompiler;
 			CompileEnvironment.bAllowAutoRTFMInstrumentation = ModuleCompileEnvironment.bAllowAutoRTFMInstrumentation;
+			CompileEnvironment.bDisableAutoRTFMInstrumentation = ModuleCompileEnvironment.bDisableAutoRTFMInstrumentation;
 			CompileEnvironment.bValidateFormatStrings = ModuleCompileEnvironment.bValidateFormatStrings;
 			CompileEnvironment.bValidateInternalApi = ModuleCompileEnvironment.bValidateInternalApi;
 		}
@@ -2105,6 +2106,7 @@ namespace UnrealBuildTool
 			if (Target.bUseAutoRTFMCompiler)
 			{
 				Result.bAllowAutoRTFMInstrumentation = Rules.bAllowAutoRTFMInstrumentation;
+				Result.bDisableAutoRTFMInstrumentation = Rules.bDisableAutoRTFMInstrumentation;
 			}
 
 			if (Result.OptimizationLevel != Rules.OptimizationLevel)
diff --git a/Engine/Source/Programs/UnrealBuildTool/Platform/Windows/VCToolChain.cs b/Engine/Source/Programs/UnrealBuildTool/Platform/Windows/VCToolChain.cs
--- a/Engine/Source/Programs/UnrealBuildTool/Platform/Windows/VCToolChain.cs
+++ b/Engine/Source/Programs/UnrealBuildTool/Platform/Windows/VCToolChain.cs
@@ -1148,7 +1148,7 @@ namespace UnrealBuildTool
 				// Additional disabled warnings for msvc. Everything below should be checked if it is necessary
 				ClangWarnings.GetVCDisabledWarnings(Arguments);
 
-				if (CompileEnvironment.bAllowAutoRTFMInstrumentation)
+				if (!CompileEnvironment.bDisableAutoRTFMInstrumentation)
 				{
 					Arguments.Add("-fautortfm");
 				}
@@ -1189,7 +1189,7 @@ namespace UnrealBuildTool
 
 			if (Target.WindowsPlatform.Compiler.IsClang())
 			{
-				if (CompileEnvironment.bAllowAutoRTFMInstrumentation)
+				if (!CompileEnvironment.bDisableAutoRTFMInstrumentation)
 				{
 					Arguments.Add("-fautortfm");
 				}
diff --git a/Engine/Source/Programs/UnrealBuildTool/System/CppCompileEnvironment.cs b/Engine/Source/Programs/UnrealBuildTool/System/CppCompileEnvironment.cs
--- a/Engine/Source/Programs/UnrealBuildTool/System/CppCompileEnvironment.cs
+++ b/Engine/Source/Programs/UnrealBuildTool/System/CppCompileEnvironment.cs
@@ -286,6 +286,11 @@ namespace UnrealBuildTool
 		/// </summary>
 		public bool bAllowAutoRTFMInstrumentation = false;
 
+		/// <summary>
+		/// Disables AutoRTFM instrumentation to this cpp file only when AutoRTFMCompiler is enabled
+		/// </summary>
+		public bool bDisableAutoRTFMInstrumentation = false;
+
 		/// <summary>
 		/// If unity builds are enabled this can be used to override if this specific module will build using Unity.
 		/// This is set using the per module configurations in BuildConfiguration.
@@ -776,6 +781,7 @@ namespace UnrealBuildTool
 			bOptimizeCode = Other.bOptimizeCode;
 			bUseAutoRTFMCompiler = Other.bUseAutoRTFMCompiler;
 			bAllowAutoRTFMInstrumentation = Other.bAllowAutoRTFMInstrumentation;
+			bDisableAutoRTFMInstrumentation = Other.bDisableAutoRTFMInstrumentation;
 			bCodeCoverage = Other.bCodeCoverage;
 			OptimizationLevel = Other.OptimizationLevel;
 			FPSemantics = Other.FPSemantics;
diff --git a/Engine/Source/Programs/UnrealBuildTool/ToolChain/ClangToolChain.cs b/Engine/Source/Programs/UnrealBuildTool/ToolChain/ClangToolChain.cs
--- a/Engine/Source/Programs/UnrealBuildTool/ToolChain/ClangToolChain.cs
+++ b/Engine/Source/Programs/UnrealBuildTool/ToolChain/ClangToolChain.cs
@@ -391,7 +391,7 @@ namespace UnrealBuildTool
 				Arguments.Add("-fpch-validate-input-files-content");
 			}
 
-			if (CompileEnvironment.bAllowAutoRTFMInstrumentation && CompileEnvironment.bUseAutoRTFMCompiler)
+			if (!CompileEnvironment.bDisableAutoRTFMInstrumentation && CompileEnvironment.bUseAutoRTFMCompiler)
 			{
 				Arguments.Add("-fautortfm");
 			}
@@ -422,7 +422,7 @@ namespace UnrealBuildTool
 					throw new BuildException($"Unsupported C standard type set: {CompileEnvironment.CStandard}");
 			}
 
-			if (CompileEnvironment.bAllowAutoRTFMInstrumentation && CompileEnvironment.bUseAutoRTFMCompiler)
+			if (!CompileEnvironment.bDisableAutoRTFMInstrumentation && CompileEnvironment.bUseAutoRTFMCompiler)
 			{
 				Arguments.Add("-fautortfm");
 			}
