Remove bAllowAutoRTFM in favor of bDisableAutoRTFM, so by default we will always enable AutoRTFM for module if using AutoRTFM compiler unless the module disables it

diff --git a/Engine/Source/Programs/UnrealBuildTool/Configuration/ModuleRules.cs b/Engine/Source/Programs/UnrealBuildTool/Configuration/ModuleRules.cs
--- a/Engine/Source/Programs/UnrealBuildTool/Configuration/ModuleRules.cs
+++ b/Engine/Source/Programs/UnrealBuildTool/Configuration/ModuleRules.cs
@@ -1185,11 +1185,6 @@ namespace UnrealBuildTool
 		/// </summary>
 		public bool bAllowConfidentialPlatformDefines { get; set; }
 
-		/// <summary>
-		/// Enables AutoRTFM instrumentation to this module only when AutoRTFMCompiler is enabled
-		/// </summary>
-		public bool bAllowAutoRTFMInstrumentation { get; set; }
-
 		/// <summary>
 		/// Disables AutoRTFM instrumentation to this module only when AutoRTFMCompiler is enabled
 		/// </summary>
diff --git a/Engine/Source/Programs/UnrealBuildTool/Configuration/UEBuildModuleCPP.cs b/Engine/Source/Programs/UnrealBuildTool/Configuration/UEBuildModuleCPP.cs
--- a/Engine/Source/Programs/UnrealBuildTool/Configuration/UEBuildModuleCPP.cs
+++ b/Engine/Source/Programs/UnrealBuildTool/Configuration/UEBuildModuleCPP.cs
@@ -1502,7 +1502,6 @@ namespace UnrealBuildTool
 			CompileEnvironment.CppStandard = ModuleCompileEnvironment.CppStandard;
 			CompileEnvironment.IncludeOrderVersion = ModuleCompileEnvironment.IncludeOrderVersion;
 			CompileEnvironment.bUseAutoRTFMCompiler = ModuleCompileEnvironment.bUseAutoRTFMCompiler;
-			CompileEnvironment.bAllowAutoRTFMInstrumentation = ModuleCompileEnvironment.bAllowAutoRTFMInstrumentation;
 			CompileEnvironment.bDisableAutoRTFMInstrumentation = ModuleCompileEnvironment.bDisableAutoRTFMInstrumentation;
 			CompileEnvironment.bValidateFormatStrings = ModuleCompileEnvironment.bValidateFormatStrings;
 			CompileEnvironment.bValidateInternalApi = ModuleCompileEnvironment.bValidateInternalApi;
@@ -2102,10 +2101,9 @@ namespace UnrealBuildTool
 
 			CompileEnvironmentDebugInfoSettings(Target, Result);
 
-			// Only enable the AutoRTFM flag if we are using the AutoRTFM compiler
+			// Only copy the AutoRTFM flag if we are using the AutoRTFM compiler
 			if (Target.bUseAutoRTFMCompiler)
 			{
-				Result.bAllowAutoRTFMInstrumentation = Rules.bAllowAutoRTFMInstrumentation;
 				Result.bDisableAutoRTFMInstrumentation = Rules.bDisableAutoRTFMInstrumentation;
 			}
 
diff --git a/Engine/Source/Programs/UnrealBuildTool/System/CppCompileEnvironment.cs b/Engine/Source/Programs/UnrealBuildTool/System/CppCompileEnvironment.cs
--- a/Engine/Source/Programs/UnrealBuildTool/System/CppCompileEnvironment.cs
+++ b/Engine/Source/Programs/UnrealBuildTool/System/CppCompileEnvironment.cs
@@ -281,11 +281,6 @@ namespace UnrealBuildTool
 		/// </summary>
 		public bool bUseAutoRTFMCompiler = false;
 
-		/// <summary>
-		/// Enables AutoRTFM instrumentation to this cpp file only when AutoRTFMCompiler is enabled
-		/// </summary>
-		public bool bAllowAutoRTFMInstrumentation = false;
-
 		/// <summary>
 		/// Disables AutoRTFM instrumentation to this cpp file only when AutoRTFMCompiler is enabled
 		/// </summary>
@@ -780,7 +775,6 @@ namespace UnrealBuildTool
 			StaticAnalyzerRulesets = new HashSet<FileReference>(Other.StaticAnalyzerRulesets);
 			bOptimizeCode = Other.bOptimizeCode;
 			bUseAutoRTFMCompiler = Other.bUseAutoRTFMCompiler;
-			bAllowAutoRTFMInstrumentation = Other.bAllowAutoRTFMInstrumentation;
 			bDisableAutoRTFMInstrumentation = Other.bDisableAutoRTFMInstrumentation;
 			bCodeCoverage = Other.bCodeCoverage;
 			OptimizationLevel = Other.OptimizationLevel;
