UnrealBuildTool: When using UBA relay mode, default encryption to SSL

#rnx

diff --git a/Engine/Source/Programs/UnrealBuildTool/Executors/UnrealBuildAccelerator/UBAAgentCoordinatorHorde.cs b/Engine/Source/Programs/UnrealBuildTool/Executors/UnrealBuildAccelerator/UBAAgentCoordinatorHorde.cs
--- a/Engine/Source/Programs/UnrealBuildTool/Executors/UnrealBuildAccelerator/UBAAgentCoordinatorHorde.cs
+++ b/Engine/Source/Programs/UnrealBuildTool/Executors/UnrealBuildAccelerator/UBAAgentCoordinatorHorde.cs
@@ -376,6 +376,12 @@ namespace UnrealBuildTool
 			ConnectionMode? connectionMode = Enum.TryParse(hordeConfig.HordeConnectionMode, true, out ConnectionMode cm) ? cm : null;
 			Encryption? encryption = Enum.TryParse(hordeConfig.HordeEncryption, true, out Encryption enc) ? enc : null;
 
+			// Default to SSL encryption for relay mode if unset
+			if (connectionMode == ConnectionMode.Relay && encryption == null)
+			{
+				encryption = Encryption.Ssl;
+			}
+
 			logger.LogInformation("Horde URL: {Server}, Pool: {Pool}, Condition: {Condition}, Connection: {Connection} HordeEncryption: {Encryption}",
 				server, hordeConfig.HordePool ?? "(none)", hordeConfig.HordeCondition ?? "(none)", connectionMode?.ToString() ?? "(none)", encryption?.ToString() ?? "(none)");
 			try
