UnrealBuildTool: Fix logic for disabling SharedPCH for cpp17 modules

#rnx

diff --git a/Engine/Source/Programs/UnrealBuildTool/Configuration/ModuleRules.cs b/Engine/Source/Programs/UnrealBuildTool/Configuration/ModuleRules.cs
--- a/Engine/Source/Programs/UnrealBuildTool/Configuration/ModuleRules.cs
+++ b/Engine/Source/Programs/UnrealBuildTool/Configuration/ModuleRules.cs
@@ -763,13 +763,7 @@ namespace UnrealBuildTool
 		{
 			get
 			{
-				if (((CppStandard.HasValue && CppStandard < Target.CppStandardEngine) || (Target.CppStandard < Target.CppStandardEngine))
-					&& (!PCHUsagePrivate.HasValue || PCHUsagePrivate == PCHUsageMode.UseSharedPCHs || PCHUsagePrivate == PCHUsageMode.UseExplicitOrSharedPCHs))
-				{
-					// SharedPCH is disallowed for modules that compile against an older CppStandard than the engine
-					return PrivatePCHHeaderFile != null ? PCHUsageMode.NoSharedPCHs : PCHUsageMode.NoPCHs;
-				}
-				else if (PCHUsagePrivate.HasValue)
+				if (PCHUsagePrivate.HasValue)
 				{
 					// Use the override
 					return PCHUsagePrivate.Value;
diff --git a/Engine/Source/Programs/UnrealBuildTool/Configuration/UEBuildModuleCPP.cs b/Engine/Source/Programs/UnrealBuildTool/Configuration/UEBuildModuleCPP.cs
--- a/Engine/Source/Programs/UnrealBuildTool/Configuration/UEBuildModuleCPP.cs
+++ b/Engine/Source/Programs/UnrealBuildTool/Configuration/UEBuildModuleCPP.cs
@@ -2109,18 +2109,18 @@ namespace UnrealBuildTool
 
 			if (Result.OptimizationLevel != Rules.OptimizationLevel)
 			{
-				Logger.LogInformation("Module {0} - Optimization level changed for module due to override. Old: {1} New: {2}", Name, Result.OptimizationLevel, Rules.OptimizationLevel);
+				Logger.LogInformation("Module {Name} - Optimization level changed for module due to override. Old: {OldOptimizationLevel} New: {NewOptimizationLevel}", Name, Result.OptimizationLevel, Rules.OptimizationLevel);
 				if (Rules.PrivatePCHHeaderFile == null)
 				{
 					if (Rules.PCHUsage != ModuleRules.PCHUsageMode.NoPCHs)
 					{
-						Logger.LogInformation("  Overriding OptimizationLevel requires a private PCH. Disabling PCH usage for {0}", Name);
+						Logger.LogInformation("  Overriding OptimizationLevel requires a private PCH. Disabling PCH usage for {Name}", Name);
 						Rules.PCHUsage = ModuleRules.PCHUsageMode.NoPCHs;
 					}
 				}
 				else if (Rules.PCHUsage == ModuleRules.PCHUsageMode.UseSharedPCHs)
 				{
-					Logger.LogInformation("  Overriding OptimizationLevel requires a private PCH. A private PCH exists but UseSharedPCHs was specified. Overriding to NoSharedPCHs for {0}", Name);
+					Logger.LogInformation("  Overriding OptimizationLevel requires a private PCH. A private PCH exists but UseSharedPCHs was specified. Overriding to NoSharedPCHs for {Name}", Name);
 					Rules.PCHUsage = ModuleRules.PCHUsageMode.NoSharedPCHs;
 				}
 				Result.OptimizationLevel = Rules.OptimizationLevel;
@@ -2132,13 +2132,13 @@ namespace UnrealBuildTool
 				{
 					if (Rules.PCHUsage != ModuleRules.PCHUsageMode.NoPCHs)
 					{
-						Logger.LogInformation($"  Overriding FPSemantics requires a private PCH. Disabling PCH usage for {Name}");
+						Logger.LogInformation("  Overriding FPSemantics requires a private PCH. Disabling PCH usage for {Name}", Name);
 						Rules.PCHUsage = ModuleRules.PCHUsageMode.NoPCHs;
 					}
 				}
 				else if (Rules.PCHUsage == ModuleRules.PCHUsageMode.UseSharedPCHs)
 				{
-					Logger.LogInformation($"  Overriding FPSemantics requires a private PCH. A private PCH exists but UseSharedPCHs was specified. Overriding to NoSharedPCHs for {Name}");
+					Logger.LogInformation("  Overriding FPSemantics requires a private PCH. A private PCH exists but UseSharedPCHs was specified. Overriding to NoSharedPCHs for {Name}", Name);
 					Rules.PCHUsage = ModuleRules.PCHUsageMode.NoSharedPCHs;
 				}
 				Result.FPSemantics = Rules.FPSemantics;
@@ -2160,6 +2160,29 @@ namespace UnrealBuildTool
 				Result.CppStandard = CppStandardVersion.Cpp20;
 			}
 
+			if (Result.CppStandard < Result.CppStandardEngine)
+			{
+				// SharedPCH is disallowed for modules that compile against an older CppStandard than the engine
+				if (Rules.PCHUsage == ModuleRules.PCHUsageMode.UseExplicitOrSharedPCHs)
+				{
+					if (Rules.PrivatePCHHeaderFile != null)
+					{
+						Logger.LogDebug("  CppStandard {CppStandard} cannot use PCHUsage {PCHUsage}, however PrivatePCHHeaderFile is set. Overriding to NoSharedPCHs for {Name}", Result.CppStandard, Rules.PCHUsage, Name);
+						Rules.PCHUsage = ModuleRules.PCHUsageMode.NoSharedPCHs;
+					}
+					else
+					{
+						Logger.LogDebug("  CppStandard {CppStandard} cannot use PCHUsage {PCHUsage}. Overriding to NoPCHs for {Name}", Result.CppStandard, Rules.PCHUsage, Name);
+						Rules.PCHUsage = ModuleRules.PCHUsageMode.NoPCHs;
+					}
+				}
+				else if (Rules.PCHUsage == ModuleRules.PCHUsageMode.UseSharedPCHs)
+				{
+					Logger.LogDebug("  CppStandard {CppStandard} cannot use PCHUsage {PCHUsage}. Overriding to NoPCHs for {Name}", Result.CppStandard, Rules.PCHUsage, Name);
+					Rules.PCHUsage = ModuleRules.PCHUsageMode.NoPCHs;
+				}
+			}
+
 			// If the module overrides the C language version, override it on the compile environment
 			if (Rules.CStandard != null)
 			{
