UnrealBuildTool: Fix config setting for TargetRules.bAllowUbaCompression

#rnx

diff --git a/Engine/Source/Programs/UnrealBuildTool/Configuration/TargetRules.cs b/Engine/Source/Programs/UnrealBuildTool/Configuration/TargetRules.cs
--- a/Engine/Source/Programs/UnrealBuildTool/Configuration/TargetRules.cs
+++ b/Engine/Source/Programs/UnrealBuildTool/Configuration/TargetRules.cs
@@ -1998,7 +1998,7 @@ namespace UnrealBuildTool
 		/// Experimental: Store object (.obj) compressed on disk. Requires UBA to link, currently MSVC only. Toggling this flag will invalidate MSVC actions.
 		/// Warning, this option is not currently compatitable with PGO or the the cl-clang linker as those are not detoured and linking will fail.
 		/// </summary>
-		[XmlConfigFile(Category = "UnrealBuildAccelerator")]
+		[XmlConfigFile(Category = "UnrealBuildAccelerator", Name = "bStoreObjFilesCompressed")]
 		[CommandLine("-UBAStoreObjFilesCompressed", Value = "true", MarkUsed = false)]
 		[RequiresUniqueBuildEnvironment]
 		public bool bAllowUbaCompression { get; set; }
