- Disabled UBA for PCH for Mac/IOS builds

diff --git a/Engine/Source/Programs/UnrealBuildTool/ToolChain/AppleToolChain.cs b/Engine/Source/Programs/UnrealBuildTool/ToolChain/AppleToolChain.cs
--- a/Engine/Source/Programs/UnrealBuildTool/ToolChain/AppleToolChain.cs
+++ b/Engine/Source/Programs/UnrealBuildTool/ToolChain/AppleToolChain.cs
@@ -314,6 +314,10 @@ namespace UnrealBuildTool
 			foreach (FileItem SourceFile in InputFiles)
 			{
 				Action CompileAction = CompileCPPFile(CompileEnvironment, SourceFile, OutputDir, ModuleName, Graph, GlobalArguments, Result);
+				if (CompileEnvironment.PrecompiledHeaderAction == PrecompiledHeaderAction.Create)
+				{
+					CompileAction.bCanExecuteInUBA = false;
+				}
 				CompileAction.PrerequisiteItems.UnionWith(FrameworkTokenFiles);
 			}
 			return Result;
