- Fixed the issue with Mac builds doing full builds every time

diff --git a/Engine/Source/Programs/UnrealBuildTool/ProjectFiles/Xcode/XcodeProjectFileGenerator.cs b/Engine/Source/Programs/UnrealBuildTool/ProjectFiles/Xcode/XcodeProjectFileGenerator.cs
--- a/Engine/Source/Programs/UnrealBuildTool/ProjectFiles/Xcode/XcodeProjectFileGenerator.cs
+++ b/Engine/Source/Programs/UnrealBuildTool/ProjectFiles/Xcode/XcodeProjectFileGenerator.cs
@@ -554,6 +554,7 @@ namespace UnrealBuildTool
 					}
 					TargetDescriptor TargetDesc = new TargetDescriptor(CurTarget.UnrealProjectFilePath, CurTarget.Name, Platform, UnrealTargetConfiguration.Development,
 						new UnrealArchitectures(Arch), new CommandLineArguments(new string[] { "-skipclangvalidation" }));
+					TargetDesc.IntermediateEnvironment = UnrealIntermediateEnvironment.GenerateProjectFiles;
 					DateTime Start = DateTime.UtcNow;
 
 					try
