[UBA]
* Fixed shutdown hangs. Was caused by UBAExecutor was never informed that processes were cancelled. (and that queued actions were not cancelled when no action was running)

diff --git a/Engine/Source/Programs/Shared/EpicGames.UBA/Impl/ProcessImpl.cs b/Engine/Source/Programs/Shared/EpicGames.UBA/Impl/ProcessImpl.cs
--- a/Engine/Source/Programs/Shared/EpicGames.UBA/Impl/ProcessImpl.cs
+++ b/Engine/Source/Programs/Shared/EpicGames.UBA/Impl/ProcessImpl.cs
@@ -173,10 +173,7 @@ namespace EpicGames.UBA
 
 		internal void RaiseExited()
 		{
-			if (ExitCode != 99999) // See ProcessCancelExitCode
-			{
-				Exited?.Invoke(this, new ExitedEventArgs(this));
-			}
+			Exited?.Invoke(this, new ExitedEventArgs(this));
 		}
 	}
 }
diff --git a/Engine/Source/Programs/UnrealBuildTool/Executors/ImmediateActionQueue.cs b/Engine/Source/Programs/UnrealBuildTool/Executors/ImmediateActionQueue.cs
--- a/Engine/Source/Programs/UnrealBuildTool/Executors/ImmediateActionQueue.cs
+++ b/Engine/Source/Programs/UnrealBuildTool/Executors/ImmediateActionQueue.cs
@@ -260,6 +260,8 @@ namespace UnrealBuildTool
 		/// </summary>
 		public bool StopCompilationAfterErrors = false;
 
+		public int CompletedActions { get => _completedActions; }
+
 		/// <summary>
 		/// Return true if the queue is done
 		/// </summary>
@@ -453,6 +455,22 @@ namespace UnrealBuildTool
 		private void CancelKeyPress(object? sender, ConsoleCancelEventArgs e)
 		{
 			Console.CancelKeyPress -= CancelKeyPress;
+
+			// We must do this and can't rely on that there are active processes that are cancelled causing a cascading cancel (force remote actions and no remote workers)
+			int completedActions = 0;
+			lock (Actions)
+			{
+				for (int actionIndex = _firstPendingAction; actionIndex != Actions.Length; ++actionIndex)
+				{
+					if (Actions[actionIndex].Status == ActionStatus.Queued)
+					{
+						Actions[actionIndex].Status = ActionStatus.Error;
+						++completedActions;
+					}
+				}
+			}
+			AddCompletedActions(completedActions);
+
 			if (!CancellationTokenSource.IsCancellationRequested)
 			{
 				Logger.LogWarning("Canceling actions...");
@@ -869,7 +887,10 @@ namespace UnrealBuildTool
 		/// <param name="action">Action being re-queued</param>
 		public void RequeueAction(LinkedAction action)
 		{
-			SetActionState(action, ActionStatus.Queued, null);
+			if (!CancellationTokenSource.IsCancellationRequested)
+			{
+				SetActionState(action, ActionStatus.Queued, null);
+			}
 		}
 
 		/// <summary>
diff --git a/Engine/Source/Programs/UnrealBuildTool/Executors/UnrealBuildAccelerator/UBAExecutor.cs b/Engine/Source/Programs/UnrealBuildTool/Executors/UnrealBuildAccelerator/UBAExecutor.cs
--- a/Engine/Source/Programs/UnrealBuildTool/Executors/UnrealBuildAccelerator/UBAExecutor.cs
+++ b/Engine/Source/Programs/UnrealBuildTool/Executors/UnrealBuildAccelerator/UBAExecutor.cs
@@ -839,11 +839,8 @@ namespace UnrealBuildTool
 
 		protected void HandleActionCancelled(ImmediateActionQueue queue, ExecuteResults? results, LinkedAction action)
 		{
-			if (_bIsCancelled)
-			{
-				ExecuteResults cancelResults = new(results?.LogLines ?? new(), Int32.MaxValue, results?.ExecutionTime ?? TimeSpan.Zero, results?.ProcessorTime ?? TimeSpan.Zero, results?.AdditionalDescription);
-				queue.OnActionCompleted(action, false, cancelResults);
-			}
+			ExecuteResults cancelResults = new(results?.LogLines ?? new(), Int32.MaxValue, results?.ExecutionTime ?? TimeSpan.Zero, results?.ProcessorTime ?? TimeSpan.Zero, results?.AdditionalDescription);
+			queue.OnActionCompleted(action, false, cancelResults);
 		}
 
 		protected void ActionFinished(ImmediateActionQueue queue, ExecuteResults results, LinkedAction action, FileItem? pchItem = null, IProcess? process = null)
