[UBA]
* New attempt at handling ctrl-c properly

diff --git a/Engine/Source/Programs/Shared/EpicGames.UBA/Impl/ProcessImpl.cs b/Engine/Source/Programs/Shared/EpicGames.UBA/Impl/ProcessImpl.cs
--- a/Engine/Source/Programs/Shared/EpicGames.UBA/Impl/ProcessImpl.cs
+++ b/Engine/Source/Programs/Shared/EpicGames.UBA/Impl/ProcessImpl.cs
@@ -173,10 +173,7 @@ namespace EpicGames.UBA
 
 		internal void RaiseExited()
 		{
-			if (ExitCode != 99999) // See ProcessCancelExitCode
-			{
-				Exited?.Invoke(this, new ExitedEventArgs(this));
-			}
+			Exited?.Invoke(this, new ExitedEventArgs(this));
 		}
 	}
 }
diff --git a/Engine/Source/Programs/UnrealBuildTool/Executors/ImmediateActionQueue.cs b/Engine/Source/Programs/UnrealBuildTool/Executors/ImmediateActionQueue.cs
--- a/Engine/Source/Programs/UnrealBuildTool/Executors/ImmediateActionQueue.cs
+++ b/Engine/Source/Programs/UnrealBuildTool/Executors/ImmediateActionQueue.cs
@@ -456,6 +456,13 @@ namespace UnrealBuildTool
 		{
 			Console.CancelKeyPress -= CancelKeyPress;
 
+			if (!CancellationTokenSource.IsCancellationRequested)
+			{
+				Logger.LogWarning("Canceling actions...");
+				CancellationTokenSource.Cancel();
+				e.Cancel = true;
+			}
+
 			// We must do this and can't rely on that there are active processes that are cancelled causing a cascading cancel (force remote actions and no remote workers)
 			int completedActions = 0;
 			lock (Actions)
@@ -470,13 +477,6 @@ namespace UnrealBuildTool
 				}
 			}
 			AddCompletedActions(completedActions);
-
-			if (!CancellationTokenSource.IsCancellationRequested)
-			{
-				Logger.LogWarning("Canceling actions...");
-				CancellationTokenSource.Cancel();
-				e.Cancel = true;
-			}
 		}
 
 		/// <summary>
@@ -887,10 +887,7 @@ namespace UnrealBuildTool
 		/// <param name="action">Action being re-queued</param>
 		public void RequeueAction(LinkedAction action)
 		{
-			if (!CancellationTokenSource.IsCancellationRequested)
-			{
-				SetActionState(action, ActionStatus.Queued, null);
-			}
+			SetActionState(action, ActionStatus.Queued, null);
 		}
 
 		/// <summary>
diff --git a/Engine/Source/Programs/UnrealBuildTool/Executors/UnrealBuildAccelerator/UBAExecutor.cs b/Engine/Source/Programs/UnrealBuildTool/Executors/UnrealBuildAccelerator/UBAExecutor.cs
--- a/Engine/Source/Programs/UnrealBuildTool/Executors/UnrealBuildAccelerator/UBAExecutor.cs
+++ b/Engine/Source/Programs/UnrealBuildTool/Executors/UnrealBuildAccelerator/UBAExecutor.cs
@@ -190,8 +190,9 @@ namespace UnrealBuildTool
 
 		private async void ActionQueueCanceled(IStorageServer? ubaStorage)
 		{
+			Server?.StopServer(); // Make sure all remove processes are returned. We can't have any callbacks after this
 			_bIsCancelled = true;
-			_session?.CancelAll();
+			_session?.CancelAll(); // Cancel all processes native side
 			ubaStorage?.SaveCasTable();
 			foreach (IUBAAgentCoordinator coordinator in _agentCoordinators)
 			{
@@ -777,6 +778,11 @@ namespace UnrealBuildTool
 				ProcessStartInfo startInfo = GetActionStartInfo(action, out FileItem? pchItem);
 				_session!.RunProcessRemote(startInfo, (s, e) =>
 				{
+					if (e.ExitCode == 99999) // Process was cancelled by executor
+					{
+						return;
+					}
+
 					Interlocked.Add(ref _remoteProcessedActions, 1);
 					if (e.ExitCode != 0 && !e.LogLines.Any())
 					{
