Only use -fautortfm if using the AutoRTFM compiler

diff --git a/Engine/Source/Programs/UnrealBuildTool/Platform/Windows/VCToolChain.cs b/Engine/Source/Programs/UnrealBuildTool/Platform/Windows/VCToolChain.cs
--- a/Engine/Source/Programs/UnrealBuildTool/Platform/Windows/VCToolChain.cs
+++ b/Engine/Source/Programs/UnrealBuildTool/Platform/Windows/VCToolChain.cs
@@ -1148,7 +1148,7 @@ namespace UnrealBuildTool
 				// Additional disabled warnings for msvc. Everything below should be checked if it is necessary
 				ClangWarnings.GetVCDisabledWarnings(Arguments);
 
-				if (!CompileEnvironment.bDisableAutoRTFMInstrumentation)
+				if (!CompileEnvironment.bDisableAutoRTFMInstrumentation && CompileEnvironment.bUseAutoRTFMCompiler)
 				{
 					Arguments.Add("-fautortfm");
 				}
@@ -1189,7 +1189,7 @@ namespace UnrealBuildTool
 
 			if (Target.WindowsPlatform.Compiler.IsClang())
 			{
-				if (!CompileEnvironment.bDisableAutoRTFMInstrumentation)
+				if (!CompileEnvironment.bDisableAutoRTFMInstrumentation && CompileEnvironment.bUseAutoRTFMCompiler)
 				{
 					Arguments.Add("-fautortfm");
 				}
