Support additional manifest namespaces in UPL commands on Android

diff --git a/Engine/Source/Programs/UnrealBuildTool/Platform/Android/UEDeployAndroid.cs b/Engine/Source/Programs/UnrealBuildTool/Platform/Android/UEDeployAndroid.cs
--- a/Engine/Source/Programs/UnrealBuildTool/Platform/Android/UEDeployAndroid.cs
+++ b/Engine/Source/Programs/UnrealBuildTool/Platform/Android/UEDeployAndroid.cs
@@ -31,6 +31,12 @@ namespace UnrealBuildTool
 		// name of the only vulkan validation layer we're interested in 
 		private const string ANDROID_VULKAN_VALIDATION_LAYER = "libVkLayer_khronos_validation.so";
 
+		// list of supported aliased namespaces in manifest. First namespace in the list is the default one
+		private static readonly List<AliasedXMLNamespace> ManifestXMLNamespaces = new() {
+				new AliasedXMLNamespace { Alias = "android", Url = "http://schemas.android.com/apk/res/android" },
+				new AliasedXMLNamespace { Alias = "dist", Url = "http://schemas.android.com/apk/distribution" },
+				new AliasedXMLNamespace { Alias = "tools", Url = "http://schemas.android.com/tools" } };
+
 		// Minimum Android SDK that must be used for Java compiling
 		readonly int MinimumSDKLevel = 30;
 
@@ -159,7 +165,7 @@ namespace UnrealBuildTool
 				}
 			}
 
-			UPL = new UnrealPluginLanguage(ProjectFile, inPluginExtraData, NDKArches, "http://schemas.android.com/apk/res/android", "xmlns:android=\"http://schemas.android.com/apk/res/android\" xmlns:tools=\"http://schemas.android.com/tools\"", UnrealTargetPlatform.Android, Logger);
+			UPL = new UnrealPluginLanguage(ProjectFile, inPluginExtraData, NDKArches, ManifestXMLNamespaces, UnrealTargetPlatform.Android, Logger);
 			UPLHashCode = UPL.GetUPLHash();
 			//			APL.SetTrace();
 		}
@@ -2810,19 +2816,7 @@ namespace UnrealBuildTool
 
 			StringBuilder Text = new StringBuilder();
 			Text.AppendLine(XML_HEADER);
-
-
-			bool bIsMakeAAREnabled = false;
-			Ini.GetBool("/Script/AndroidSingleInstanceServiceEditor.AndroidSingleInstanceServiceRuntimeSettings", "bEnableASISPlugin", out bIsMakeAAREnabled);
-
-			if (bIsMakeAAREnabled)
-			{				
-				Text.AppendLine("<manifest xmlns:android=\"http://schemas.android.com/apk/res/android\" xmlns:dist=\"http://schemas.android.com/apk/distribution\" xmlns:tools=\"http://schemas.android.com/tools\"");
-			}
-			else
-			{
-				Text.AppendLine("<manifest xmlns:android=\"http://schemas.android.com/apk/res/android\" xmlns:tools=\"http://schemas.android.com/tools\"");
-			}
+			Text.AppendLine($"<manifest {String.Join(" ", ManifestXMLNamespaces)}");
 
 			if (ExtraManifestNodeTags != null)
 			{
@@ -2853,14 +2847,6 @@ namespace UnrealBuildTool
 				Text.AppendLine("\t</queries>");
 			}
 
-			if (bIsMakeAAREnabled)
-			{
-				Text.AppendLine("\t<dist:module dist:instant=\"true\" />");
-				Text.AppendLine("\t<queries>");
-				Text.AppendLine(String.Format("\t\t<package android:name=\"{0}\" />", PackageName));
-				Text.AppendLine("\t</queries>");
-			}
-
 			Text.AppendLine("\t<!-- Application Definition -->");
 			Text.AppendLine("\t<application android:label=\"@string/app_name\"");
 			Text.AppendLine("\t             android:icon=\"@drawable/icon\"");
diff --git a/Engine/Source/Programs/UnrealBuildTool/Platform/IOS/UEDeployIOS.cs b/Engine/Source/Programs/UnrealBuildTool/Platform/IOS/UEDeployIOS.cs
--- a/Engine/Source/Programs/UnrealBuildTool/Platform/IOS/UEDeployIOS.cs
+++ b/Engine/Source/Programs/UnrealBuildTool/Platform/IOS/UEDeployIOS.cs
@@ -998,7 +998,7 @@ namespace UnrealBuildTool
 
 			string RelativeEnginePath = Unreal.EngineDirectory.MakeRelativeTo(DirectoryReference.GetCurrentDirectory());
 
-			UnrealPluginLanguage UPL = new UnrealPluginLanguage(ProjectFile, UPLScripts, ProjectArches, "", "", UnrealTargetPlatform.IOS, Logger);
+			UnrealPluginLanguage UPL = new UnrealPluginLanguage(ProjectFile, UPLScripts, ProjectArches, null, UnrealTargetPlatform.IOS, Logger);
 
 			// Passing in true for distribution is not ideal here but given the way that ios packaging happens and this call chain it seems unavoidable for now, maybe there is a way to correctly pass it in that I can't find?
 			UPL.Init(ProjectArches, true, RelativeEnginePath, BundlePath, ProjectDirectory, Config.ToString(), false);
diff --git a/Engine/Source/Programs/UnrealBuildTool/System/UnrealPluginLanguage.cs b/Engine/Source/Programs/UnrealBuildTool/System/UnrealPluginLanguage.cs
--- a/Engine/Source/Programs/UnrealBuildTool/System/UnrealPluginLanguage.cs
+++ b/Engine/Source/Programs/UnrealBuildTool/System/UnrealPluginLanguage.cs
@@ -487,13 +487,26 @@ namespace UnrealBuildTool
 	 * 
 	 */
 
+	// Helper struct to initialize aliased namespaces
+	struct AliasedXMLNamespace
+	{
+		public string Alias { get; set; }
+		public string Url { get; set; }
+
+		public override string ToString()
+		{
+			return String.Format("xmlns:{0}=\"{1}\"", Alias, Url);
+		}
+	}
+
 	class UnrealPluginLanguage
 	{
 		/** The merged XML program to run */
 		private XDocument XDoc;
 
 		/** XML namespace */
-		private XNamespace XMLNameSpace;
+		private XNamespace XMLDefaultNameSpace;
+		private Dictionary<string, XNamespace> XMLNameSpaceAliases = new Dictionary<string, XNamespace>();
 		private string XMLRootDefinition;
 
 		private UnrealTargetPlatform TargetPlatform;
@@ -550,7 +563,8 @@ namespace UnrealBuildTool
 		private string? LastError;
 		private ILogger Logger;
 
-		public UnrealPluginLanguage(FileReference? InProjectFile, List<string> InXMLFiles, List<string> InArchitectures, string InXMLNameSpace, string InRootDefinition, UnrealTargetPlatform InTargetPlatform, ILogger InLogger)
+		// First entry in InXMLNameSpaceAliases will be used as default namespace
+		public UnrealPluginLanguage(FileReference? InProjectFile, List<string> InXMLFiles, List<string> InArchitectures, List<AliasedXMLNamespace>? InXMLNameSpaceAliases, UnrealTargetPlatform InTargetPlatform, ILogger InLogger)
 		{
 			ProjectFile = InProjectFile;
 			Logger = InLogger;
@@ -560,14 +574,30 @@ namespace UnrealBuildTool
 			Contexts = new Dictionary<string, UPLContext>();
 			GlobalContext = new UPLContext("", "", "");
 			ContextIndex = 0;
+			
+			if (InXMLNameSpaceAliases != null && InXMLNameSpaceAliases.Count > 0)
+			{
+				XMLDefaultNameSpace = InXMLNameSpaceAliases[0].Url;
+
+				StringBuilder RootNamespaces = new StringBuilder();
+				foreach (AliasedXMLNamespace AliasToNamespace in InXMLNameSpaceAliases)
+				{
+					XMLNameSpaceAliases[AliasToNamespace.Alias] = AliasToNamespace.Url;
+					RootNamespaces.AppendFormat("xmlns:{0}=\"{1}\" ", AliasToNamespace.Alias, AliasToNamespace.Url);
+				}
 
-			XMLNameSpace = InXMLNameSpace;
-			XMLRootDefinition = InRootDefinition;
+				XMLRootDefinition = String.Join(" ", RootNamespaces);
+			}
+			else
+			{
+				XMLDefaultNameSpace = "";
+				XMLRootDefinition = "";
+			}
 			TargetPlatform = InTargetPlatform;
 
 			string PathPrefix = Path.GetFileName(Directory.GetCurrentDirectory()).Equals("Source") ? ".." : "Engine";
 
-			XDoc = XDocument.Parse("<root " + InRootDefinition + "></root>");
+			XDoc = XDocument.Parse("<root " + XMLRootDefinition + "></root>");
 			foreach (string Basename in InXMLFiles)
 			{
 				string Filename = Path.Combine(PathPrefix, Basename.Replace("\\", "/"));
@@ -1010,19 +1040,31 @@ namespace UnrealBuildTool
 			}
 		}
 
-		private void AddAttribute(XElement Element, string Name, string Value)
+		XNamespace? TrimNamespaceAliasFromName(ref string Name)
 		{
-			XAttribute? Attribute;
 			int Index = Name.IndexOf(":");
 			if (Index >= 0)
 			{
+				XNamespace? Namespace;
+				string NamespaceAlias = Name.Substring(0, Index);
+				if (!XMLNameSpaceAliases.TryGetValue(NamespaceAlias, out Namespace))
+				{
+					Namespace = XMLDefaultNameSpace;
+				}
 				Name = Name.Substring(Index + 1);
-				Attribute = Element.Attribute(XMLNameSpace + Name);
+				return Namespace;
 			}
 			else
 			{
-				Attribute = Element.Attribute(Name);
+				return null;
 			}
+		}
+
+
+		private void AddAttribute(XElement Element, string Name, string Value)
+		{
+			XNamespace? XMLNameSpace = TrimNamespaceAliasFromName(ref Name);
+			XAttribute? Attribute = Element.Attribute(XMLNameSpace != null? XMLNameSpace + Name : Name);
 
 			if (Attribute != null)
 			{
@@ -1030,30 +1072,14 @@ namespace UnrealBuildTool
 			}
 			else
 			{
-				if (Index >= 0)
-				{
-					Element.Add(new XAttribute(XMLNameSpace + Name, Value));
-				}
-				else
-				{
-					Element.Add(new XAttribute(Name, Value));
-				}
+				Element.Add(new XAttribute(XMLNameSpace != null ? XMLNameSpace + Name : Name, Value));
 			}
 		}
 
 		private void RemoveAttribute(XElement Element, string Name)
 		{
-			XAttribute? Attribute;
-			int Index = Name.IndexOf(":");
-			if (Index >= 0)
-			{
-				Name = Name.Substring(Index + 1);
-				Attribute = Element.Attribute(XMLNameSpace + Name);
-			}
-			else
-			{
-				Attribute = Element.Attribute(Name);
-			}
+			XNamespace? XMLNameSpace = TrimNamespaceAliasFromName(ref Name);
+			XAttribute? Attribute = Element.Attribute(XMLNameSpace != null ? XMLNameSpace + Name : Name);
 
 			Attribute?.Remove();
 		}
@@ -1383,14 +1409,14 @@ namespace UnrealBuildTool
 
 					case "addPermission":
 						{
-							string? Name = GetAttributeWithNamespace(CurrentContext, Node, XMLNameSpace, "name");
+							string? Name = GetAttributeWithNamespace(CurrentContext, Node, XMLDefaultNameSpace, "name");
 							if (Name != null)
 							{
 								// make sure it isn't already added
 								bool bFound = false;
 								foreach (XElement Element in XMLWork.Descendants("uses-permission"))
 								{
-									XAttribute? Attribute = Element.Attribute(XMLNameSpace + "name");
+									XAttribute? Attribute = Element.Attribute(XMLDefaultNameSpace + "name");
 									if (Attribute != null)
 									{
 										if (Attribute.Value == Name)
@@ -1420,12 +1446,12 @@ namespace UnrealBuildTool
 
 					case "removePermission":
 						{
-							string? Name = GetAttributeWithNamespace(CurrentContext, Node, XMLNameSpace, "name");
+							string? Name = GetAttributeWithNamespace(CurrentContext, Node, XMLDefaultNameSpace, "name");
 							if (Name != null)
 							{
 								foreach (XElement Element in XMLWork.Descendants("uses-permission"))
 								{
-									XAttribute? Attribute = Element.Attribute(XMLNameSpace + "name");
+									XAttribute? Attribute = Element.Attribute(XMLDefaultNameSpace + "name");
 									if (Attribute != null)
 									{
 										if (Attribute.Value == Name)
@@ -1441,14 +1467,14 @@ namespace UnrealBuildTool
 
 					case "addFeature":
 						{
-							string? Name = GetAttributeWithNamespace(CurrentContext, Node, XMLNameSpace, "name");
+							string? Name = GetAttributeWithNamespace(CurrentContext, Node, XMLDefaultNameSpace, "name");
 							if (Name != null)
 							{
 								// make sure it isn't already added
 								bool bFound = false;
 								foreach (XElement Element in XMLWork.Descendants("uses-feature"))
 								{
-									XAttribute? Attribute = Element.Attribute(XMLNameSpace + "name");
+									XAttribute? Attribute = Element.Attribute(XMLDefaultNameSpace + "name");
 									if (Attribute != null)
 									{
 										if (Attribute.Value == Name)
@@ -1478,12 +1504,12 @@ namespace UnrealBuildTool
 
 					case "removeFeature":
 						{
-							string? Name = GetAttributeWithNamespace(CurrentContext, Node, XMLNameSpace, "name");
+							string? Name = GetAttributeWithNamespace(CurrentContext, Node, XMLDefaultNameSpace, "name");
 							if (Name != null)
 							{
 								foreach (XElement Element in XMLWork.Descendants("uses-feature"))
 								{
-									XAttribute? Attribute = Element.Attribute(XMLNameSpace + "name");
+									XAttribute? Attribute = Element.Attribute(XMLDefaultNameSpace + "name");
 									if (Attribute != null)
 									{
 										if (Attribute.Value == Name)
@@ -1499,14 +1525,14 @@ namespace UnrealBuildTool
 
 					case "addLibrary":
 						{
-							string? Name = GetAttributeWithNamespace(CurrentContext, Node, XMLNameSpace, "name");
+							string? Name = GetAttributeWithNamespace(CurrentContext, Node, XMLDefaultNameSpace, "name");
 							if (Name != null)
 							{
 								// make sure it isn't already added
 								bool bFound = false;
 								foreach (XElement Element in XMLWork.Descendants("uses-library"))
 								{
-									XAttribute? Attribute = Element.Attribute(XMLNameSpace + "name");
+									XAttribute? Attribute = Element.Attribute(XMLDefaultNameSpace + "name");
 									if (Attribute != null)
 									{
 										if (Attribute.Value == Name)
@@ -1536,12 +1562,12 @@ namespace UnrealBuildTool
 
 					case "removeLibrary":
 						{
-							string? Name = GetAttributeWithNamespace(CurrentContext, Node, XMLNameSpace, "name");
+							string? Name = GetAttributeWithNamespace(CurrentContext, Node, XMLDefaultNameSpace, "name");
 							if (Name != null)
 							{
 								foreach (XElement Element in XMLWork.Descendants("uses-library"))
 								{
-									XAttribute? Attribute = Element.Attribute(XMLNameSpace + "name");
+									XAttribute? Attribute = Element.Attribute(XMLDefaultNameSpace + "name");
 									if (Attribute != null)
 									{
 										if (Attribute.Value == Name)
@@ -2318,17 +2344,8 @@ namespace UnrealBuildTool
 									}
 								}
 
-								XAttribute? Attribute;
-								int Index = Name.IndexOf(":");
-								if (Index >= 0)
-								{
-									Name = Name.Substring(Index + 1);
-									Attribute = Element.Attribute(XMLNameSpace + Name);
-								}
-								else
-								{
-									Attribute = Element.Attribute(Name);
-								}
+								XNamespace? XMLNameSpace = TrimNamespaceAliasFromName(ref Name);
+								XAttribute? Attribute = Element.Attribute(XMLNameSpace != null ? XMLNameSpace + Name : Name);
 
 								CurrentContext.StringVariables[Result] = (Attribute != null) ? Attribute.Value : "";
 							}
@@ -2582,7 +2599,8 @@ namespace UnrealBuildTool
 							{
 								if (Value != null)
 								{
-									XElement NewElement = new XElement(Value);
+									XNamespace? XMLNameSpace = TrimNamespaceAliasFromName(ref Value);
+									XElement NewElement = new XElement(XMLNameSpace != null ? XMLNameSpace + Value : Value);
 									if (Text != null)
 									{
 										NewElement.Value = Text;
