UnrealBuildTool: Handle spaces in VSWorkspace build tool paths by adding quotes

#rnx

diff --git a/Engine/Source/Programs/UnrealBuildTool/ProjectFiles/VisualStudioWorkspace/VSWorkspaceProjectFile.cs b/Engine/Source/Programs/UnrealBuildTool/ProjectFiles/VisualStudioWorkspace/VSWorkspaceProjectFile.cs
--- a/Engine/Source/Programs/UnrealBuildTool/ProjectFiles/VisualStudioWorkspace/VSWorkspaceProjectFile.cs
+++ b/Engine/Source/Programs/UnrealBuildTool/ProjectFiles/VisualStudioWorkspace/VSWorkspaceProjectFile.cs
@@ -287,9 +287,9 @@ namespace UnrealBuildTool
 
 			return new TargetBuildInfo()
 			{
-				BuildCmd = $"{BuildCommandBuilder.BuildScript.FullName} {BuildArguments}",
-				RebuildCmd = $"{BuildCommandBuilder.RebuildScript.FullName} {BuildArguments}",
-				CleanCmd = $"{BuildCommandBuilder.CleanScript.FullName} {BuildArguments}",
+				BuildCmd = $"\"{BuildCommandBuilder.BuildScript.FullName}\" {BuildArguments}",
+				RebuildCmd = $"\"{BuildCommandBuilder.RebuildScript.FullName}\" {BuildArguments}",
+				CleanCmd = $"\"{BuildCommandBuilder.CleanScript.FullName}\" {BuildArguments}",
 				PrimaryOutput = Target.Binaries[0].OutputFilePath.FullName,
 				BuildByDefault = bBuildByDefault,
 			};
