[UBT]
* Added so uba executor loads uba config file.
* Added config file which enables some options for cache client (validate cas outputs and disabling directory preparsing

diff --git a/Engine/Source/Programs/Shared/EpicGames.UBA/IConfig.cs b/Engine/Source/Programs/Shared/EpicGames.UBA/IConfig.cs
new file mode 100644
--- /dev/null
+++ b/Engine/Source/Programs/Shared/EpicGames.UBA/IConfig.cs
@@ -0,0 +1,20 @@
+// Copyright Epic Games, Inc. All Rights Reserved.
+
+namespace EpicGames.UBA
+{
+	/// <summary>
+	/// Base interface for uba config file
+	/// </summary>
+	public interface IConfig
+	{
+		/// <summary>
+		/// Load a config file
+		/// </summary>
+		/// <param name="configFile">The name of the config file</param>
+		/// <returns>The IConfig</returns>
+		public static IConfig LoadConfig(string configFile)
+		{
+			return new ConfigImpl(configFile);
+		}
+	}
+}
\ No newline at end of file
diff --git a/Engine/Source/Programs/Shared/EpicGames.UBA/Impl/ConfigImpl.cs b/Engine/Source/Programs/Shared/EpicGames.UBA/Impl/ConfigImpl.cs
new file mode 100644
--- /dev/null
+++ b/Engine/Source/Programs/Shared/EpicGames.UBA/Impl/ConfigImpl.cs
@@ -0,0 +1,19 @@
+// Copyright Epic Games, Inc. All Rights Reserved.
+
+using System.Runtime.InteropServices;
+
+namespace EpicGames.UBA
+{
+	internal class ConfigImpl : IConfig
+	{
+		#region DllImport
+		[DllImport("UbaHost", CharSet = CharSet.Auto)]
+		static extern bool Config_Load(string configFile);
+		#endregion
+
+		public ConfigImpl(string configFile)
+		{
+			Config_Load(configFile);
+		}
+	}
+}
\ No newline at end of file
diff --git a/Engine/Source/Programs/UnrealBuildTool/Executors/UnrealBuildAccelerator/UBAAgentCoordinatorHorde.cs b/Engine/Source/Programs/UnrealBuildTool/Executors/UnrealBuildAccelerator/UBAAgentCoordinatorHorde.cs
--- a/Engine/Source/Programs/UnrealBuildTool/Executors/UnrealBuildAccelerator/UBAAgentCoordinatorHorde.cs
+++ b/Engine/Source/Programs/UnrealBuildTool/Executors/UnrealBuildAccelerator/UBAAgentCoordinatorHorde.cs
@@ -132,13 +132,10 @@ namespace UnrealBuildTool
 			_client = _serviceProvider.GetRequiredService<IHordeClient>().CreateComputeClient();
 
 			_logger.LogInformation("Creating tool bundle...");
-			DirectoryReference ubaDir;
+			DirectoryReference ubaDir = UBAExecutor.UbaBinariesDir;
 			List<string> agentFiles = new();
 			if (OperatingSystem.IsWindows())
 			{
-#pragma warning disable CA1308 // Normalize strings to uppercase
-				ubaDir = DirectoryReference.Combine(Unreal.EngineDirectory, "Binaries", "Win64", "UnrealBuildAccelerator", RuntimeInformation.ProcessArchitecture.ToString().ToLowerInvariant());
-#pragma warning restore CA1308 // Normalize strings to uppercase
 				agentFiles.Add("UbaAgent.exe");
 				if (useSentry)
 				{
@@ -157,24 +154,10 @@ namespace UnrealBuildTool
 					agentFiles.Add("libclang_rt.tsan.so"); // Needs to be copied from autosdk
 					agentFiles.Add("llvm-symbolizer"); // Needs to be copied from autosdk
 				}
-
-				if (RuntimeInformation.ProcessArchitecture == Architecture.X64)
-				{
-					ubaDir = DirectoryReference.Combine(Unreal.EngineDirectory, "Binaries", "Linux", "UnrealBuildAccelerator");
-				}
-				else if (RuntimeInformation.ProcessArchitecture == Architecture.Arm64)
-				{
-					ubaDir = DirectoryReference.Combine(Unreal.EngineDirectory, "Binaries", "LinuxArm64", "UnrealBuildAccelerator");
-				}
-				else
-				{
-					throw new PlatformNotSupportedException();
-				}
 			}
 			else if (OperatingSystem.IsMacOS())
 			{
 				agentFiles.Add("UbaAgent");
-				ubaDir = DirectoryReference.Combine(Unreal.EngineDirectory, "Binaries", "Mac", "UnrealBuildAccelerator");
 			}
 			else
 			{
diff --git a/Engine/Source/Programs/UnrealBuildTool/Executors/UnrealBuildAccelerator/UBAExecutor.cs b/Engine/Source/Programs/UnrealBuildTool/Executors/UnrealBuildAccelerator/UBAExecutor.cs
--- a/Engine/Source/Programs/UnrealBuildTool/Executors/UnrealBuildAccelerator/UBAExecutor.cs
+++ b/Engine/Source/Programs/UnrealBuildTool/Executors/UnrealBuildAccelerator/UBAExecutor.cs
@@ -68,6 +68,35 @@ namespace UnrealBuildTool
 			return EpicGames.UBA.Utils.IsAvailable();
 		}
 
+		public static DirectoryReference UbaBinariesDir
+		{
+			get
+			{
+				if (OperatingSystem.IsWindows())
+				{
+					#pragma warning disable CA1308 // Normalize strings to uppercase
+					return DirectoryReference.Combine(Unreal.EngineDirectory, "Binaries", "Win64", "UnrealBuildAccelerator", RuntimeInformation.ProcessArchitecture.ToString().ToLowerInvariant());
+					#pragma warning restore CA1308 // Normalize strings to uppercase
+				}
+				else if (OperatingSystem.IsLinux())
+				{
+					if (RuntimeInformation.ProcessArchitecture == Architecture.X64)
+					{
+						return DirectoryReference.Combine(Unreal.EngineDirectory, "Binaries", "Linux", "UnrealBuildAccelerator");
+					}
+					else if (RuntimeInformation.ProcessArchitecture == Architecture.Arm64)
+					{
+						return DirectoryReference.Combine(Unreal.EngineDirectory, "Binaries", "LinuxArm64", "UnrealBuildAccelerator");
+					}
+				}
+				else if (OperatingSystem.IsMacOS())
+				{
+					return DirectoryReference.Combine(Unreal.EngineDirectory, "Binaries", "Mac", "UnrealBuildAccelerator");
+				}
+				throw new PlatformNotSupportedException();
+			}
+		}
+
 		public UBAExecutor(int maxLocalActions, bool bAllCores, bool bCompactOutput, Microsoft.Extensions.Logging.ILogger logger, CommandLineArguments? additionalArguments = null)
 			: base(maxLocalActions, bAllCores, bCompactOutput, logger)
 		{
@@ -292,6 +321,10 @@ namespace UnrealBuildTool
 					_ = Task.Run(LaunchVisualizer);
 				}
 
+				string arch = RuntimeInformation.ProcessArchitecture.ToString().ToLowerInvariant();
+				FileReference configFile = FileReference.Combine(UbaBinariesDir, "UbaHost.toml");
+				EpicGames.UBA.IConfig.LoadConfig(configFile.FullName);
+
 				using EpicGames.UBA.ILogger ubaLogger = EpicGames.UBA.ILogger.CreateLogger(logger);
 				using (Server = IServer.CreateServer(UBAConfig.MaxWorkers, UBAConfig.SendSize, ubaLogger, UBAConfig.bUseQuic))
 				{
