IAS on IOS - engine changes

diff --git a/Engine/Source/Programs/UnrealBuildTool/Platform/IOS/IOSToolChain.cs b/Engine/Source/Programs/UnrealBuildTool/Platform/IOS/IOSToolChain.cs
--- a/Engine/Source/Programs/UnrealBuildTool/Platform/IOS/IOSToolChain.cs
+++ b/Engine/Source/Programs/UnrealBuildTool/Platform/IOS/IOSToolChain.cs
@@ -171,7 +171,8 @@ namespace UnrealBuildTool
 		{
 			base.GetCompileArguments_Debugging(CompileEnvironment, Arguments);
 
-			Arguments.Add("-fvisibility=hidden"); // hides the linker warnings with PhysX
+			Arguments.Add("-fvisibility=hidden");
+			Arguments.Add("-fvisibility-inlines-hidden");
 		}
 
 		/// <inheritdoc/>
@@ -208,6 +209,19 @@ namespace UnrealBuildTool
 		{
 			base.GetCompileArguments_Global(CompileEnvironment, Arguments);
 
+			if (CompileEnvironment.Configuration == CppConfiguration.Shipping)
+			{
+				Arguments.Add("-ffunction-sections");
+				Arguments.Add("-fdata-sections");
+				Arguments.Add("-fno-unwind-tables");
+				Arguments.Add("-fno-asynchronous-unwind-tables");
+
+				if (!CompileEnvironment.bPGOProfile)
+				{
+					Arguments.Add("-fno-use-cxa-atexit");
+				}
+			}
+
 			// What architecture(s) to build for
 			Arguments.Add(FormatArchitectureArg(CompileEnvironment.Architectures));
 
@@ -323,6 +337,8 @@ namespace UnrealBuildTool
 			Arguments.Add("-ObjC");
 			// Arguments.Add("-v");
 
+			Arguments.Add("-Wl,-O3");
+
 			// use LTO if desired (like VCToolchain does)
 			if (LinkEnvironment.bAllowLTCG)
 			{
diff --git a/Engine/Source/Runtime/Core/Private/IOS/IOSPlatformFile.cpp b/Engine/Source/Runtime/Core/Private/IOS/IOSPlatformFile.cpp
--- a/Engine/Source/Runtime/Core/Private/IOS/IOSPlatformFile.cpp
+++ b/Engine/Source/Runtime/Core/Private/IOS/IOSPlatformFile.cpp
@@ -15,9 +15,9 @@
 #include "Async/MappedFileHandle.h"
 #include <sys/mman.h>
 
-//#if PLATFORM_IOS
-//#include "IPlatformFileSandboxWrapper.h"
-//#endif
+#if PLATFORM_USE_PLATFORM_FILE_MANAGED_STORAGE_WRAPPER
+#include "HAL/IPlatformFileManagedStorageWrapper.h"
+#endif //PLATFORM_USE_PLATFORM_FILE_MANAGED_STORAGE_WRAPPER
 
 // make an FTimeSpan object that represents the "epoch" for time_t (from a stat struct)
 const FDateTime IOSEpoch(1970, 1, 1);
@@ -1048,6 +1048,13 @@ FString FIOSPlatformFile::ConvertToIOSPath(const FString& Filename, bool bForWri
 	{
 		return Result;
 	}
+
+	if (Result.StartsWith(TEXT("~/")))
+	{
+		static FString ReadPathBase = FString([[NSBundle mainBundle]bundlePath] );
+		Result.ReplaceInline(TEXT("~"), TEXT(""));
+		return ReadPathBase + Result;
+	}
 	
 	FPaths::MakePlatformFilename(Result);
 	
diff --git a/Engine/Source/Runtime/Experimental/IoStore/OnDemand/Private/IoStoreOnDemand.cpp b/Engine/Source/Runtime/Experimental/IoStore/OnDemand/Private/IoStoreOnDemand.cpp
--- a/Engine/Source/Runtime/Experimental/IoStore/OnDemand/Private/IoStoreOnDemand.cpp
+++ b/Engine/Source/Runtime/Experimental/IoStore/OnDemand/Private/IoStoreOnDemand.cpp
@@ -339,7 +339,12 @@ static void LoadCaCerts()
  */
 static TUniquePtr<FArchive> CreateReaderFromPlatformPackage(const FString& RelPath)
 {
-	const FString AbsPath = FPaths::Combine(FGenericPlatformMisc::RootDir(), RelPath);
+#if PLATFORM_IOS
+    // IOS OpenRead assumes it is in cookeddata, using ~ for the base path tells it to use the package base path instead
+    const FString AbsPath = FPaths::Combine(TEXT("~"), RelPath);
+#else
+    const FString AbsPath = FPaths::Combine(FGenericPlatformMisc::RootDir(), RelPath);
+#endif
 	if (TUniquePtr<IFileHandle> File(IPlatformFile::GetPlatformPhysical().OpenRead(*AbsPath)); File.IsValid())
 	{
 #if PLATFORM_ANDROID
