[UE] Review suggestions: renamed the bTrackRHIResourceInfo to bTrackRHIResourceInfoForTest
#rnx

diff --git a/Engine/Source/Programs/UnrealBuildTool/Configuration/ReadOnlyTargetRules.cs b/Engine/Source/Programs/UnrealBuildTool/Configuration/ReadOnlyTargetRules.cs
--- a/Engine/Source/Programs/UnrealBuildTool/Configuration/ReadOnlyTargetRules.cs
+++ b/Engine/Source/Programs/UnrealBuildTool/Configuration/ReadOnlyTargetRules.cs
@@ -172,7 +172,7 @@ namespace UnrealBuildTool
 
 		public bool bUseIris => Inner.bUseIris;
 
-		public bool bTrackRHIResourceInfo => Inner.bTrackRHIResourceInfo;
+		public bool bTrackRHIResourceInfoForTest => Inner.bTrackRHIResourceInfoForTest;
 
 		[Obsolete("Deprecated in UE5.4 - No longer used.")]
 		public bool bCompileIntelMetricsDiscovery => Inner.bCompileIntelMetricsDiscovery;
diff --git a/Engine/Source/Programs/UnrealBuildTool/Configuration/TargetRules.cs b/Engine/Source/Programs/UnrealBuildTool/Configuration/TargetRules.cs
--- a/Engine/Source/Programs/UnrealBuildTool/Configuration/TargetRules.cs
+++ b/Engine/Source/Programs/UnrealBuildTool/Configuration/TargetRules.cs
@@ -1020,7 +1020,7 @@ namespace UnrealBuildTool
 		/// Useful for ListShaderMaps and ListShaderLibraries commands.
 		/// </summary>
 		[RequiresUniqueBuildEnvironment]
-		public bool bTrackRHIResourceInfo { get; set; } = false;
+		public bool bTrackRHIResourceInfoForTest { get; set; } = false;
 
 		/// <summary>
 		/// Whether we are compiling editor code or not. Prefer the more explicit bCompileAgainstEditor instead.
diff --git a/Engine/Source/Runtime/RHI/RHI.Build.cs b/Engine/Source/Runtime/RHI/RHI.Build.cs
--- a/Engine/Source/Runtime/RHI/RHI.Build.cs
+++ b/Engine/Source/Runtime/RHI/RHI.Build.cs
@@ -24,7 +24,7 @@ public class RHI : ModuleRules
 				// let the user decide to use for Test
 				if (Target.Configuration == UnrealTargetConfiguration.DebugGame ||
 					Target.Configuration == UnrealTargetConfiguration.Development ||
-					Target.Configuration == UnrealTargetConfiguration.Test && Target.bTrackRHIResourceInfo)
+					(Target.Configuration == UnrealTargetConfiguration.Test && Target.bTrackRHIResourceInfoForTest))
 				{
 					PublicDefinitions.Add("RHI_WANT_RESOURCE_INFO=1");
 				}
