UnrealBuildTool: Error earlier if there are duplicate named source files in a module, as this will cause intermediate file conflicts in NonUnity builds

#rnx

diff --git a/Engine/Source/Programs/UnrealBuildTool/Configuration/UEBuildModuleCPP.cs b/Engine/Source/Programs/UnrealBuildTool/Configuration/UEBuildModuleCPP.cs
--- a/Engine/Source/Programs/UnrealBuildTool/Configuration/UEBuildModuleCPP.cs
+++ b/Engine/Source/Programs/UnrealBuildTool/Configuration/UEBuildModuleCPP.cs
@@ -34,6 +34,9 @@ namespace UnrealBuildTool
 			public readonly List<FileItem> SwiftFiles = new List<FileItem>();
 			public readonly List<FileItem> RCFiles = new List<FileItem>();
 			public readonly List<FileItem> ISPCFiles = new List<FileItem>();
+
+			public IEnumerable<FileItem> AllHeaderFiles => HeaderFiles.Concat(ISPCHeaderFiles);
+			public IEnumerable<FileItem> AllCompiledFiles => IXXFiles.Concat(CPPFiles).Concat(CFiles).Concat(CCFiles).Concat(MFiles).Concat(MMFiles).Concat(SwiftFiles).Concat(RCFiles).Concat(ISPCFiles);
 		}
 
 		/// <summary>
@@ -396,13 +399,35 @@ namespace UnrealBuildTool
 			Dictionary<DirectoryItem, FileItem[]> DirectoryToSourceFiles = new Dictionary<DirectoryItem, FileItem[]>();
 			InputFileCollection InputFiles = FindInputFiles(Target.Platform, DirectoryToSourceFiles, Logger);
 
+			{
+				Dictionary<string, List<FileItem>> nameToPath = new();
+				foreach (FileItem InputFile in InputFiles.AllCompiledFiles)
+				{
+					string name = InputFile.Name.ToUpperInvariant();
+					if (!nameToPath.ContainsKey(name))
+					{
+						nameToPath.Add(name, new List<FileItem>());
+					}
+					nameToPath[name].Add(InputFile);
+				}
+				IEnumerable<KeyValuePair<string, List<FileItem>>> fileConflicts = nameToPath.Where(item => item.Value.Count > 1);
+				if (fileConflicts.Any())
+				{
+					Logger.LogInformation("Input filename conflicts:");
+					foreach (KeyValuePair<string, List<FileItem>> item in fileConflicts)
+					{
+						item.Value.ForEach(x => Logger.LogInformation("* {Path}", x));
+					}
+					throw new BuildException("Multiple input files found with duplicate filenames, this is is not allowed as intermediate output files for non-unity builds will conflict");
+				}
+			}
+
 			foreach (KeyValuePair<DirectoryItem, FileItem[]> Pair in DirectoryToSourceFiles)
 			{
 				Graph.AddSourceFiles(Pair.Key, Pair.Value);
 			}
 
-			Graph.AddHeaderFiles(InputFiles.HeaderFiles.ToArray());
-			Graph.AddHeaderFiles(InputFiles.ISPCHeaderFiles.ToArray());
+			Graph.AddHeaderFiles(InputFiles.AllHeaderFiles.ToArray());
 
 			// We are building with IWYU and thismodule does not support it, early out
 			if (Target.bIWYU && Rules.IWYUSupport == IWYUSupport.None)
@@ -945,7 +970,14 @@ namespace UnrealBuildTool
 			{
 				if (DirectoryLookupCache.DirectoryExists(Directory))
 				{
-					return DirectoryLookupCache.EnumerateFiles(Directory).ToHashSet();
+					try
+					{
+						return DirectoryLookupCache.EnumerateFiles(Directory).ToHashSet();
+					}
+					catch (Exception)
+					{
+						return null;
+					}
 				}
 				else
 				{
