fix for IOS linker failure

diff --git a/Engine/Source/Programs/UnrealBuildTool/Platform/IOS/IOSToolChain.cs b/Engine/Source/Programs/UnrealBuildTool/Platform/IOS/IOSToolChain.cs
--- a/Engine/Source/Programs/UnrealBuildTool/Platform/IOS/IOSToolChain.cs
+++ b/Engine/Source/Programs/UnrealBuildTool/Platform/IOS/IOSToolChain.cs
@@ -337,7 +337,7 @@ namespace UnrealBuildTool
 			Arguments.Add("-ObjC");
 			// Arguments.Add("-v");
 
-			Arguments.Add("-Wl,-O3");
+			//Arguments.Add("-Wl,-O3");
 
 			// use LTO if desired (like VCToolchain does)
 			if (LinkEnvironment.bAllowLTCG)
