Enable Visual Studio's "keep the console window open when debugging stops" in UBT

See https://developercommunity.visualstudio.com/t/Automatically-close-the-console-when-de/10475212#T-ND10674286

diff --git a/Engine/Source/Programs/UnrealBuildTool/ProjectFiles/VisualStudio/VCProject.cs b/Engine/Source/Programs/UnrealBuildTool/ProjectFiles/VisualStudio/VCProject.cs
--- a/Engine/Source/Programs/UnrealBuildTool/ProjectFiles/VisualStudio/VCProject.cs
+++ b/Engine/Source/Programs/UnrealBuildTool/ProjectFiles/VisualStudio/VCProject.cs
@@ -2484,6 +2484,13 @@ namespace UnrealBuildTool
 					VCProjectFileContent.AppendLine("    <NMakeCompile>");
 					VCProjectFileContent.AppendLine("      <NMakeCompileFileCommandLine>$(BuildBatchScript) {0} -WorkingDir=$(MSBuildProjectDirectory) -Files=$(SelectedFiles)</NMakeCompileFileCommandLine>", BuildArguments);
 					VCProjectFileContent.AppendLine("    </NMakeCompile>");
+					if (TargetRulesObject.bIsBuildingConsoleApplication)
+					{
+						// Let Visual Studio keep the console window open when debugging stops. Ignored by the build.
+						VCProjectFileContent.AppendLine("    <Link>");
+						VCProjectFileContent.AppendLine("      <SubSystem>Console</SubSystem>");
+						VCProjectFileContent.AppendLine("    </Link>");
+					}
 					if (ProjectFileGenerator.bVisualStudioLinux && TargetRulesObject.Platform.IsInGroup(UnrealPlatformGroup.Linux))
 					{
 						VCProjectFileContent.AppendLine("    <PostBuildEvent>");
