UnrealBuildTool: Add RequiresUniqueBuildEnvironment to TargetType

#rnx

diff --git a/Engine/Source/Programs/UnrealBuildTool/Configuration/TargetRules.cs b/Engine/Source/Programs/UnrealBuildTool/Configuration/TargetRules.cs
--- a/Engine/Source/Programs/UnrealBuildTool/Configuration/TargetRules.cs
+++ b/Engine/Source/Programs/UnrealBuildTool/Configuration/TargetRules.cs
@@ -640,6 +640,7 @@ namespace UnrealBuildTool
 		/// <summary>
 		/// The type of target.
 		/// </summary>
+		[RequiresUniqueBuildEnvironment]
 		public global::UnrealBuildTool.TargetType Type { get; set; } = global::UnrealBuildTool.TargetType.Game;
 
 		/// <summary>
