UnrealBuildTool: Replace usage of obsolete FormatterServices.GetUninitializedObject

#rnx

diff --git a/Engine/Source/Programs/UnrealBuildTool/Configuration/TargetRules.cs b/Engine/Source/Programs/UnrealBuildTool/Configuration/TargetRules.cs
--- a/Engine/Source/Programs/UnrealBuildTool/Configuration/TargetRules.cs
+++ b/Engine/Source/Programs/UnrealBuildTool/Configuration/TargetRules.cs
@@ -6,7 +6,7 @@ using System.Diagnostics.CodeAnalysis;
 using System.IO;
 using System.Linq;
 using System.Reflection;
-using System.Runtime.Serialization;
+using System.Runtime.CompilerServices;
 using EpicGames.Core;
 using Microsoft.Extensions.Logging;
 using UnrealBuildBase;
@@ -2917,7 +2917,8 @@ namespace UnrealBuildTool
 		/// <returns>Target instance</returns>
 		public static TargetRules Create(Type rulesType, TargetInfo targetInfo, FileReference? baseFile, FileReference? platformFile, IEnumerable<FileReference>? targetFiles, BuildSettingsVersion? defaultBuildSettings, ILogger logger)
 		{
-			TargetRules rules = (TargetRules)FormatterServices.GetUninitializedObject(rulesType);
+
+			TargetRules rules = (TargetRules)RuntimeHelpers.GetUninitializedObject(rulesType);
 			if (defaultBuildSettings.HasValue)
 			{
 				rules.DefaultBuildSettings = defaultBuildSettings.Value;
diff --git a/Engine/Source/Programs/UnrealBuildTool/Configuration/TestTargetRules.cs b/Engine/Source/Programs/UnrealBuildTool/Configuration/TestTargetRules.cs
--- a/Engine/Source/Programs/UnrealBuildTool/Configuration/TestTargetRules.cs
+++ b/Engine/Source/Programs/UnrealBuildTool/Configuration/TestTargetRules.cs
@@ -3,7 +3,7 @@
 using System;
 using System.IO;
 using System.Reflection;
-using System.Runtime.Serialization;
+using System.Runtime.CompilerServices;
 using EpicGames.Core;
 using UnrealBuildBase;
 
@@ -209,7 +209,7 @@ namespace UnrealBuildTool
 		public static TestTargetRules Create(TargetRules rules, TargetInfo targetInfo)
 		{
 			Type testTargetRulesType = typeof(TestTargetRules);
-			TestTargetRules testRules = (TestTargetRules)FormatterServices.GetUninitializedObject(testTargetRulesType);
+			TestTargetRules testRules = (TestTargetRules)RuntimeHelpers.GetUninitializedObject(testTargetRulesType);
 
 			// Initialize the logger before calling the constructor
 			testRules.Logger = rules.Logger;
diff --git a/Engine/Source/Programs/UnrealBuildTool/System/RulesAssembly.cs b/Engine/Source/Programs/UnrealBuildTool/System/RulesAssembly.cs
--- a/Engine/Source/Programs/UnrealBuildTool/System/RulesAssembly.cs
+++ b/Engine/Source/Programs/UnrealBuildTool/System/RulesAssembly.cs
@@ -5,7 +5,7 @@ using System.Collections.Generic;
 using System.Diagnostics.CodeAnalysis;
 using System.Linq;
 using System.Reflection;
-using System.Runtime.Serialization;
+using System.Runtime.CompilerServices;
 using EpicGames.Core;
 using Microsoft.Extensions.Logging;
 using UnrealBuildBase;
@@ -520,7 +520,8 @@ namespace UnrealBuildTool
 			try
 			{
 				// Create an uninitialized ModuleRules object and set some defaults.
-				ModuleRules RulesObject = (ModuleRules)FormatterServices.GetUninitializedObject(RulesObjectType);
+				ModuleRules RulesObject = (ModuleRules)RuntimeHelpers.GetUninitializedObject(RulesObjectType);
+
 				// even if we created a platform-extension version of the module rules, we are pretending to be
 				// the base type, so that no one else needs to manage this
 				RulesObject.Name = ModuleName;
