[UBA]
* Disabled pch timestamp check on windows since this prevents pch from being able to be cached (this also prevents them from being deterministic)

diff --git a/Engine/Source/Programs/UnrealBuildTool/Platform/Windows/VCToolChain.cs b/Engine/Source/Programs/UnrealBuildTool/Platform/Windows/VCToolChain.cs
--- a/Engine/Source/Programs/UnrealBuildTool/Platform/Windows/VCToolChain.cs
+++ b/Engine/Source/Programs/UnrealBuildTool/Platform/Windows/VCToolChain.cs
@@ -1791,6 +1791,12 @@ namespace UnrealBuildTool
 					if (CompileEnvironment.PrecompiledHeaderAction != PrecompiledHeaderAction.None)
 					{
 						CompileAction.ArtifactMode |= ArtifactMode.AbsolutePath; // Unfortunately we require matching absolute paths for pch to be cached
+
+						if (Target.WindowsPlatform.Compiler.IsClang())
+						{
+							CompileAction.Arguments.Add("-Xclang -fno-pch-timestamp"); // This is needed to prevent check on timestamp stored inside pch
+							CompileAction.Arguments.Add("-Xclang -fvalidate-ast-input-files-content"); // Validate PCH inputs by content if mtime check fails
+						}
 					}
 				}
 
