Really really disable C5105 when doing PVS analysis as disabling it in PlatformCompilerSetup.h AND disabling it DirectHeaderCompiling.h just wasn't enough
#rnx

diff --git a/Engine/Source/Programs/UnrealBuildTool/Platform/Windows/PVSToolChain.cs b/Engine/Source/Programs/UnrealBuildTool/Platform/Windows/PVSToolChain.cs
--- a/Engine/Source/Programs/UnrealBuildTool/Platform/Windows/PVSToolChain.cs
+++ b/Engine/Source/Programs/UnrealBuildTool/Platform/Windows/PVSToolChain.cs
@@ -790,7 +790,7 @@ namespace UnrealBuildTool
 			CppCompileEnvironment PreprocessCompileEnvironment = new CppCompileEnvironment(CompileEnvironment);
 			PreprocessCompileEnvironment.bPreprocessOnly = true;
 			PreprocessCompileEnvironment.bEnableUndefinedIdentifierWarnings = false; // Not sure why THIRD_PARTY_INCLUDES_START doesn't pick this up; the _Pragma appears in the preprocessed output. Perhaps in preprocess-only mode the compiler doesn't respect these?
-			PreprocessCompileEnvironment.AdditionalArguments += " /wd4005 /wd4828";
+			PreprocessCompileEnvironment.AdditionalArguments += " /wd4005 /wd4828 /wd5105";
 			PreprocessCompileEnvironment.Definitions.Add("PVS_STUDIO");
 
 			List<IExternalAction> PreprocessActions = new List<IExternalAction>();
