UnrealBuildTool: Remove usage of obsolate BinarySerializer and replace with Json

#rnx

diff --git a/Engine/Source/Programs/AutomationTool/Turnkey/LocalCache.cs b/Engine/Source/Programs/AutomationTool/Turnkey/LocalCache.cs
--- a/Engine/Source/Programs/AutomationTool/Turnkey/LocalCache.cs
+++ b/Engine/Source/Programs/AutomationTool/Turnkey/LocalCache.cs
@@ -1,4 +1,4 @@
-﻿// Copyright Epic Games, Inc. All Rights Reserved.
+// Copyright Epic Games, Inc. All Rights Reserved.
 
 using System;
 using System.Collections.Generic;
@@ -147,11 +147,6 @@ namespace Turnkey
 			string Str = JsonSerializer.Serialize(Object);
 			Directory.CreateDirectory(Path.GetDirectoryName(Filename));
 			File.WriteAllText(Filename, Str);
-
-// 			using (FileStream Stream = new FileStream(Filename, FileMode.Create))
-// 			{
-// 				new BinaryFormatter().Serialize(Stream, Object);
-// 			}
 		}
 
 		static SavedCache DeserializeObject(string Filename)
@@ -168,14 +163,6 @@ namespace Turnkey
 				{
 					return JsonSerializer.Deserialize<SavedCache>(Str);
 				}
-
-// 				if (File.Exists(Filename))
-// 				{
-// 					using (FileStream Stream = new FileStream(Filename, FileMode.Open))
-// 					{
-// 						return (SavedCache)new BinaryFormatter().Deserialize(Stream);
-// 					}
-// 				}
 			}
 			catch(Exception Ex)
 			{
diff --git a/Engine/Source/Programs/Shared/EpicGames.Core.Tests/JsonSerializerUtilsTests.cs b/Engine/Source/Programs/Shared/EpicGames.Core.Tests/JsonSerializerUtilsTests.cs
new file mode 100644
--- /dev/null
+++ b/Engine/Source/Programs/Shared/EpicGames.Core.Tests/JsonSerializerUtilsTests.cs
@@ -0,0 +1,89 @@
+// Copyright Epic Games, Inc. All Rights Reserved.
+
+using System;
+using System.IO;
+using Microsoft.VisualStudio.TestTools.UnitTesting;
+
+namespace EpicGames.Core.Tests
+{
+	[TestClass]
+	public class JsonSerializerUtilsTests
+	{
+		private struct TestObject : IEquatable<TestObject>
+		{
+			public string _stringFieldPublic;
+			public string StringPropertyPublic { get; set; }
+
+			public TestObject(int testData)
+			{
+				_stringFieldPublic = testData.ToString();
+				StringPropertyPublic = (testData + 1).ToString();
+			}
+
+			public bool Equals(TestObject other) => _stringFieldPublic == other._stringFieldPublic && StringPropertyPublic == other.StringPropertyPublic;
+
+			public override bool Equals(object? obj) => obj is TestObject objTyped && Equals(objTyped);
+
+			public override int GetHashCode() => HashCode.Combine(_stringFieldPublic, StringPropertyPublic);
+		}
+
+		private static DirectoryReference s_tempDirectory = null!;
+
+		public JsonSerializerUtilsTests()
+		{
+			s_tempDirectory = CreateTempDir();
+		}
+
+		[TestCleanup]
+		public void RemoveTempDir()
+		{
+			if (DirectoryReference.Exists(s_tempDirectory))
+			{
+				DirectoryReference.Delete(s_tempDirectory, true);
+			}
+		}
+
+		[TestInitialize]
+		public void InitializeTest()
+		{
+		}
+
+		[TestMethod]
+		public void Save()
+		{
+			FileReference path = FileReference.Combine(s_tempDirectory, "save.json");
+			JsonSerializerUtils.Save(path, new TestObject(Int32.MaxValue));
+			Assert.IsTrue(FileReference.Exists(path));
+		}
+
+		[TestMethod]
+		public void SaveIfDifferent()
+		{
+			FileReference path = FileReference.Combine(s_tempDirectory, "saveIfDifferent.json");
+			JsonSerializerUtils.Save(path, new TestObject(0));
+			Assert.IsTrue(FileReference.Exists(path));
+			FileInfo initial = path.ToFileInfo();
+			JsonSerializerUtils.SaveIfDifferent(path, new TestObject(Int32.MaxValue));
+			FileInfo same = path.ToFileInfo();
+			Assert.IsTrue(initial.LastWriteTime.Equals(same.LastWriteTime));
+		}
+
+		[TestMethod]
+		public void Load()
+		{
+			FileReference path = FileReference.Combine(s_tempDirectory, "load.json");
+			TestObject obj = new TestObject(Int32.MaxValue);
+			JsonSerializerUtils.Save(path, obj);
+			Assert.IsTrue(FileReference.Exists(path));
+			TestObject obj2 = JsonSerializerUtils.Load<TestObject>(path);
+			Assert.IsTrue(obj.Equals(obj2));
+		}
+
+		private static DirectoryReference CreateTempDir()
+		{
+			string tempDir = Path.Join(Path.GetTempPath(), "epicgames-core-tests-" + Guid.NewGuid().ToString()[..8]);
+			Directory.CreateDirectory(tempDir);
+			return new DirectoryReference(tempDir);
+		}
+	}
+}
diff --git a/Engine/Source/Programs/Shared/EpicGames.Core/BinaryFormatterUtils.cs b/Engine/Source/Programs/Shared/EpicGames.Core/BinaryFormatterUtils.cs
deleted file mode 100644
diff --git a/Engine/Source/Programs/Shared/EpicGames.Core/DirectoryReference.cs b/Engine/Source/Programs/Shared/EpicGames.Core/DirectoryReference.cs
--- a/Engine/Source/Programs/Shared/EpicGames.Core/DirectoryReference.cs
+++ b/Engine/Source/Programs/Shared/EpicGames.Core/DirectoryReference.cs
@@ -7,6 +7,8 @@ using System.Diagnostics.CodeAnalysis;
 using System.Globalization;
 using System.IO;
 using System.Linq;
+using System.Text.Json;
+using System.Text.Json.Serialization;
 
 namespace EpicGames.Core
 {
@@ -14,6 +16,7 @@ namespace EpicGames.Core
 	/// Representation of an absolute directory path. Allows fast hashing and comparisons.
 	/// </summary>
 	[Serializable]
+	[JsonConverter(typeof(DirectoryReferenceJsonConverter))]
 	[TypeConverter(typeof(DirectoryReferenceTypeConverter))]
 	public class DirectoryReference : FileSystemReference, IEquatable<DirectoryReference>, IComparable<DirectoryReference>
 	{
@@ -357,10 +360,31 @@ namespace EpicGames.Core
 		#endregion
 	}
 
+	/// <summary>
+	/// Json converter to/from strings
+	/// </summary>
+	sealed class DirectoryReferenceJsonConverter : JsonConverter<DirectoryReference>
+	{
+		/// <inheritdoc/>
+		public override bool CanConvert(Type typeToConvert) => typeToConvert == typeof(string) || base.CanConvert(typeToConvert);
+
+		/// <inheritdoc/>
+		public override DirectoryReference? Read(ref Utf8JsonReader reader, Type typeToConvert, JsonSerializerOptions options) => DirectoryReference.FromString(reader.GetString());
+
+		/// <inheritdoc/>
+		public override void Write(Utf8JsonWriter writer, DirectoryReference value, JsonSerializerOptions options) => writer.WriteStringValue(value.ToString());
+
+		/// <inheritdoc/>
+		public override DirectoryReference ReadAsPropertyName(ref Utf8JsonReader reader, Type typeToConvert, JsonSerializerOptions options) => Read(ref reader, typeToConvert, options)!;
+
+		/// <inheritdoc/>
+		public override void WriteAsPropertyName(Utf8JsonWriter writer, DirectoryReference value, JsonSerializerOptions options) => writer.WritePropertyName(value.ToString());
+	}
+
 	/// <summary>
 	/// Type converter to/from strings
 	/// </summary>
-	class DirectoryReferenceTypeConverter : TypeConverter
+	sealed class DirectoryReferenceTypeConverter : TypeConverter
 	{
 		/// <inheritdoc/>
 		public override bool CanConvertFrom(ITypeDescriptorContext? context, Type sourceType)
diff --git a/Engine/Source/Programs/Shared/EpicGames.Core/FileReference.cs b/Engine/Source/Programs/Shared/EpicGames.Core/FileReference.cs
--- a/Engine/Source/Programs/Shared/EpicGames.Core/FileReference.cs
+++ b/Engine/Source/Programs/Shared/EpicGames.Core/FileReference.cs
@@ -8,6 +8,8 @@ using System.Globalization;
 using System.IO;
 using System.Linq;
 using System.Text;
+using System.Text.Json;
+using System.Text.Json.Serialization;
 using System.Threading;
 using System.Threading.Tasks;
 
@@ -17,6 +19,7 @@ namespace EpicGames.Core
 	/// Representation of an absolute file path. Allows fast hashing and comparisons.
 	/// </summary>
 	[Serializable]
+	[JsonConverter(typeof(FileReferenceJsonConverter))]
 	[TypeConverter(typeof(FileReferenceTypeConverter))]
 	public class FileReference : FileSystemReference, IEquatable<FileReference>, IComparable<FileReference>
 	{
@@ -729,10 +732,31 @@ namespace EpicGames.Core
 		#endregion
 	}
 
+	/// <summary>
+	/// Json converter to/from strings
+	/// </summary>
+	sealed class FileReferenceJsonConverter : JsonConverter<FileReference>
+	{
+		/// <inheritdoc/>
+		public override bool CanConvert(Type typeToConvert) => typeToConvert == typeof(string) || base.CanConvert(typeToConvert);
+
+		/// <inheritdoc/>
+		public override FileReference? Read(ref Utf8JsonReader reader, Type typeToConvert, JsonSerializerOptions options) => FileReference.FromString(reader.GetString());
+
+		/// <inheritdoc/>
+		public override void Write(Utf8JsonWriter writer, FileReference value, JsonSerializerOptions options) => writer.WriteStringValue(value.ToString());
+
+		/// <inheritdoc/>
+		public override FileReference ReadAsPropertyName(ref Utf8JsonReader reader, Type typeToConvert, JsonSerializerOptions options) => Read(ref reader, typeToConvert, options)!;
+
+		/// <inheritdoc/>
+		public override void WriteAsPropertyName(Utf8JsonWriter writer, FileReference value, JsonSerializerOptions options) => writer.WritePropertyName(value.ToString());
+	}
+
 	/// <summary>
 	/// Type converter to/from strings
 	/// </summary>
-	class FileReferenceTypeConverter : TypeConverter
+	sealed class FileReferenceTypeConverter : TypeConverter
 	{
 		/// <inheritdoc/>
 		public override bool CanConvertFrom(ITypeDescriptorContext? context, Type sourceType)
diff --git a/Engine/Source/Programs/Shared/EpicGames.Core/JsonSerializerUtils.cs b/Engine/Source/Programs/Shared/EpicGames.Core/JsonSerializerUtils.cs
new file mode 100644
--- /dev/null
+++ b/Engine/Source/Programs/Shared/EpicGames.Core/JsonSerializerUtils.cs
@@ -0,0 +1,59 @@
+// Copyright Epic Games, Inc. All Rights Reserved.
+
+using System;
+using System.IO;
+using System.Text.Json;
+
+namespace EpicGames.Core
+{
+	/// <summary>
+	/// Utility functions for serializing using the JsonSerializer, to replace BinaryFormatterUtils.
+	/// </summary>
+	public static class JsonSerializerUtils
+	{
+		static JsonSerializerOptions Options => new JsonSerializerOptions() { IncludeFields = true };
+
+		/// <summary>
+		/// Load an object from a file on disk, using the binary formatter
+		/// </summary>
+		/// <param name="location">File to read from</param>
+		/// <returns>Instance of the object that was read from disk</returns>
+		public static T Load<T>(FileReference location)
+		{
+			using FileStream stream = new FileStream(location.FullName, FileMode.Open, FileAccess.Read);
+			JsonSerializerOptions options = new JsonSerializerOptions();
+			return JsonSerializer.Deserialize<T>(stream, Options)
+				?? throw new TypeLoadException($"Type {typeof(T)} cannot be loaded from {location}");
+		}
+
+		/// <summary>
+		/// Saves a file to disk, using the binary formatter
+		/// </summary>
+		/// <param name="location">File to write to</param>
+		/// <param name="obj">Object to serialize</param>
+		public static void Save<T>(FileReference location, T obj)
+		{
+			DirectoryReference.CreateDirectory(location.Directory);
+			using FileStream stream = new FileStream(location.FullName, FileMode.Create, FileAccess.Write);
+			JsonSerializer.Serialize<T>(stream, obj, Options);
+		}
+
+		/// <summary>
+		/// Saves a file to disk using the binary formatter, without updating the timestamp if it hasn't changed
+		/// </summary>
+		/// <param name="location">File to write to</param>
+		/// <param name="obj">Object to serialize</param>
+		public static void SaveIfDifferent<T>(FileReference location, T obj)
+		{
+			byte[] contents;
+			using (MemoryStream stream = new MemoryStream())
+			{
+				JsonSerializer.Serialize(stream, obj, Options);
+				contents = stream.ToArray();
+			}
+
+			DirectoryReference.CreateDirectory(location.Directory);
+			FileReference.WriteAllBytesIfDifferent(location, contents);
+		}
+	}
+}
diff --git a/Engine/Source/Programs/Shared/EpicGames.Core/Md5Hash.cs b/Engine/Source/Programs/Shared/EpicGames.Core/Md5Hash.cs
--- a/Engine/Source/Programs/Shared/EpicGames.Core/Md5Hash.cs
+++ b/Engine/Source/Programs/Shared/EpicGames.Core/Md5Hash.cs
@@ -6,6 +6,9 @@ using System.ComponentModel;
 using System.Globalization;
 using System.IO;
 using System.Security.Cryptography;
+using System.Text.Json;
+using System.Text.Json.Serialization;
+
 
 #pragma warning disable CA5351 // Do Not Use Broken Cryptographic Algorithms
 
@@ -14,6 +17,7 @@ namespace EpicGames.Core
 	/// <summary>
 	/// Struct representing a strongly typed Md5Hash value
 	/// </summary>
+	[JsonConverter(typeof(Md5HashJsonConverter))]
 	[TypeConverter(typeof(Md5HashTypeConverter))]
 	public readonly struct Md5Hash : IEquatable<Md5Hash>, IComparable<Md5Hash>
 	{
@@ -89,20 +93,21 @@ namespace EpicGames.Core
 		/// </summary>
 		/// <param name="text"></param>
 		/// <returns></returns>
-		public static Md5Hash Parse(string text)
-		{
-			return new Md5Hash(StringUtils.ParseHexString(text));
-		}
+		public static Md5Hash Parse(string text) => new Md5Hash(StringUtils.ParseHexString(text));
 
 		/// <summary>
 		/// Parses a digest from the given hex string
 		/// </summary>
 		/// <param name="text"></param>
 		/// <returns></returns>
-		public static Md5Hash Parse(Utf8String text)
-		{
-			return new Md5Hash(StringUtils.ParseHexString(text.Span));
-		}
+		public static Md5Hash Parse(ReadOnlySpan<byte> text) => new Md5Hash(StringUtils.ParseHexString(text));
+
+		/// <summary>
+		/// Parses a digest from the given hex string
+		/// </summary>
+		/// <param name="text"></param>
+		/// <returns></returns>
+		public static Md5Hash Parse(Utf8String text) => Parse(text.Span);
 
 		/// <inheritdoc cref="IComparable{T}.CompareTo(T)"/>
 		public int CompareTo(Md5Hash other)
@@ -174,20 +179,26 @@ namespace EpicGames.Core
 		/// </summary>
 		/// <param name="reader"></param>
 		/// <returns></returns>
-		public static Md5Hash ReadMd5Hash(this MemoryReader reader)
-		{
-			return new Md5Hash(reader.ReadFixedLengthBytes(Md5Hash.NumBytes));
-		}
+		public static Md5Hash ReadMd5Hash(this MemoryReader reader) => new Md5Hash(reader.ReadFixedLengthBytes(Md5Hash.NumBytes));
 
 		/// <summary>
 		/// Write an <see cref="Md5Hash"/> to a memory writer
 		/// </summary>
 		/// <param name="writer"></param>
 		/// <param name="hash"></param>
-		public static void WriteMd5Hash(this MemoryWriter writer, Md5Hash hash)
-		{
-			writer.WriteFixedLengthBytes(hash.Span);
-		}
+		public static void WriteMd5Hash(this MemoryWriter writer, Md5Hash hash) => writer.WriteFixedLengthBytes(hash.Span);
+	}
+
+	/// <summary>
+	/// Json converter to/from strings
+	/// </summary>
+	sealed class Md5HashJsonConverter : JsonConverter<Md5Hash>
+	{
+		/// <inheritdoc/>
+		public override Md5Hash Read(ref Utf8JsonReader reader, Type typeToConvert, JsonSerializerOptions options) => Md5Hash.Parse(reader.GetUtf8String());
+
+		/// <inheritdoc/>
+		public override void Write(Utf8JsonWriter writer, Md5Hash value, JsonSerializerOptions options) => writer.WriteStringValue(value.Span);
 	}
 
 	/// <summary>
@@ -196,15 +207,9 @@ namespace EpicGames.Core
 	sealed class Md5HashTypeConverter : TypeConverter
 	{
 		/// <inheritdoc/>
-		public override bool CanConvertFrom(ITypeDescriptorContext? context, Type sourceType)
-		{
-			return sourceType == typeof(string);
-		}
+		public override bool CanConvertFrom(ITypeDescriptorContext? context, Type sourceType) => sourceType == typeof(string);
 
 		/// <inheritdoc/>
-		public override object ConvertFrom(ITypeDescriptorContext? context, CultureInfo? culture, object value)
-		{
-			return Md5Hash.Parse((string)value);
-		}
+		public override object ConvertFrom(ITypeDescriptorContext? context, CultureInfo? culture, object value) => Md5Hash.Parse((string)value);
 	}
 }
diff --git a/Engine/Source/Programs/Shared/EpicGames.Core/Obsolete/BinaryFormatterUtils.cs b/Engine/Source/Programs/Shared/EpicGames.Core/Obsolete/BinaryFormatterUtils.cs
new file mode 100644
--- /dev/null
+++ b/Engine/Source/Programs/Shared/EpicGames.Core/Obsolete/BinaryFormatterUtils.cs
@@ -0,0 +1,34 @@
+// Copyright Epic Games, Inc. All Rights Reserved.
+
+using System;
+
+namespace EpicGames.Core
+{
+	/// <summary>
+	/// Utility functions for serializing using BinaryFormatter, should not be used.
+	/// </summary>
+	[Obsolete("BinaryFormatter serialization and deserialization are obsolete. See https://aka.ms/binaryformatter for more information.")]
+	public static class BinaryFormatterUtils
+	{
+		/// <summary>
+		/// Load an object from a file on disk, using the json serializer
+		/// </summary>
+		/// <param name="location">File to read from</param>
+		/// <returns>Instance of the object that was read from disk</returns>
+		public static T Load<T>(FileReference location) => JsonSerializerUtils.Load<T>(location);
+
+		/// <summary>
+		/// Saves a file to disk, using the json serializer
+		/// </summary>
+		/// <param name="location">File to write to</param>
+		/// <param name="obj">Object to serialize</param>
+		public static void Save(FileReference location, object obj) => JsonSerializerUtils.Save(location, obj);
+
+		/// <summary>
+		/// Saves a file to disk using the binary serializer, without updating the timestamp if it hasn't changed
+		/// </summary>
+		/// <param name="location">File to write to</param>
+		/// <param name="obj">Object to serialize</param>
+		public static void SaveIfDifferent(FileReference location, object obj) => JsonSerializerUtils.SaveIfDifferent(location, obj);
+	}
+}
diff --git a/Engine/Source/Programs/UnrealBuildTool/Configuration/UEBuildTarget.cs b/Engine/Source/Programs/UnrealBuildTool/Configuration/UEBuildTarget.cs
--- a/Engine/Source/Programs/UnrealBuildTool/Configuration/UEBuildTarget.cs
+++ b/Engine/Source/Programs/UnrealBuildTool/Configuration/UEBuildTarget.cs
@@ -11,6 +11,7 @@ using System.Reflection;
 using System.Runtime.InteropServices;
 using System.Runtime.Serialization;
 using System.Text;
+using System.Text.Json;
 using System.Text.Json.Serialization;
 using System.Threading.Tasks;
 using EpicGames.Core;
@@ -24,7 +25,9 @@ namespace UnrealBuildTool
 	/// <summary>
 	/// The platform we're building for
 	/// </summary>
-	[Serializable, TypeConverter(typeof(UnrealTargetPlatformTypeConverter))]
+	[Serializable]
+	[JsonConverter(typeof(UnrealTargetPlatformJsonConverter))]
+	[TypeConverter(typeof(UnrealTargetPlatformTypeConverter))]
 	public partial struct UnrealTargetPlatform : ISerializable
 	{
 		#region Private/boilerplate
@@ -284,7 +287,14 @@ namespace UnrealBuildTool
 		public static UnrealTargetPlatform VisionOS { get; } = FindOrAddByName("VisionOS");
 	}
 
-	internal class UnrealTargetPlatformTypeConverter : TypeConverter
+	internal sealed class UnrealTargetPlatformJsonConverter : JsonConverter<UnrealTargetPlatform>
+	{
+		public override UnrealTargetPlatform Read(ref Utf8JsonReader reader, Type typeToConvert, JsonSerializerOptions options) => UnrealTargetPlatform.Parse(reader.GetString() ?? String.Empty);
+
+		public override void Write(Utf8JsonWriter writer, UnrealTargetPlatform value, JsonSerializerOptions options) => writer.WriteStringValue(value.ToString());
+	}
+
+	internal sealed class UnrealTargetPlatformTypeConverter : TypeConverter
 	{
 		public override bool CanConvertFrom(ITypeDescriptorContext? context, Type sourceType)
 		{
@@ -571,7 +581,9 @@ namespace UnrealBuildTool
 	/// <summary>
 	/// The architecture we're building for
 	/// </summary>
-	[Serializable, TypeConverter(typeof(UnrealArchPlatformTypeConverter))]
+	[Serializable]
+	[JsonConverter(typeof(UnrealArchPlatformJsonConverter))]
+	[TypeConverter(typeof(UnrealArchPlatformTypeConverter))]
 	public partial struct UnrealArch : ISerializable
 	{
 		#region Private/boilerplate
@@ -843,7 +855,14 @@ namespace UnrealBuildTool
 		public static Lazy<UnrealArch> Host { get; } = new(() => UnrealArchitectureConfig.ForPlatform(BuildHostPlatform.Current.Platform).GetHostArchitecture());
 	}
 
-	internal class UnrealArchPlatformTypeConverter : TypeConverter
+	internal sealed class UnrealArchPlatformJsonConverter : JsonConverter<UnrealArch>
+	{
+		public override UnrealArch Read(ref Utf8JsonReader reader, Type typeToConvert, JsonSerializerOptions options) => UnrealArch.Parse(reader.GetString() ?? String.Empty);
+
+		public override void Write(Utf8JsonWriter writer, UnrealArch value, JsonSerializerOptions options) => writer.WriteStringValue(value.ToString());
+	}
+
+	internal sealed class UnrealArchPlatformTypeConverter : TypeConverter
 	{
 		public override bool CanConvertFrom(ITypeDescriptorContext? context, Type sourceType)
 		{
@@ -889,6 +908,7 @@ namespace UnrealBuildTool
 	/// A collection of one or more architecetures
 	/// </summary>
 	[Serializable]
+	[JsonConverter(typeof(UnrealArchitecturesJsonConverter))]
 	public class UnrealArchitectures
 	{
 		/// <summary>
@@ -898,6 +918,9 @@ namespace UnrealBuildTool
 		public UnrealArchitectures(IEnumerable<UnrealArch> Architectures)
 		{
 			this.Architectures = new(Architectures.Distinct());
+
+			// standardize order so that passing in { X64, Arm64 } will always result in "arm64+x64" filenames, etc
+			this.Architectures.SortBy(x => x.ToString());
 		}
 
 		/// <summary>
@@ -1030,6 +1053,18 @@ namespace UnrealBuildTool
 		}
 	}
 
+	/// <summary>
+	/// Json converter to/from strings
+	/// </summary>
+	sealed class UnrealArchitecturesJsonConverter : JsonConverter<UnrealArchitectures>
+	{
+		/// <inheritdoc/>
+		public override UnrealArchitectures? Read(ref Utf8JsonReader reader, Type typeToConvert, JsonSerializerOptions options) => UnrealArchitectures.FromString(reader.GetString(), null);
+
+		/// <inheritdoc/>
+		public override void Write(Utf8JsonWriter writer, UnrealArchitectures value, JsonSerializerOptions options) => writer.WriteStringValue(value.ToString());
+	}
+
 	/// <summary>
 	/// The class of platform. See Utils.GetPlatformsInClass().
 	/// </summary>
@@ -3100,7 +3135,7 @@ namespace UnrealBuildTool
 				// is just to depend on the engine DLL timestamps. Instead, we add a makefile dependency on it, causing it to be
 				// regenerated if missing, and only add it as a prereq if we're updating the .target file.
 				FileItem InfoFileItem = FileItem.GetItemByFileReference(InfoFile);
-				BinaryFormatterUtils.SaveIfDifferent(InfoFile, Info);
+				JsonSerializerUtils.SaveIfDifferent(InfoFile, Info);
 				InfoFileItem.ResetCachedInfo();
 				Makefile.InternalDependencies.Add(InfoFileItem);
 
diff --git a/Engine/Source/Programs/UnrealBuildTool/Executors/ImmediateActionQueue.cs b/Engine/Source/Programs/UnrealBuildTool/Executors/ImmediateActionQueue.cs
--- a/Engine/Source/Programs/UnrealBuildTool/Executors/ImmediateActionQueue.cs
+++ b/Engine/Source/Programs/UnrealBuildTool/Executors/ImmediateActionQueue.cs
@@ -260,7 +260,7 @@ namespace UnrealBuildTool
 		/// </summary>
 		public bool StopCompilationAfterErrors = false;
 
-		public int CompletedActions { get => _completedActions; }
+		public int CompletedActions => _completedActions;
 
 		/// <summary>
 		/// Return true if the queue is done
diff --git a/Engine/Source/Programs/UnrealBuildTool/Modes/WriteMetadataMode.cs b/Engine/Source/Programs/UnrealBuildTool/Modes/WriteMetadataMode.cs
--- a/Engine/Source/Programs/UnrealBuildTool/Modes/WriteMetadataMode.cs
+++ b/Engine/Source/Programs/UnrealBuildTool/Modes/WriteMetadataMode.cs
@@ -112,7 +112,7 @@ namespace UnrealBuildTool
 		private Task<int> ExecuteInternal(CommandLineArguments Arguments, ILogger Logger)
 		{
 			// Read the target info
-			WriteMetadataTargetInfo TargetInfo = BinaryFormatterUtils.Load<WriteMetadataTargetInfo>(Arguments.GetFileReference("-Input="));
+			WriteMetadataTargetInfo TargetInfo = JsonSerializerUtils.Load<WriteMetadataTargetInfo>(Arguments.GetFileReference("-Input="));
 			bool bNoManifestChanges = Arguments.HasOption("-NoManifestChanges");
 			int VersionNumber = Arguments.GetInteger("-Version=");
 			Arguments.CheckAllArgumentsUsed();
diff --git a/Engine/Source/Programs/UnrealBuildTool/Platform/Android/UEDeployAndroid.cs b/Engine/Source/Programs/UnrealBuildTool/Platform/Android/UEDeployAndroid.cs
--- a/Engine/Source/Programs/UnrealBuildTool/Platform/Android/UEDeployAndroid.cs
+++ b/Engine/Source/Programs/UnrealBuildTool/Platform/Android/UEDeployAndroid.cs
@@ -1668,7 +1668,7 @@ namespace UnrealBuildTool
 			StartInfo.RedirectStandardOutput = true;
 			StartInfo.RedirectStandardError = true;
 
-			string ProcOutput = string.Empty;
+			string ProcOutput = String.Empty;
 
 			Process Proc = new Process();
 			Proc.StartInfo = StartInfo;
diff --git a/Engine/Source/Programs/UnrealBuildTool/Platform/IOS/IOSPostBuildSyncMode.cs b/Engine/Source/Programs/UnrealBuildTool/Platform/IOS/IOSPostBuildSyncMode.cs
--- a/Engine/Source/Programs/UnrealBuildTool/Platform/IOS/IOSPostBuildSyncMode.cs
+++ b/Engine/Source/Programs/UnrealBuildTool/Platform/IOS/IOSPostBuildSyncMode.cs
@@ -2,6 +2,7 @@
 
 using System;
 using System.Collections.Generic;
+using System.Text.Json.Serialization;
 using System.Threading.Tasks;
 using EpicGames.Core;
 using Microsoft.Extensions.Logging;
@@ -32,6 +33,16 @@ namespace UnrealBuildTool
 		public bool bForDistribution = false;
 		public bool bBuildAsFramework = false;
 
+		/// <summary>
+		/// Private constructor, for serialization.
+		/// </summary>
+		[JsonConstructor]
+#pragma warning disable CS8618 // Non-nullable field must contain a non-null value when exiting constructor. Consider adding the 'required' modifier or declaring as nullable.
+		private IOSPostBuildSyncTarget()
+#pragma warning restore CS8618 // Non-nullable field must contain a non-null value when exiting constructor. Consider adding the 'required' modifier or declaring as nullable.
+		{
+		}
+
 		public IOSPostBuildSyncTarget(ReadOnlyTargetRules Target, FileReference OutputPath, DirectoryReference? ProjectIntermediateDirectory, List<string> UPLScripts, VersionNumber SdkVersion, Dictionary<string, DirectoryReference> FrameworkNameToSourceDir)
 		{
 			Platform = Target.Platform;
@@ -65,7 +76,7 @@ namespace UnrealBuildTool
 		[CommandLine("-XmlConfigCache=")]
 		public FileReference? XmlConfigCache = null;
 
-		// this isn't actually used, but is helpful to pass -legacyxcode along in CreatePostBuildSyncAction, and UBT won't
+		// this isn't actually used, but is hZelpful to pass -legacyxcode along in CreatePostBuildSyncAction, and UBT won't
 		// complain that nothing is using it, because where we _do_ use it is outside the normal cmdline parsing functionality
 		[CommandLine("-LegacyXcode")]
 		public bool bLegacyXcode;
@@ -78,7 +89,7 @@ namespace UnrealBuildTool
 			IOSToolChainSettings.SelectXcode(false, Logger);
 
 			// Run the PostBuildSync command
-			IOSPostBuildSyncTarget Target = BinaryFormatterUtils.Load<IOSPostBuildSyncTarget>(InputFile!);
+			IOSPostBuildSyncTarget Target = JsonSerializerUtils.Load<IOSPostBuildSyncTarget>(InputFile!);
 			IOSToolChain.PostBuildSync(Target, Logger);
 
 			return Task.FromResult(0);
diff --git a/Engine/Source/Programs/UnrealBuildTool/Platform/IOS/IOSToolChain.cs b/Engine/Source/Programs/UnrealBuildTool/Platform/IOS/IOSToolChain.cs
--- a/Engine/Source/Programs/UnrealBuildTool/Platform/IOS/IOSToolChain.cs
+++ b/Engine/Source/Programs/UnrealBuildTool/Platform/IOS/IOSToolChain.cs
@@ -1193,7 +1193,7 @@ namespace UnrealBuildTool
 
 				IOSPostBuildSyncTarget PostBuildSyncTarget = new IOSPostBuildSyncTarget(Target, Executable.Location, BinaryLinkEnvironment.IntermediateDirectory, UPLScripts, SdkVersion, FrameworkNameToSourceDir);
 				FileReference PostBuildSyncFile = FileReference.Combine(BinaryLinkEnvironment.IntermediateDirectory!, "PostBuildSync.dat");
-				BinaryFormatterUtils.Save(PostBuildSyncFile, PostBuildSyncTarget);
+				JsonSerializerUtils.Save(PostBuildSyncFile, PostBuildSyncTarget);
 
 				string PostBuildSyncArguments = String.Format("-Input=\"{0}\" -XmlConfigCache=\"{1}\" -remoteini=\"{2}\" -legacyxcode", PostBuildSyncFile, XmlConfig.CacheFile, UnrealBuildTool.GetRemoteIniPath());
 
diff --git a/Engine/Source/Programs/UnrealBuildTool/System/HotReload.cs b/Engine/Source/Programs/UnrealBuildTool/System/HotReload.cs
--- a/Engine/Source/Programs/UnrealBuildTool/System/HotReload.cs
+++ b/Engine/Source/Programs/UnrealBuildTool/System/HotReload.cs
@@ -119,7 +119,7 @@ namespace UnrealBuildTool
 		/// <returns>New hot reload state instance</returns>
 		public static HotReloadState Load(FileReference Location)
 		{
-			return BinaryFormatterUtils.Load<HotReloadState>(Location);
+			return JsonSerializerUtils.Load<HotReloadState>(Location);
 		}
 
 		/// <summary>
@@ -129,7 +129,7 @@ namespace UnrealBuildTool
 		public void Save(FileReference Location)
 		{
 			DirectoryReference.CreateDirectory(Location.Directory);
-			BinaryFormatterUtils.Save(Location, this);
+			JsonSerializerUtils.Save(Location, this);
 		}
 	}
 
@@ -1304,7 +1304,7 @@ namespace UnrealBuildTool
 					{
 						throw new Exception(String.Format("Unable to find metadata file to patch action graph ({0})", TargetInfoFile));
 					}
-					WriteMetadataTargetInfo TargetInfo = BinaryFormatterUtils.Load<WriteMetadataTargetInfo>(TargetInfoFile);
+					WriteMetadataTargetInfo TargetInfo = JsonSerializerUtils.Load<WriteMetadataTargetInfo>(TargetInfoFile);
 
 					// Update the module names
 					bool bHasUpdatedModuleNames = false;
@@ -1328,7 +1328,7 @@ namespace UnrealBuildTool
 					if (bHasUpdatedModuleNames)
 					{
 						FileReference HotReloadTargetInfoFile = FileReference.Combine(TargetInfoFile.Directory, "Metadata-HotReload.dat");
-						BinaryFormatterUtils.SaveIfDifferent(HotReloadTargetInfoFile, TargetInfo);
+						JsonSerializerUtils.SaveIfDifferent(HotReloadTargetInfoFile, TargetInfo);
 
 						Action NewAction = new Action(Action.Inner);
 
diff --git a/Engine/Source/Programs/UnrealBuildTool/System/ModuleManifest.cs b/Engine/Source/Programs/UnrealBuildTool/System/ModuleManifest.cs
--- a/Engine/Source/Programs/UnrealBuildTool/System/ModuleManifest.cs
+++ b/Engine/Source/Programs/UnrealBuildTool/System/ModuleManifest.cs
@@ -5,6 +5,7 @@ using System.Collections.Generic;
 using System.Diagnostics.CodeAnalysis;
 using System.IO;
 using System.Linq;
+using System.Text.Json.Serialization;
 using EpicGames.Core;
 
 namespace UnrealBuildTool
@@ -32,10 +33,22 @@ namespace UnrealBuildTool
 		/// <summary>
 		/// Constructs the module map with the given changelist
 		/// </summary>
-		/// <param name="InBuildId">The unique build id</param>
-		public ModuleManifest(string InBuildId)
+		/// <param name="BuildId">The unique build id</param>
+		public ModuleManifest(string BuildId)
 		{
-			BuildId = InBuildId;
+			this.BuildId = BuildId;
+		}
+
+		/// <summary>
+		/// Constructs the module map with the given changelist
+		/// </summary>
+		/// <param name="BuildId">The unique build id</param>
+		/// <param name="ModuleNameToFileName">Module name to filename mapping</param>
+		[JsonConstructor]
+		public ModuleManifest(string BuildId, Dictionary<string, string> ModuleNameToFileName)
+		{
+			this.BuildId = BuildId;
+			this.ModuleNameToFileName = ModuleNameToFileName;
 		}
 
 		/// <summary>
diff --git a/Engine/Source/Programs/UnrealBuildTool/System/TargetReceipt.cs b/Engine/Source/Programs/UnrealBuildTool/System/TargetReceipt.cs
--- a/Engine/Source/Programs/UnrealBuildTool/System/TargetReceipt.cs
+++ b/Engine/Source/Programs/UnrealBuildTool/System/TargetReceipt.cs
@@ -6,6 +6,7 @@ using System.Diagnostics.CodeAnalysis;
 using System.IO;
 using System.Linq;
 using System.Text.Json;
+using System.Text.Json.Serialization;
 using System.Xml.Serialization;
 using EpicGames.Core;
 using UnrealBuildBase;
@@ -72,6 +73,7 @@ namespace UnrealBuildTool
 		/// <summary>
 		/// Private constructor, for serialization.
 		/// </summary>
+		[JsonConstructor]
 		private BuildProduct()
 		{
 			Path = null!;
@@ -153,6 +155,7 @@ namespace UnrealBuildTool
 		/// <summary>
 		/// Private constructor, for serialization.
 		/// </summary>
+		[JsonConstructor]
 		private RuntimeDependency()
 		{
 			Path = null!;
@@ -239,6 +242,16 @@ namespace UnrealBuildTool
 		[XmlAttribute]
 		public string Value;
 
+		/// <summary>
+		/// Private constructor, for serialization.
+		/// </summary>
+		[JsonConstructor]
+		private ReceiptProperty()
+		{
+			Name = String.Empty;
+			Value = String.Empty;
+		}
+
 		/// <summary>
 		/// Construct a property with the given name and value
 		/// </summary>
@@ -337,6 +350,17 @@ namespace UnrealBuildTool
 		/// </summary>
 		public List<ReceiptProperty> AdditionalProperties = new List<ReceiptProperty>();
 
+		/// <summary>
+		/// Private constructor, for serialization.
+		/// </summary>
+		[JsonConstructor]
+		private TargetReceipt()
+		{
+			TargetName = String.Empty;
+			Architectures = new UnrealArchitectures(Array.Empty<UnrealArch>());
+			Version = new();
+		}
+
 		/// <summary>
 		/// Constructor
 		/// </summary>
diff --git a/Engine/Source/Programs/UnrealBuildTool/ToolChain/AppleToolChain.cs b/Engine/Source/Programs/UnrealBuildTool/ToolChain/AppleToolChain.cs
--- a/Engine/Source/Programs/UnrealBuildTool/ToolChain/AppleToolChain.cs
+++ b/Engine/Source/Programs/UnrealBuildTool/ToolChain/AppleToolChain.cs
@@ -7,6 +7,7 @@ using System.IO;
 using System.Linq;
 using System.Security.Cryptography.X509Certificates;
 using System.Text;
+using System.Text.Json.Serialization;
 using System.Text.RegularExpressions;
 using System.Threading.Tasks;
 using EpicGames.Core;
@@ -988,6 +989,16 @@ namespace UnrealBuildTool
 		public DirectoryReference ProjectIntermediateDirectory;
 		public FileReference StubOutputPath;
 
+		/// <summary>
+		/// Private constructor, for serialization.
+		/// </summary>
+		[JsonConstructor]
+#pragma warning disable CS8618 // Non-nullable field must contain a non-null value when exiting constructor. Consider adding the 'required' modifier or declaring as nullable.
+		private ApplePostBuildSyncTarget()
+#pragma warning restore CS8618 // Non-nullable field must contain a non-null value when exiting constructor. Consider adding the 'required' modifier or declaring as nullable.
+		{
+		}
+
 		public ApplePostBuildSyncTarget(ReadOnlyTargetRules Target, FileItem Executable, DirectoryReference IntermediateDir)
 		{
 			Platform = Target.Platform;
@@ -1026,7 +1037,7 @@ namespace UnrealBuildTool
 			Arguments.CheckAllArgumentsUsed();
 
 			// Run the PostBuildSync command
-			ApplePostBuildSyncTarget Target = BinaryFormatterUtils.Load<ApplePostBuildSyncTarget>(InputFile!);
+			ApplePostBuildSyncTarget Target = JsonSerializerUtils.Load<ApplePostBuildSyncTarget>(InputFile!);
 			int ExitCode = PostBuildSync(Target, Logger);
 
 			return Task.FromResult(ExitCode);
@@ -1187,7 +1198,7 @@ namespace UnrealBuildTool
 		{
 			ApplePostBuildSyncTarget PostBuildSync = new(Target, Executable, IntermediateDir);
 			FileReference PostBuildSyncFile = FileReference.Combine(IntermediateDir!, "PostBuildSync.dat");
-			BinaryFormatterUtils.Save(PostBuildSyncFile, PostBuildSync);
+			JsonSerializerUtils.Save(PostBuildSyncFile, PostBuildSync);
 
 			string PostBuildSyncArguments = String.Format("-modernxcode -Input=\"{0}\" -XmlConfigCache=\"{1}\" -remoteini=\"{2}\"", PostBuildSyncFile, XmlConfig.CacheFile, UnrealBuildTool.GetRemoteIniPath());
 			Action PostBuildSyncAction = Graph.CreateRecursiveAction<ApplePostBuildSyncMode>(ActionType.CreateAppBundle, PostBuildSyncArguments);
