[UBT] added `DeprecatedEngineVersion` to plugin descriptor to be able to warn when a plugin is deprecated and will be removed in future engine versions.

diff --git a/Engine/Source/Programs/UnrealBuildTool/Configuration/UEBuildPlugin.cs b/Engine/Source/Programs/UnrealBuildTool/Configuration/UEBuildPlugin.cs
--- a/Engine/Source/Programs/UnrealBuildTool/Configuration/UEBuildPlugin.cs
+++ b/Engine/Source/Programs/UnrealBuildTool/Configuration/UEBuildPlugin.cs
@@ -1,5 +1,6 @@
 // Copyright Epic Games, Inc. All Rights Reserved.
 
+using System;
 using System.Collections.Generic;
 using System.Linq;
 using EpicGames.Core;
@@ -114,6 +115,11 @@ namespace UnrealBuildTool
 						logger.LogError("Plugin '{PluginName}' cannot depend on plugin '{DependencyPluginName}' because it is disallowed.", Name, dependencyPlugin.Name);
 						anyErrors = true;
 					}
+					else if (!String.IsNullOrEmpty(dependencyPlugin.Descriptor.DeprecatedEngineVersion))
+					{
+						logger.LogWarning("Plugin '{PluginName}' depends on plugin '{DependencyPluginName}' which was deprecated in {EngineVersion} and will soon be removed. Please update your dependencies.",
+							Name, dependencyPlugin.Name, dependencyPlugin.Descriptor.DeprecatedEngineVersion);
+					}
 				}
 			}
 
diff --git a/Engine/Source/Programs/UnrealBuildTool/System/PluginDescriptor.cs b/Engine/Source/Programs/UnrealBuildTool/System/PluginDescriptor.cs
--- a/Engine/Source/Programs/UnrealBuildTool/System/PluginDescriptor.cs
+++ b/Engine/Source/Programs/UnrealBuildTool/System/PluginDescriptor.cs
@@ -112,6 +112,11 @@ namespace UnrealBuildTool
 		/// Sets the version of the engine that this plugin is compatible with.
 		/// </summary>
 		public string? EngineVersion;
+		
+		/// <summary>
+		/// Sets the version of the engine at which this plugin has been deprecated.
+		/// </summary>
+		public string? DeprecatedEngineVersion;
 
 		/// <summary>4
 		/// If true, this plugin from a platform extension extending another plugin */
@@ -284,6 +289,7 @@ namespace UnrealBuildTool
 			RawObject.TryGetStringField("MarketplaceURL", out MarketplaceURL);
 			RawObject.TryGetStringField("SupportURL", out SupportURL);
 			RawObject.TryGetStringField("EngineVersion", out EngineVersion);
+			RawObject.TryGetStringField("DeprecatedEngineVersion", out DeprecatedEngineVersion);
 			RawObject.TryGetStringArrayField("SupportedPrograms", out SupportedPrograms);
 			RawObject.TryGetBoolField("bIsPluginExtension", out bIsPluginExtension);
 
@@ -436,6 +442,10 @@ namespace UnrealBuildTool
 			{
 				Writer.WriteValue("EngineVersion", EngineVersion);
 			}
+			if (!String.IsNullOrEmpty(DeprecatedEngineVersion))
+			{
+				Writer.WriteValue("DeprecatedEngineVersion", DeprecatedEngineVersion);
+			}
 			if (!String.IsNullOrEmpty(VersePath))
 			{
 				Writer.WriteValue("VersePath", VersePath);
@@ -548,6 +558,10 @@ namespace UnrealBuildTool
 			{
 				CachedJson.AddOrSetFieldValue("EngineVersion", EngineVersion);
 			}
+			if (!String.IsNullOrEmpty(DeprecatedEngineVersion))
+			{
+				CachedJson.AddOrSetFieldValue("DeprecatedEngineVersion", DeprecatedEngineVersion);
+			}
 			if (!String.IsNullOrEmpty(VersePath))
 			{
 				CachedJson.AddOrSetFieldValue("VersePath", VersePath);
