Make a force disable AutoRTFM bool

diff --git a/Engine/Source/Programs/UnrealBuildTool/Configuration/ReadOnlyTargetRules.cs b/Engine/Source/Programs/UnrealBuildTool/Configuration/ReadOnlyTargetRules.cs
--- a/Engine/Source/Programs/UnrealBuildTool/Configuration/ReadOnlyTargetRules.cs
+++ b/Engine/Source/Programs/UnrealBuildTool/Configuration/ReadOnlyTargetRules.cs
@@ -136,6 +136,8 @@ namespace UnrealBuildTool
 
 		public bool bUseAutoRTFMCompiler => Inner.bUseAutoRTFMCompiler;
 
+		public bool bForceNoAutoRTFMCompiler => Inner.bForceNoAutoRTFMCompiler;
+
 		[Obsolete("Deprecated in UE5.1 - No longer used as Chaos is always enabled.")]
 		public bool bCompileChaos => Inner.bCompileChaos;
 
@@ -705,4 +707,4 @@ namespace UnrealBuildTool
 			}
 		}
 	}
-}
\ No newline at end of file
+}
diff --git a/Engine/Source/Programs/UnrealBuildTool/Configuration/TargetRules.cs b/Engine/Source/Programs/UnrealBuildTool/Configuration/TargetRules.cs
--- a/Engine/Source/Programs/UnrealBuildTool/Configuration/TargetRules.cs
+++ b/Engine/Source/Programs/UnrealBuildTool/Configuration/TargetRules.cs
@@ -868,11 +868,16 @@ namespace UnrealBuildTool
 		/// Whether to use the AutoRTFM Clang compiler.
 		/// </summary>
 		[RequiresUniqueBuildEnvironment]
-		[CommandLine("-NoUseAutoRTFM", Value = "false")]
-		[CommandLine("-UseAutoRTFM", Value = "true")]
-		[XmlConfigFile(Category = "BuildConfiguration")]
 		public bool bUseAutoRTFMCompiler { get; set; }
 
+		/// <summary>
+		/// Whether to use force AutoRTFM Clang compiler off.
+		/// </summary>
+		[RequiresUniqueBuildEnvironment]
+		[CommandLine("-NoUseAutoRTFM", Value = "true")]
+		[XmlConfigFile(Category = "BuildConfiguration")]
+		public bool bForceNoAutoRTFMCompiler { get; set; } = false;
+
 		/// <summary>
 		/// Whether to compile the Chaos physics plugin.
 		/// </summary>
diff --git a/Engine/Source/Programs/UnrealBuildTool/Configuration/UEBuildModuleCPP.cs b/Engine/Source/Programs/UnrealBuildTool/Configuration/UEBuildModuleCPP.cs
--- a/Engine/Source/Programs/UnrealBuildTool/Configuration/UEBuildModuleCPP.cs
+++ b/Engine/Source/Programs/UnrealBuildTool/Configuration/UEBuildModuleCPP.cs
@@ -2129,12 +2129,12 @@ namespace UnrealBuildTool
 			Result.DeterministicWarningLevel = Rules.DeterministicWarningLevel;
 			Result.bValidateFormatStrings = Rules.bValidateFormatStrings;
 			Result.bValidateInternalApi = Rules.bValidateInternalApi;
-			Result.bUseAutoRTFMCompiler = Target.bUseAutoRTFMCompiler;
+			Result.bUseAutoRTFMCompiler = Target.bUseAutoRTFMCompiler && !Target.bForceNoAutoRTFMCompiler;
 
 			CompileEnvironmentDebugInfoSettings(Target, Result);
 
 			// Only copy the AutoRTFM flag if we are using the AutoRTFM compiler
-			if (Target.bUseAutoRTFMCompiler)
+			if (Result.bUseAutoRTFMCompiler)
 			{
 				Result.bDisableAutoRTFMInstrumentation = Rules.bDisableAutoRTFMInstrumentation;
 			}
diff --git a/Engine/Source/Programs/UnrealBuildTool/Platform/Linux/UEBuildLinux.cs b/Engine/Source/Programs/UnrealBuildTool/Platform/Linux/UEBuildLinux.cs
--- a/Engine/Source/Programs/UnrealBuildTool/Platform/Linux/UEBuildLinux.cs
+++ b/Engine/Source/Programs/UnrealBuildTool/Platform/Linux/UEBuildLinux.cs
@@ -672,7 +672,7 @@ namespace UnrealBuildTool
 				Logger.LogWarning("Warning: bPreferThinLTO is set, but LTO is disabled. Flag will have no effect");
 			}
 
-			if (Target.bUseAutoRTFMCompiler)
+			if (Target.bUseAutoRTFMCompiler && !Target.bForceNoAutoRTFMCompiler)
 			{
 				Options |= ClangToolChainOptions.UseAutoRTFMCompiler;
 			}
diff --git a/Engine/Source/Programs/UnrealBuildTool/Platform/Windows/UEBuildWindows.cs b/Engine/Source/Programs/UnrealBuildTool/Platform/Windows/UEBuildWindows.cs
--- a/Engine/Source/Programs/UnrealBuildTool/Platform/Windows/UEBuildWindows.cs
+++ b/Engine/Source/Programs/UnrealBuildTool/Platform/Windows/UEBuildWindows.cs
@@ -1155,7 +1155,7 @@ namespace UnrealBuildTool
 				Logger.LogInformation("Defaulting static analyzer output type to text");
 			}
 
-			if (Target.bUseAutoRTFMCompiler)
+			if (Target.bUseAutoRTFMCompiler && !Target.bForceNoAutoRTFMCompiler)
 			{
 				// We check the static analyzer, and only allow using the AutoRTFM compiler if
 				// we are not doing static analysis, or are explicitly wanting to use clang.
