[UBT]
* Fixed pgo directory path not using double backslash

diff --git a/Engine/Source/Programs/UnrealBuildTool/Platform/Linux/LinuxToolChain.cs b/Engine/Source/Programs/UnrealBuildTool/Platform/Linux/LinuxToolChain.cs
--- a/Engine/Source/Programs/UnrealBuildTool/Platform/Linux/LinuxToolChain.cs
+++ b/Engine/Source/Programs/UnrealBuildTool/Platform/Linux/LinuxToolChain.cs
@@ -838,8 +838,8 @@ namespace UnrealBuildTool
 				//
 				Arguments.Add("-Wno-backend-plugin");
 
-				Log.TraceInformationOnce("Enabling Profile Guided Optimization (PGO). Linking will take a while.");
-				Arguments.Add(String.Format("-fprofile-instr-use=\"{0}\"", Path.Combine(LinkEnvironment.PGODirectory!, LinkEnvironment.PGOFilenamePrefix!)));
+				DirectoryReference? PGODir = DirectoryReference.FromString(LinkEnvironment.PGODirectory!);
+				Arguments.Add($"-fprofile-instr-use=\"{NormalizeCommandLinePath(DirectoryReference.Combine(PGODir!, LinkEnvironment.PGOFilenamePrefix!))}\"");
 			}
 			else if (LinkEnvironment.bPGOProfile)
 			{
diff --git a/Engine/Source/Programs/UnrealBuildTool/ToolChain/ClangToolChain.cs b/Engine/Source/Programs/UnrealBuildTool/ToolChain/ClangToolChain.cs
--- a/Engine/Source/Programs/UnrealBuildTool/ToolChain/ClangToolChain.cs
+++ b/Engine/Source/Programs/UnrealBuildTool/ToolChain/ClangToolChain.cs
@@ -655,7 +655,8 @@ namespace UnrealBuildTool
 			if (CompileEnvironment.bPGOOptimize)
 			{
 				Log.TraceInformationOnce("Enabling Profile Guided Optimization (PGO). Linking will take a while.");
-				Arguments.Add($"-fprofile-instr-use=\"{Path.Combine(CompileEnvironment.PGODirectory!, CompileEnvironment.PGOFilenamePrefix!)}\"");
+				DirectoryReference? PGODir = DirectoryReference.FromString(CompileEnvironment.PGODirectory!);
+				Arguments.Add($"-fprofile-instr-use=\"{NormalizeCommandLinePath(DirectoryReference.Combine(PGODir!, CompileEnvironment.PGOFilenamePrefix!))}\"");
 			}
 			else if (CompileEnvironment.bPGOProfile)
 			{
