UnrealBuildTool: Workaround for setting uba compression in the TargetRules

#rnx

diff --git a/Engine/Source/Programs/UnrealBuildTool/Configuration/TargetRules.cs b/Engine/Source/Programs/UnrealBuildTool/Configuration/TargetRules.cs
--- a/Engine/Source/Programs/UnrealBuildTool/Configuration/TargetRules.cs
+++ b/Engine/Source/Programs/UnrealBuildTool/Configuration/TargetRules.cs
@@ -2980,6 +2980,14 @@ namespace UnrealBuildTool
 				throw new NotSupportedException("Logger property must be initialized by the caller.");
 			}
 
+			// TODO: Need to get the compression setting from UBA so we can invalidate actions if the compression changes.
+			// Figure out a better way to do this.
+			{
+				UnrealBuildAcceleratorConfig ubaConfig = new();
+				XmlConfig.ApplyTo(ubaConfig);
+				bAllowUbaCompression = ubaConfig.bStoreObjFilesCompressed;
+			}
+
 			// Read settings from config files
 			Dictionary<ConfigDependencyKey, IReadOnlyList<string>?> configValues = new Dictionary<ConfigDependencyKey, IReadOnlyList<string>?>();
 			foreach (object configurableObject in GetConfigurableObjects())
