UnrealBuildTool: Move additional target compile args to after the file type specific args, so any previous arguments will be overwritten if desired, or if they need to be added after other args, such as args that depend on the cppstd version

diff --git a/Engine/Source/Programs/UnrealBuildTool/Platform/Windows/VCToolChain.cs b/Engine/Source/Programs/UnrealBuildTool/Platform/Windows/VCToolChain.cs
--- a/Engine/Source/Programs/UnrealBuildTool/Platform/Windows/VCToolChain.cs
+++ b/Engine/Source/Programs/UnrealBuildTool/Platform/Windows/VCToolChain.cs
@@ -1007,11 +1007,6 @@ namespace UnrealBuildTool
 				AddDefinition(Arguments, "_WINDOWS");
 				AddDefinition(Arguments, "WIN32");
 			}
-
-			if (!String.IsNullOrEmpty(CompileEnvironment.AdditionalArguments))
-			{
-				Arguments.Add(CompileEnvironment.AdditionalArguments);
-			}
 		}
 
 		protected virtual void AppendCLArguments_H(CppCompileEnvironment CompileEnvironment, List<string> Arguments)
@@ -1853,6 +1848,12 @@ namespace UnrealBuildTool
 					AppendCLArguments_CPP(CompileEnvironment, CompileAction.Arguments);
 				}
 
+				// Add additional arguments to the argument list, must be the final arguments added
+				if (!String.IsNullOrEmpty(CompileEnvironment.AdditionalArguments))
+				{
+					CompileAction.Arguments.Add(CompileEnvironment.AdditionalArguments);
+				}
+
 				if (Target.WindowsPlatform.Compiler.IsClang())
 				{
 					// If we are using the AutoRTFM compiler, we make the compile action depend on the version of the compiler itself.
diff --git a/Engine/Source/Programs/UnrealBuildTool/ToolChain/ClangToolChain.cs b/Engine/Source/Programs/UnrealBuildTool/ToolChain/ClangToolChain.cs
--- a/Engine/Source/Programs/UnrealBuildTool/ToolChain/ClangToolChain.cs
+++ b/Engine/Source/Programs/UnrealBuildTool/ToolChain/ClangToolChain.cs
@@ -860,9 +860,6 @@ namespace UnrealBuildTool
 			{
 				GetCompileArguments_Analyze(CompileEnvironment, Arguments);
 			}
-
-			// Add additional arguments to the argument list.
-			GetCompileArguments_AdditionalArgs(CompileEnvironment, Arguments);
 		}
 
 		protected virtual string GetFileNameFromExtension(string AbsolutePath, string Extension)
@@ -1018,6 +1015,9 @@ namespace UnrealBuildTool
 				Arguments.Add(GetOutputFileArgument(OutputFile));
 			}
 
+			// Add additional arguments to the argument list, must be the final arguments added
+			GetCompileArguments_AdditionalArgs(CompileEnvironment, Arguments);
+
 			return OutputFile;
 		}
 
