[UBA]
* Added configuration to bucket id for cache server

diff --git a/Engine/Source/Programs/UnrealBuildTool/Executors/UnrealBuildAccelerator/UBAExecutor.cs b/Engine/Source/Programs/UnrealBuildTool/Executors/UnrealBuildAccelerator/UBAExecutor.cs
--- a/Engine/Source/Programs/UnrealBuildTool/Executors/UnrealBuildAccelerator/UBAExecutor.cs
+++ b/Engine/Source/Programs/UnrealBuildTool/Executors/UnrealBuildAccelerator/UBAExecutor.cs
@@ -605,8 +605,9 @@ namespace UnrealBuildTool
 
 			using (Blake3.Hasher hasher = Blake3.Hasher.New())
 			{
-				// Use platform to chose bucket since there will never be any cache hits between platforms
+				// Use platform and config to chose bucket since there will never be any cache hits between platforms or configs
 				hasher.Update(System.Text.Encoding.UTF8.GetBytes(action.Target.Platform.ToString()));
+				hasher.Update(new byte[] { (byte)action.Target.Configuration });
 
 				// Absolute path is set for actions that uses pch.
 				// And since pch contains absolute paths we unfortunately can't share cache data between machines that have different paths
