Add plists to the allowed resources that can be copied on iOS during staging

#rnx

diff --git a/Engine/Source/Programs/UnrealBuildTool/ProjectFiles/Xcode/XcodeProject.File.cs b/Engine/Source/Programs/UnrealBuildTool/ProjectFiles/Xcode/XcodeProject.File.cs
--- a/Engine/Source/Programs/UnrealBuildTool/ProjectFiles/Xcode/XcodeProject.File.cs
+++ b/Engine/Source/Programs/UnrealBuildTool/ProjectFiles/Xcode/XcodeProject.File.cs
@@ -126,6 +126,7 @@ namespace UnrealBuildTool.XcodeProjectXcconfig
 			".lproj",
 			".png",
 			".xcprivacy",
+			".plist",
 			// no extensions will be folders, like "cookeddata" that we want to copy as a resource
 			"",
 		};
