UnrealBuildTool: Track in UBA telemetry the success or failure of the coordinators connecting, as well as the time delay until the first remote agent is used

#rnx

diff --git a/Engine/Source/Programs/UnrealBuildTool/Executors/UnrealBuildAccelerator/UBAAgentCoordinatorHorde.cs b/Engine/Source/Programs/UnrealBuildTool/Executors/UnrealBuildAccelerator/UBAAgentCoordinatorHorde.cs
--- a/Engine/Source/Programs/UnrealBuildTool/Executors/UnrealBuildAccelerator/UBAAgentCoordinatorHorde.cs
+++ b/Engine/Source/Programs/UnrealBuildTool/Executors/UnrealBuildAccelerator/UBAAgentCoordinatorHorde.cs
@@ -582,6 +582,7 @@ namespace UnrealBuildTool
 			_cancellationSource = new CancellationTokenSource();
 			_hordeSessionTask = UBAHordeSession.TryCreateHordeSessionAsync(HordeConfig, executor, _ubaConfig.bStrict, _logger, _cancellationSource.Token);
 			await _hordeSessionTask;
+			executor.AgentCoordinatorInitialized(this, _hordeSessionTask.IsCompletedSuccessfully && _hordeSessionTask.Result != null);
 		}
 
 		public void Start(ImmediateActionQueue queue, Func<LinkedAction, bool> canRunRemotely)
diff --git a/Engine/Source/Programs/UnrealBuildTool/Executors/UnrealBuildAccelerator/UBAExecutor.cs b/Engine/Source/Programs/UnrealBuildTool/Executors/UnrealBuildAccelerator/UBAExecutor.cs
--- a/Engine/Source/Programs/UnrealBuildTool/Executors/UnrealBuildAccelerator/UBAExecutor.cs
+++ b/Engine/Source/Programs/UnrealBuildTool/Executors/UnrealBuildAccelerator/UBAExecutor.cs
@@ -36,11 +36,19 @@ namespace UnrealBuildTool
 		uint _actionsQueuedThatCanRunRemotely = UInt32.MaxValue;
 		readonly ThreadedLogger _threadedLogger;
 
+		DateTime _ubaStartTimeUtc = DateTime.UtcNow;
+		TimeSpan _ubaDurationWaitingForRemote = TimeSpan.Zero;
+
 		// Tracking for LinkedActions that failed remotely that should be retried locally
 		readonly ConcurrentDictionary<LinkedAction, bool> _localRetryActions = new();
 		// Tracking for LinkedActions that failed locally that should be retried without UBA
 		readonly ConcurrentDictionary<LinkedAction, bool> _forcedRetryActions = new();
 
+		// Tracking for successful coordinator connections
+		int _successfulCoordinatorConnections = 0;
+		// Tracking for failed coordinator connections
+		int _failedCoordinatorConnections = 0;
+
 		// Tracking for all actions processed locally
 		int _localProcessedActions = 0;
 		// Tracking for all actions processed remotely
@@ -179,6 +187,18 @@ namespace UnrealBuildTool
 			return BitConverter.ToString(bytes).Replace("-", "", StringComparison.OrdinalIgnoreCase).ToLowerInvariant(); // "1234567890abcdef1234567890abcdef";
 		}
 
+		public void AgentCoordinatorInitialized(IUBAAgentCoordinator coordinator, bool successful)
+		{
+			if (successful)
+			{
+				Interlocked.Add(ref _successfulCoordinatorConnections, 1);
+			}
+			else
+			{
+				Interlocked.Add(ref _failedCoordinatorConnections, 1);
+			}
+		}
+
 		class UBAArtifactCache : IArtifactCache
 		{
 			public ArtifactCacheState State => ArtifactCacheState.Available;
@@ -440,7 +460,7 @@ namespace UnrealBuildTool
 		/// <returns>True if all the tasks successfully executed, or false if any of them failed.</returns>
 		bool ExecuteActionsInternal(IEnumerable<LinkedAction> inputActions, ISessionServer session, Microsoft.Extensions.Logging.ILogger logger, IActionArtifactCache? actionArtifactCache, System.Action onCancel)
 		{
-			DateTime startTimeUTC = DateTime.UtcNow;
+			_ubaStartTimeUtc = DateTime.UtcNow;
 			using ImmediateActionQueue queue = CreateActionQueue(inputActions, actionArtifactCache, UBAConfig.CacheMaxWorkers, logger);
 			int actionLimit = Math.Min(NumParallelProcesses, queue.TotalActions);
 			queue.CreateAutomaticRunner(action => RunActionLocal(queue, action), bUseActionWeights, actionLimit, NumParallelProcesses);
@@ -496,8 +516,12 @@ namespace UnrealBuildTool
 				bool res = queue.RunTillDone().Result; // Using inline wait to avoid possible thread switch
 
 				queue.GetActionResultCounts(out int totalActions, out int succeededActions, out int failedActions, out int cacheHitActions, out int cacheMissActions);
-				telemetryEvent = new TelemetryExecutorUBAEvent(Name, startTimeUTC, res, totalActions, succeededActions, failedActions, cacheHitActions, cacheMissActions,
-					_localProcessedActions, _remoteProcessedActions, _localRetryActions.Count, _forcedRetryActions.Count, DateTime.UtcNow);
+				telemetryEvent = new TelemetryExecutorUBAEvent(Name, _ubaStartTimeUtc, res, totalActions, succeededActions, failedActions, cacheHitActions, cacheMissActions,
+					_localProcessedActions, _remoteProcessedActions,
+					_localRetryActions.Count, _forcedRetryActions.Count,
+					!UBAConfig.bDisableRemote ? _agentCoordinators.Count : 0, !UBAConfig.bDisableRemote ? _successfulCoordinatorConnections : 0, !UBAConfig.bDisableRemote ? _failedCoordinatorConnections : 0,
+					!UBAConfig.bDisableRemote ? _ubaDurationWaitingForRemote : TimeSpan.Zero,
+					DateTime.UtcNow);
 
 				return res;
 			}
@@ -770,6 +794,10 @@ namespace UnrealBuildTool
 
 			return () =>
 			{
+				if (_ubaDurationWaitingForRemote == TimeSpan.Zero)
+				{
+					_ubaDurationWaitingForRemote = DateTime.UtcNow - _ubaStartTimeUtc;
+				}
 				if (_bIsCancelled)
 				{
 					HandleActionCancelled(queue, null, action);
diff --git a/Engine/Source/Programs/UnrealBuildTool/System/Telemetry.cs b/Engine/Source/Programs/UnrealBuildTool/System/Telemetry.cs
--- a/Engine/Source/Programs/UnrealBuildTool/System/Telemetry.cs
+++ b/Engine/Source/Programs/UnrealBuildTool/System/Telemetry.cs
@@ -151,19 +151,34 @@ namespace UnrealBuildTool
 		public int RetriedLocalActions { get; }
 		public int RetriedDisabledActions { get; }
 
+		public int TotalCoordinators { get; }
+		public int SucceededCoordinators { get; }
+		public int FailedCoordinators { get; }
+
+		public TimeSpan DurationWaitingForRemote { get; }
+
 		public double RetriedLocalRate => RetriedLocalActions / (double)TotalActions;
 		public double RetriedDisabledRate => RetriedDisabledActions / (double)TotalActions;
 
 		public double LocalUsage => LocalActions / (double)(TotalActions + RetriedLocalActions + RetriedDisabledActions - SkippedActions);
 		public double RemoteUsage => RemoteActions / (double)(TotalActions + RetriedLocalActions + RetriedDisabledActions - SkippedActions);
 
-		public TelemetryExecutorUBAEvent(string executor, DateTime startUTC, bool result, int totalActions, int succeededActions, int failedActions, int cacheHitActions, int cacheMissActions, int localActions, int remoteActions, int retriedLocalActions, int retriedDisabledActions, DateTime timestamp)
+		public TelemetryExecutorUBAEvent(string executor, DateTime startUTC, bool result, int totalActions, int succeededActions, int failedActions, int cacheHitActions, int cacheMissActions,
+			int localActions, int remoteActions,
+			int retriedLocalActions, int retriedDisabledActions,
+			int totalCoordinators, int succeededCoordinators, int failedCoordinators,
+			TimeSpan durationWaitingForRemote,
+			DateTime timestamp)
 			: base(executor, startUTC, result, totalActions, succeededActions, failedActions, cacheHitActions, cacheMissActions, timestamp)
 		{
 			LocalActions = localActions;
 			RemoteActions = remoteActions;
 			RetriedLocalActions = retriedLocalActions;
 			RetriedDisabledActions = retriedDisabledActions;
+			TotalCoordinators = totalCoordinators;
+			SucceededCoordinators = succeededCoordinators;
+			FailedCoordinators = failedCoordinators;
+			DurationWaitingForRemote = durationWaitingForRemote;
 		}
 	}
 
