UnrealBuildTool: -fno-eliminate-unused-debug-type for clang-cl (pending for llvm 19)

#rnx

diff --git a/Engine/Source/Programs/UnrealBuildTool/Platform/Windows/VCToolChain.cs b/Engine/Source/Programs/UnrealBuildTool/Platform/Windows/VCToolChain.cs
--- a/Engine/Source/Programs/UnrealBuildTool/Platform/Windows/VCToolChain.cs
+++ b/Engine/Source/Programs/UnrealBuildTool/Platform/Windows/VCToolChain.cs
@@ -826,16 +826,27 @@ namespace UnrealBuildTool
 					Arguments.Add("/Z7");
 				}
 
-				// https://clang.llvm.org/docs/UsersManual.html#cmdoption-gline-tables-only
-				if (Target.WindowsPlatform.Compiler.IsClang() && CompileEnvironment.bDebugLineTablesOnly)
+				if (Target.WindowsPlatform.Compiler.IsClang())
 				{
-					Arguments.Add("-gline-tables-only");
-				}
+					VersionNumber ClangVersion = Target.WindowsPlatform.Compiler == WindowsCompiler.Intel ? MicrosoftPlatformSDK.GetClangVersionForIntelCompiler(EnvVars.CompilerPath) : EnvVars.CompilerVersion;
 
-				// https://clang.llvm.org/docs/UsersManual.html#cmdoption-fstandalone-debug
-				if (Target.WindowsPlatform.Compiler.IsClang() && Target.WindowsPlatform.bClangStandaloneDebug)
-				{
-					Arguments.Add("-fstandalone-debug");
+					// https://clang.llvm.org/docs/UsersManual.html#cmdoption-gline-tables-only
+					if (CompileEnvironment.bDebugLineTablesOnly)
+					{
+						Arguments.Add("-gline-tables-only");
+					}
+
+					// https://clang.llvm.org/docs/UsersManual.html#cmdoption-fstandalone-debug
+					if (Target.WindowsPlatform.bClangStandaloneDebug)
+					{
+						Arguments.Add("-fstandalone-debug");
+					}
+
+					// https://clang.llvm.org/docs/UsersManual.html#cmdoption-feliminate-unused-debug-types
+					if (ClangVersion >= new VersionNumber(19))
+					{
+						Arguments.Add("-fno-eliminate-unused-debug-type");
+					}
 				}
 			}
 
