#UE Add windows build option -NoLinkerDebugInfo to allow CreateDebugInfo to be false in the linker if you are also supplying -NoDebugInfo.

diff --git a/Engine/Source/Programs/UnrealBuildTool/Platform/Windows/UEBuildWindows.cs b/Engine/Source/Programs/UnrealBuildTool/Platform/Windows/UEBuildWindows.cs
--- a/Engine/Source/Programs/UnrealBuildTool/Platform/Windows/UEBuildWindows.cs
+++ b/Engine/Source/Programs/UnrealBuildTool/Platform/Windows/UEBuildWindows.cs
@@ -537,6 +537,13 @@ namespace UnrealBuildTool
 		[CommandLine("-StripPrivateSymbols")]
 		public bool bStripPrivateSymbols = false;
 
+		/// <summary>
+		/// If you supply -NoDebugInfo, windows platforms still create debug info while linking. Set this to true to not create debug info while linking in this circumstance
+		/// </summary>
+		[XmlConfigFile(Category = "WindowsPlatform")]
+		[CommandLine("-NoLinkerDebugInfo")]
+		public bool bNoLinkerDebugInfo = false;
+
 		/// <summary>
 		/// Set page size to allow for larger than 4GB PDBs to be generated by the msvc linker.
 		/// Will default to 16384 for monolithic editor builds.
@@ -902,6 +909,8 @@ namespace UnrealBuildTool
 
 		public bool bStripPrivateSymbols => Inner.bStripPrivateSymbols;
 
+		public bool bNoLinkerDebugInfo => Inner.bNoLinkerDebugInfo;
+
 		public uint? PdbPageSize => Inner.PdbPageSize;
 
 		public string? PdbAlternatePath => Inner.PdbAlternatePath;
@@ -1768,7 +1777,10 @@ namespace UnrealBuildTool
 			//       useful call stacks, even though compiler-generate debug info may be disabled.  This gives
 			//       us much of the build-time savings of fully-disabled debug info, without giving up call
 			//       data completely.
-			GlobalLinkEnvironment.bCreateDebugInfo = true;
+			if (!Target.WindowsPlatform.bNoLinkerDebugInfo)
+			{
+				GlobalLinkEnvironment.bCreateDebugInfo = true;
+			}
 		}
 
 		/// <summary>
