UnrealBuildTool: Don't spam toolchain info if multiple targets all have unpreferred settings

#rnx

diff --git a/Engine/Source/Programs/UnrealBuildTool/Platform/Windows/UEBuildWindows.cs b/Engine/Source/Programs/UnrealBuildTool/Platform/Windows/UEBuildWindows.cs
--- a/Engine/Source/Programs/UnrealBuildTool/Platform/Windows/UEBuildWindows.cs
+++ b/Engine/Source/Programs/UnrealBuildTool/Platform/Windows/UEBuildWindows.cs
@@ -1248,6 +1248,17 @@ namespace UnrealBuildTool
 			Target.WindowsPlatform.ToolchainVersion = Target.WindowsPlatform.Environment.ToolChainVersion.ToString();
 			Target.WindowsPlatform.WindowsSdkVersion = Target.WindowsPlatform.Environment.WindowsSdkVersion.ToString();
 
+			ValidateToolchainVersion(Target);
+		}
+
+		static bool _toolchainWarningLogged = false;
+		void ValidateToolchainVersion(TargetRules Target)
+		{
+			if (_toolchainWarningLogged || Target.WindowsPlatform.Environment == null)
+			{
+				return;
+			}
+
 			// Ensure we're using a recent enough version of Clang given the MSVC version
 			if (Target.WindowsPlatform.Compiler.IsClang() && !MicrosoftPlatformSDK.IgnoreToolchainErrors)
 			{
@@ -1263,6 +1274,7 @@ namespace UnrealBuildTool
 			{
 				if (!MicrosoftPlatformSDK.IsPreferredVersion(Target.WindowsPlatform.Compiler, Target.WindowsPlatform.Environment.CompilerVersion))
 				{
+					_toolchainWarningLogged = true;
 					VersionNumber preferred = MicrosoftPlatformSDK.GetLatestPreferredVersion(Target.WindowsPlatform.Compiler);
 					MicrosoftPlatformSDK.DumpAllToolChainInstallations(Target.WindowsPlatform.Compiler, Target.Architecture, Logger);
 					if (Target.WindowsPlatform.ToolchainVersionWarningLevel == WarningLevel.Error)
@@ -1274,6 +1286,7 @@ namespace UnrealBuildTool
 
 				if (Target.WindowsPlatform.Compiler != Target.WindowsPlatform.ToolChain && !MicrosoftPlatformSDK.IsPreferredVersion(Target.WindowsPlatform.ToolChain, Target.WindowsPlatform.Environment.ToolChainVersion))
 				{
+					_toolchainWarningLogged = true;
 					VersionNumber preferred = MicrosoftPlatformSDK.GetLatestPreferredVersion(Target.WindowsPlatform.ToolChain);
 					MicrosoftPlatformSDK.DumpAllToolChainInstallations(Target.WindowsPlatform.ToolChain, Target.Architecture, Logger);
 					if (Target.WindowsPlatform.ToolchainVersionWarningLevel == WarningLevel.Error)
@@ -1283,18 +1296,6 @@ namespace UnrealBuildTool
 					Logger.LogInformation("{Toolchain} toolchain version {Version} is not a preferred version. Please use a preferred toolchain such as {PreferredVersion}", WindowsPlatform.GetCompilerName(Target.WindowsPlatform.ToolChain), Target.WindowsPlatform.Environment.ToolChainVersion, preferred);
 				}
 			}
-
-			//			@Todo: Still getting reports of frequent OOM issues with this enabled as of 15.7.
-			//			// Enable fast PDB linking if we're on VS2017 15.7 or later. Previous versions have OOM issues with large projects.
-			//			if(!Target.bFormalBuild && !Target.bUseFastPDBLinking.HasValue && Target.WindowsPlatform.Compiler.IsMSVC())
-			//			{
-			//				VersionNumber Version;
-			//				DirectoryReference ToolChainDir;
-			//				if(TryGetVCToolChainDir(Target.WindowsPlatform.Compiler, Target.WindowsPlatform.CompilerVersion, out Version, out ToolChainDir) && Version >= new VersionNumber(14, 14, 26316))
-			//				{
-			//					Target.bUseFastPDBLinking = true;
-			//				}
-			//			}
 		}
 
 		/// <summary>
