UnrealBuildTool: Fix analysis issue

#rnx

diff --git a/Engine/Source/Programs/UnrealBuildTool/Platform/Windows/UEBuildWindows.cs b/Engine/Source/Programs/UnrealBuildTool/Platform/Windows/UEBuildWindows.cs
--- a/Engine/Source/Programs/UnrealBuildTool/Platform/Windows/UEBuildWindows.cs
+++ b/Engine/Source/Programs/UnrealBuildTool/Platform/Windows/UEBuildWindows.cs
@@ -1109,7 +1109,6 @@ namespace UnrealBuildTool
 		/// <summary>
 		/// Validate a target's settings
 		/// </summary>
-		[SupportedOSPlatform("windows")]
 		public override void ValidateTarget(TargetRules Target)
 		{
 			if (Platform == UnrealTargetPlatform.Win64)
@@ -1238,17 +1237,20 @@ namespace UnrealBuildTool
 				Target.bCompileISPC = false; // The version of ISPC we currently use does not support Windows Aarch64
 			}
 
-			// Initialize the VC environment for the target, and set all the version numbers to the concrete values we chose
-			Target.WindowsPlatform.Environment = CreateVCEnvironment(Target);
+			if (OperatingSystem.IsWindows())
+			{
+				// Initialize the VC environment for the target, and set all the version numbers to the concrete values we chose
+				Target.WindowsPlatform.Environment = CreateVCEnvironment(Target);
 
-			// pull some things from it
-			Target.WindowsPlatform.Compiler = Target.WindowsPlatform.Environment.Compiler;
-			Target.WindowsPlatform.CompilerVersion = Target.WindowsPlatform.Environment.CompilerVersion.ToString();
-			Target.WindowsPlatform.ToolChain = Target.WindowsPlatform.Environment.ToolChain;
-			Target.WindowsPlatform.ToolchainVersion = Target.WindowsPlatform.Environment.ToolChainVersion.ToString();
-			Target.WindowsPlatform.WindowsSdkVersion = Target.WindowsPlatform.Environment.WindowsSdkVersion.ToString();
+				// pull some things from it
+				Target.WindowsPlatform.Compiler = Target.WindowsPlatform.Environment.Compiler;
+				Target.WindowsPlatform.CompilerVersion = Target.WindowsPlatform.Environment.CompilerVersion.ToString();
+				Target.WindowsPlatform.ToolChain = Target.WindowsPlatform.Environment.ToolChain;
+				Target.WindowsPlatform.ToolchainVersion = Target.WindowsPlatform.Environment.ToolChainVersion.ToString();
+				Target.WindowsPlatform.WindowsSdkVersion = Target.WindowsPlatform.Environment.WindowsSdkVersion.ToString();
 
-			ValidateToolchainVersion(Target);
+				ValidateToolchainVersion(Target);
+			}
 		}
 
 		static bool _toolchainWarningLogged = false;
