UnrealBuildTool: Improve missing precompiled module error

#rnx

diff --git a/Engine/Source/Programs/UnrealBuildTool/Configuration/UEBuildModuleCPP.cs b/Engine/Source/Programs/UnrealBuildTool/Configuration/UEBuildModuleCPP.cs
--- a/Engine/Source/Programs/UnrealBuildTool/Configuration/UEBuildModuleCPP.cs
+++ b/Engine/Source/Programs/UnrealBuildTool/Configuration/UEBuildModuleCPP.cs
@@ -371,8 +371,17 @@ namespace UnrealBuildTool
 			{
 				if (!FileReference.Exists(PrecompiledManifestLocation))
 				{
-					throw new BuildException("Missing precompiled manifest for '{0}', '{1}'. This module was most likely not flagged for being included in a precompiled build - set 'PrecompileForTargets = PrecompileTargetsType.Any;' in {0}.build.cs to override." +
-						" If part of a plugin, also check if its 'Type' is correct.", Name, PrecompiledManifestLocation);
+					Logger.LogError("Missing precompiled manifest for '{Name}', '{Manifest}", Name, PrecompiledManifestLocation);
+					Logger.LogInformation("This module was most likely not flagged for being included in a precompiled build - set 'PrecompileForTargets = PrecompileTargetsType.Any;' in {Name}.Build.cs to override.", Name);
+					if (Rules.Plugin != null)
+					{
+						Logger.LogInformation("As it is part of the plugin '{PluginName}', also check if its 'Type' is correct.", Rules.Plugin.Name);
+					}
+					if (ReferenceStackParentModules != null)
+					{
+						Logger.LogInformation("Dependent modules '{OtherModules}'", String.Join(' ', ReferenceStackParentModules.OrderBy(x => x.Name).Select(x => x.Name)));
+					}
+					throw new BuildLogEventException("Missing precompiled manifest for '{Name}', '{Manifest}" , Name, PrecompiledManifestLocation);
 				}
 
 				PrecompiledManifest Manifest = PrecompiledManifest.Read(PrecompiledManifestLocation);
@@ -381,7 +390,7 @@ namespace UnrealBuildTool
 					FileItem ObjectFile = FileItem.GetItemByFileReference(OutputFile);
 					if (!ObjectFile.Exists)
 					{
-						throw new BuildException("Missing object file {0} listed in {1}", OutputFile, PrecompiledManifestLocation);
+						throw new BuildLogEventException("Missing object file {OutputFile} listed in {Manifest}", OutputFile, PrecompiledManifestLocation);
 					}
 					LinkInputFiles.Add(ObjectFile);
 				}
