Use Joe.K change to use ForceNoAutoRTFMCompiler in the Get of bUseAutoRTFMCompiler

diff --git a/Engine/Source/Programs/UnrealBuildTool/Configuration/ReadOnlyTargetRules.cs b/Engine/Source/Programs/UnrealBuildTool/Configuration/ReadOnlyTargetRules.cs
--- a/Engine/Source/Programs/UnrealBuildTool/Configuration/ReadOnlyTargetRules.cs
+++ b/Engine/Source/Programs/UnrealBuildTool/Configuration/ReadOnlyTargetRules.cs
@@ -136,8 +136,6 @@ namespace UnrealBuildTool
 
 		public bool bUseAutoRTFMCompiler => Inner.bUseAutoRTFMCompiler;
 
-		public bool bForceNoAutoRTFMCompiler => Inner.bForceNoAutoRTFMCompiler;
-
 		[Obsolete("Deprecated in UE5.1 - No longer used as Chaos is always enabled.")]
 		public bool bCompileChaos => Inner.bCompileChaos;
 
diff --git a/Engine/Source/Programs/UnrealBuildTool/Configuration/TargetRules.cs b/Engine/Source/Programs/UnrealBuildTool/Configuration/TargetRules.cs
--- a/Engine/Source/Programs/UnrealBuildTool/Configuration/TargetRules.cs
+++ b/Engine/Source/Programs/UnrealBuildTool/Configuration/TargetRules.cs
@@ -868,7 +868,12 @@ namespace UnrealBuildTool
 		/// Whether to use the AutoRTFM Clang compiler.
 		/// </summary>
 		[RequiresUniqueBuildEnvironment]
-		public bool bUseAutoRTFMCompiler { get; set; }
+		public bool bUseAutoRTFMCompiler
+		{
+			get => bForceNoAutoRTFMCompiler ? false : _bUseAutoRTFMCompilerPrivate;
+			set => _bUseAutoRTFMCompilerPrivate = value;
+		}
+		private bool _bUseAutoRTFMCompilerPrivate = false;
 
 		/// <summary>
 		/// Whether to use force AutoRTFM Clang compiler off.
diff --git a/Engine/Source/Programs/UnrealBuildTool/Configuration/UEBuildModuleCPP.cs b/Engine/Source/Programs/UnrealBuildTool/Configuration/UEBuildModuleCPP.cs
--- a/Engine/Source/Programs/UnrealBuildTool/Configuration/UEBuildModuleCPP.cs
+++ b/Engine/Source/Programs/UnrealBuildTool/Configuration/UEBuildModuleCPP.cs
@@ -2138,7 +2138,7 @@ namespace UnrealBuildTool
 			Result.DeterministicWarningLevel = Rules.DeterministicWarningLevel;
 			Result.bValidateFormatStrings = Rules.bValidateFormatStrings;
 			Result.bValidateInternalApi = Rules.bValidateInternalApi;
-			Result.bUseAutoRTFMCompiler = Target.bUseAutoRTFMCompiler && !Target.bForceNoAutoRTFMCompiler;
+			Result.bUseAutoRTFMCompiler = Target.bUseAutoRTFMCompiler;
 
 			CompileEnvironmentDebugInfoSettings(Target, Result);
 
diff --git a/Engine/Source/Programs/UnrealBuildTool/Platform/Linux/UEBuildLinux.cs b/Engine/Source/Programs/UnrealBuildTool/Platform/Linux/UEBuildLinux.cs
--- a/Engine/Source/Programs/UnrealBuildTool/Platform/Linux/UEBuildLinux.cs
+++ b/Engine/Source/Programs/UnrealBuildTool/Platform/Linux/UEBuildLinux.cs
@@ -672,7 +672,7 @@ namespace UnrealBuildTool
 				Logger.LogWarning("Warning: bPreferThinLTO is set, but LTO is disabled. Flag will have no effect");
 			}
 
-			if (Target.bUseAutoRTFMCompiler && !Target.bForceNoAutoRTFMCompiler)
+			if (Target.bUseAutoRTFMCompiler)
 			{
 				Options |= ClangToolChainOptions.UseAutoRTFMCompiler;
 			}
diff --git a/Engine/Source/Programs/UnrealBuildTool/Platform/Windows/UEBuildWindows.cs b/Engine/Source/Programs/UnrealBuildTool/Platform/Windows/UEBuildWindows.cs
--- a/Engine/Source/Programs/UnrealBuildTool/Platform/Windows/UEBuildWindows.cs
+++ b/Engine/Source/Programs/UnrealBuildTool/Platform/Windows/UEBuildWindows.cs
@@ -1163,7 +1163,7 @@ namespace UnrealBuildTool
 				Logger.LogInformation("Defaulting static analyzer output type to text");
 			}
 
-			if (Target.bUseAutoRTFMCompiler && !Target.bForceNoAutoRTFMCompiler)
+			if (Target.bUseAutoRTFMCompiler)
 			{
 				// We check the static analyzer, and only allow using the AutoRTFM compiler if
 				// we are not doing static analysis, or are explicitly wanting to use clang.
