UnrealBuildTool: Add missing json converter for telemetry duration

#rnx

diff --git a/Engine/Source/Programs/UnrealBuildTool/System/Telemetry.cs b/Engine/Source/Programs/UnrealBuildTool/System/Telemetry.cs
--- a/Engine/Source/Programs/UnrealBuildTool/System/Telemetry.cs
+++ b/Engine/Source/Programs/UnrealBuildTool/System/Telemetry.cs
@@ -155,6 +155,7 @@ namespace UnrealBuildTool
 		public int SucceededCoordinators { get; }
 		public int FailedCoordinators { get; }
 
+		[JsonConverter(typeof(DurationDataRouterEventConverter))]
 		public TimeSpan DurationWaitingForRemote { get; }
 
 		public double RetriedLocalRate => RetriedLocalActions / (double)TotalActions;
