UBA horde settings from INI

#rnx

diff --git a/Engine/Source/Programs/UnrealBuildTool/Executors/UnrealBuildAccelerator/UBAAgentCoordinatorHorde.cs b/Engine/Source/Programs/UnrealBuildTool/Executors/UnrealBuildAccelerator/UBAAgentCoordinatorHorde.cs
--- a/Engine/Source/Programs/UnrealBuildTool/Executors/UnrealBuildAccelerator/UBAAgentCoordinatorHorde.cs
+++ b/Engine/Source/Programs/UnrealBuildTool/Executors/UnrealBuildAccelerator/UBAAgentCoordinatorHorde.cs
@@ -547,11 +547,24 @@ namespace UnrealBuildTool
 
 	class UBAAgentCoordinatorHorde : IUBAAgentCoordinator, IDisposable
 	{
-		public UBAAgentCoordinatorHorde(ILogger logger, UnrealBuildAcceleratorConfig ubaConfig, CommandLineArguments? additionalArguments = null)
+		public string? Server => HordeConfig.HordeServer;
+		public bool Enabled => !HordeConfig.bDisableHorde;
+
+		public UBAAgentCoordinatorHorde(ILogger logger, UnrealBuildAcceleratorConfig ubaConfig, CommandLineArguments? additionalArguments = null, DirectoryReference? projectDir = null)
 		{
 			_logger = logger;
 			_ubaConfig = ubaConfig;
 
+			ConfigCache.ReadSettings(projectDir, BuildHostPlatform.Current.Platform, HordeConfig);
+			if (HordeConfig.HordeEnabled?.Equals("BuildMachineOnly", StringComparison.OrdinalIgnoreCase) == true && !Unreal.IsBuildMachine())
+			{
+				// reset
+				HordeConfig = new();
+			}
+			if (HordeConfig.HordeEnabled?.Equals("False", StringComparison.OrdinalIgnoreCase) == true)
+			{
+				HordeConfig.bDisableHorde = true;
+			}
 			XmlConfig.ApplyTo(HordeConfig);
 			additionalArguments?.ApplyTo(HordeConfig);
 
diff --git a/Engine/Source/Programs/UnrealBuildTool/Executors/UnrealBuildAccelerator/UBAExecutor.cs b/Engine/Source/Programs/UnrealBuildTool/Executors/UnrealBuildAccelerator/UBAExecutor.cs
--- a/Engine/Source/Programs/UnrealBuildTool/Executors/UnrealBuildAccelerator/UBAExecutor.cs
+++ b/Engine/Source/Programs/UnrealBuildTool/Executors/UnrealBuildAccelerator/UBAExecutor.cs
@@ -105,17 +105,24 @@ namespace UnrealBuildTool
 			}
 		}
 
-		public UBAExecutor(int maxLocalActions, bool bAllCores, bool bCompactOutput, Microsoft.Extensions.Logging.ILogger logger, CommandLineArguments? additionalArguments = null)
+		public UBAExecutor(int maxLocalActions, bool bAllCores, bool bCompactOutput, Microsoft.Extensions.Logging.ILogger logger, IEnumerable<TargetDescriptor> targetDescriptors)
 			: base(maxLocalActions, bAllCores, bCompactOutput, logger)
 		{
 			XmlConfig.ApplyTo(this);
 			XmlConfig.ApplyTo(UBAConfig);
 			CommandLine.ParseArguments(Environment.GetCommandLineArgs(), this, logger);
-			additionalArguments?.ApplyTo(this);
-			additionalArguments?.ApplyTo(UBAConfig);
+
+			List<UBAAgentCoordinatorHorde> hordeAgentCoordinators = new();
+			foreach (TargetDescriptor targetDescriptor in targetDescriptors)
+			{
+				targetDescriptor.AdditionalArguments.ApplyTo(this);
+				targetDescriptor.AdditionalArguments.ApplyTo(UBAConfig);
+				hordeAgentCoordinators.Add(new UBAAgentCoordinatorHorde(logger, UBAConfig, targetDescriptor.AdditionalArguments, targetDescriptor.ProjectFile?.Directory));
+			}
+			hordeAgentCoordinators.RemoveAll(x => !x.Enabled);
+			_agentCoordinators.AddRange(hordeAgentCoordinators.DistinctBy(x => x.Server));
 
 			_threadedLogger = new ThreadedLogger(logger);
-			_agentCoordinators.Add(new UBAAgentCoordinatorHorde(logger, UBAConfig, additionalArguments));
 		}
 
 		private void PrintConfiguration()
@@ -982,8 +989,8 @@ namespace UnrealBuildTool
 			return EpicGames.UBA.Utils.IsAvailable();
 		}
 
-		public UBALocalExecutor(int maxLocalActions, bool bAllCores, bool bCompactOutput, Microsoft.Extensions.Logging.ILogger logger, CommandLineArguments? additionalArguments = null)
-			: base(maxLocalActions, bAllCores, bCompactOutput, logger, additionalArguments)
+		public UBALocalExecutor(int maxLocalActions, bool bAllCores, bool bCompactOutput, Microsoft.Extensions.Logging.ILogger logger, IEnumerable<TargetDescriptor> targetDescriptors)
+			: base(maxLocalActions, bAllCores, bCompactOutput, logger, targetDescriptors)
 		{
 			UBAConfig.bDisableRemote = true;
 			UBAConfig.bForceBuildAllRemote = false;
diff --git a/Engine/Source/Programs/UnrealBuildTool/Executors/UnrealBuildAccelerator/UnrealBuildAcceleratorHordeConfig.cs b/Engine/Source/Programs/UnrealBuildTool/Executors/UnrealBuildAccelerator/UnrealBuildAcceleratorHordeConfig.cs
--- a/Engine/Source/Programs/UnrealBuildTool/Executors/UnrealBuildAccelerator/UnrealBuildAcceleratorHordeConfig.cs
+++ b/Engine/Source/Programs/UnrealBuildTool/Executors/UnrealBuildAccelerator/UnrealBuildAcceleratorHordeConfig.cs
@@ -15,6 +15,7 @@ namespace UnrealBuildTool
 		/// <summary>
 		/// Uri of the Horde server
 		/// </summary>
+		[ConfigFile(ConfigHierarchyType.Engine, "Horde", "ServerUrl")]
 		[XmlConfigFile(Category = "Horde", Name = "Server")]
 		[CommandLine("-UBAHorde=")]
 		public string? HordeServer { get; set; }
@@ -22,6 +23,7 @@ namespace UnrealBuildTool
 		/// <summary>
 		/// Auth token for the Horde server
 		/// </summary>
+		[ConfigFile(ConfigHierarchyType.Engine, "Horde", "UbaToken")]
 		[XmlConfigFile(Category = "Horde", Name = "Token")]
 		[CommandLine("-UBAHordeToken=")]
 		public string? HordeToken { get; set; }
@@ -67,6 +69,7 @@ namespace UnrealBuildTool
 		/// <summary>
 		/// Pool for the Horde agent to assign if no override current platform doesn't have it set
 		/// </summary>
+		[ConfigFile(ConfigHierarchyType.Engine, "Horde", "UbaPool")]
 		[XmlConfigFile(Category = "Horde", Name = "Pool")]
 		public string? DefaultHordePool { get; set; }
 
@@ -79,24 +82,28 @@ namespace UnrealBuildTool
 		/// <summary>
 		/// Pool for the Horde agent to assign when on Linux
 		/// </summary>
+		[ConfigFile(ConfigHierarchyType.Engine, "Horde", "UbaLinuxPool")]
 		[XmlConfigFile(Category = "Horde", Name = "LinuxPool")]
 		public string? LinuxHordePool { get; set; }
 
 		/// <summary>
 		/// Pool for the Horde agent to assign when on Mac
 		/// </summary>
+		[ConfigFile(ConfigHierarchyType.Engine, "Horde", "UbaMacPool")]
 		[XmlConfigFile(Category = "Horde", Name = "MacPool")]
 		public string? MacHordePool { get; set; }
 
 		/// <summary>
 		/// Pool for the Horde agent to assign when on Windows
 		/// </summary>
+		[ConfigFile(ConfigHierarchyType.Engine, "Horde", "UbaWindowsPool")]
 		[XmlConfigFile(Category = "Horde", Name = "WindowsPool")]
 		public string? WindowsHordePool { get; set; }
 
 		/// <summary>
 		/// Requirements for the Horde agent to assign
 		/// </summary>
+		[ConfigFile(ConfigHierarchyType.Engine, "Horde", "UbaRequirements")]
 		[XmlConfigFile(Category = "Horde", Name = "Requirements")]
 		[CommandLine("-UBAHordeRequirements=")]
 		public string? HordeCondition { get; set; }
@@ -104,6 +111,7 @@ namespace UnrealBuildTool
 		/// <summary>
 		/// Which ip UBA server should give to agents. This will invert so host listens and agents connect
 		/// </summary>
+		[ConfigFile(ConfigHierarchyType.Engine, "Horde", "UbaLocalHost")]
 		[XmlConfigFile(Category = "Horde", Name = "LocalHost")]
 		[CommandLine("-UBAHordeHost")]
 		public string HordeHost { get; set; } = String.Empty;
@@ -111,6 +119,7 @@ namespace UnrealBuildTool
 		/// <summary>
 		/// Max cores allowed to be used by build session
 		/// </summary>
+		[ConfigFile(ConfigHierarchyType.Engine, "Horde", "UbaMaxCores")]
 		[XmlConfigFile(Category = "UnrealBuildAccelerator")]
 		[CommandLine("-UBAHordeMaxCores")]
 		public int HordeMaxCores { get; set; } = 576;
@@ -118,6 +127,7 @@ namespace UnrealBuildTool
 		/// <summary>
 		/// How long UBT should wait to ask for help. Useful in build configs where machine can delay remote work and still get same wall time results (pch dependencies etc)
 		/// </summary>
+		[ConfigFile(ConfigHierarchyType.Engine, "Horde", "UbaDelay")]
 		[XmlConfigFile(Category = "UnrealBuildAccelerator")]
 		[CommandLine("-UBAHordeDelay")]
 		public int HordeDelay { get; set; } = 0;
@@ -125,6 +135,7 @@ namespace UnrealBuildTool
 		/// <summary>
 		/// Allow use of Wine. Only applicable to Horde agents running Linux. Can still be ignored if Wine executable is not set on agent.
 		/// </summary>
+		[ConfigFile(ConfigHierarchyType.Engine, "Horde", "UbaAllowWine")]
 		[XmlConfigFile(Category = "UnrealBuildAccelerator")]
 		[CommandLine("-UBAHordeAllowWine", Value = "true")]
 		public bool bHordeAllowWine { get; set; } = true;
@@ -133,6 +144,7 @@ namespace UnrealBuildTool
 		/// Connection mode for agent/compute communication
 		/// <see cref="ConnectionMode" /> for valid modes.
 		/// </summary>
+		[ConfigFile(ConfigHierarchyType.Engine, "Horde", "UbaConnectionMode")]
 		[XmlConfigFile(Category = "Horde", Name = "ConnectionMode")]
 		[CommandLine("-UBAHordeConnectionMode=")]
 		public string? HordeConnectionMode { get; set; }
@@ -141,6 +153,7 @@ namespace UnrealBuildTool
 		/// Encryption to use for agent/compute communication. Note that UBA agent uses its own encryption.
 		/// <see cref="Encryption" /> for valid modes.
 		/// </summary>
+		[ConfigFile(ConfigHierarchyType.Engine, "Horde", "UbaEncryption")]
 		[XmlConfigFile(Category = "Horde", Name = "Encryption")]
 		[CommandLine("-UBAHordeEncryption=")]
 		public string? HordeEncryption { get; set; }
@@ -148,6 +161,7 @@ namespace UnrealBuildTool
 		/// <summary>
 		/// Sentry URL to send box data to. Optional.
 		/// </summary>
+		[ConfigFile(ConfigHierarchyType.Engine, "Horde", "UbaSentryUrl")]
 		[XmlConfigFile(Category = "Horde")]
 		[CommandLine("-UBASentryUrl=")]
 		public string? UBASentryUrl { get; set; }
@@ -158,5 +172,12 @@ namespace UnrealBuildTool
 		[XmlConfigFile(Category = "UnrealBuildAccelerator")]
 		[CommandLine("-UBADisableHorde")]
 		public bool bDisableHorde { get; set; } = false;
+
+		/// <summary>
+		/// Enabled parameter for use by INI config, expects one of [True, False, BuildMachineOnly]
+		/// </summary>
+		[ConfigFile(ConfigHierarchyType.Engine, "Horde", "UbaEnabled")]
+		public string? HordeEnabled { get; set; } = null;
+
 	}
 }
\ No newline at end of file
diff --git a/Engine/Source/Programs/UnrealBuildTool/System/ActionGraph.cs b/Engine/Source/Programs/UnrealBuildTool/System/ActionGraph.cs
--- a/Engine/Source/Programs/UnrealBuildTool/System/ActionGraph.cs
+++ b/Engine/Source/Programs/UnrealBuildTool/System/ActionGraph.cs
@@ -395,7 +395,7 @@ namespace UnrealBuildTool
 						// Intentionally not checking MinActionsForRemote
 						if (BuildConfiguration.bAllowUBAExecutor && UBAExecutor.IsAvailable())
 						{
-							return new UBAExecutor(BuildConfiguration.MaxParallelActions, BuildConfiguration.bAllCores, BuildConfiguration.bCompactOutput, Logger, TargetDescriptors.FirstOrDefault()?.AdditionalArguments);
+							return new UBAExecutor(BuildConfiguration.MaxParallelActions, BuildConfiguration.bAllCores, BuildConfiguration.bCompactOutput, Logger, TargetDescriptors);
 						}
 						return null;
 					}
@@ -424,7 +424,7 @@ namespace UnrealBuildTool
 
 			if (BuildConfiguration.bAllowUBALocalExecutor && UBALocalExecutor.IsAvailable())
 			{
-				return new UBALocalExecutor(BuildConfiguration.MaxParallelActions, BuildConfiguration.bAllCores, BuildConfiguration.bCompactOutput, Logger, TargetDescriptors.FirstOrDefault()?.AdditionalArguments);
+				return new UBALocalExecutor(BuildConfiguration.MaxParallelActions, BuildConfiguration.bAllCores, BuildConfiguration.bCompactOutput, Logger, TargetDescriptors);
 			}
 
 			return new ParallelExecutor(BuildConfiguration.MaxParallelActions, BuildConfiguration.bAllCores, BuildConfiguration.bCompactOutput, Logger);
