Fix Android Vulkan SM5 issues

RHI_RAYTRACING was always defined to 0.
Lumen was not enabled for SP_VULKAN_SM5_ANDROID

diff --git a/Engine/Source/Programs/UnrealBuildTool/Platform/Android/UEBuildAndroid.cs b/Engine/Source/Programs/UnrealBuildTool/Platform/Android/UEBuildAndroid.cs
--- a/Engine/Source/Programs/UnrealBuildTool/Platform/Android/UEBuildAndroid.cs
+++ b/Engine/Source/Programs/UnrealBuildTool/Platform/Android/UEBuildAndroid.cs
@@ -47,9 +47,15 @@ namespace UnrealBuildTool
 		/// <summary>
 		/// Enables runtime ray tracing support.
 		/// </summary>
-		[ConfigFile(ConfigHierarchyType.Engine, "/Script/AndroidTargetPlatform.AndroidTargetSettings")]
+		[ConfigFile(ConfigHierarchyType.Engine, "/Script/AndroidRuntimeSettings.AndroidRuntimeSettings")]
 		public bool bEnableRayTracing = false;
 
+		/// <summary>
+		/// Enables the desktop renderer.
+		/// </summary>
+		[ConfigFile(ConfigHierarchyType.Engine, "/Script/AndroidRuntimeSettings.AndroidRuntimeSettings")]
+		public bool bSupportsVulkanSM5 = false;
+
 		/// <summary>
 		/// Enables ASIS plugin and STANDALONE support.
 		/// </summary>
@@ -95,6 +101,8 @@ namespace UnrealBuildTool
 
 		public bool bEnableRayTracing => Inner.bEnableRayTracing;
 
+		public bool bSupportsVulkanSM5 => Inner.bSupportsVulkanSM5;
+
 		public bool bEnableASISPlugin => Inner.bEnableASISPlugin;
 
 		public AndroidTargetRules TargetRules => Inner;
@@ -505,7 +513,7 @@ namespace UnrealBuildTool
 			CompileEnvironment.Definitions.Add("USE_NULL_RHI=0");
  			CompileEnvironment.Definitions.Add("WITH_RHI_BREADCRUMBS=1");
 
-			if (Target.AndroidPlatform.bEnableRayTracing)
+			if (Target.AndroidPlatform.bEnableRayTracing && Target.AndroidPlatform.bSupportsVulkanSM5)
 			{
 				Logger.LogInformation("Compiling with ray tracing enabled");
 				CompileEnvironment.Definitions.Add("RHI_RAYTRACING=1");
diff --git a/Engine/Source/Runtime/RenderCore/Private/RenderUtils.cpp b/Engine/Source/Runtime/RenderCore/Private/RenderUtils.cpp
--- a/Engine/Source/Runtime/RenderCore/Private/RenderUtils.cpp
+++ b/Engine/Source/Runtime/RenderCore/Private/RenderUtils.cpp
@@ -2058,9 +2058,10 @@ bool DoesPlatformSupportLumenGI(EShaderPlatform Platform, bool bSkipProjectCheck
 	#endif // WITH_EDITOR
 
 	const bool bMetalSM5 = Platform == SP_METAL_SM5 || ParentPreviewPlatform == SP_METAL_SM5;
+	const bool bVulkanSM5 = Platform == SP_VULKAN_SM5 || Platform == SP_VULKAN_SM5_ANDROID;
 	return (bSkipProjectCheck || bLumenSupported)
 		&& FDataDrivenShaderPlatformInfo::GetSupportsLumenGI(Platform)
-		&& (IsFeatureLevelSupported(Platform, ERHIFeatureLevel::SM6) || Platform == SP_VULKAN_SM5 || bMetalSM5) // Android and Mac can't rely on SM6 yet, but want to run Lumen
+		&& (IsFeatureLevelSupported(Platform, ERHIFeatureLevel::SM6) || bVulkanSM5 || bMetalSM5) // Android and Mac can't rely on SM6 yet, but want to run Lumen
 		&& !IsForwardShadingEnabled(Platform);
 }
 
