UnrealBuildTool: Load Horde url from environment if a build machine and it is set

#rnx

diff --git a/Engine/Source/Programs/UnrealBuildTool/Executors/UnrealBuildAccelerator/UBAAgentCoordinatorHorde.cs b/Engine/Source/Programs/UnrealBuildTool/Executors/UnrealBuildAccelerator/UBAAgentCoordinatorHorde.cs
--- a/Engine/Source/Programs/UnrealBuildTool/Executors/UnrealBuildAccelerator/UBAAgentCoordinatorHorde.cs
+++ b/Engine/Source/Programs/UnrealBuildTool/Executors/UnrealBuildAccelerator/UBAAgentCoordinatorHorde.cs
@@ -561,6 +561,14 @@ namespace UnrealBuildTool
 				// reset
 				HordeConfig = new();
 			}
+			if (Unreal.IsBuildMachine())
+			{
+				string? hordeUrl = System.Environment.GetEnvironmentVariable("UE_HORDE_URL");
+				if (!String.IsNullOrEmpty(hordeUrl))
+				{
+					HordeConfig.HordeServer = hordeUrl;
+				}
+			}	
 			if (HordeConfig.HordeEnabled?.Equals("False", StringComparison.OrdinalIgnoreCase) == true)
 			{
 				HordeConfig.bDisableHorde = true;
