Update to PVS 7.31

diff --git a/Engine/Source/Programs/UnrealBuildTool/System/ActionGraph.cs b/Engine/Source/Programs/UnrealBuildTool/System/ActionGraph.cs
--- a/Engine/Source/Programs/UnrealBuildTool/System/ActionGraph.cs
+++ b/Engine/Source/Programs/UnrealBuildTool/System/ActionGraph.cs
@@ -297,11 +297,25 @@ namespace UnrealBuildTool
 
 				bool ShouldWarn(FileItem item)
 				{
-					return item.Location.FullName.Length > Unreal.RootDirectory.FullName.Length + BuildConfiguration.MaxNestedPathLength &&
-						item.Location.IsUnderDirectory(Unreal.RootDirectory) &&
-						!(item.Location.ContainsName("Intermediate", 0) && item.Location.ContainsName("H", 1)) && // Ignore -IncludeHeader items
-						(item.Location.ContainsName("Restricted", 0) == false) && //Be more relaxed for internal only code
-						item.Location.ContainsName("NotForLicensees", 0) == false;
+					if (item.Location.FullName.Length <= Unreal.RootDirectory.FullName.Length + BuildConfiguration.MaxNestedPathLength)
+						return false;
+
+					if (!item.Location.IsUnderDirectory(Unreal.RootDirectory))
+						return false;
+
+					if (!item.Location.FullName.Contains("/Restricted/NotForLicensees/", StringComparison.OrdinalIgnoreCase)) //Be more relaxed for internal only code
+						return false;
+
+					if (item.Location.FullName.Contains("/Intermediate/", StringComparison.OrdinalIgnoreCase))
+					{
+						if (item.Location.ContainsName("H", 1)) // Ignore -IncludeHeader items
+							return false;
+
+						if (item.Location.FullName.EndsWith(".i.PVS-Studio.log")) // Ignore PVS Studio intermediate items
+							return false;
+					}
+
+					return true;
 				}
 
 				HashSet<FileReference> FailPaths = new();
