UnrealBuildTool: Add a `bExplicitTargetForType` flag to TargetRules, to opt out of consideration by GetTargetNameByType.

This prevents myriad -TargetType=Editor invocations of UBT from suddenly becoming ambiguous as we introduce a new pseudo-Editor target.

diff --git a/Engine/Source/Programs/UnrealBuildTool/Configuration/TargetRules.cs b/Engine/Source/Programs/UnrealBuildTool/Configuration/TargetRules.cs
--- a/Engine/Source/Programs/UnrealBuildTool/Configuration/TargetRules.cs
+++ b/Engine/Source/Programs/UnrealBuildTool/Configuration/TargetRules.cs
@@ -817,6 +817,11 @@ namespace UnrealBuildTool
 		/// </summary>
 		public bool bGenerateProgramProject { get; set; }
 
+		/// <summary>
+		/// If true, GetTargetNameByType will not consider this target, thereby disambiguating -TargetType=X.
+		/// </summary>
+		public bool bExplicitTargetForType { get; set; }
+
 		/// <summary>
 		/// Whether the target should be included in the default solution build configuration
 		/// Setting this to false will skip building when running in the IDE
diff --git a/Engine/Source/Programs/UnrealBuildTool/System/RulesAssembly.cs b/Engine/Source/Programs/UnrealBuildTool/System/RulesAssembly.cs
--- a/Engine/Source/Programs/UnrealBuildTool/System/RulesAssembly.cs
+++ b/Engine/Source/Programs/UnrealBuildTool/System/RulesAssembly.cs
@@ -903,7 +903,7 @@ namespace UnrealBuildTool
 			foreach (KeyValuePair<string, FileReference> TargetPair in TargetNameToTargetFile)
 			{
 				TargetRules? Rules = CreateTargetRulesInstance(TargetPair.Key + "Target", new TargetInfo(TargetPair.Key, Platform, Configuration, null, ProjectFile, null), Logger);
-				if (Rules != null && Rules.Type == Type)
+				if (Rules != null && Rules.Type == Type && !Rules.bExplicitTargetForType)
 				{
 					Matches.Add(TargetPair.Key);
 				}
