UnrealBuildTool: Allow for ForceIncludeFiles to be set in ModuleRules

#rnx

diff --git a/Engine/Source/Programs/UnrealBuildTool/Configuration/ModuleRules.cs b/Engine/Source/Programs/UnrealBuildTool/Configuration/ModuleRules.cs
--- a/Engine/Source/Programs/UnrealBuildTool/Configuration/ModuleRules.cs
+++ b/Engine/Source/Programs/UnrealBuildTool/Configuration/ModuleRules.cs
@@ -740,6 +740,11 @@ namespace UnrealBuildTool
 
 		private FPSemanticsMode? FPSemanticsPrivate = null;
 
+		/// <summary>
+		/// Header files that should be force included for every source file in this module.
+		/// </summary>
+		public List<string> ForceIncludeFiles { get; } = new List<string>();
+
 		/// <summary>
 		/// Explicit private PCH for this module. Implies that this module will not use a shared PCH.
 		/// </summary>
diff --git a/Engine/Source/Programs/UnrealBuildTool/Configuration/UEBuildModule.cs b/Engine/Source/Programs/UnrealBuildTool/Configuration/UEBuildModule.cs
--- a/Engine/Source/Programs/UnrealBuildTool/Configuration/UEBuildModule.cs
+++ b/Engine/Source/Programs/UnrealBuildTool/Configuration/UEBuildModule.cs
@@ -1341,6 +1341,8 @@ namespace UnrealBuildTool
 			Writer.WriteValue("Rules", RulesFile.FullName);
 			Writer.WriteValue("PCHUsage", Rules.PCHUsage.ToString());
 
+			ExportJsonStringArray(Writer, "ForceIncludeModules", Rules.ForceIncludeFiles);
+
 			if (Rules.PrivatePCHHeaderFile != null)
 			{
 				Writer.WriteValue("PrivatePCH", FileReference.Combine(ModuleDirectory, Rules.PrivatePCHHeaderFile).FullName);
diff --git a/Engine/Source/Programs/UnrealBuildTool/Configuration/UEBuildModuleCPP.cs b/Engine/Source/Programs/UnrealBuildTool/Configuration/UEBuildModuleCPP.cs
--- a/Engine/Source/Programs/UnrealBuildTool/Configuration/UEBuildModuleCPP.cs
+++ b/Engine/Source/Programs/UnrealBuildTool/Configuration/UEBuildModuleCPP.cs
@@ -1791,7 +1791,7 @@ namespace UnrealBuildTool
 
 						CompileEnvironment = new CppCompileEnvironment(CompileEnvironment);
 						CompileEnvironment.Definitions.Clear();
-						CompileEnvironment.ForceIncludeFiles.Add(PrivateDefinitionsFileItem);
+						CompileEnvironment.ForceIncludeFiles.Insert(0, PrivateDefinitionsFileItem);
 						CompileEnvironment.PrecompiledHeaderAction = PrecompiledHeaderAction.Include;
 						CompileEnvironment.PrecompiledHeaderIncludeFilename = Instance.HeaderFile.Location;
 						CompileEnvironment.PCHInstance = Instance;
@@ -1830,7 +1830,7 @@ namespace UnrealBuildTool
 					CompileEnvironment.Definitions.Clear();
 
 					FileItem PrivateDefinitionsFileItem = Graph.CreateIntermediateTextFile(PrivateDefinitionsFile, Writer.ToString());
-					CompileEnvironment.ForceIncludeFiles.Add(PrivateDefinitionsFileItem);
+					CompileEnvironment.ForceIncludeFiles.Insert(0, PrivateDefinitionsFileItem);
 					return PrivateDefinitionsFileItem;
 				}
 			}
@@ -2284,6 +2284,25 @@ namespace UnrealBuildTool
 			// Setup the compile environment for the module.
 			SetupPrivateCompileEnvironment(Result.UserIncludePaths, Result.SystemIncludePaths, Result.ModuleInterfacePaths, Result.Definitions, Result.AdditionalFrameworks, Result.AdditionalPrerequisites, Rules.bLegacyPublicIncludePaths, Rules.bLegacyParentIncludePaths);
 
+			foreach (string ForceIncludeFile in Rules.ForceIncludeFiles)
+			{
+				if (Path.IsPathFullyQualified(ForceIncludeFile))
+				{
+					Result.ForceIncludeFiles.Add(FileItem.GetItemByPath(ForceIncludeFile));
+				}
+				else
+				{
+					FileReference? Path = Result.UserIncludePaths.Concat(Result.SystemIncludePaths).Select(x => FileReference.Combine(x, ForceIncludeFile)).FirstOrDefault(x => FileReference.Exists(x));
+					if (Path == null)
+					{
+						Logger.LogWarning("Unable to resolve force include path '{Path}'. Please verify that it can be found in the include directories for module '{Module}'", ForceIncludeFile, Name);
+						continue;
+					}
+					Result.ForceIncludeFiles.Add(FileItem.GetItemByFileReference(Path));
+				}
+				
+			}
+
 			return Result;
 		}
 
