UnrealBuildTool: Don't allow uba for any msvc clang link

#rnx

diff --git a/Engine/Source/Programs/UnrealBuildTool/Platform/Windows/VCToolChain.cs b/Engine/Source/Programs/UnrealBuildTool/Platform/Windows/VCToolChain.cs
--- a/Engine/Source/Programs/UnrealBuildTool/Platform/Windows/VCToolChain.cs
+++ b/Engine/Source/Programs/UnrealBuildTool/Platform/Windows/VCToolChain.cs
@@ -2806,7 +2806,7 @@ namespace UnrealBuildTool
 				LinkAction.bCanExecuteRemotely = true;
 			}
 
-			if ((Target.WindowsPlatform.Compiler.IsIntel() && Target.WindowsPlatform.bAllowClangLinker) || LinkEnvironment.bPGOOptimize || LinkEnvironment.bPGOProfile)
+			if ((Target.WindowsPlatform.Compiler.IsClang() && Target.WindowsPlatform.bAllowClangLinker) || LinkEnvironment.bPGOOptimize || LinkEnvironment.bPGOProfile)
 			{
 				LinkAction.bCanExecuteInUBA = false; // Disabled for now. Should revisit to see why it is not working
 			}
