UnrealBuildTool: Add support for automatic computer cluster assignment in Horde

- If cluster ID is set to "_auto", UBT will ask Horde for a suitable cluster ID to run compute tasks in
- Horde uses internal and external IP for the cluster selection
- The primary use-case is geographical load balancing and selecting the compute cluster closest to the initiating UBT session
- When enabled, pool config is ignored and instead conditions for the correct platform is used for selecting an agent

diff --git a/Engine/Source/Programs/Shared/EpicGames.Horde/Compute/Clients/AgentComputeClient.cs b/Engine/Source/Programs/Shared/EpicGames.Horde/Compute/Clients/AgentComputeClient.cs
--- a/Engine/Source/Programs/Shared/EpicGames.Horde/Compute/Clients/AgentComputeClient.cs
+++ b/Engine/Source/Programs/Shared/EpicGames.Horde/Compute/Clients/AgentComputeClient.cs
@@ -19,10 +19,15 @@ namespace EpicGames.Horde.Compute.Clients
 	/// </summary>
 	public sealed class AgentComputeClient : IComputeClient
 	{
+		private static readonly ClusterId s_cluster = new ("_agent");
+		
 		class LeaseImpl : IComputeLease
 		{
 			readonly IAsyncEnumerator<RemoteComputeSocket> _source;
-
+			
+			/// <inheritdoc/>
+			public ClusterId Cluster { get; } = s_cluster;
+			
 			/// <inheritdoc/>
 			public IReadOnlyList<string> Properties { get; } = new List<string>();
 
@@ -73,7 +78,13 @@ namespace EpicGames.Horde.Compute.Clients
 
 		/// <inheritdoc/>
 		public ValueTask DisposeAsync() => new ValueTask();
-
+		
+		/// <inheritdoc/>
+		public Task<ClusterId> GetClusterAsync(Requirements? requirements, string? requestId, ConnectionMetadataRequest? connection, ILogger logger, CancellationToken cancellationToken = default)
+		{
+			return Task.FromResult(s_cluster);
+		}
+		
 		/// <inheritdoc/>
 		public async Task<IComputeLease?> TryAssignWorkerAsync(ClusterId? clusterId, Requirements? requirements, string? requestId, ConnectionMetadataRequest? connection, ILogger logger, CancellationToken cancellationToken)
 		{
diff --git a/Engine/Source/Programs/Shared/EpicGames.Horde/Compute/Clients/LocalComputeClient.cs b/Engine/Source/Programs/Shared/EpicGames.Horde/Compute/Clients/LocalComputeClient.cs
--- a/Engine/Source/Programs/Shared/EpicGames.Horde/Compute/Clients/LocalComputeClient.cs
+++ b/Engine/Source/Programs/Shared/EpicGames.Horde/Compute/Clients/LocalComputeClient.cs
@@ -16,8 +16,11 @@ namespace EpicGames.Horde.Compute.Clients
 	/// </summary>
 	public sealed class LocalComputeClient : IComputeClient
 	{
+		private static readonly ClusterId s_cluster = new ("_local");
+		
 		class LeaseImpl : IComputeLease
 		{
+			public ClusterId Cluster { get; } = s_cluster;
 			public IReadOnlyList<string> Properties { get; } = new List<string>();
 			public IReadOnlyDictionary<string, int> AssignedResources => new Dictionary<string, int>();
 			public RemoteComputeSocket Socket => _socket;
@@ -82,7 +85,13 @@ namespace EpicGames.Horde.Compute.Clients
 			await worker.RunAsync(socket, cancellationToken);
 			await socket.CloseAsync(cancellationToken);
 		}
-
+		
+		/// <inheritdoc/>
+		public Task<ClusterId> GetClusterAsync(Requirements? requirements, string? requestId, ConnectionMetadataRequest? connection, ILogger logger, CancellationToken cancellationToken = default)
+		{
+			return Task.FromResult(s_cluster);
+		}
+		
 		/// <inheritdoc/>
 		public Task<IComputeLease?> TryAssignWorkerAsync(ClusterId? clusterId, Requirements? requirements, string? requestId, ConnectionMetadataRequest? connection, ILogger logger, CancellationToken cancellationToken)
 		{
diff --git a/Engine/Source/Programs/Shared/EpicGames.Horde/Compute/Clients/ServerComputeClient.cs b/Engine/Source/Programs/Shared/EpicGames.Horde/Compute/Clients/ServerComputeClient.cs
--- a/Engine/Source/Programs/Shared/EpicGames.Horde/Compute/Clients/ServerComputeClient.cs
+++ b/Engine/Source/Programs/Shared/EpicGames.Horde/Compute/Clients/ServerComputeClient.cs
@@ -115,6 +115,7 @@ namespace EpicGames.Horde.Compute.Clients
 		public const int NonceLength = 64;
 
 		record LeaseInfo(
+			ClusterId Cluster,
 			IReadOnlyList<string> Properties,
 			IReadOnlyDictionary<string, int> AssignedResources,
 			RemoteComputeSocket Socket,
@@ -124,6 +125,7 @@ namespace EpicGames.Horde.Compute.Clients
 
 		class LeaseImpl : IComputeLease
 		{
+			public ClusterId Cluster => _source.Current.Cluster;
 			public IReadOnlyList<string> Properties => _source.Current.Properties;
 			public IReadOnlyDictionary<string, int> AssignedResources => _source.Current.AssignedResources;
 			public RemoteComputeSocket Socket => _source.Current.Socket;
@@ -216,7 +218,34 @@ namespace EpicGames.Horde.Compute.Clients
 		{
 			_cancellationSource.Dispose();
 		}
-
+		
+		/// <inheritdoc/>
+		public async Task<ClusterId> GetClusterAsync(Requirements? requirements, string? requestId, ConnectionMetadataRequest? connection, ILogger logger, CancellationToken cancellationToken = default)
+		{
+			AssignComputeRequest request = new()
+			{
+				Requirements = requirements,
+				RequestId = requestId,
+				Connection = connection,
+				Protocol = (int)ComputeProtocol.Latest
+			};
+			
+			using HttpResponseMessage httpResponse = await HordeHttpClient.PostAsync(_httpClient, "api/v2/compute/_cluster", request, _cancellationSource.Token);
+			if (!httpResponse.IsSuccessStatusCode)
+			{
+				string body = await httpResponse.Content.ReadAsStringAsync(cancellationToken);
+				throw new ComputeClientException($"Unable to find suitable cluster. HTTP status code {httpResponse.StatusCode}: {body}");
+			}
+			
+			GetClusterResponse? response = await httpResponse.Content.ReadFromJsonAsync<GetClusterResponse>(HordeHttpClient.JsonSerializerOptions, cancellationToken);
+			if (response == null)
+			{
+				throw new InvalidOperationException();
+			}
+			
+			return response.ClusterId;
+		}
+		
 		/// <inheritdoc/>
 		public async Task<IComputeLease?> TryAssignWorkerAsync(ClusterId? clusterId, Requirements? requirements, string? requestId, ConnectionMetadataRequest? connection, ILogger logger, CancellationToken cancellationToken)
 		{
@@ -349,7 +378,7 @@ namespace EpicGames.Horde.Compute.Clients
 
 			await using ComputeTransport transport = await CreateTransportAsync(socket, response, cancellationToken);
 			await using RemoteComputeSocket computeSocket = new(transport, (ComputeProtocol)response.Protocol, workerLogger);
-			yield return new LeaseInfo(response.Properties, response.AssignedResources, computeSocket, response.Ip, response.ConnectionMode, response.Ports);
+			yield return new LeaseInfo(response.ClusterId, response.Properties, response.AssignedResources, computeSocket, response.Ip, response.ConnectionMode, response.Ports);
 		}
 
 		private static async Task<ComputeTransport> CreateTransportAsync(Socket socket, AssignComputeResponse response, CancellationToken cancellationToken)
diff --git a/Engine/Source/Programs/Shared/EpicGames.Horde/Compute/IComputeClient.cs b/Engine/Source/Programs/Shared/EpicGames.Horde/Compute/IComputeClient.cs
--- a/Engine/Source/Programs/Shared/EpicGames.Horde/Compute/IComputeClient.cs
+++ b/Engine/Source/Programs/Shared/EpicGames.Horde/Compute/IComputeClient.cs
@@ -13,6 +13,16 @@ namespace EpicGames.Horde.Compute
 	/// </summary>
 	public interface IComputeClient : IAsyncDisposable
 	{
+		/// <summary>
+		/// Find the most suitable cluster to execute a given compute assignment request
+		/// </summary>
+		/// <param name="requirements">Requirements for the agent</param>
+		/// <param name="requestId">Optional ID identifying the request over multiple calls, such as retrying the same request</param>
+		/// <param name="connection">Optional preference of connection details</param>
+		/// <param name="logger">Logger for output from this worker</param>
+		/// <param name="cancellationToken">Cancellation token for the operation</param>
+		public Task<ClusterId> GetClusterAsync(Requirements? requirements, string? requestId, ConnectionMetadataRequest? connection, ILogger logger, CancellationToken cancellationToken = default);
+		
 		/// <summary>
 		/// Adds a new remote request
 		/// </summary>
diff --git a/Engine/Source/Programs/Shared/EpicGames.Horde/Compute/IComputeLease.cs b/Engine/Source/Programs/Shared/EpicGames.Horde/Compute/IComputeLease.cs
--- a/Engine/Source/Programs/Shared/EpicGames.Horde/Compute/IComputeLease.cs
+++ b/Engine/Source/Programs/Shared/EpicGames.Horde/Compute/IComputeLease.cs
@@ -12,6 +12,11 @@ namespace EpicGames.Horde.Compute
 	/// </summary>
 	public interface IComputeLease : IAsyncDisposable
 	{
+		/// <summary>
+		/// Compute cluster ID
+		/// </summary>
+		ClusterId Cluster { get; }
+		
 		/// <summary>
 		/// Properties of the remote machine
 		/// </summary>
diff --git a/Engine/Source/Programs/UnrealBuildTool/Executors/UnrealBuildAccelerator/UBAAgentCoordinatorHorde.cs b/Engine/Source/Programs/UnrealBuildTool/Executors/UnrealBuildAccelerator/UBAAgentCoordinatorHorde.cs
--- a/Engine/Source/Programs/UnrealBuildTool/Executors/UnrealBuildAccelerator/UBAAgentCoordinatorHorde.cs
+++ b/Engine/Source/Programs/UnrealBuildTool/Executors/UnrealBuildAccelerator/UBAAgentCoordinatorHorde.cs
@@ -10,6 +10,7 @@ using System.Threading;
 using System.Threading.Tasks;
 using EpicGames.Core;
 using EpicGames.Horde;
+using EpicGames.Horde.Agents;
 using EpicGames.Horde.Common;
 using EpicGames.Horde.Compute;
 using EpicGames.Horde.Compute.Clients;
@@ -33,7 +34,7 @@ namespace UnrealBuildTool
 		readonly Guid _id = Guid.NewGuid();
 
 		const string ResourceLogicalCores = "LogicalCores";
-		readonly ClusterId _clusterId = new("default");
+		ClusterId? _clusterId;
 
 		readonly string? _pool;
 		readonly bool _allowWine;
@@ -225,7 +226,7 @@ namespace UnrealBuildTool
 		{
 			if (_client == null)
 			{
-				throw new InvalidOperationException("Call init first");
+				throw new InvalidOperationException($"Session not initialized. Call {nameof(InitAsync)} first");
 			}
 
 			PrefixLogger workerLogger = new($"[Worker{_workerId}]", _logger);
@@ -248,6 +249,8 @@ namespace UnrealBuildTool
 					Encryption = _encryption,
 					Ports = { { UbaPortName, UbaPort }, { UbaProxyPortName, UbaProxyPort } }
 				};
+				
+				await ResolveClusterIdAsync(hordeConfig, requirements, requestId, cmr, workerLogger, cancellationToken);
 				lease = await _client.TryAssignWorkerAsync(_clusterId, requirements, requestId, cmr, workerLogger, cancellationToken);
 				if (lease == null)
 				{
@@ -324,13 +327,13 @@ namespace UnrealBuildTool
 
 		async Task UpdateCpuCoreNeedAsync(int targetCoreCount, CancellationToken cancellationToken = default)
 		{
-			if (_client != null && _pool != null)
+			if (_client != null && _pool != null && _clusterId != null)
 			{
 				_logger.LogDebug("Setting CPU core need to {TargetCoreCount}", targetCoreCount);
 				Dictionary<string, int> resourceNeeds = new() { { ResourceLogicalCores, targetCoreCount } };
 				try
 				{
-					await _client.DeclareResourceNeedsAsync(_clusterId, _pool, resourceNeeds, cancellationToken);
+					await _client.DeclareResourceNeedsAsync(_clusterId.Value, _pool, resourceNeeds, cancellationToken);
 				}
 				catch (Exception e)
 				{
@@ -338,6 +341,31 @@ namespace UnrealBuildTool
 				}
 			}
 		}
+		
+		private async Task ResolveClusterIdAsync(UnrealBuildAcceleratorHordeConfig config, Requirements requirements, string requestId, ConnectionMetadataRequest cmr, ILogger logger, CancellationToken cancellationToken = default)
+		{
+			if (_client == null)
+			{
+				throw new InvalidOperationException($"Session not initialized. Call {nameof(InitAsync)} first");
+			}
+			
+			if (_clusterId == null)
+			{
+				if (config.HordeCluster == UnrealBuildAcceleratorHordeConfig.ClusterAuto)
+				{
+					_clusterId = await _client.GetClusterAsync(requirements, requestId, cmr, logger, cancellationToken);
+				}
+				else if (!String.IsNullOrEmpty(config.HordeCluster))
+				{
+					_clusterId = new ClusterId(config.HordeCluster);
+				}
+				else
+				{
+					_clusterId = new ClusterId(UnrealBuildAcceleratorHordeConfig.ClusterDefault);
+				}
+				_logger.LogInformation("Horde cluster resolved as {ClusterId}", _clusterId.ToString());
+			}
+		}
 
 		public static async Task<UBAHordeSession?> TryCreateHordeSessionAsync(UnrealBuildAcceleratorHordeConfig hordeConfig, UBAExecutor executor, bool bStrictErrors, ILogger logger, CancellationToken cancellationToken = default)
 		{
@@ -655,21 +683,7 @@ namespace UnrealBuildTool
 							break;
 						}
 
-						Requirements requirements = new()
-						{
-							Exclusive = true
-						};
-
-						if (!String.IsNullOrEmpty(HordeConfig.HordePool))
-						{
-							requirements.Pool = HordeConfig.HordePool;
-						}
-
-						if (HordeConfig.HordeCondition != null)
-						{
-							requirements.Condition = Condition.Parse(HordeConfig.HordeCondition);
-						}
-
+						Requirements requirements = GetRequirements(HordeConfig);
 						if (!await hordeSession.AddWorkerAsync(requirements, HordeConfig, _cancellationSource.Token))
 						{
 							_logger.LogDebug("No additional workers available");
@@ -739,6 +753,48 @@ namespace UnrealBuildTool
 				_timer = null;
 			}
 		}
+		
+		private static Requirements GetRequirements(UnrealBuildAcceleratorHordeConfig config)
+		{
+			Requirements requirements = new() { Exclusive = true };
+			
+			if (config.HordeCluster == UnrealBuildAcceleratorHordeConfig.ClusterAuto)
+			{
+				string condition = "";
+				if (OperatingSystem.IsWindows())
+				{
+					condition = $"({KnownPropertyNames.OsFamily} == 'Windows' || {KnownPropertyNames.WineEnabled} == 'true')";
+				}
+				else if (OperatingSystem.IsMacOS())
+				{
+					condition = $"{KnownPropertyNames.OsFamily} == 'MacOS'";
+				}
+				else if (OperatingSystem.IsLinux())
+				{
+					condition = $"{KnownPropertyNames.OsFamily} == 'Linux'";
+				}
+				
+				if (config.HordeCondition != null)
+				{
+					condition = " " + config.HordeCondition;
+				}
+				requirements.Condition = Condition.Parse(condition);
+			}
+			else
+			{
+				if (!String.IsNullOrEmpty(config.HordePool))
+				{
+					requirements.Pool = config.HordePool;
+				}
+				
+				if (config.HordeCondition != null)
+				{
+					requirements.Condition = Condition.Parse(config.HordeCondition);
+				}
+			}
+			
+			return requirements;
+		}
 
 		readonly ILogger _logger;
 		readonly UnrealBuildAcceleratorConfig _ubaConfig;
diff --git a/Engine/Source/Programs/UnrealBuildTool/Executors/UnrealBuildAccelerator/UnrealBuildAcceleratorHordeConfig.cs b/Engine/Source/Programs/UnrealBuildTool/Executors/UnrealBuildAccelerator/UnrealBuildAcceleratorHordeConfig.cs
--- a/Engine/Source/Programs/UnrealBuildTool/Executors/UnrealBuildAccelerator/UnrealBuildAcceleratorHordeConfig.cs
+++ b/Engine/Source/Programs/UnrealBuildTool/Executors/UnrealBuildAccelerator/UnrealBuildAcceleratorHordeConfig.cs
@@ -107,6 +107,25 @@ namespace UnrealBuildTool
 		[XmlConfigFile(Category = "Horde", Name = "Requirements")]
 		[CommandLine("-UBAHordeRequirements=")]
 		public string? HordeCondition { get; set; }
+		
+		/// <summary>
+		/// Default compute cluster ID if nothing is set
+		/// </summary>
+		public const string ClusterDefault = "default";
+		
+		/// <summary>
+		/// Compute cluster ID for automatically resolving the most suitable cluster
+		/// </summary>
+		public const string ClusterAuto = "_auto";
+		
+		/// <summary>
+		/// Compute cluster ID to use in Horde. Set to "_auto" to let Horde server resolve a suitable cluster.
+		/// In multi-region setups this is can simplify configuration of UBT/UBA a lot.
+		/// </summary>
+		[ConfigFile(ConfigHierarchyType.Engine, "Horde", "UbaCluster")]
+		[XmlConfigFile(Category = "Horde", Name = "Cluster")]
+		[CommandLine("-UBAHordeCluster=")]
+		public string? HordeCluster { get; set; }
 
 		/// <summary>
 		/// Which ip UBA server should give to agents. This will invert so host listens and agents connect
