[UBA]
* Added display log entry for uba cache connection succeeded/failed plus time it took

diff --git a/Engine/Source/Programs/UnrealBuildTool/Executors/UnrealBuildAccelerator/UBAExecutor.cs b/Engine/Source/Programs/UnrealBuildTool/Executors/UnrealBuildAccelerator/UBAExecutor.cs
--- a/Engine/Source/Programs/UnrealBuildTool/Executors/UnrealBuildAccelerator/UBAExecutor.cs
+++ b/Engine/Source/Programs/UnrealBuildTool/Executors/UnrealBuildAccelerator/UBAExecutor.cs
@@ -370,14 +370,15 @@ namespace UnrealBuildTool
 								port = Int32.Parse(nameAndPort[1]);
 							}
 
-							if (_cacheClient.Connect(nameAndPort[0], port))
+							System.Diagnostics.Stopwatch stopwatch = System.Diagnostics.Stopwatch.StartNew();
+							bool cacheSuccess = _cacheClient.Connect(nameAndPort[0], port);
+							long totalMs = stopwatch.ElapsedMilliseconds;
+							string successText = cacheSuccess ? "Connected to" : "Failed to connect to";
+							logger.LogInformation("UbaCache - {SuccessText} {Name}:{Port} ({Seconds}.{Milliseconds}s)", successText, nameAndPort[0], port, totalMs / 1000, totalMs % 1000);
+							if (cacheSuccess)
 							{
 								actionArtifactCache = new UBAActionArtifactCache(this);
 							}
-							else
-							{
-								logger.LogInformation("Timed out trying to connect to cache server on {Name}:{Port}. Cache will be disabled", nameAndPort[0], port);
-							}
 						}
 
 						ubaStorage = ubaStorageServer;
