[UBA]
* Added status information from ubt/horde in trace stream (to show in visualizer)

diff --git a/Engine/Source/Programs/Shared/EpicGames.UBA/ISessionServer.cs b/Engine/Source/Programs/Shared/EpicGames.UBA/ISessionServer.cs
--- a/Engine/Source/Programs/Shared/EpicGames.UBA/ISessionServer.cs
+++ b/Engine/Source/Programs/Shared/EpicGames.UBA/ISessionServer.cs
@@ -221,6 +221,24 @@ namespace EpicGames.UBA
 		/// <param name="exitCode">The exit code of the external process</param>
 		public abstract void EndExternalProcess(uint id, uint exitCode);
 
+		/// <summary>
+		/// Writes external status to the uba trace stream which can then be visualized by ubavisualizer
+		/// </summary>
+		/// <param name="processesTotal">Total processes in session</param>
+		/// <param name="processesDone">Processes done in session</param>
+		/// <param name="errorCount">Number of errors in session</param>
+		public abstract void UpdateProgress(uint processesTotal, uint processesDone, uint errorCount);
+
+		/// <summary>
+		/// Writes external status to the uba trace stream which can then be visualized by ubavisualizer
+		/// </summary>
+		/// <param name="statusRow">Row of status text. Reuse one index to show one line in visualizer</param>
+		/// <param name="statusColumn">The identation of status name that will be shown in visualizer</param>
+		/// <param name="statusText">The status text that will be shown in visualizer</param>
+		/// <param name="statusType">The status type</param>
+		/// <param name="statusLink">Optional hyperlink that can be used to make text clickable in visualizer</param>
+		public abstract void UpdateStatus(uint statusRow, uint statusColumn, string statusText, LogEntryType statusType, string? statusLink = null);
+
 		/// <summary>
 		/// Set a custom cas key for a process's tracked inputs
 		/// </summary>
diff --git a/Engine/Source/Programs/Shared/EpicGames.UBA/Impl/SessionServerImpl.cs b/Engine/Source/Programs/Shared/EpicGames.UBA/Impl/SessionServerImpl.cs
--- a/Engine/Source/Programs/Shared/EpicGames.UBA/Impl/SessionServerImpl.cs
+++ b/Engine/Source/Programs/Shared/EpicGames.UBA/Impl/SessionServerImpl.cs
@@ -94,6 +94,12 @@ namespace EpicGames.UBA
 		[DllImport("UbaHost", CharSet = CharSet.Auto)]
 		static extern void SessionServer_EndExternalProcess(IntPtr server, uint id, uint exitCode);
 
+		[DllImport("UbaHost", CharSet = CharSet.Auto)]
+		static extern void SessionServer_UpdateProgress(IntPtr server, uint processesTotal, uint processesDone, uint errorCount);
+
+		[DllImport("UbaHost", CharSet = CharSet.Auto)]
+		static extern uint SessionServer_UpdateStatus(IntPtr server, uint statusRow, uint statusColumn, string statusText, byte statusType, string? statusLink);
+
 		[DllImport("UbaHost", CharSet = CharSet.Auto)]
 		static extern IntPtr SessionServer_RunProcess(IntPtr server, IntPtr info, bool async, bool enableDetour);
 
@@ -201,8 +207,13 @@ namespace EpicGames.UBA
 		public void RegisterNewFiles(params string[] files) => Array.ForEach(files, (file) => SessionServer_RegisterNewFile(_handle, file));
 
 		public uint BeginExternalProcess(string description) => SessionServer_BeginExternalProcess(_handle, description);
+
 		public void EndExternalProcess(uint id, uint exitCode) => SessionServer_EndExternalProcess(_handle, id, exitCode);
 
+		public void UpdateProgress(uint processesTotal, uint processesDone, uint errorCount) => SessionServer_UpdateProgress(_handle, processesTotal, processesDone, errorCount);
+
+		public void UpdateStatus(uint statusRow, uint statusColumn, string statusText, LogEntryType statusType, string? statusLink) => SessionServer_UpdateStatus(_handle, statusRow, statusColumn, statusText, (byte)statusType, statusLink);
+
 		public void SetCustomCasKeyFromTrackedInputs(string file, string workingDirectory, IProcess process) => SessionServer_SetCustomCasKeyFromTrackedInputs(_handle, process.GetHandle(), file, workingDirectory);
 
 		public void PrintSummary() => SessionServer_PrintSummary(_handle);
diff --git a/Engine/Source/Programs/UnrealBuildTool/Executors/ImmediateActionQueue.cs b/Engine/Source/Programs/UnrealBuildTool/Executors/ImmediateActionQueue.cs
--- a/Engine/Source/Programs/UnrealBuildTool/Executors/ImmediateActionQueue.cs
+++ b/Engine/Source/Programs/UnrealBuildTool/Executors/ImmediateActionQueue.cs
@@ -261,6 +261,8 @@ namespace UnrealBuildTool
 		public bool StopCompilationAfterErrors = false;
 
 		public int CompletedActions { get => _completedActions; }
+		public int CacheHitActions { get => _cacheHitActions; }
+		public int CacheMissActions { get => _cacheMissActions; }
 
 		/// <summary>
 		/// Return true if the queue is done
@@ -272,6 +274,11 @@ namespace UnrealBuildTool
 		/// </summary>
 		public Action<LinkedAction>? OnArtifactsRead = null;
 
+		/// <summary>
+		///  Action that can be to notify when artifacts have been missed
+		/// </summary>
+		public Action<LinkedAction>? OnArtifactsMiss = null;
+
 		/// <summary>
 		/// Collection of available runners
 		/// </summary>
@@ -429,12 +436,13 @@ namespace UnrealBuildTool
 						if (success)
 						{
 							Interlocked.Increment(ref _cacheHitActions);
-							OnArtifactsRead?.Invoke(action);
 							OnActionCompleted(action, success, s_copiedFromCacheResults);
+							OnArtifactsRead?.Invoke(action);
 						}
 						else
 						{
 							Interlocked.Increment(ref _cacheMissActions);
+							OnArtifactsMiss?.Invoke(action);
 							RequeueAction(action);
 						}
 					});
diff --git a/Engine/Source/Programs/UnrealBuildTool/Executors/UnrealBuildAccelerator/UBAAgentCoordinator.cs b/Engine/Source/Programs/UnrealBuildTool/Executors/UnrealBuildAccelerator/UBAAgentCoordinator.cs
--- a/Engine/Source/Programs/UnrealBuildTool/Executors/UnrealBuildAccelerator/UBAAgentCoordinator.cs
+++ b/Engine/Source/Programs/UnrealBuildTool/Executors/UnrealBuildAccelerator/UBAAgentCoordinator.cs
@@ -3,9 +3,12 @@
 using System;
 using System.Threading.Tasks;
 using EpicGames.Core;
+using EpicGames.UBA;
 
 namespace UnrealBuildTool
 {
+	// statusRow, statusColumn, statusText, statusType, statusLink
+	using StatusUpdateAction = Action<uint, uint, string, LogEntryType, string?>;
 
 	interface IUBAAgentCoordinator
 	{
@@ -13,10 +16,12 @@ namespace UnrealBuildTool
 
 		Task InitAsync(UBAExecutor executor);
 
-		void Start(ImmediateActionQueue queue, Func<LinkedAction, bool> canRunRemotely);
+		void Start(ImmediateActionQueue queue, Func<LinkedAction, bool> canRunRemotely, StatusUpdateAction updateStatus);
 
 		void Stop();
 
 		Task CloseAsync();
+
+		void Done();
 	}
 }
\ No newline at end of file
diff --git a/Engine/Source/Programs/UnrealBuildTool/Executors/UnrealBuildAccelerator/UBAAgentCoordinatorHorde.cs b/Engine/Source/Programs/UnrealBuildTool/Executors/UnrealBuildAccelerator/UBAAgentCoordinatorHorde.cs
--- a/Engine/Source/Programs/UnrealBuildTool/Executors/UnrealBuildAccelerator/UBAAgentCoordinatorHorde.cs
+++ b/Engine/Source/Programs/UnrealBuildTool/Executors/UnrealBuildAccelerator/UBAAgentCoordinatorHorde.cs
@@ -26,6 +26,9 @@ using UnrealBuildBase;
 
 namespace UnrealBuildTool
 {
+	// statusRow, statusColumn, statusText, statusType, statusLink
+	using StatusUpdateAction = Action<uint, uint, string, EpicGames.UBA.LogEntryType, string?>;
+
 	class UBAHordeSession : IAsyncDisposable
 	{
 		/// <summary>
@@ -46,6 +49,7 @@ namespace UnrealBuildTool
 		readonly ILogger _logger;
 
 		readonly UBAExecutor _owner;
+		internal StatusUpdateAction? _updateStatus;
 
 		readonly BundleStorageClient _storage = BundleStorageClient.CreateInMemory(NullLogger.Instance);
 		BlobLocator _ubaAgentLocator;
@@ -54,15 +58,16 @@ namespace UnrealBuildTool
 		readonly string _crypto;
 		IComputeClient? _client;
 
-		public struct Worker
+		public class Worker
 		{
-			public Task BackgroundTask { get; set; }
+			public Stopwatch? StartTime { get; init; }
+			public string? Ip { get; init; }
+			public ConnectionMetadataPort? Port { get; init; }
+			public ConnectionMetadataPort? ProxyPort { get; init; }
+
 			public int NumLogicalCores { get; set; }
-			public Stopwatch StartTime { get; set; }
-			public bool Started { get; set; }
-			public string Ip { get; set; }
-			public ConnectionMetadataPort Port { get; set; }
-			public ConnectionMetadataPort ProxyPort { get; set; }
+			public bool Active { get; set; }
+			public Task? BackgroundTask { get; set; }
 		}
 
 		readonly List<Worker> _workers = new();
@@ -103,7 +108,7 @@ namespace UnrealBuildTool
 
 				for (int idx = _workers.Count - 1; idx >= 0; idx--)
 				{
-					await _workers[idx].BackgroundTask;
+					await _workers[idx].BackgroundTask!;
 					_workers.RemoveAt(idx);
 				}
 			}
@@ -181,33 +186,46 @@ namespace UnrealBuildTool
 			return handle.GetLocator();
 		}
 
-		public int NumLogicalCores { get; private set; } = 0;
-
 		public async void RemoveCompleteWorkersAsync()
 		{
 			for (int idx = 0; idx < _workers.Count; idx++)
 			{
 				Worker worker = _workers[idx];
-				if (worker.BackgroundTask.IsCompleted)
+				if (worker.BackgroundTask!.IsCompleted)
 				{
 					await worker.BackgroundTask;
-					NumLogicalCores -= worker.NumLogicalCores;
-					_workers.RemoveAt(idx--);
+					lock (_workers)
+						_workers.RemoveAt(idx--);
 				}
 			}
 		}
 
-		public int QueuedUpCores()
+		public void GetCoreCount(out int active, out int queued, out int activeAgents, out int queuedAgents)
 		{
-			int count = 0;
-			foreach (Worker worker in _workers)
+			active = 0;
+			queued = 0;
+			activeAgents = 0;
+			queuedAgents = 0;
+			lock (_workers)
 			{
-				if (!worker.Started)
+				foreach (Worker worker in _workers)
 				{
-					count += worker.NumLogicalCores;
+					if (worker.NumLogicalCores == 0)
+					{
+						continue;
+					}
+					if (worker.Active)
+					{
+						++activeAgents;
+						active += worker.NumLogicalCores;
+					}
+					else
+					{
+						++queuedAgents;
+						queued += worker.NumLogicalCores;
+					}
 				}
 			}
-			return count;
 		}
 
 		int _workerId = 0;
@@ -222,7 +240,7 @@ namespace UnrealBuildTool
 				"LogicalCores"
 			};
 
-		public async Task<bool> AddWorkerAsync(Requirements requirements, UnrealBuildAcceleratorHordeConfig hordeConfig, CancellationToken cancellationToken)
+		public async Task<bool> AddWorkerAsync(Requirements requirements, UnrealBuildAcceleratorHordeConfig hordeConfig, CancellationToken cancellationToken, int activeCores)
 		{
 			if (_client == null)
 			{
@@ -256,7 +274,7 @@ namespace UnrealBuildTool
 				{
 					_logger.LogDebug("Unable to assign a remote worker");
 
-					int missingNumCores = Math.Max(0, _maxCores - NumLogicalCores);
+					int missingNumCores = Math.Max(0, _maxCores - activeCores);
 					await UpdateCpuCoreNeedAsync(missingNumCores, cancellationToken);
 					return false;
 				}
@@ -310,10 +328,11 @@ namespace UnrealBuildTool
 					ProxyPort = ubaProxyPort,
 				};
 				worker.BackgroundTask = RunWorkerAsync(worker, lease, locator, exeName, workerLogger, hordeConfig, _cancellationTokenSource.Token);
-				_workers.Add(worker);
+				lock (_workers)
+					_workers.Add(worker);
+				UpdateHordeStatus(null);
 				lease = null; // Will be disposed by RunWorkerAsync
 
-				NumLogicalCores += numLogicalCores;
 				return true;
 			}
 			finally
@@ -325,10 +344,14 @@ namespace UnrealBuildTool
 			}
 		}
 
+		int _targetCoreCount;
+
 		async Task UpdateCpuCoreNeedAsync(int targetCoreCount, CancellationToken cancellationToken = default)
 		{
-			if (_client != null && _pool != null && _clusterId != null)
+			if (_client != null && _pool != null && _clusterId != null && targetCoreCount != _targetCoreCount)
 			{
+				_targetCoreCount = targetCoreCount;
+
 				_logger.LogDebug("Setting CPU core need to {TargetCoreCount}", targetCoreCount);
 				Dictionary<string, int> resourceNeeds = new() { { ResourceLogicalCores, targetCoreCount } };
 				try
@@ -484,7 +507,8 @@ namespace UnrealBuildTool
 						}
 						else
 						{
-							arguments.Add($"-Listen={self.Port.AgentPort}");
+							arguments.Add($"-Listen={self.Port!.AgentPort}");
+							arguments.Add("-ListenTimeout=10");
 						}
 
 						if (!String.IsNullOrEmpty(_crypto))
@@ -498,7 +522,7 @@ namespace UnrealBuildTool
 						{
 							arguments.Add($"-Sentry=\"{hordeConfig.UBASentryUrl}\"");
 						}
-						arguments.Add("-ProxyPort=" + self.ProxyPort.AgentPort);
+						arguments.Add("-ProxyPort=" + self.ProxyPort!.AgentPort);
 						if (_owner.UBAConfig.bUseQuic)
 						{
 							arguments.Add("-quic");
@@ -516,6 +540,7 @@ namespace UnrealBuildTool
 						arguments.Add("-Dir=%UE_HORDE_SHARED_DIR%\\Uba");
 						arguments.Add("-Eventfile=%UE_HORDE_TERMINATION_SIGNAL_FILE%");
 						arguments.Add("-MaxIdle=15");
+
 						if (_owner.UBAConfig.bLogEnabled)
 						{
 							arguments.Add("-Log");
@@ -528,7 +553,7 @@ namespace UnrealBuildTool
 						ExecuteProcessFlags execFlags = _allowWine ? ExecuteProcessFlags.UseWine : ExecuteProcessFlags.None;
 						await using AgentManagedProcess process = await channel.ExecuteAsync(executable, arguments, null, null, execFlags, cancellationToken);
 						bool shouldConnect = !useListen;
-						self.Started = true;
+						bool isFirstRead = true;
 						string? line;
 
 						while ((line = await process.ReadLineAsync(cancellationToken)) != null)
@@ -537,14 +562,26 @@ namespace UnrealBuildTool
 
 							if (shouldConnect && line.Contains("Listening on", StringComparison.OrdinalIgnoreCase)) // This log entry means that the agent is ready for connections.
 							{
-								long totalMs = self.StartTime.ElapsedMilliseconds;
+								long totalMs = self.StartTime!.ElapsedMilliseconds;
 								logger.LogInformation("Connecting to UbaAgent on {Ip}:{Port} (local agent port {AgentPort}) {Seconds}.{Milliseconds} seconds after assigned",
-									self.Ip, self.Port.Port, self.Port.AgentPort, totalMs / 1000, totalMs % 1000);
+									self.Ip, self.Port!.Port, self.Port.AgentPort, totalMs / 1000, totalMs % 1000);
 
-								_owner.Server!.AddClient(self.Ip, self.Port.Port, _crypto);
+								if (!_owner.Server!.AddClient(self.Ip!, self.Port.Port, _crypto))
+								{
+									break;
+								}
 								shouldConnect = false;
 							}
+
+							if (isFirstRead)
+							{
+								isFirstRead = false;
+								self.Active = true;
+								UpdateHordeStatus(null);
+							}
 						}
+						self.NumLogicalCores = 0;
+						UpdateHordeStatus(null);
 						logger.LogDebug("Shutting down process");
 						int ExitCode = await process.WaitForExitAsync(cancellationToken);
 						if (ExitCode != 0)
@@ -559,9 +596,14 @@ namespace UnrealBuildTool
 			}
 			catch (TimeoutException)
 			{
+				self.NumLogicalCores = 0;
+				UpdateHordeStatus(null);
 			}
 			catch (Exception ex)
 			{
+				self.NumLogicalCores = 0;
+				UpdateHordeStatus(null);
+
 				if (!cancellationToken.IsCancellationRequested)
 				{
 					logger.Log(_strict ? LogLevel.Error : LogLevel.Debug, KnownLogEvents.Systemic_Horde_Compute, ex, "Exception in worker task: {Ex}", ex.ToString());
@@ -571,6 +613,34 @@ namespace UnrealBuildTool
 				}
 			}
 		}
+
+		string _additionalStatus = "";
+		string _lastStatus = "";
+
+		internal void UpdateHordeStatus(string? additionalStatus)
+		{
+			if (additionalStatus != null)
+			{
+				_additionalStatus = additionalStatus;
+			}
+
+			if (_updateStatus == null)
+			{
+				return;
+			}
+
+			int activeCores;
+			int queuedCores;
+			int activeAgents;
+			int queuedAgents;
+			GetCoreCount(out activeCores, out queuedCores, out activeAgents, out queuedAgents);
+			string status = $"Running. {activeAgents} agent{(activeAgents != 1 ? "s" : "")} ({activeCores} cores){_additionalStatus} {(queuedAgents != 0 ? $"(Preparing {queuedAgents} agent{(queuedAgents != 1 ? "s" : "")})" : "")}";
+			if (_lastStatus != status)
+			{
+				_lastStatus = status;
+				_updateStatus(0, 6, status, EpicGames.UBA.LogEntryType.Info, null);
+			}
+		}
 	}
 
 	class UBAAgentCoordinatorHorde : IUBAAgentCoordinator, IDisposable
@@ -628,26 +698,49 @@ namespace UnrealBuildTool
 				return;
 			}
 
+			_executor = executor;
 			_cancellationSource = new CancellationTokenSource();
 			_hordeSessionTask = UBAHordeSession.TryCreateHordeSessionAsync(HordeConfig, executor, _ubaConfig.bStrict, _logger, _cancellationSource.Token);
 			await _hordeSessionTask;
 			executor.AgentCoordinatorInitialized(this, _hordeSessionTask.IsCompletedSuccessfully && _hordeSessionTask.Result != null);
 		}
 
-		public void Start(ImmediateActionQueue queue, Func<LinkedAction, bool> canRunRemotely)
+		public void Start(ImmediateActionQueue queue, Func<LinkedAction, bool> canRunRemotely, StatusUpdateAction updateStatus)
 		{
+			_updateStatus = updateStatus;
+
 			int timerPeriod = 5000;
 			bool shownNoAgentsFoundMessage = false;
 
-			if (_hordeSessionTask == null)
+			UBAHordeSession? hordeSession = _hordeSessionTask?.Result;
+
+			string? link = null;
+			if (HordeConfig.HordeServer != null)
+			{
+				link = $"{HordeConfig.HordeServer!.TrimEnd('/')}/agents";
+				if (HordeConfig.HordePool != string.Empty)
+				{
+					link += $"?agent={HordeConfig.HordePool}";
+				}
+			}
+			updateStatus(0, 1, "Horde", EpicGames.UBA.LogEntryType.Info, link);
+
+			if (hordeSession == null)
 			{
+				updateStatus(0, 6, "Not created (check log)", EpicGames.UBA.LogEntryType.Info, null);
 				return;
 			}
+
+			hordeSession._updateStatus = _updateStatus;
+			hordeSession.UpdateHordeStatus(null);
+
+			int requestCounter = 0;
+
 			_timer = new(async (_) =>
 			{
 				_timer?.Change(Timeout.Infinite, Timeout.Infinite);
 
-				if (_cancellationSource!.IsCancellationRequested || _hordeSessionTask == null)
+				if (_hordeSessionTask == null)
 				{
 					return;
 				}
@@ -659,10 +752,11 @@ namespace UnrealBuildTool
 					return;
 				}
 
-				hordeSession.RemoveCompleteWorkersAsync();
+				//hordeSession.RemoveCompleteWorkersAsync();
 
-				if (queue.IsDone)
+				if (queue.IsDone || _cancellationSource!.IsCancellationRequested)
 				{
+					hordeSession.UpdateHordeStatus(" - Requests stopped");
 					return;
 				}
 
@@ -671,26 +765,44 @@ namespace UnrealBuildTool
 
 				try
 				{
-					double queueWeight = queue.EnumerateReadyToCompileActions().Where(x => canRunRemotely(x)).Sum(x => x.Weight);
-
-					queueWeight -= hordeSession.QueuedUpCores();
 					while (true)
 					{
-						int currentLogicalCores = hordeSession.NumLogicalCores;
+						if (_cancellationSource!.IsCancellationRequested)
+						{
+							hordeSession.UpdateHordeStatus(" - Requests stopped");
+							return;
+						}
+
+						hordeSession.RemoveCompleteWorkersAsync();
+
+						double queueWeight = queue.EnumerateReadyToCompileActions().Where(x => canRunRemotely(x)).Sum(x => x.Weight);
+
+						int activeCores;
+						int queuedCores;
+						int activeAgents;
+						int queuedAgents;
+						hordeSession.GetCoreCount(out activeCores, out queuedCores, out activeAgents, out queuedAgents);
+
+						queueWeight -= queuedCores;
+
+						bool isSatisfied = queueWeight <= queueThreshold || (activeCores + queuedCores) >= HordeConfig.HordeMaxCores || _cancellationSource!.IsCancellationRequested;
 
-						if (queueWeight <= queueThreshold || currentLogicalCores >= HordeConfig.HordeMaxCores || _cancellationSource!.IsCancellationRequested)
+						hordeSession.UpdateHordeStatus(isSatisfied ? " - Requests paused" : $" - Requesting agent{("....".Substring(0, (requestCounter++ % 4)))}");
+
+						if (isSatisfied)
 						{
 							break;
 						}
 
 						Requirements requirements = GetRequirements(HordeConfig);
-						if (!await hordeSession.AddWorkerAsync(requirements, HordeConfig, _cancellationSource.Token))
+						if (!await hordeSession.AddWorkerAsync(requirements, HordeConfig, _cancellationSource.Token, activeCores))
 						{
-							_logger.LogDebug("No additional workers available");
+							await Task.Delay(500); // Sleep a little bit to make sure previous horde status is visible
+							string status = (queuedAgents + activeAgents) == 0 ? " - No agents available" : " - No additional agents available";
+							hordeSession.UpdateHordeStatus(status);
+							_logger.LogDebug(status);
 							break;
 						}
-						int coresAdded = hordeSession.NumLogicalCores - currentLogicalCores;
-						queueWeight -= coresAdded;
 					}
 				}
 				catch (NoComputeAgentsFoundException ex)
@@ -735,6 +847,14 @@ namespace UnrealBuildTool
 			}
 		}
 
+		public void Done()
+		{
+			if (_updateStatus != null)
+			{
+				_updateStatus(0, 6, "Done", EpicGames.UBA.LogEntryType.Info, null);
+			}
+		}
+
 		public void Dispose()
 		{
 			Stop();
@@ -800,8 +920,10 @@ namespace UnrealBuildTool
 		readonly UnrealBuildAcceleratorConfig _ubaConfig;
 		UnrealBuildAcceleratorHordeConfig HordeConfig { get; init; } = new();
 
+		UBAExecutor? _executor;
 		CancellationTokenSource? _cancellationSource;
 		Task<UBAHordeSession?>? _hordeSessionTask;
+		StatusUpdateAction? _updateStatus;
 		Timer? _timer;
 	}
 }
diff --git a/Engine/Source/Programs/UnrealBuildTool/Executors/UnrealBuildAccelerator/UBAExecutor.cs b/Engine/Source/Programs/UnrealBuildTool/Executors/UnrealBuildAccelerator/UBAExecutor.cs
--- a/Engine/Source/Programs/UnrealBuildTool/Executors/UnrealBuildAccelerator/UBAExecutor.cs
+++ b/Engine/Source/Programs/UnrealBuildTool/Executors/UnrealBuildAccelerator/UBAExecutor.cs
@@ -25,6 +25,7 @@ namespace UnrealBuildTool
 		public string Crypto { get; private set; } = String.Empty;
 		public IServer? Server { get; private set; }
 		ISessionServer? _session;
+		object _sessionLock = new();
 		ICacheClient? _cacheClient;
 		EpicGames.UBA.ILogger? _ubaLogger;
 		readonly List<IUBAAgentCoordinator> _agentCoordinators = new();
@@ -33,6 +34,7 @@ namespace UnrealBuildTool
 		bool _bIsRemoteActionsAllowed = true;
 		readonly object _actionsChangedLock = new();
 		bool _bActionsChanged = true;
+		uint _errorCount = 0;
 		uint _actionsQueuedThatCanRunRemotely = UInt32.MaxValue;
 		readonly ThreadedLogger _threadedLogger;
 
@@ -125,6 +127,14 @@ namespace UnrealBuildTool
 			_threadedLogger = new ThreadedLogger(logger);
 		}
 
+		public void UpdateStatus(uint statusRow, uint statusColumn, string statusText, LogEntryType statusType, string? statusLink)
+		{
+			lock (_sessionLock)
+			{
+				_session?.UpdateStatus(statusRow, statusColumn, statusText, statusType, statusLink);
+			}
+		}
+
 		private void PrintConfiguration()
 		{
 			_threadedLogger.LogInformation("  Storage capacity {StoreCapacityGb}Gb", UBAConfig.StoreCapacityGb);
@@ -244,15 +254,20 @@ namespace UnrealBuildTool
 			public Task FlushChangesAsync(CancellationToken cancellationToken) => Task.CompletedTask;
 		}
 
-		private async void ActionQueueCanceled(IStorageServer? ubaStorage)
+		private void ActionQueueCanceled(IStorageServer? ubaStorage)
 		{
 			Server?.StopServer(); // Make sure all remove processes are returned. We can't have any callbacks after this
 			_bIsCancelled = true;
 			_session?.CancelAll(); // Cancel all processes native side
 			ubaStorage?.SaveCasTable();
-			foreach (IUBAAgentCoordinator coordinator in _agentCoordinators)
+
+			// We need the lock here since things are happening in parallel.
+			lock (_sessionLock)
 			{
-				await coordinator.CloseAsync();
+				foreach (IUBAAgentCoordinator coordinator in _agentCoordinators)
+				{
+					_agentCoordinators.ForEach(ac => ac.CloseAsync().Wait(2000)); // Give coordinators some time to close (this makes coordinators like horde return resources faster)
+				}
 			}
 		}
 
@@ -266,7 +281,6 @@ namespace UnrealBuildTool
 			if (inputActions.Count() < NumParallelProcesses && !UBAConfig.bForceBuildAllRemote)
 			{
 				UBAConfig.bDisableRemote = true;
-				UBAConfig.bForceBuildAllRemote = false;
 				UBAConfig.Zone = "local";
 			}
 
@@ -358,51 +372,75 @@ namespace UnrealBuildTool
 					_ubaLogger = ubaLogger;
 					using IStorageServer ubaStorageServer = IStorageServer.CreateStorageServer(Server, ubaLogger, new StorageServerCreateInfo(_rootDirRef.FullName, ((ulong)UBAConfig.StoreCapacityGb) * 1000 * 1000 * 1000, !UBAConfig.bStoreRaw, UBAConfig.Zone));
 					using ISessionServerCreateInfo serverCreateInfo = ISessionServerCreateInfo.CreateSessionServerCreateInfo(ubaStorageServer, Server, ubaLogger, new SessionServerCreateInfo(_rootDirRef.FullName, ubaTraceFile.FullName.Replace('\\', '/'), UBAConfig.bDisableCustomAlloc, false, UBAConfig.bResetCas, UBAConfig.bWriteToDisk, UBAConfig.bDetailedTrace, !UBAConfig.bDisableWaitOnMem, UBAConfig.bAllowKillOnMem, UBAConfig.bStoreObjFilesCompressed));
-					using (_session = ISessionServer.CreateSessionServer(serverCreateInfo))
-					using (_cacheClient = ICacheClient.CreateCacheClient(_session, UBAConfig.bReportCacheMissReason))
 					{
-						if (!String.IsNullOrEmpty(UBAConfig.CacheServer))
+						try
 						{
-							string[] nameAndPort = UBAConfig.CacheServer.Split(':');
-							int port = 1347;
-							if (nameAndPort.Length > 1)
+							_session = ISessionServer.CreateSessionServer(serverCreateInfo);
+							_cacheClient = ICacheClient.CreateCacheClient(_session, UBAConfig.bReportCacheMissReason);
+							if (!String.IsNullOrEmpty(UBAConfig.CacheServer))
 							{
-								port = Int32.Parse(nameAndPort[1]);
+								string[] nameAndPort = UBAConfig.CacheServer.Split(':');
+								int port = 1347;
+								if (nameAndPort.Length > 1)
+								{
+									port = Int32.Parse(nameAndPort[1]);
+								}
+
+								_session.UpdateStatus(1, 1, "Cache", LogEntryType.Info, null);
+								_session.UpdateStatus(1, 6, "Connecting...", LogEntryType.Info);
+
+								System.Diagnostics.Stopwatch stopwatch = System.Diagnostics.Stopwatch.StartNew();
+								bool cacheSuccess = _cacheClient.Connect(nameAndPort[0], port);
+								long totalMs = stopwatch.ElapsedMilliseconds;
+								string successText = cacheSuccess ? "Connected to" : "Failed to connect to";
+								logger.LogInformation("UbaCache - {SuccessText} {Name}:{Port} ({Seconds}.{Milliseconds}s)", successText, nameAndPort[0], port, totalMs / 1000, totalMs % 1000);
+								if (cacheSuccess)
+								{
+									_session.UpdateStatus(1, 6, "Connected", LogEntryType.Info);
+									actionArtifactCache = new UBAActionArtifactCache(this);
+								}
+								else
+								{
+									_session.UpdateStatus(1, 6, "Not connected", LogEntryType.Info);
+								}
 							}
 
-							System.Diagnostics.Stopwatch stopwatch = System.Diagnostics.Stopwatch.StartNew();
-							bool cacheSuccess = _cacheClient.Connect(nameAndPort[0], port);
-							long totalMs = stopwatch.ElapsedMilliseconds;
-							string successText = cacheSuccess ? "Connected to" : "Failed to connect to";
-							logger.LogInformation("UbaCache - {SuccessText} {Name}:{Port} ({Seconds}.{Milliseconds}s)", successText, nameAndPort[0], port, totalMs / 1000, totalMs % 1000);
-							if (cacheSuccess)
+							ubaStorage = ubaStorageServer;
+
+							if (!UBAConfig.bDisableRemote)
 							{
-								actionArtifactCache = new UBAActionArtifactCache(this);
+								Server.StartServer(UBAConfig.Host, UBAConfig.Port, Crypto);
+							}
+							else
+							{
+								_session.DisableRemoteExecution();
 							}
-						}
 
-						ubaStorage = ubaStorageServer;
+							bool success = ExecuteActionsInternal(inputActions, _session, logger, actionArtifactCache, () => ActionQueueCanceled(ubaStorage));
 
-						if (!UBAConfig.bDisableRemote)
-						{
-							Server.StartServer(UBAConfig.Host, UBAConfig.Port, Crypto);
-						}
-
-						bool success = ExecuteActionsInternal(inputActions, _session, logger, actionArtifactCache, () => ActionQueueCanceled(ubaStorage));
+							if (!UBAConfig.bDisableRemote)
+							{
+								Server.StopServer();
+							}
 
-						//_cacheClient.RequestServerShutdown("");
+							if (UBAConfig.bPrintSummary)
+							{
+								_session.PrintSummary();
+							}
 
-						if (!UBAConfig.bDisableRemote)
-						{
-							Server.StopServer();
+							return success && !_bIsCancelled;
 						}
-
-						if (UBAConfig.bPrintSummary)
+						finally
 						{
-							_session.PrintSummary();
+							_cacheClient?.Dispose();
+							_cacheClient = null;
+							lock (_sessionLock)
+							{
+								_agentCoordinators.ForEach(ac => ac.Done());
+								_session?.Dispose();
+								_session = null;
+							}
 						}
-
-						return success && !_bIsCancelled;
 					}
 				}
 			}
@@ -475,6 +513,11 @@ namespace UnrealBuildTool
 			ImmediateActionQueueRunner remoteRunner = queue.CreateManualRunner(action => RunActionRemote(queue, action));
 			queue.CancellationToken.Register(onCancel);
 
+			queue.OnArtifactsRead += (action) => { UpdateCacheProgress(queue); UpdateProgress(queue); };
+			queue.OnArtifactsMiss += (action) => UpdateCacheProgress(queue);
+
+			UpdateProgress(queue);
+
 			// Start the queue
 			queue.Start();
 
@@ -489,10 +532,7 @@ namespace UnrealBuildTool
 					if (count <= NumParallelProcesses)
 					{
 						_bIsRemoteActionsAllowed = false;
-						foreach (IUBAAgentCoordinator coordinator in _agentCoordinators)
-						{
-							coordinator.Stop();
-						}
+						_agentCoordinators.ForEach(ac => ac.Stop());
 						session!.DisableRemoteExecution();
 					}
 					else
@@ -516,9 +556,12 @@ namespace UnrealBuildTool
 
 			try
 			{
+				uint statusUpdateCounter = 10;
 				foreach (IUBAAgentCoordinator coordinator in _agentCoordinators)
 				{
-					coordinator.Start(queue, CanRunRemotely);
+					uint statusUpdateIndex = statusUpdateCounter;
+					coordinator.Start(queue, CanRunRemotely, (sr, sc, st, t, sl) => UpdateStatus(statusUpdateIndex + sr, sc, st, t, sl));
+					statusUpdateCounter += 10;
 				}
 
 				bool res = queue.RunTillDone().Result; // Using inline wait to avoid possible thread switch
@@ -535,9 +578,9 @@ namespace UnrealBuildTool
 			}
 			finally
 			{
-				foreach (IUBAAgentCoordinator coordinator in _agentCoordinators)
+				if (_bIsRemoteActionsAllowed)
 				{
-					coordinator.Stop();
+					_agentCoordinators.ForEach(ac => ac.Stop());
 				}
 			}
 		}
@@ -592,7 +635,7 @@ namespace UnrealBuildTool
 				UserData = action,
 				Description = description,
 				Configuration = action.bIsGCCCompiler ? EpicGames.UBA.ProcessStartInfo.CommonProcessConfigs.CompileClang : EpicGames.UBA.ProcessStartInfo.CommonProcessConfigs.CompileMsvc,
-				LogFile = UBAConfig.bLogEnabled ? action.Inner.ProducedItems.First().Location.GetFileName() : null,
+				LogFile = UBAConfig.bLogEnabled ? (action.Inner.ProducedItems.First().Location.GetFileName() + ".log") : null,
 			};
 
 			pchItem = null;
@@ -935,12 +978,28 @@ namespace UnrealBuildTool
 
 			queue.OnActionCompleted(action, success, results);
 
+			if (!success)
+			{
+				++_errorCount;
+			}
+			UpdateProgress(queue);
+
 			lock (_actionsChangedLock)
 			{
 				_bActionsChanged = true;
 			}
 		}
-
+		void UpdateProgress(ImmediateActionQueue queue)
+		{
+			int completedActions = queue.CompletedActions;
+			int totalActions = queue.TotalActions;
+			int percent = completedActions * 100 / totalActions;
+			_session!.UpdateProgress((uint)totalActions, (uint)completedActions, _errorCount);
+		}
+		void UpdateCacheProgress(ImmediateActionQueue queue)
+		{
+			_session!.UpdateStatus(1, 6, $"Hits {queue.CacheHitActions} Misses {queue.CacheMissActions}", LogEntryType.Info);
+		}
 		void RemoteActionFailedNoOutput(ImmediateActionQueue queue, LinkedAction action, int exitCode, string executingHost)
 		{
 			if (_bIsCancelled)
