UBT: Adjust cluster usage for Horde and UBA

- Include cluster from config with the rest of the config log line
- If HordeCluster is specified utlizied those prefered conditions to give it priority instead of only when ClustAuto is specified

#rnx

diff --git a/Engine/Source/Programs/UnrealBuildTool/Executors/UnrealBuildAccelerator/UBAAgentCoordinatorHorde.cs b/Engine/Source/Programs/UnrealBuildTool/Executors/UnrealBuildAccelerator/UBAAgentCoordinatorHorde.cs
--- a/Engine/Source/Programs/UnrealBuildTool/Executors/UnrealBuildAccelerator/UBAAgentCoordinatorHorde.cs
+++ b/Engine/Source/Programs/UnrealBuildTool/Executors/UnrealBuildAccelerator/UBAAgentCoordinatorHorde.cs
@@ -386,7 +386,7 @@ namespace UnrealBuildTool
 				{
 					_clusterId = new ClusterId(UnrealBuildAcceleratorHordeConfig.ClusterDefault);
 				}
-				_logger.LogInformation("Horde cluster resolved as {ClusterId}", _clusterId.ToString());
+				_logger.LogInformation("Horde cluster resolved as '{ClusterId}'", _clusterId.ToString());
 			}
 		}
 
@@ -416,8 +416,8 @@ namespace UnrealBuildTool
 				encryption = Encryption.Ssl;
 			}
 
-			logger.LogInformation("Horde URL: {Server}, Pool: {Pool}, Condition: {Condition}, Connection: {Connection} HordeEncryption: {Encryption}",
-				server, hordeConfig.HordePool ?? "(none)", hordeConfig.HordeCondition ?? "(none)", connectionMode?.ToString() ?? "(none)", encryption?.ToString() ?? "(none)");
+			logger.LogInformation("Horde URL: {Server}, Pool: {Pool}, Cluster {Cluster}, Condition: {Condition}, Connection: {Connection} HordeEncryption: {Encryption}",
+				server, hordeConfig.HordePool ?? "(none)", hordeConfig.HordeCluster ?? "(none)", hordeConfig.HordeCondition ?? "(none)", connectionMode?.ToString() ?? "(none)", encryption?.ToString() ?? "(none)");
 			try
 			{
 				bool allowWine = hordeConfig.bHordeAllowWine && OperatingSystem.IsWindows();
@@ -878,8 +878,9 @@ namespace UnrealBuildTool
 		{
 			Requirements requirements = new() { Exclusive = true };
 			
-			if (config.HordeCluster == UnrealBuildAcceleratorHordeConfig.ClusterAuto)
+			if (!string.IsNullOrEmpty(config.HordeCluster))
 			{
+				// Apply filtering options when a cluster, either explicit or auto is specified, ensuring the cluster options take precidence.
 				string condition = "";
 				if (OperatingSystem.IsWindows())
 				{
