Adding UE support for custom generated code to be compiled similar to default generated code.
This is done via new exporter option called CompileOutput.

diff --git a/Engine/Source/Programs/Shared/EpicGames.UHT/Tables/UhtExporterTable.cs b/Engine/Source/Programs/Shared/EpicGames.UHT/Tables/UhtExporterTable.cs
--- a/Engine/Source/Programs/Shared/EpicGames.UHT/Tables/UhtExporterTable.cs
+++ b/Engine/Source/Programs/Shared/EpicGames.UHT/Tables/UhtExporterTable.cs
@@ -31,6 +31,14 @@ namespace EpicGames.UHT.Tables
 		/// The exporter should be run by default
 		/// </summary>
 		Default = 1 << 0,
+		
+		/// <summary>
+		/// Generated CPP files that match any filter in the CppFilters property of this UhtExporter will be included in the compilation process.
+		/// For this option, all custom generated files should be placed into the same folder as the default unreal generated files.
+		/// Make sure the CppFilters property is set up properly and does not conflict with other exporters (including the default).
+		/// Header files should be included as normal in code.
+		/// </summary>
+		CompileOutput = 1 << 1,
 	}
 
 	/// <summary>
diff --git a/Engine/Source/Programs/UnrealBuildTool/Configuration/UEBuildModuleCPP.cs b/Engine/Source/Programs/UnrealBuildTool/Configuration/UEBuildModuleCPP.cs
--- a/Engine/Source/Programs/UnrealBuildTool/Configuration/UEBuildModuleCPP.cs
+++ b/Engine/Source/Programs/UnrealBuildTool/Configuration/UEBuildModuleCPP.cs
@@ -537,6 +537,10 @@ namespace UnrealBuildTool
 							string Key = Prefix + FileName.Substring(0, FileName.Length - ".gen.cpp".Length);
 							GeneratedFiles.Add(Key, File);
 						}
+						else if (GeneratedCPPCompileEnvironment.FileMatchesExtraGeneratedCPPTypes(FileName))
+						{
+							GeneratedFiles.Add(FileName, File);
+						}
 					}
 				}
 
@@ -761,7 +765,7 @@ namespace UnrealBuildTool
 			{
 				Unity.GetAdaptiveFiles(Target, CPPFiles, InputFiles.HeaderFiles, CompileEnvironment, WorkingSet, Rules.ShortName ?? Name, IntermediateDirectory, Graph,
 					out List<FileItem> NormalFiles, out List<FileItem> AdaptiveFiles);
-				if (!NormalFiles.Where(file => !file.HasExtension(".gen.cpp")).Any())
+				if (!NormalFiles.Where(file => !file.HasExtension(".gen.cpp") && !GeneratedCPPCompileEnvironment.FileMatchesExtraGeneratedCPPTypes(file.FullName)).Any())
 				{
 					NormalFiles = CPPFiles;
 					AdaptiveFiles.RemoveAll(new HashSet<FileItem>(NormalFiles).Contains);
diff --git a/Engine/Source/Programs/UnrealBuildTool/Configuration/UEBuildTarget.cs b/Engine/Source/Programs/UnrealBuildTool/Configuration/UEBuildTarget.cs
--- a/Engine/Source/Programs/UnrealBuildTool/Configuration/UEBuildTarget.cs
+++ b/Engine/Source/Programs/UnrealBuildTool/Configuration/UEBuildTarget.cs
@@ -14,6 +14,8 @@ using System.Text;
 using System.Text.Json.Serialization;
 using System.Threading.Tasks;
 using EpicGames.Core;
+using EpicGames.UHT.Tables;
+using EpicGames.UHT.Types;
 using EpicGames.UHT.Utils;
 using Microsoft.Extensions.Logging;
 using OpenTracing.Util;
@@ -2588,6 +2590,8 @@ namespace UnrealBuildTool
 			{
 				await ExternalExecution.ExecuteHeaderToolIfNecessaryAsync(BuildConfiguration, ProjectFile, Makefile, TargetName, WorkingSet, Logger);
 			}
+			
+			GatherExtraGeneratedCPPFileTypes(Makefile, TargetName, ref GlobalCompileEnvironment);
 
 			// Prefetch directory items for UHT folders since they are going to be used later
 			foreach (UEBuildModuleCPP Module in ModulesToGenerateHeadersFor)
@@ -6006,5 +6010,29 @@ namespace UnrealBuildTool
 			}
 			return NewPathList;
 		}
+		
+		private static void GatherExtraGeneratedCPPFileTypes(TargetMakefile Makefile, string TargetName, ref CppCompileEnvironment InCompileEnvironment)
+		{
+			UhtStdFileManager fileManager = new UhtStdFileManager();
+			FileReference moduleInfoFileName = ExternalExecution.GetUHTModuleInfoFileName(Makefile, TargetName);
+			if (fileManager.ReadSource(moduleInfoFileName.FullName, out UhtSourceFragment fragment))
+			{
+				UHTManifest? manifest = System.Text.Json.JsonSerializer.Deserialize<UHTManifest>(fragment.Data.ToString());
+
+				if (manifest != null && manifest.UhtPlugins.Count > 0)
+				{
+					UhtTables tables = new();
+					tables.AddPlugins(manifest.UhtPlugins);
+
+					foreach (UhtExporter exporter in tables.ExporterTable)
+					{
+						if (exporter.Options.HasAnyFlags(UhtExporterOptions.CompileOutput))
+						{
+							InCompileEnvironment.ExtraGeneratedCPPFileTypes.AddRange(exporter.CppFilters);
+						}
+					}
+				}
+			}
+		}
 	}
 }
diff --git a/Engine/Source/Programs/UnrealBuildTool/Platform/Windows/PVSToolChain.cs b/Engine/Source/Programs/UnrealBuildTool/Platform/Windows/PVSToolChain.cs
--- a/Engine/Source/Programs/UnrealBuildTool/Platform/Windows/PVSToolChain.cs
+++ b/Engine/Source/Programs/UnrealBuildTool/Platform/Windows/PVSToolChain.cs
@@ -778,7 +778,7 @@ namespace UnrealBuildTool
 			}
 
 			// Ignore generated files
-			if (InputFiles.All(x => x.Location.GetFileName().EndsWith(".gen.cpp")))
+			if (InputFiles.All(x => x.Location.GetFileName().EndsWith(".gen.cpp") || CompileEnvironment.FileMatchesExtraGeneratedCPPTypes(x.Location.GetFileName())))
 			{
 				return new CPPOutput();
 			}
diff --git a/Engine/Source/Programs/UnrealBuildTool/System/CppCompileEnvironment.cs b/Engine/Source/Programs/UnrealBuildTool/System/CppCompileEnvironment.cs
--- a/Engine/Source/Programs/UnrealBuildTool/System/CppCompileEnvironment.cs
+++ b/Engine/Source/Programs/UnrealBuildTool/System/CppCompileEnvironment.cs
@@ -2,6 +2,7 @@
 
 using System.Collections.Generic;
 using System.Linq;
+using System.IO.Enumeration;
 using EpicGames.Core;
 using UnrealBuildBase;
 
@@ -570,6 +571,11 @@ namespace UnrealBuildTool
 		/// A dictionary of the source file items and the inlined gen.cpp files contained in it
 		/// </summary>
 		public Dictionary<FileItem, List<FileItem>> FileInlineGenCPPMap = new();
+		
+		/// <summary>
+		/// Non-default naming conventions that generated CPP files can have (this uses a filter to check the file name for matches) 
+		/// </summary>
+		public List<string> ExtraGeneratedCPPFileTypes = new List<string>();
 
 		/// <summary>
 		/// FileItems with colliding names. (Which means they would overwrite each other in intermediate folder
@@ -802,6 +808,7 @@ namespace UnrealBuildTool
 			bCheckSystemHeadersForModification = Other.bCheckSystemHeadersForModification;
 			ForceIncludeFiles.AddRange(Other.ForceIncludeFiles);
 			AdditionalPrerequisites.AddRange(Other.AdditionalPrerequisites);
+			ExtraGeneratedCPPFileTypes = Other.ExtraGeneratedCPPFileTypes;
 			CollidingNames = Other.CollidingNames;
 			FileInlineGenCPPMap = new Dictionary<FileItem, List<FileItem>>(Other.FileInlineGenCPPMap);
 			Definitions.AddRange(Other.Definitions);
@@ -836,5 +843,18 @@ namespace UnrealBuildTool
 		{
 			return Instance?.Output.GetPrecompiledHeaderFile(Architectures.SingleArchitecture);
 		}
+		
+		public bool FileMatchesExtraGeneratedCPPTypes(string FileName)
+		{
+			foreach (string fileFilter in ExtraGeneratedCPPFileTypes)
+			{
+				if (FileSystemName.MatchesSimpleExpression(fileFilter, FileName, true))
+				{
+					return true;
+				}
+			}
+
+			return false;
+		}
 	}
 }
diff --git a/Engine/Source/Programs/UnrealBuildTool/System/Unity.cs b/Engine/Source/Programs/UnrealBuildTool/System/Unity.cs
--- a/Engine/Source/Programs/UnrealBuildTool/System/Unity.cs
+++ b/Engine/Source/Programs/UnrealBuildTool/System/Unity.cs
@@ -200,8 +200,8 @@ namespace UnrealBuildTool
 					// declared in the header but are required to link. If they are placed later in the list, you will see
 					// compile errors because the templated function is instantiated but is defined later in the same translation unit
 					// which results in 'error C2908: explicit specialization; '*****' has already been instantiated'
-					bool bAIsGenerated = A.AbsolutePath.EndsWith(".gen.cpp");
-					bool bBIsGenerated = B.AbsolutePath.EndsWith(".gen.cpp");
+					bool bAIsGenerated = A.AbsolutePath.EndsWith(".gen.cpp") || CompileEnvironment.FileMatchesExtraGeneratedCPPTypes(A.AbsolutePath);
+					bool bBIsGenerated = B.AbsolutePath.EndsWith(".gen.cpp") || CompileEnvironment.FileMatchesExtraGeneratedCPPTypes(B.AbsolutePath);
 					if (bAIsGenerated && !bBIsGenerated)
 					{
 						return -1;
@@ -371,7 +371,7 @@ namespace UnrealBuildTool
 					UnhandledHeaderFilesInWorkingSet.Remove(HeaderFile);
 				}
 
-				bool bAdaptive = bHeaderInWorkingSet || WorkingSet.Contains(CPPFile);
+				bool bAdaptive = (bHeaderInWorkingSet || WorkingSet.Contains(CPPFile)) && !CompileEnvironment.FileMatchesExtraGeneratedCPPTypes(CPPFile.FullName);
 				List<FileItem> Files = bAdaptive ? AdaptiveFiles : NormalFiles;
 				Files.Add(CPPFile);
 			}
