C# config cache fix so that it doesn't ignore extra commandline arguments because the IEnumerable Append method doesn't mutate the container it is called on, but returns a copy with the specified elements added.

diff --git a/Engine/Source/Programs/UnrealBuildTool/System/ConfigCache.cs b/Engine/Source/Programs/UnrealBuildTool/System/ConfigCache.cs
--- a/Engine/Source/Programs/UnrealBuildTool/System/ConfigCache.cs
+++ b/Engine/Source/Programs/UnrealBuildTool/System/ConfigCache.cs
@@ -284,22 +284,28 @@ namespace UnrealBuildTool
 		/// <returns>The requested config hierarchy</returns>
 		public static ConfigHierarchy ReadHierarchy(ConfigHierarchyType Type, DirectoryReference? ProjectDir, UnrealTargetPlatform Platform, string CustomConfig = "", string[]? CustomArgs = null, DirectoryReference? HotfixDir = null)
 		{
-			CommandLineArguments CombinedArgs = new CommandLineArguments(Environment.GetCommandLineArgs());
-
+			CommandLineArguments CombinedArgs;
 			if (CustomArgs != null)
 			{
+				List<string> CommandLineArgs = new List<string>(Environment.GetCommandLineArgs());
 				foreach (string CustomArg in CustomArgs)
 				{
 					if (CustomArg.StartsWith("-", StringComparison.InvariantCultureIgnoreCase))
 					{
-						CombinedArgs.Append(CustomArg);
+						CommandLineArgs.Add(CustomArg);
 					}
 					else
 					{
-						CombinedArgs.Append("-" + CustomArg);
+						CommandLineArgs.Add("-" + CustomArg);
 					}
 				}
+				CombinedArgs = new CommandLineArguments(CommandLineArgs.ToArray());
 			}
+			else
+			{
+				CombinedArgs = new CommandLineArguments(Environment.GetCommandLineArgs());
+			}
+
 
 			return ReadHierarchy(Type, ProjectDir, Platform, CustomConfig, CombinedArgs, HotfixDir);
 		}
