PR #12106: PVS-Studio - SN-DBS integration fix

diff --git a/Engine/Source/Programs/UnrealBuildTool/Platform/Windows/PVSToolChain.cs b/Engine/Source/Programs/UnrealBuildTool/Platform/Windows/PVSToolChain.cs
--- a/Engine/Source/Programs/UnrealBuildTool/Platform/Windows/PVSToolChain.cs
+++ b/Engine/Source/Programs/UnrealBuildTool/Platform/Windows/PVSToolChain.cs
@@ -650,7 +650,12 @@ namespace UnrealBuildTool
 			}
 			yield return DirectoryItem.GetItemByDirectoryReference(Unreal.RootDirectory);
 		}
-
+  
+		public override void SetEnvironmentVariables()
+		{
+			Target.WindowsPlatform.Environment?.SetEnvironmentVariables();
+		}
+			 
 		static Version GetAnalyzerVersion(FileReference AnalyzerPath)
 		{
 			string Output = String.Empty;
@@ -931,12 +936,8 @@ namespace UnrealBuildTool
 				Arguments.Add($"--cfg \"{ConfigFileItem.AbsolutePath}\"");
 				Arguments.Add($"--i-file=\"{PreprocessedFileItem.AbsolutePath}\"");
 				Arguments.Add($"--analysis-mode {(uint)Settings.ModeFlags}");
-
-				if (LicenseFile != null)
-				{
-					Arguments.Add($"--lic-file \"{LicenseFile}\"");
-					AnalyzeAction.PrerequisiteItems.Add(FileItem.GetItemByFileReference(LicenseFile));
-				}
+				Arguments.Add($"--lic-name \"{ApplicationSettings?.UserName}\" --lic-key \"{ApplicationSettings?.SerialNumber}\"");
+	
 				AnalyzeAction.CommandArguments = String.Join(' ', Arguments);
 
 				AnalyzeAction.PrerequisiteItems.Add(ConfigFileItem);
