[UBA]
* Added -nocopy to UbaVisualizer spawned by UBT since it is already copied somewhere

diff --git a/Engine/Source/Programs/UnrealBuildTool/Executors/UnrealBuildAccelerator/UBAExecutor.cs b/Engine/Source/Programs/UnrealBuildTool/Executors/UnrealBuildAccelerator/UBAExecutor.cs
--- a/Engine/Source/Programs/UnrealBuildTool/Executors/UnrealBuildAccelerator/UBAExecutor.cs
+++ b/Engine/Source/Programs/UnrealBuildTool/Executors/UnrealBuildAccelerator/UBAExecutor.cs
@@ -484,7 +484,7 @@ namespace UnrealBuildTool
 				}
 				if (FileReference.Exists(tempPath))
 				{
-					System.Diagnostics.ProcessStartInfo psi = new(BuildHostPlatform.Current.Shell.FullName, $" /C start \"\" \"{tempPath.FullName}\" -listen")
+					System.Diagnostics.ProcessStartInfo psi = new(BuildHostPlatform.Current.Shell.FullName, $" /C start \"\" \"{tempPath.FullName}\" -listen -nocopy")
 					{
 						WorkingDirectory = System.IO.Path.GetTempPath(),
 						WindowStyle = System.Diagnostics.ProcessWindowStyle.Hidden,
