UnrealBuildTool: Fix preprocessed path not being printed for clang toolchains

#rnx

diff --git a/Engine/Source/Programs/UnrealBuildTool/ToolChain/ClangSpecificFileAction.cs b/Engine/Source/Programs/UnrealBuildTool/ToolChain/ClangSpecificFileAction.cs
--- a/Engine/Source/Programs/UnrealBuildTool/ToolChain/ClangSpecificFileAction.cs
+++ b/Engine/Source/Programs/UnrealBuildTool/ToolChain/ClangSpecificFileAction.cs
@@ -116,6 +116,12 @@ namespace UnrealBuildTool
 			FileItem ProducedItem = FileItem.GetItemByFileReference(FileReference.Combine(OutputDir, UniqueDummyName + ".n"));
 			Action.ProducedItems.Add(ProducedItem);
 
+			if (RspLines.Any(x => x.Contains(".cpp.i")))
+			{
+				FileItem PreprocessedItem = FileItem.GetItemByFileReference(FileReference.Combine(OutputDir, UniqueDummyName + ".i"));
+				Action.ProducedItems.Add(PreprocessedItem);
+			}
+
 			FileItem ResponseFile = FileItem.GetItemByPath(Action.CommandArguments.Substring(1).Trim('"'));
 			File.WriteAllLines(ResponseFile.FullName, NewRspLines);
 
