Clang Analysis: Enable unix.Malloc as there are only a few outstanding warnings

#rnx

diff --git a/Engine/Source/Programs/UnrealBuildTool/Configuration/ModuleRules.cs b/Engine/Source/Programs/UnrealBuildTool/Configuration/ModuleRules.cs
--- a/Engine/Source/Programs/UnrealBuildTool/Configuration/ModuleRules.cs
+++ b/Engine/Source/Programs/UnrealBuildTool/Configuration/ModuleRules.cs
@@ -1059,9 +1059,8 @@ namespace UnrealBuildTool
 			"cplusplus.NewDelete",
 			"cplusplus.NewDeleteLeaks",
 			"unix.cstring.NullArg",
-			"unix.Malloc",
 
-			// Needs evalulation
+			// Needs evaluation
 			"deadcode.DeadStores",
 		};
 
