[UBA]
* Added so failing to connect to horde agent outputs exception message in visualizer

diff --git a/Engine/Source/Programs/UnrealBuildTool/Executors/UnrealBuildAccelerator/UBAAgentCoordinatorHorde.cs b/Engine/Source/Programs/UnrealBuildTool/Executors/UnrealBuildAccelerator/UBAAgentCoordinatorHorde.cs
--- a/Engine/Source/Programs/UnrealBuildTool/Executors/UnrealBuildAccelerator/UBAAgentCoordinatorHorde.cs
+++ b/Engine/Source/Programs/UnrealBuildTool/Executors/UnrealBuildAccelerator/UBAAgentCoordinatorHorde.cs
@@ -817,6 +817,7 @@ namespace UnrealBuildTool
 				{
 					if (!_cancellationSource!.IsCancellationRequested)
 					{
+						hordeSession.UpdateHordeStatus(" - " + ex.Message);
 						_logger.Log(_ubaConfig.bStrict ? LogLevel.Warning : LogLevel.Information, KnownLogEvents.Systemic_Horde_Compute, ex, "Unable to get worker: {Ex}", ex.ToString());
 					}
 				}
