Add support for a new type of Restricted folder, LimitedAccess
LimitedAccess will allow Plugins and Programs to be shared externally with a limited subset of users under alternate agreements or for varying purposes without requiring bespoke, problematic, custom permissions that reach in to the NotForLicencees folders or bespoke, problematic, custom distribution processes for providing the code.

diff --git a/Engine/Config/Localization/Category.ini b/Engine/Config/Localization/Category.ini
--- a/Engine/Config/Localization/Category.ini
+++ b/Engine/Config/Localization/Category.ini
@@ -34,6 +34,7 @@ IncludePathFilters=Source/Developer/*
 IncludePathFilters=Plugins/*
 ExcludePathFilters=*/NoRedist/*
 ExcludePathFilters=*/NotForLicensees/*
+ExcludePathFilters=*/LimitedAccess/*
 InputKeys=Category
 OutputNamespaces=UObjectCategory
 OutputKeys="{MetaDataValue}"
diff --git a/Engine/Config/Localization/Editor.ini b/Engine/Config/Localization/Editor.ini
--- a/Engine/Config/Localization/Editor.ini
+++ b/Engine/Config/Localization/Editor.ini
@@ -30,6 +30,7 @@ ExcludePathFilters=Plugins/Online/*
 ExcludePathFilters=Plugins/Runtime/HTTPChunkInstaller/*
 ExcludePathFilters=*/NoRedist/*
 ExcludePathFilters=*/NotForLicensees/*
+ExcludePathFilters=*/LimitedAccess/*
 ExcludePathFilters=*/ThirdParty/*
 FileNameFilters=*.cpp
 FileNameFilters=*.h
@@ -60,6 +61,7 @@ IncludePathFilters=Source/Developer/*
 IncludePathFilters=Plugins/*
 ExcludePathFilters=*/NoRedist/*
 ExcludePathFilters=*/NotForLicensees/*
+ExcludePathFilters=*/LimitedAccess/*
 FieldTypesToInclude=Property
 ; general settings
 FieldOwnerTypesToInclude=DeveloperSettings
diff --git a/Engine/Config/Localization/Engine.ini b/Engine/Config/Localization/Engine.ini
--- a/Engine/Config/Localization/Engine.ini
+++ b/Engine/Config/Localization/Engine.ini
@@ -35,6 +35,7 @@ ExcludePathFilters=Plugins/Online/*
 ExcludePathFilters=Plugins/Runtime/HTTPChunkInstaller/*
 ExcludePathFilters=*/NoRedist/*
 ExcludePathFilters=*/NotForLicensees/*
+ExcludePathFilters=*/LimitedAccess/*
 ExcludePathFilters=*/ThirdParty/*
 FileNameFilters=*.cpp
 FileNameFilters=*.h
diff --git a/Engine/Config/Localization/Keywords.ini b/Engine/Config/Localization/Keywords.ini
--- a/Engine/Config/Localization/Keywords.ini
+++ b/Engine/Config/Localization/Keywords.ini
@@ -29,6 +29,7 @@ IncludePathFilters=Source/Developer/*
 IncludePathFilters=Plugins/*
 ExcludePathFilters=*/NoRedist/*
 ExcludePathFilters=*/NotForLicensees/*
+ExcludePathFilters=*/LimitedAccess/*
 InputKeys=Keywords
 OutputNamespaces=UObjectKeywords
 OutputKeys="{FieldPath}"
diff --git a/Engine/Config/Localization/PropertyNames.ini b/Engine/Config/Localization/PropertyNames.ini
--- a/Engine/Config/Localization/PropertyNames.ini
+++ b/Engine/Config/Localization/PropertyNames.ini
@@ -34,6 +34,7 @@ IncludePathFilters=Source/Developer/*
 IncludePathFilters=Plugins/*
 ExcludePathFilters=*/NoRedist/*
 ExcludePathFilters=*/NotForLicensees/*
+ExcludePathFilters=*/LimitedAccess/*
 InputKeys=DisplayName
 OutputNamespaces=UObjectDisplayNames
 OutputKeys="{FieldPath}"
diff --git a/Engine/Config/Localization/ToolTips.ini b/Engine/Config/Localization/ToolTips.ini
--- a/Engine/Config/Localization/ToolTips.ini
+++ b/Engine/Config/Localization/ToolTips.ini
@@ -34,6 +34,7 @@ IncludePathFilters=Source/Developer/*
 IncludePathFilters=Plugins/*
 ExcludePathFilters=*/NoRedist/*
 ExcludePathFilters=*/NotForLicensees/*
+ExcludePathFilters=*/LimitedAccess/*
 InputKeys=ToolTip
 OutputNamespaces=UObjectToolTips
 OutputKeys="{FieldPath}"
diff --git a/Engine/Source/Editor/UnrealEd/Private/Commandlets/PackageUtilities.cpp b/Engine/Source/Editor/UnrealEd/Private/Commandlets/PackageUtilities.cpp
--- a/Engine/Source/Editor/UnrealEd/Private/Commandlets/PackageUtilities.cpp
+++ b/Engine/Source/Editor/UnrealEd/Private/Commandlets/PackageUtilities.cpp
@@ -220,7 +220,7 @@ bool NormalizePackageNames( TArray<FString> PackageNames, TArray<FString>& Packa
 
 			if ( (PackageFilter&NORMALIZE_ExcludeNoRedistPackages) != 0 )
 			{
-				if (Filename.Contains(TEXT("/NoRedist/")) || Filename.Contains(TEXT("/NotForLicensees/")) || Filename.Contains(TEXT("/EpicInternal/")))
+				if (Filename.Contains(TEXT("/NoRedist/")) || Filename.Contains(TEXT("/NotForLicensees/")) || Filename.Contains(TEXT("/LimitedAccess/")) || Filename.Contains(TEXT("/EpicInternal/")))
 				{
 					PackagePathNames.RemoveAt(PackageIndex);
 					continue;
diff --git a/Engine/Source/Editor/UnrealEd/Public/PackageHelperFunctions.h b/Engine/Source/Editor/UnrealEd/Public/PackageHelperFunctions.h
--- a/Engine/Source/Editor/UnrealEd/Public/PackageHelperFunctions.h
+++ b/Engine/Source/Editor/UnrealEd/Public/PackageHelperFunctions.h
@@ -37,7 +37,7 @@ enum EPackageNormalizationFlags
 	NORMALIZE_ExcludeNonDeveloperPackages  = 0x10,
 	/** do not include packages inside the Engine/Content folders in the result array; only relevant if the input array is empty */
 	NORMALIZE_ExcludeEnginePackages		= 0x20,
-	/** do not include packages inside NoRedist or NotForLicensees folders */
+	/** do not include packages inside NoRedist, NotForLicensees, or LimitedAccess folders */
 	NORMALIZE_ExcludeNoRedistPackages	= 0x40,
 	/** do not include localized packages */
 	NORMALIZE_ExcludeLocalizedPackages	= 0x80,
diff --git a/Engine/Source/Programs/AutomationTool/BuildGraph/BuildGraph.cs b/Engine/Source/Programs/AutomationTool/BuildGraph/BuildGraph.cs
--- a/Engine/Source/Programs/AutomationTool/BuildGraph/BuildGraph.cs
+++ b/Engine/Source/Programs/AutomationTool/BuildGraph/BuildGraph.cs
@@ -947,7 +947,7 @@ namespace AutomationTool
 			if (file.IsUnderDirectory(engineDirectory))
 			{
 				string[] pathFragments = file.MakeRelativeTo(engineDirectory).Split(Path.DirectorySeparatorChar, Path.AltDirectorySeparatorChar);
-				if (pathFragments.All(x => !x.Equals("NotForLicensees", StringComparison.OrdinalIgnoreCase) && !x.Equals("NoRedist", StringComparison.OrdinalIgnoreCase)))
+				if (pathFragments.All(x => !x.Equals("NotForLicensees", StringComparison.OrdinalIgnoreCase) && !x.Equals("NoRedist", StringComparison.OrdinalIgnoreCase) && !x.Equals("LimitedAccess", StringComparison.OrdinalIgnoreCase)))
 				{
 					return true;
 				}
diff --git a/Engine/Source/Programs/AutomationTool/Scripts/CopyBuildToStagingDirectory.Automation.cs b/Engine/Source/Programs/AutomationTool/Scripts/CopyBuildToStagingDirectory.Automation.cs
--- a/Engine/Source/Programs/AutomationTool/Scripts/CopyBuildToStagingDirectory.Automation.cs
+++ b/Engine/Source/Programs/AutomationTool/Scripts/CopyBuildToStagingDirectory.Automation.cs
@@ -1830,7 +1830,7 @@ namespace AutomationScripts
 				Message.Append("\nIf these files are intended to be distributed in packaged builds, move the source files out of a restricted folder, or remap them during staging using the following syntax in DefaultGame.ini:");
 				Message.Append("\n[Staging]");
 				Message.Append("\n+RemapDirectories=(From=\"Foo/NoRedist\", To=\"Foo\")");
-				if (RestrictedNames.Any(x => x != "NotForLicensees" && x != "NoRedist")) // We don't ever want internal stuff white-listing folders like this
+				if (RestrictedNames.Any(x => x != "NotForLicensees" && x != "NoRedist" && x != "LimitedAccess")) // We don't ever want internal stuff white-listing folders like this
 				{
 					Message.Append("\nAlternatively, allow list them using this syntax in DefaultGame.ini:");
 					Message.Append("\n[Staging]");
diff --git a/Engine/Source/Programs/UnrealBuildTool/Configuration/TestModuleRules.cs b/Engine/Source/Programs/UnrealBuildTool/Configuration/TestModuleRules.cs
--- a/Engine/Source/Programs/UnrealBuildTool/Configuration/TestModuleRules.cs
+++ b/Engine/Source/Programs/UnrealBuildTool/Configuration/TestModuleRules.cs
@@ -20,6 +20,7 @@ namespace UnrealBuildTool
 		private static readonly XNamespace SchemaInstance = XNamespace.Get("http://www.w3.org/2001/XMLSchema-instance");
 		private static readonly XNamespace SchemaLocation = XNamespace.Get("http://www.epicgames.com/BuildGraph ../../Build/Graph/Schema.xsd");
 		private static readonly List<string> RestrictedFoldersNonPlatform = new List<string>() {
+			RestrictedFolder.LimitedAccess.ToString(),
 			RestrictedFolder.NotForLicensees.ToString(),
 			RestrictedFolder.NoRedist.ToString(),
 			RestrictedFolder.EpicInternal.ToString(),
diff --git a/Engine/Source/Programs/UnrealBuildTool/Platform/Android/UEDeployAndroid.cs b/Engine/Source/Programs/UnrealBuildTool/Platform/Android/UEDeployAndroid.cs
--- a/Engine/Source/Programs/UnrealBuildTool/Platform/Android/UEDeployAndroid.cs
+++ b/Engine/Source/Programs/UnrealBuildTool/Platform/Android/UEDeployAndroid.cs
@@ -4293,9 +4293,11 @@ namespace UnrealBuildTool
 			string UnrealPreBuiltFilesPath = GetUnrealPreBuiltFilePath(EngineDirectory);
 			string UnrealBuildFilesPath_NFL = GetUnrealBuildFilePath(Path.Combine(EngineDirectory, "Restricted", "NotForLicensees"));
 			string UnrealBuildFilesPath_NR = GetUnrealBuildFilePath(Path.Combine(EngineDirectory, "Restricted", "NoRedist"));
+			string UnrealBuildFilesPath_LA = GetUnrealBuildFilePath(Path.Combine(EngineDirectory, "Restricted", "LimitedAccess"));
 			string GameBuildFilesPath = Path.Combine(ProjectDirectory, "Build", "Android");
 			string GameBuildFilesPath_NFL = Path.Combine(Path.Combine(ProjectDirectory, "Restricted", "NotForLicensees"), "Build", "Android");
 			string GameBuildFilesPath_NR = Path.Combine(Path.Combine(ProjectDirectory, "Restricted", "NoRedist"), "Build", "Android");
+			string GameBuildFilesPath_LA = Path.Combine(Path.Combine(ProjectDirectory, "Restricted", "LimitedAccess"), "Build", "Android");
 
 			// get a list of unique NDK architectures enabled for build
 			List<string> NDKArches = new List<string>();
@@ -4970,9 +4972,11 @@ namespace UnrealBuildTool
 				CopyFileDirectory(UnrealBuildFilesPath, UnrealBuildPath, Replacements, ExcludeFolders);
 				CopyFileDirectory(UnrealBuildFilesPath_NFL, UnrealBuildPath, Replacements, ExcludeFolders);
 				CopyFileDirectory(UnrealBuildFilesPath_NR, UnrealBuildPath, Replacements, ExcludeFolders);
+				CopyFileDirectory(UnrealBuildFilesPath_LA, UnrealBuildPath, Replacements, ExcludeFolders);
 				CopyFileDirectory(GameBuildFilesPath, UnrealBuildPath, Replacements, ExcludeFolders);
 				CopyFileDirectory(GameBuildFilesPath_NFL, UnrealBuildPath, Replacements, ExcludeFolders);
 				CopyFileDirectory(GameBuildFilesPath_NR, UnrealBuildPath, Replacements, ExcludeFolders);
+				CopyFileDirectory(GameBuildFilesPath_LA, UnrealBuildPath, Replacements, ExcludeFolders);
 
 				// Parse Gradle filters (may have been replaced by above copies)
 				ParseFilterFile(Path.Combine(UnrealBuildPath, "GradleFilter.txt"));
diff --git a/Engine/Source/Programs/UnrealBuildTool/ProjectFiles/ProjectFileGenerator.cs b/Engine/Source/Programs/UnrealBuildTool/ProjectFiles/ProjectFileGenerator.cs
--- a/Engine/Source/Programs/UnrealBuildTool/ProjectFiles/ProjectFileGenerator.cs
+++ b/Engine/Source/Programs/UnrealBuildTool/ProjectFiles/ProjectFileGenerator.cs
@@ -1760,7 +1760,8 @@ namespace UnrealBuildTool
 		private static string[] NoRedistFolders = new string[]
 		{
 			Path.DirectorySeparatorChar + "NoRedist" + Path.DirectorySeparatorChar,
-			Path.DirectorySeparatorChar + "NotForLicensees" + Path.DirectorySeparatorChar
+			Path.DirectorySeparatorChar + "NotForLicensees" + Path.DirectorySeparatorChar,
+			Path.DirectorySeparatorChar + "LimitedAccess" + Path.DirectorySeparatorChar
 		};
 
 		/// <summary>
diff --git a/Engine/Source/Programs/UnrealBuildTool/System/ConfigHierarchy.cs b/Engine/Source/Programs/UnrealBuildTool/System/ConfigHierarchy.cs
--- a/Engine/Source/Programs/UnrealBuildTool/System/ConfigHierarchy.cs
+++ b/Engine/Source/Programs/UnrealBuildTool/System/ConfigHierarchy.cs
@@ -1028,11 +1028,13 @@ namespace UnrealBuildTool
 			// Restricted Locations
 			new ConfigLayerExpansion { Before1 = "{ENGINE}/", After1 = "{ENGINE}/Restricted/NotForLicensees/", Before2 = "{PROJECT}/Config/", After2 = "{RESTRICTEDPROJECT_NFL}/Config/" },
 			new ConfigLayerExpansion { Before1 = "{ENGINE}/", After1 = "{ENGINE}/Restricted/NoRedist/",         Before2 = "{PROJECT}/Config/", After2 = "{RESTRICTEDPROJECT_NR}/Config/" },
+			new ConfigLayerExpansion { Before1 = "{ENGINE}/", After1 = "{ENGINE}/Restricted/LimitedAccess/",    Before2 = "{PROJECT}/Config/", After2 = "{RESTRICTEDPROJECT_LA}/Config/" },
 			// Platform Extensions
 			new ConfigLayerExpansion { Before1 = "{ENGINE}/Config/{PLATFORM}/", After1 = "{EXTENGINE}/Config/",    Before2 = "{PROJECT}/Config/{PLATFORM}/", After2 = "{EXTPROJECT}/Config/" },
 			// Platform Extensions in Restricted Locations
 			new ConfigLayerExpansion { Before1 = "{ENGINE}/Config/{PLATFORM}/", After1 = "{ENGINE}/Restricted/NotForLicensees/Platforms/{PLATFORM}/Config/",   Before2 = "{PROJECT}/Config/{PLATFORM}/", After2 = "{RESTRICTEDPROJECT_NFL}/Platforms/{PLATFORM}/{OPT_SUBDIR}Config/" },
 			new ConfigLayerExpansion { Before1 = "{ENGINE}/Config/{PLATFORM}/", After1 = "{ENGINE}/Restricted/NoRedist/Platforms/{PLATFORM}/Config/",          Before2 = "{PROJECT}/Config/{PLATFORM}/", After2 = "{RESTRICTEDPROJECT_NR}/Platforms/{PLATFORM}/{OPT_SUBDIR}Config/" },
+			new ConfigLayerExpansion { Before1 = "{ENGINE}/Config/{PLATFORM}/", After1 = "{ENGINE}/Restricted/LimitedAccess/Platforms/{PLATFORM}/Config/",     Before2 = "{PROJECT}/Config/{PLATFORM}/", After2 = "{RESTRICTEDPROJECT_LA}/Platforms/{PLATFORM}/{OPT_SUBDIR}Config/" },
 		};
 
 		private static string PerformBasicReplacements(string InString, string BaseIniName, string CustomConfig)
@@ -1091,6 +1093,7 @@ namespace UnrealBuildTool
 			{
 				DirectoryReference NFLDir;
 				DirectoryReference NRDir;
+				DirectoryReference LADir;
 				string OptionalSubDir = "";
 
 				if (ProjectDir.IsUnderDirectory(Unreal.EngineDirectory))
@@ -1098,11 +1101,13 @@ namespace UnrealBuildTool
 					OptionalSubDir = ProjectDir.MakeRelativeTo(Unreal.EngineDirectory) + "/";
 					NFLDir = DirectoryReference.Combine(Unreal.EngineDirectory, "Restricted/NotForLicensees");
 					NRDir = DirectoryReference.Combine(Unreal.EngineDirectory, "Restricted/NoRedist");
+					LADir = DirectoryReference.Combine(Unreal.EngineDirectory, "Restricted/LimitedAccess");
 				}
 				else
 				{
 					NFLDir = DirectoryReference.Combine(ProjectDir, "Restricted/NotForLicensees");
 					NRDir = DirectoryReference.Combine(ProjectDir, "Restricted/NoRedist");
+					LADir = DirectoryReference.Combine(ProjectDir, "Restricted/LimitedAccess");
 				}
 
 				if (ProjectDir.IsUnderDirectory(NFLDir))
@@ -1113,6 +1118,10 @@ namespace UnrealBuildTool
 				{
 					OptionalSubDir = ProjectDir.MakeRelativeTo(NRDir) + "/";
 				}
+				else if (ProjectDir.IsUnderDirectory(LADir))
+				{
+					OptionalSubDir = ProjectDir.MakeRelativeTo(LADir) + "/";
+				}
 
 				string PlatformExtensionProjectConfigDir = DirectoryReference.Combine(ProjectDir, "Platforms", PlatformName).FullName;
 
@@ -1120,6 +1129,7 @@ namespace UnrealBuildTool
 				OutString = OutString.Replace("{EXTPROJECT}", PlatformExtensionProjectConfigDir);
 				OutString = OutString.Replace("{RESTRICTEDPROJECT_NFL}", NFLDir.FullName);
 				OutString = OutString.Replace("{RESTRICTEDPROJECT_NR}", NRDir.FullName);
+				OutString = OutString.Replace("{RESTRICTEDPROJECT_LA}", LADir.FullName);
 				OutString = OutString.Replace("{OPT_SUBDIR}", OptionalSubDir);
 
 			}
diff --git a/Engine/Source/Programs/UnrealBuildTool/System/RestrictedFolders.cs b/Engine/Source/Programs/UnrealBuildTool/System/RestrictedFolders.cs
--- a/Engine/Source/Programs/UnrealBuildTool/System/RestrictedFolders.cs
+++ b/Engine/Source/Programs/UnrealBuildTool/System/RestrictedFolders.cs
@@ -189,6 +189,11 @@ namespace UnrealBuildTool
 		/// </summary>
 		public static RestrictedFolder CarefullyRedist = FindOrAddByName("CarefullyRedist");
 
+		/// <summary>
+		/// Epic Employees and Contractors as well as some explicitly permitted external groups
+		/// </summary>
+		public static RestrictedFolder LimitedAccess = FindOrAddByName("LimitedAccess");
+
 		/// <summary>
 		/// Epic Employees and Contractors
 		/// </summary>
diff --git a/Engine/Source/Runtime/Core/Private/Misc/ConfigContext.cpp b/Engine/Source/Runtime/Core/Private/Misc/ConfigContext.cpp
--- a/Engine/Source/Runtime/Core/Private/Misc/ConfigContext.cpp
+++ b/Engine/Source/Runtime/Core/Private/Misc/ConfigContext.cpp
@@ -108,11 +108,13 @@ void FConfigContext::CachePaths()
 			FPaths::MakePathRelativeTo(RelativeDir, *(EngineRootDir + TEXT("/")));
 			ProjectNotForLicenseesDir = FPaths::Combine(EngineRootDir, TEXT("Restricted/NotForLicensees"), RelativeDir);
 			ProjectNoRedistDir = FPaths::Combine(EngineRootDir, TEXT("Restricted/NoRedist"), RelativeDir);
+			ProjectLimitedAccessDir = FPaths::Combine(EngineRootDir, TEXT("Restricted/LimitedAccess"), RelativeDir);
 		}
 		else
 		{
 			ProjectNotForLicenseesDir = FPaths::Combine(ProjectRootDir, TEXT("Restricted/NotForLicensees"));
 			ProjectNoRedistDir = FPaths::Combine(ProjectRootDir, TEXT("Restricted/NoRedist"));
+			ProjectLimitedAccessDir = FPaths::Combine(ProjectRootDir, TEXT("Restricted/LimitedAccess"));
 		}
 		
 		// if we explicitly don't want project configs, then make a limited layer set without any {PROJECT} paths
@@ -640,6 +642,7 @@ FString FConfigContext::PerformFinalExpansions(const FString& InString, const FS
 	OutString = OutString.Replace(TEXT("{PROJECT}"), *ProjectRootDir);
 	OutString = OutString.Replace(TEXT("{RESTRICTEDPROJECT_NFL}"), *ProjectNotForLicenseesDir);
 	OutString = OutString.Replace(TEXT("{RESTRICTEDPROJECT_NR}"), *ProjectNoRedistDir);
+	OutString = OutString.Replace(TEXT("{RESTRICTEDPROJECT_LA}"), *ProjectLimitedAccessDir);
 
 	if (FPaths::IsUnderDirectory(ProjectRootDir, ProjectNotForLicenseesDir))
 	{
@@ -655,6 +658,13 @@ FString FConfigContext::PerformFinalExpansions(const FString& InString, const FS
 
 		OutString = OutString.Replace(TEXT("{OPT_SUBDIR}"), *(RelativeDir + TEXT("/")));
 	}
+	else if (FPaths::IsUnderDirectory(ProjectRootDir, ProjectLimitedAccessDir))
+	{
+		FString RelativeDir = ProjectRootDir;
+		FPaths::MakePathRelativeTo(RelativeDir, *(ProjectLimitedAccessDir + TEXT("/")));
+
+		OutString = OutString.Replace(TEXT("{OPT_SUBDIR}"), *(RelativeDir + TEXT("/")));
+	}
 	else if (FPaths::IsUnderDirectory(ProjectRootDir, EngineRootDir))
 	{
 		FString RelativeDir = ProjectRootDir;
diff --git a/Engine/Source/Runtime/Core/Private/Misc/Paths.cpp b/Engine/Source/Runtime/Core/Private/Misc/Paths.cpp
--- a/Engine/Source/Runtime/Core/Private/Misc/Paths.cpp
+++ b/Engine/Source/Runtime/Core/Private/Misc/Paths.cpp
@@ -375,6 +375,10 @@ FString FPaths::ConvertPath(const FString& Path, EPathConversion Method, const T
 		case EPathConversion::Engine_NoRedist:
 		case EPathConversion::Project_NoRedist:
 			return bAppendSuffix ? FPaths::Combine(Prefix, TEXT("Restricted/NoRedist"), Suffix) : FPaths::Combine(Prefix, TEXT("Restricted/NoRedist"));
+
+		case EPathConversion::Engine_LimitedAccess:
+		case EPathConversion::Project_LimitedAccess:
+			return bAppendSuffix ? FPaths::Combine(Prefix, TEXT("Restricted/LimitedAccess"), Suffix) : FPaths::Combine(Prefix, TEXT("Restricted/LimitedAccess"));
 	}
 	
 	return TEXT("");
@@ -774,6 +778,7 @@ const TArray<FString>& FPaths::GetRestrictedFolderNames()
 
 	if (!StaticData.bRestrictedFolderNamesInitialized)
 	{
+		StaticData.RestrictedFolderNames.Add(TEXT("LimitedAccess"));
 		StaticData.RestrictedFolderNames.Add(TEXT("NotForLicensees"));
 		StaticData.RestrictedFolderNames.Add(TEXT("NoRedist"));
 		StaticData.RestrictedFolderNames.Add(TEXT("CarefullyRedist"));
diff --git a/Engine/Source/Runtime/Core/Private/Modules/ModuleManager.cpp b/Engine/Source/Runtime/Core/Private/Modules/ModuleManager.cpp
--- a/Engine/Source/Runtime/Core/Private/Modules/ModuleManager.cpp
+++ b/Engine/Source/Runtime/Core/Private/Modules/ModuleManager.cpp
@@ -1446,7 +1446,7 @@ void FModuleManager::AddExtraBinarySearchPaths()
 	if (!bExtraBinarySearchPathsAdded)
 	{
 		// Ensure that dependency dlls can be found in restricted sub directories
-		TArray<FString> RestrictedFolderNames = { TEXT("NoRedist"), TEXT("NotForLicensees"), TEXT("CarefullyRedist") };
+		TArray<FString> RestrictedFolderNames = { TEXT("NoRedist"), TEXT("NotForLicensees"), TEXT("CarefullyRedist"), TEXT("LimitedAccess") };
 		for (FName PlatformName : FDataDrivenPlatformInfoRegistry::GetConfidentialPlatforms())
 		{
 			RestrictedFolderNames.Add(PlatformName.ToString());
@@ -1521,7 +1521,7 @@ void FModuleManager::AddBinariesDirectory(const TCHAR *InDirectory, bool bIsGame
 	FPlatformProcess::AddDllDirectory(InDirectory);
 
 	// Also recurse into restricted sub-folders, if they exist
-	const TCHAR* RestrictedFolderNames[] = { TEXT("NoRedist"), TEXT("NotForLicensees"), TEXT("CarefullyRedist") };
+	const TCHAR* RestrictedFolderNames[] = { TEXT("NoRedist"), TEXT("NotForLicensees"), TEXT("CarefullyRedist"), TEXT("LimitedAccess") };
 	for (const TCHAR* RestrictedFolderName : RestrictedFolderNames)
 	{
 		FString RestrictedFolder = FPaths::Combine(InDirectory, RestrictedFolderName);
diff --git a/Engine/Source/Runtime/Core/Public/Misc/ConfigContext.h b/Engine/Source/Runtime/Core/Public/Misc/ConfigContext.h
--- a/Engine/Source/Runtime/Core/Public/Misc/ConfigContext.h
+++ b/Engine/Source/Runtime/Core/Public/Misc/ConfigContext.h
@@ -215,6 +215,7 @@ public:
 	FName ConfigFileTag;
 	
 	// useful strings that are used alot when walking the hierarchy
+	FString ProjectLimitedAccessDir;
 	FString ProjectNotForLicenseesDir;
 	FString ProjectNoRedistDir;
 	TMap<FString, FPerPlatformDirs> PerPlatformDirs;
diff --git a/Engine/Source/Runtime/Core/Public/Misc/ConfigHierarchy.h b/Engine/Source/Runtime/Core/Public/Misc/ConfigHierarchy.h
--- a/Engine/Source/Runtime/Core/Public/Misc/ConfigHierarchy.h
+++ b/Engine/Source/Runtime/Core/Public/Misc/ConfigHierarchy.h
@@ -97,6 +97,11 @@ inline FConfigLayerExpansion GConfigExpansions[] =
 		TEXT("{PROJECT}/Config/"),				TEXT("{RESTRICTEDPROJECT_NR}/Config/"),
 		EConfigExpansionFlags::ForUncooked 
 	},
+	{ 
+		TEXT("{ENGINE}/"),						TEXT("{ENGINE}/Restricted/LimitedAccess/"),			
+		TEXT("{PROJECT}/Config/"),				TEXT("{RESTRICTEDPROJECT_LA}/Config/"),
+		EConfigExpansionFlags::ForUncooked | EConfigExpansionFlags::ForCooked
+	},
 
 	// Platform Extensions
 	{
@@ -127,6 +132,11 @@ inline FConfigLayerExpansion GConfigExpansions[] =
 		TEXT("{PROJECT}/Config/{PLATFORM}/"),	TEXT("{RESTRICTEDPROJECT_NR}/Platforms/{PLATFORM}/{OPT_SUBDIR}Config/"), 
 		EConfigExpansionFlags::ForUncooked // | EConfigExpansionFlags::ForPlugin
 	},
+	{
+		TEXT("{ENGINE}/Config/{PLATFORM}/"),	TEXT("{ENGINE}/Restricted/LimitedAccess/Platforms/{PLATFORM}/{OPT_SUBDIR}Config/"),
+		TEXT("{PROJECT}/Config/{PLATFORM}/"),	TEXT("{RESTRICTEDPROJECT_LA}/Platforms/{PLATFORM}/{OPT_SUBDIR}Config/"), 
+		EConfigExpansionFlags::ForUncooked | EConfigExpansionFlags::ForCooked // | EConfigExpansionFlags::ForPlugin
+	},
 };
 
 
diff --git a/Engine/Source/Runtime/Core/Public/Misc/Paths.h b/Engine/Source/Runtime/Core/Public/Misc/Paths.h
--- a/Engine/Source/Runtime/Core/Public/Misc/Paths.h
+++ b/Engine/Source/Runtime/Core/Public/Misc/Paths.h
@@ -205,11 +205,13 @@ public:
 		Engine_PlatformExtension,
 		Engine_NotForLicensees,
 		Engine_NoRedist,
+		Engine_LimitedAccess,
 
 		Project_First,
 		Project_PlatformExtension = Project_First,
 		Project_NotForLicensees,
 		Project_NoRedist,
+		Project_LimitedAccess
 	};
 	
 	/**
