New msvc analysis warnings

#rnx

diff --git a/Engine/Source/Programs/UnrealBuildTool/Platform/Windows/VCToolChain.cs b/Engine/Source/Programs/UnrealBuildTool/Platform/Windows/VCToolChain.cs
--- a/Engine/Source/Programs/UnrealBuildTool/Platform/Windows/VCToolChain.cs
+++ b/Engine/Source/Programs/UnrealBuildTool/Platform/Windows/VCToolChain.cs
@@ -412,6 +412,11 @@ namespace UnrealBuildTool
 
 				// Re-evalulate new analysis warnings at a later time
 				Arguments.Add("/wd6031"); // return value ignored: called-function could return unexpected value
+
+				if (CompileEnvironment.SystemIncludePaths.Concat(CompileEnvironment.SharedSystemIncludePaths).Any(x => x.FullName.Contains("GoogleTest", StringComparison.OrdinalIgnoreCase)))
+				{
+					Arguments.Add("/wd6326"); // Potential comparison of a constant with another constant
+				}
 			}
 
 			// Prevents the compiler from displaying its logo for each invocation.
diff --git a/Engine/Source/Runtime/CUDA/Source/Private/CudaModule.cpp b/Engine/Source/Runtime/CUDA/Source/Private/CudaModule.cpp
--- a/Engine/Source/Runtime/CUDA/Source/Private/CudaModule.cpp
+++ b/Engine/Source/Runtime/CUDA/Source/Private/CudaModule.cpp
@@ -97,11 +97,18 @@ void FCUDAModule::UnloadCuda()
 
 bool FCUDAModule::IsDeviceIndexRHISelected(int DeviceIndex)
 {
+#ifdef _MSC_VER
+#	pragma warning(push)
+#	pragma warning(disable:6398) // The address-of a field cannot be null in well-defined code
+#endif
 	if (DriverLibrary == nullptr || &DriverApiPtrs == nullptr)
 	{
 		UE_LOG(LogCUDA, Warning, TEXT("Failed to check IsRHISelectedDevice at %s:%u"), ANSI_TO_TCHAR(__FILE__), __LINE__);
 		return false;
 	}
+#ifdef _MSC_VER
+#	pragma warning(pop)
+#endif
 
 	// Get the current CUDA device.
 	CUdevice cuDevice;
diff --git a/Engine/Source/Runtime/Navmesh/Private/Detour/DetourNavLinkBuilder.cpp b/Engine/Source/Runtime/Navmesh/Private/Detour/DetourNavLinkBuilder.cpp
--- a/Engine/Source/Runtime/Navmesh/Private/Detour/DetourNavLinkBuilder.cpp
+++ b/Engine/Source/Runtime/Navmesh/Private/Detour/DetourNavLinkBuilder.cpp
@@ -207,7 +207,8 @@ namespace UE::Detour::NavLink::Private
 		case DT_LINK_ACTION_JUMP_OVER:
 			return config.jumpOverConfig.filterDistanceThreshold;
 		default:
-			dtAssert(false);
+			const bool dtval = false;
+			dtAssert(dtval);
 			return 100.f;
 		}
 	}
