Always use the clang linker on Windows with the AutoRTFM compiler to speed up link time by 6x.

diff --git a/Engine/Source/Programs/UnrealBuildTool/Platform/Windows/UEBuildWindows.cs b/Engine/Source/Programs/UnrealBuildTool/Platform/Windows/UEBuildWindows.cs
--- a/Engine/Source/Programs/UnrealBuildTool/Platform/Windows/UEBuildWindows.cs
+++ b/Engine/Source/Programs/UnrealBuildTool/Platform/Windows/UEBuildWindows.cs
@@ -1175,6 +1175,9 @@ namespace UnrealBuildTool
 					case StaticAnalyzer.Clang:
 					case StaticAnalyzer.Default:
 						Target.WindowsPlatform.Compiler = WindowsCompiler.ClangRTFM;
+
+						// We always use clang linker with AutoRTFM.
+						Target.WindowsPlatform.bAllowClangLinker = true;
 						break;
 				}
 			}
