[UBT]
* Made cache record types readonly to fix warning on farm

diff --git a/Engine/Source/Programs/Shared/EpicGames.UBA/ICacheClient.cs b/Engine/Source/Programs/Shared/EpicGames.UBA/ICacheClient.cs
--- a/Engine/Source/Programs/Shared/EpicGames.UBA/ICacheClient.cs
+++ b/Engine/Source/Programs/Shared/EpicGames.UBA/ICacheClient.cs
@@ -39,7 +39,7 @@ namespace EpicGames.UBA
 	/// </summary>
 	/// <param name="Success">Is set to true if succeeded in fetching artifacts</param>
 	/// <param name="LogLines">Contains log lines if any</param>
-	public record struct FetchFromCacheResult(bool Success, List<string> LogLines);
+	public readonly record struct FetchFromCacheResult(bool Success, List<string> LogLines);
 
 	/// <summary>
 	/// Base interface for a cache client
diff --git a/Engine/Source/Programs/UnrealBuildTool/Artifacts/Artifacts.cs b/Engine/Source/Programs/UnrealBuildTool/Artifacts/Artifacts.cs
--- a/Engine/Source/Programs/UnrealBuildTool/Artifacts/Artifacts.cs
+++ b/Engine/Source/Programs/UnrealBuildTool/Artifacts/Artifacts.cs
@@ -235,7 +235,7 @@ namespace UnrealBuildTool.Artifacts
 	/// </summary>
 	/// <param name="Success">Is set to true if succeeded in fetching artifacts</param>
 	/// <param name="LogLines">Contains log lines if any</param>
-	public record struct ActionArtifactResult(bool Success, List<string> LogLines);
+	public readonly record struct ActionArtifactResult(bool Success, List<string> LogLines);
 
 	/// <summary>
 	/// Interface for action specific support of artifacts.
