[UBT]
* Enabled uba for clang linker now when lld-link.exe is handled properly

diff --git a/Engine/Source/Programs/UnrealBuildTool/Platform/Windows/VCToolChain.cs b/Engine/Source/Programs/UnrealBuildTool/Platform/Windows/VCToolChain.cs
--- a/Engine/Source/Programs/UnrealBuildTool/Platform/Windows/VCToolChain.cs
+++ b/Engine/Source/Programs/UnrealBuildTool/Platform/Windows/VCToolChain.cs
@@ -2811,7 +2811,7 @@ namespace UnrealBuildTool
 				LinkAction.bCanExecuteRemotely = true;
 			}
 
-			if ((Target.WindowsPlatform.Compiler.IsClang() && Target.WindowsPlatform.bAllowClangLinker) || LinkEnvironment.bPGOOptimize || LinkEnvironment.bPGOProfile)
+			if (LinkEnvironment.bPGOOptimize || LinkEnvironment.bPGOProfile)
 			{
 				LinkAction.bCanExecuteInUBA = false; // Disabled for now. Should revisit to see why it is not working
 			}
