Fix code comments being the wrong way around.

#rnx

- The code documentation for these two members were switched, probably a copy paste error at somepoint.

diff --git a/Engine/Source/Programs/UnrealBuildTool/System/EncryptionAndSigning.cs b/Engine/Source/Programs/UnrealBuildTool/System/EncryptionAndSigning.cs
--- a/Engine/Source/Programs/UnrealBuildTool/System/EncryptionAndSigning.cs
+++ b/Engine/Source/Programs/UnrealBuildTool/System/EncryptionAndSigning.cs
@@ -148,12 +148,12 @@ namespace UnrealBuildTool
 			public bool bDataCryptoRequired { get; set; } = false;
 
 			/// <summary>
-			/// Config setting to enable pak signing
+			/// Config setting to enable pak encryption
 			/// </summary>
 			public bool PakEncryptionRequired { get; set; } = true;
 
 			/// <summary>
-			/// Config setting to enable pak encryption
+			/// Config setting to enable pak signing
 			/// </summary>
 			public bool PakSigningRequired { get; set; } = true;
 
