UnealBuildTool: Don't try to strip exports for a clang pch as no object is created

#rnx

diff --git a/Engine/Source/Programs/UnrealBuildTool/ToolChain/ClangToolChain.cs b/Engine/Source/Programs/UnrealBuildTool/ToolChain/ClangToolChain.cs
--- a/Engine/Source/Programs/UnrealBuildTool/ToolChain/ClangToolChain.cs
+++ b/Engine/Source/Programs/UnrealBuildTool/ToolChain/ClangToolChain.cs
@@ -1199,7 +1199,7 @@ namespace UnrealBuildTool
 			FileItem CompilerResponseFileItem = Graph.CreateIntermediateTextFile(ResponseFileName, ResponseFileContents);
 			string CommandArguments = GetResponseFileArgument(CompilerResponseFileItem);
 
-			if (bMergeModules)
+			if (bMergeModules && CompileEnvironment.PrecompiledHeaderAction != PrecompiledHeaderAction.Create)
 			{
 				// EXTRACTEXPORTS can only be interpreted by UBA.. so this action won't build outside uba
 				CommandArguments += " /EXTRACTEXPORTS";
