UnrealBuildTool: Correct object extension for merge module output based on platform

#rnx

diff --git a/Engine/Source/Programs/UnrealBuildTool/Configuration/UEBuildBinary.cs b/Engine/Source/Programs/UnrealBuildTool/Configuration/UEBuildBinary.cs
--- a/Engine/Source/Programs/UnrealBuildTool/Configuration/UEBuildBinary.cs
+++ b/Engine/Source/Programs/UnrealBuildTool/Configuration/UEBuildBinary.cs
@@ -878,7 +878,7 @@ namespace UnrealBuildTool
 				bStripUnusedExports = true;
 
 				string Name = OutputFilePaths.First().GetFileNameWithoutExtension();
-				FileReference extraObj = FileReference.Combine(IntermediateDirectory!, $"{Name}.extra.obj");
+				FileReference extraObj = FileReference.Combine(IntermediateDirectory!, $"{Name}.extra.{(Target.Platform.IsInGroup(UnrealPlatformGroup.Microsoft) ? "obj" : "o")}");
 				BinaryLinkEnvironment.InputFiles.Add(FileItem.GetItemByFileReference(extraObj));
 			}
 
diff --git a/Engine/Source/Programs/UnrealBuildTool/Configuration/UEBuildTarget.cs b/Engine/Source/Programs/UnrealBuildTool/Configuration/UEBuildTarget.cs
--- a/Engine/Source/Programs/UnrealBuildTool/Configuration/UEBuildTarget.cs
+++ b/Engine/Source/Programs/UnrealBuildTool/Configuration/UEBuildTarget.cs
@@ -3044,7 +3044,7 @@ namespace UnrealBuildTool
 				arguments.AddRange(toStrip.Select(path => $"/S:{path.Location.ChangeExtension(".exi")}").OrderBy(x => x));
 				arguments.AddRange(otherObjects.Select(path => $"/D:{path.Location.ChangeExtension(".exi")}").OrderBy(x => x));
 
-				FileReference extraObj = FileReference.Combine(IntermediateDirectory, $"{Name}.extra.obj");
+				FileReference extraObj = FileReference.Combine(IntermediateDirectory, $"{Name}.extra.{(CompileEnvironment.Platform.IsInGroup(UnrealPlatformGroup.Microsoft) ? "obj" : "o")}");
 				arguments.Add($"/O:{extraObj}");
 
 				MakefileBuilder.CreateIntermediateTextFile(stripRsp, arguments);
