UnrealBuildTool: Don't rewrite .props file if unchanged during project generation

#rnx

diff --git a/Engine/Source/Programs/UnrealBuildTool/ProjectFiles/ProjectFileGenerator.cs b/Engine/Source/Programs/UnrealBuildTool/ProjectFiles/ProjectFileGenerator.cs
--- a/Engine/Source/Programs/UnrealBuildTool/ProjectFiles/ProjectFileGenerator.cs
+++ b/Engine/Source/Programs/UnrealBuildTool/ProjectFiles/ProjectFileGenerator.cs
@@ -610,20 +610,16 @@ namespace UnrealBuildTool
 		/// Creates a .props file next to each project which specifies the path to the engine directory
 		/// </summary>
 		/// <param name="PropsFile">The properties file path</param>
-		void CreateProjectPropsFile(FileReference PropsFile)
-		{
-			using (FileStream Stream = FileReference.Open(PropsFile, FileMode.Create, FileAccess.Write, FileShare.Read))
-			{
-				using (StreamWriter Writer = new StreamWriter(Stream, Encoding.UTF8))
-				{
-					Writer.WriteLine("<?xml version=\"1.0\" encoding=\"utf-8\"?>");
-					Writer.WriteLine("<Project ToolsVersion=\"Current\" xmlns=\"http://schemas.microsoft.com/developer/msbuild/2003\">");
-					Writer.WriteLine("\t<PropertyGroup>");
-					Writer.WriteLine("\t\t<EngineDir Condition=\"'$(EngineDir)' == ''\">{0}</EngineDir>", Unreal.EngineDirectory);
-					Writer.WriteLine("\t</PropertyGroup>");
-					Writer.WriteLine("</Project>");
-				}
-			}
+		static void CreateProjectPropsFile(FileReference PropsFile)
+		{
+			using StringWriter Writer = new StringWriter();
+			Writer.WriteLine("<?xml version=\"1.0\" encoding=\"utf-8\"?>");
+			Writer.WriteLine("<Project ToolsVersion=\"Current\" xmlns=\"http://schemas.microsoft.com/developer/msbuild/2003\">");
+			Writer.WriteLine("\t<PropertyGroup>");
+			Writer.WriteLine("\t\t<EngineDir Condition=\"'$(EngineDir)' == ''\">{0}</EngineDir>", Unreal.EngineDirectory);
+			Writer.WriteLine("\t</PropertyGroup>");
+			Writer.WriteLine("</Project>");
+			FileReference.WriteAllTextIfDifferent(PropsFile, Writer.ToString());
 		}
 
 		/// <summary>
