UnrealBuildTool: Don't use a parallel loop for intellisense generation as it is causing issues with MSBuild where the working directory changes. Does not appear to significantly affect project generation time

#rnx

diff --git a/Engine/Source/Programs/UnrealBuildTool/ProjectFiles/ProjectFileGenerator.cs b/Engine/Source/Programs/UnrealBuildTool/ProjectFiles/ProjectFileGenerator.cs
--- a/Engine/Source/Programs/UnrealBuildTool/ProjectFiles/ProjectFileGenerator.cs
+++ b/Engine/Source/Programs/UnrealBuildTool/ProjectFiles/ProjectFileGenerator.cs
@@ -2139,10 +2139,10 @@ namespace UnrealBuildTool
 					int NumTargets = Targets.Count;
 					int NumTasks = NumTargets;
 					int TasksFinished = 0;
-					System.Threading.Tasks.Parallel.For(0, NumTargets, TargetIndex =>
+					foreach (Tuple<ProjectFile, ProjectTarget> Item in Targets)
 					{
-						ProjectFile TargetProjectFile = Targets[TargetIndex].Item1;
-						ProjectTarget CurTarget = Targets[TargetIndex].Item2;
+						ProjectFile TargetProjectFile = Item.Item1;
+						ProjectTarget CurTarget = Item.Item2;
 
 						// Ignore projects for platforms we can't build for
 						if (!CurTarget.SupportedPlatforms.Any(x => GetIntelliSensePlatforms().Contains(x)))
@@ -2166,11 +2166,8 @@ namespace UnrealBuildTool
 
 						try
 						{
-							foreach (UnrealTargetPlatform IntellisensePlatform in GetIntelliSensePlatforms())
+							foreach (UnrealTargetPlatform IntellisensePlatform in GetIntelliSensePlatforms().Where(x => CurTarget.SupportedPlatforms.Contains(x)))
 							{
-								if (!CurTarget.SupportedPlatforms.Contains(IntellisensePlatform))
-									continue;
-
 								// Get the architecture from the target platform
 								UnrealArchitectures DefaultArchitecture = UnrealArchitectureConfig.ForPlatform(IntellisensePlatform).ActiveArchitectures(CurTarget.UnrealProjectFilePath, CurTarget.Name);
 
@@ -2226,7 +2223,7 @@ namespace UnrealBuildTool
 							Interlocked.Increment(ref TasksFinished);
 							Progress.Write(TasksFinished, NumTasks);
 						}
-					});
+					}
 				}
 			}
 		}
