UnrealBuildTool: Don't use a parallel loop for intellisense generation as it is causing issues with MSBuild where the working directory changes. Does not appear to significantly affect project generation time

#rnx

diff --git a/Engine/Source/Programs/UnrealBuildTool/ProjectFiles/ProjectFileGenerator.cs b/Engine/Source/Programs/UnrealBuildTool/ProjectFiles/ProjectFileGenerator.cs
--- a/Engine/Source/Programs/UnrealBuildTool/ProjectFiles/ProjectFileGenerator.cs
+++ b/Engine/Source/Programs/UnrealBuildTool/ProjectFiles/ProjectFileGenerator.cs
@@ -2131,29 +2131,20 @@ namespace UnrealBuildTool
 		/// <return>Whether the process was successful or not</return>
 		private void GenerateIntelliSenseData(string[] Arguments, List<Tuple<ProjectFile, ProjectTarget>> Targets, ILogger Logger)
 		{
-			if (ShouldGenerateIntelliSenseData() && Targets.Count > 0)
+			if (ShouldGenerateIntelliSenseData() && Targets.Any())
 			{
 				string ProgressInfoText = RuntimePlatform.IsWindows ? "Binding IntelliSense data..." : "Generating data for project indexing...";
 				using (ProgressWriter Progress = new ProgressWriter(ProgressInfoText, true, Logger))
 				{
-					int NumTargets = Targets.Count;
+					// Ignore projects for platforms we can't build for
+					List<Tuple<ProjectFile, ProjectTarget>> SupportedTargets = Targets.Where(x => x.Item2.SupportedPlatforms.Any(x => GetIntelliSensePlatforms().Contains(x))).ToList();
+					int NumTargets = SupportedTargets.Count;
 					int NumTasks = NumTargets;
 					int TasksFinished = 0;
-					System.Threading.Tasks.Parallel.For(0, NumTargets, TargetIndex =>
+					foreach (Tuple<ProjectFile, ProjectTarget> Item in SupportedTargets)
 					{
-						ProjectFile TargetProjectFile = Targets[TargetIndex].Item1;
-						ProjectTarget CurTarget = Targets[TargetIndex].Item2;
-
-						// Ignore projects for platforms we can't build for
-						if (!CurTarget.SupportedPlatforms.Any(x => GetIntelliSensePlatforms().Contains(x)))
-						{
-							lock (Progress)
-							{
-								Interlocked.Increment(ref TasksFinished);
-								Progress.Write(TasksFinished, NumTasks);
-							}
-							return;
-						}
+						ProjectFile TargetProjectFile = Item.Item1;
+						ProjectTarget CurTarget = Item.Item2;
 
 						Logger.LogDebug("Found target: {Target}", CurTarget.Name);
 
@@ -2166,11 +2157,8 @@ namespace UnrealBuildTool
 
 						try
 						{
-							foreach (UnrealTargetPlatform IntellisensePlatform in GetIntelliSensePlatforms())
+							foreach (UnrealTargetPlatform IntellisensePlatform in GetIntelliSensePlatforms().Where(x => CurTarget.SupportedPlatforms.Contains(x)))
 							{
-								if (!CurTarget.SupportedPlatforms.Contains(IntellisensePlatform))
-									continue;
-
 								// Get the architecture from the target platform
 								UnrealArchitectures DefaultArchitecture = UnrealArchitectureConfig.ForPlatform(IntellisensePlatform).ActiveArchitectures(CurTarget.UnrealProjectFilePath, CurTarget.Name);
 
@@ -2226,7 +2214,7 @@ namespace UnrealBuildTool
 							Interlocked.Increment(ref TasksFinished);
 							Progress.Write(TasksFinished, NumTasks);
 						}
-					});
+					}
 				}
 			}
 		}
