UnrealBuildTool: Clarify error when a precompiled build references a module that isn available

#rnx

diff --git a/Engine/Source/Programs/UnrealBuildTool/Configuration/UEBuildModuleCPP.cs b/Engine/Source/Programs/UnrealBuildTool/Configuration/UEBuildModuleCPP.cs
--- a/Engine/Source/Programs/UnrealBuildTool/Configuration/UEBuildModuleCPP.cs
+++ b/Engine/Source/Programs/UnrealBuildTool/Configuration/UEBuildModuleCPP.cs
@@ -371,8 +371,8 @@ namespace UnrealBuildTool
 			{
 				if (!FileReference.Exists(PrecompiledManifestLocation))
 				{
-					Logger.LogError("Missing precompiled manifest for '{Name}', '{Manifest}", Name, PrecompiledManifestLocation);
-					Logger.LogInformation("This module was most likely not flagged for being included in a precompiled build - set 'PrecompileForTargets = PrecompileTargetsType.Any;' in {Name}.Build.cs to override.", Name);
+					Logger.LogError("Missing precompiled manifest for '{Name}', '{Manifest}. This module can not be referenced in a monolithic precompiled build, remove this reference or migrate to a fully compiled source build.", Name, PrecompiledManifestLocation);
+					Logger.LogInformation("This module was most likely not flagged during a release for being included in a precompiled build - set 'PrecompileForTargets = PrecompileTargetsType.Any;' in {Name}.Build.cs to override.", Name);
 					if (Rules.Plugin != null)
 					{
 						Logger.LogInformation("As it is part of the plugin '{PluginName}', also check if its 'Type' is correct.", Rules.Plugin.Name);
