[UBT]
* Fixed so cache path is also removing log line spam

diff --git a/Engine/Source/Programs/UnrealBuildTool/Executors/UnrealBuildAccelerator/UBAExecutor.cs b/Engine/Source/Programs/UnrealBuildTool/Executors/UnrealBuildAccelerator/UBAExecutor.cs
--- a/Engine/Source/Programs/UnrealBuildTool/Executors/UnrealBuildAccelerator/UBAExecutor.cs
+++ b/Engine/Source/Programs/UnrealBuildTool/Executors/UnrealBuildAccelerator/UBAExecutor.cs
@@ -226,6 +226,11 @@ namespace UnrealBuildTool
 			public Task FlushChangesAsync(CancellationToken cancellationToken) => Task.CompletedTask;
 		}
 
+		private static void RemoveLogLineSpam(List<string> logLines)
+		{
+			logLines.RemoveAll((line) => (line.StartsWith("   Creating library ", StringComparison.OrdinalIgnoreCase) && line.EndsWith(".exp", StringComparison.OrdinalIgnoreCase)) || line.EndsWith("file(s) copied.", StringComparison.OrdinalIgnoreCase));
+		}
+
 		class UBAActionArtifactCache : IActionArtifactCache
 		{
 			readonly UBAExecutor _executor;
@@ -246,6 +251,7 @@ namespace UnrealBuildTool
 					using (IRootPaths rootPaths = _executor.GetActionRootPaths(action))
 					{
 						FetchFromCacheResult result = _executor._cacheClient!.FetchFromCache(rootPaths, bucket, startInfo);
+						RemoveLogLineSpam(result.LogLines);
 						return new ActionArtifactResult(result.Success, result.LogLines);
 					}
 				}, cancellationToken, TaskCreationOptions.LongRunning | TaskCreationOptions.PreferFairness, TaskScheduler.Default);
@@ -828,8 +834,7 @@ namespace UnrealBuildTool
 					TimeSpan processorTime = process.TotalProcessorTime;
 					TimeSpan executionTime = process.TotalWallTime;
 					List<string> logLines = process.LogLines;
-					logLines.RemoveAll((line) => (line.StartsWith("   Creating library ", StringComparison.OrdinalIgnoreCase) && line.EndsWith(".exp", StringComparison.OrdinalIgnoreCase)) || line.EndsWith("file(s) copied.", StringComparison.OrdinalIgnoreCase));
-
+					RemoveLogLineSpam(logLines);
 					string? additionalDescription = !enableDetour ? "(UBA disabled)" : null;
 					ActionFinished(queue, new ExecuteResults(logLines, process.ExitCode, executionTime, processorTime, additionalDescription), action, pchItem, process);
 				}
