UnrealBuildTool: Clarify message when a build target is unsupported for a specific platform

#rnx

diff --git a/Engine/Source/Programs/UnrealBuildTool/Configuration/UEBuildTarget.cs b/Engine/Source/Programs/UnrealBuildTool/Configuration/UEBuildTarget.cs
--- a/Engine/Source/Programs/UnrealBuildTool/Configuration/UEBuildTarget.cs
+++ b/Engine/Source/Programs/UnrealBuildTool/Configuration/UEBuildTarget.cs
@@ -1296,9 +1296,13 @@ namespace UnrealBuildTool
 				{
 					throw new BuildException("{0} targets are not currently supported from this engine distribution.", RulesObject.Type);
 				}
+				else if (!InstalledPlatformInfo.IsValid(null, RulesObject.Platform, null, EProjectType.Code, InstalledPlatformState.Supported))
+				{
+					throw new BuildException("The {0} platform is not currently supported from this engine distribution.", RulesObject.Platform);
+				}
 				else if (!InstalledPlatformInfo.IsValid(RulesObject.Type, RulesObject.Platform, null, EProjectType.Code, InstalledPlatformState.Supported))
 				{
-					throw new BuildException("The {0} platform is not supported from this engine distribution.", RulesObject.Platform);
+					throw new BuildException("{0} targets for the {1} platform are not currently supported from this engine distribution.", RulesObject.Type, RulesObject.Platform);
 				}
 				else
 				{
