UnrealBuildTool: LTO is allowed for modular linux builds as lld is always used to link

#rnx

diff --git a/Engine/Source/Programs/UnrealBuildTool/Platform/Linux/UEBuildLinux.cs b/Engine/Source/Programs/UnrealBuildTool/Platform/Linux/UEBuildLinux.cs
--- a/Engine/Source/Programs/UnrealBuildTool/Platform/Linux/UEBuildLinux.cs
+++ b/Engine/Source/Programs/UnrealBuildTool/Platform/Linux/UEBuildLinux.cs
@@ -265,11 +265,6 @@ namespace UnrealBuildTool
 				}
 			}
 
-			if (Target.bAllowLTCG && Target.LinkType != TargetLinkType.Monolithic)
-			{
-				throw new BuildException("LTO (LTCG) for modular builds is not supported (lld is not currently used for dynamic libraries).");
-			}
-
 			if (Target.GlobalDefinitions.Contains("USE_NULL_RHI=1"))
 			{
 				Target.bCompileCEF3 = false;
