UnrealBuildTool: Enable fat clang objects when module merging for clang and lto could be enabled

#rnx

diff --git a/Engine/Source/Programs/UnrealBuildTool/ToolChain/ClangToolChain.cs b/Engine/Source/Programs/UnrealBuildTool/ToolChain/ClangToolChain.cs
--- a/Engine/Source/Programs/UnrealBuildTool/ToolChain/ClangToolChain.cs
+++ b/Engine/Source/Programs/UnrealBuildTool/ToolChain/ClangToolChain.cs
@@ -665,6 +665,11 @@ namespace UnrealBuildTool
 				Arguments.Add("-fprofile-generate");
 			}
 
+			if (bMergeModules && (Options.HasFlag(ClangToolChainOptions.EnableLinkTimeOptimization) || CompileEnvironment.bAllowLTCG || CompileEnvironment.bPGOOptimize || CompileEnvironment.bPGOProfile))
+			{
+				Arguments.Add("-ffat-lto-objects");
+			}
+
 			if (!CompileEnvironment.bUseInlining)
 			{
 				Arguments.Add("-fno-inline-functions");
