Xcode 16 compatibility

diff --git a/Engine/Config/Apple/Apple_SDK.json b/Engine/Config/Apple/Apple_SDK.json
--- a/Engine/Config/Apple/Apple_SDK.json
+++ b/Engine/Config/Apple/Apple_SDK.json
@@ -2,7 +2,7 @@
 	// Xcode versions
 	"MainVersion": "14.1",
 	"MinVersion": "14.1.0",
-	"MaxVersion": "15.9.0",
+	"MaxVersion": "16.9.0",
 
 	// The versions on Windows are iTunes versions
 	"MinVersion_Win64": "1100.0.0.0",
diff --git a/Engine/Source/Programs/UnrealBuildTool/ToolChain/AppleToolChain.cs b/Engine/Source/Programs/UnrealBuildTool/ToolChain/AppleToolChain.cs
--- a/Engine/Source/Programs/UnrealBuildTool/ToolChain/AppleToolChain.cs
+++ b/Engine/Source/Programs/UnrealBuildTool/ToolChain/AppleToolChain.cs
@@ -488,6 +488,14 @@ namespace UnrealBuildTool
 			Arguments.Add("-Wno-unknown-warning-option");
 			Arguments.Add("-Wno-range-loop-analysis");
 			Arguments.Add("-Wno-single-bit-bitfield-constant-conversion");
+
+			// Disable warnings for Xcode 16 for now
+			if (Info.ClangVersion.CompareTo(new Version(16, 0, 0)) >= 0)
+			{
+            	Arguments.Add("-Wno-shadow");
+				Arguments.Add("-Wno-invalid-unevaluated-string");
+				Arguments.Add("-Wno-deprecated-this-capture");
+			}
 		}
 
 		/// <inheritdoc/>
@@ -805,23 +813,12 @@ namespace UnrealBuildTool
 
 		protected virtual void GetLinkArguments_Global(LinkEnvironment LinkEnvironment, List<string> Arguments)
 		{
-			// Temp solution for UE-191350
-
-			// This Regex will extract a 2-4 segment version code from string containing a 2-5 segment code 
-			// ie: if given string: "Apple clang version 14.0.0 (clang-1400.0.17.3.1)"
-			//     it'll extract: "1400.0.17.3"  (note the dropped 5th segment)
-			Match M = Regex.Match(Info.ClangVersionString, @"(\(clang-(?<ver>\d+.\d+(.(\d+))?(.(\d+))?)(.(\d+))?\))");
-			if (M.Success)
-			{
-				string LibString = M.Groups["ver"].ToString();
-				Version? LibVersion = new Version(LibString);
-				// The Apple's new linker in Xcode 15 beta 5 (clang version 1500.0.38.1) has issues with some templated classes and dynamic linking.
-				// Fall back to using the classic.
-				if (LibVersion.CompareTo(new Version(1500, 0, 38, 1)) >= 0)
-				{
-					Arguments.Add(" -ld_classic");
-				}
-			}
+			// The Apple's new linker in Xcode 15 beta 5 (clang version 1500.0.38.1) has issues with some templated classes and dynamic linking.
+			// Fall back to using the classic.
+			if (Info.ClangVersion.CompareTo(new Version(15, 0, 38)) >= 0 && Info.ClangVersion.CompareTo(new Version(16, 0, 0)) < 0)
+            {
+                Arguments.Add(" -ld_classic");
+            }
 		}
 
 		#region Stub Xcode Projects
