UnrealBuildTool: Support adding additional debug visualizer files to be linked when a module is referenced with ModuleRules.PublicDebugVisualizerPaths

diff --git a/Engine/Source/Programs/UnrealBuildTool/Configuration/ModuleRules.cs b/Engine/Source/Programs/UnrealBuildTool/Configuration/ModuleRules.cs
--- a/Engine/Source/Programs/UnrealBuildTool/Configuration/ModuleRules.cs
+++ b/Engine/Source/Programs/UnrealBuildTool/Configuration/ModuleRules.cs
@@ -1268,6 +1268,11 @@ namespace UnrealBuildTool
 		/// </summary>
 		public List<string> PublicAdditionalLibraries = new();
 
+		/// <summary>
+		/// List of additional debug visualizers (.natvis and .natstepfilter) exposed to other modules - typically used for External (third party) modules
+		/// </summary>
+		public List<string> PublicDebugVisualizerPaths = new();
+
 		/// <summary>
 		/// Per-architecture lists of dependencies for linking to ignore (useful when building for multiple architectures, and a lib only is needed for one architecture), it's up to the Toolchain to use this
 		/// </summary>
diff --git a/Engine/Source/Programs/UnrealBuildTool/Configuration/UEBuildModule.cs b/Engine/Source/Programs/UnrealBuildTool/Configuration/UEBuildModule.cs
--- a/Engine/Source/Programs/UnrealBuildTool/Configuration/UEBuildModule.cs
+++ b/Engine/Source/Programs/UnrealBuildTool/Configuration/UEBuildModule.cs
@@ -197,7 +197,10 @@ namespace UnrealBuildTool
 		/// </summary>
 		public bool bDependsOnVerse = false;
 
-		public List<FileItem> NatvisFiles = new();
+		/// <summary>
+		/// Set of debug visualizer paths
+		/// </summary>
+		public HashSet<FileItem> NatvisFiles = new();
 
 		/// <summary>
 		/// Constructor
@@ -299,9 +302,28 @@ namespace UnrealBuildTool
 			// get the module directories from the module
 			ModuleDirectories = Rules.GetAllModuleDirectories();
 
-			foreach (DirectoryItem Directory in ModuleDirectories.Select(x => DirectoryItem.GetItemByDirectoryReference(x)))
+			// Add any additional debug visualizers 
+			foreach (string natVisPath in HashSetFromOptionalEnumerableStringParameter(Rules.PublicDebugVisualizerPaths))
+			{
+				FileItem natVisItem = FileItem.GetItemByPath(natVisPath);
+				if (!natVisItem.HasExtension(".natvis") && !natVisItem.HasExtension(".natstepfilter"))
+				{
+					Log.TraceWarningTask(RulesFile, $"Referenced Debug Visualizer '{natVisItem}' is not a .natvis or .natstepfilter file");
+				}
+				else if (!natVisItem.Exists)
+				{
+					Log.TraceWarningTask(RulesFile, $"Referenced Debug Visualizer '{natVisItem}' does not exist");
+				}
+				else
+				{
+					NatvisFiles.Add(natVisItem);
+				}
+			}
+
+			// Add any debug visualizers found in the module directories
+			foreach (DirectoryItem directory in ModuleDirectories.Select(x => DirectoryItem.GetItemByDirectoryReference(x)))
 			{
-				NatvisFiles.AddRange(Directory.EnumerateFiles().Where(x => x.HasExtension(".natvis") || x.HasExtension(".natstepfilter")));
+				NatvisFiles.UnionWith(directory.EnumerateFiles().Where(x => x.HasExtension(".natvis") || x.HasExtension(".natstepfilter")));
 			}
 		}
 
