- Added support for converting Apple clang version to LLVM version

diff --git a/Engine/Config/Apple/Apple_SDK.json b/Engine/Config/Apple/Apple_SDK.json
--- a/Engine/Config/Apple/Apple_SDK.json
+++ b/Engine/Config/Apple/Apple_SDK.json
@@ -3,8 +3,21 @@
 	"MainVersion": "14.1",
 	"MinVersion": "14.1.0",
 	"MaxVersion": "16.9.0",
+	// !!!
+	// NOTE: If you update the MaxVersion, double check the AppleVersionToLLVMVersion array below!!!
+	// !!!
 
 	// The versions on Windows are iTunes versions
 	"MinVersion_Win64": "1100.0.0.0",
-	"MaxVersion_Win64": "8999.0"
+	"MaxVersion_Win64": "8999.0",
+
+	// This is not a version range, but a mapping of Xcode clang versions to LLVM versions, for shared version checks with other clangs
+	// Version mapping can be found at https://en.wikipedia.org/wiki/Xcode#Toolchain_versions
+	// The first version is the first version that is using the second version source LLVM
+	"AppleVersionToLLVMVersions": [
+		"14.0.0-14.0.0", // Xcode 14.0 -> LLVM 14.0
+		"14.0.3-15.0.0", // Xcode 14.3 -> LLVM 15.0
+		"15.0.0-16.0.0", // Xcode 15.0 -> LLVM 16.0
+		"16.0.0-17.0.6", // Xcode 16.0 -> LLVM 17.0.6
+	],
 }
diff --git a/Engine/Source/Programs/UnrealBuildTool/Platform/IOS/IOSToolChain.cs b/Engine/Source/Programs/UnrealBuildTool/Platform/IOS/IOSToolChain.cs
--- a/Engine/Source/Programs/UnrealBuildTool/Platform/IOS/IOSToolChain.cs
+++ b/Engine/Source/Programs/UnrealBuildTool/Platform/IOS/IOSToolChain.cs
@@ -76,7 +76,7 @@ namespace UnrealBuildTool
 		{
 			FileReference CompilerPath = FileReference.Combine(Settings.ToolchainDir, IOSCompiler);
 			FileReference ArchiverPath = FileReference.Combine(Settings.ToolchainDir, IOSArchiver);
-			return new AppleToolChainInfo(IOSToolChainSettings.XcodeDeveloperDir, CompilerPath, ArchiverPath, Logger);
+			return new AppleToolChainInfo(UnrealTargetPlatform.IOS, IOSToolChainSettings.XcodeDeveloperDir, CompilerPath, ArchiverPath, Logger);
 		}
 
 		public override string GetSDKVersion()
diff --git a/Engine/Source/Programs/UnrealBuildTool/Platform/Mac/MacToolChain.cs b/Engine/Source/Programs/UnrealBuildTool/Platform/Mac/MacToolChain.cs
--- a/Engine/Source/Programs/UnrealBuildTool/Platform/Mac/MacToolChain.cs
+++ b/Engine/Source/Programs/UnrealBuildTool/Platform/Mac/MacToolChain.cs
@@ -71,7 +71,7 @@ namespace UnrealBuildTool
 		{
 			FileReference CompilerPath = FileReference.Combine(Settings.ToolchainDir, MacCompiler);
 			FileReference ArchiverPath = FileReference.Combine(Settings.ToolchainDir, MacArchiver);
-			return new AppleToolChainInfo(MacToolChainSettings.XcodeDeveloperDir, CompilerPath, ArchiverPath, Logger);
+			return new AppleToolChainInfo(UnrealTargetPlatform.Mac, MacToolChainSettings.XcodeDeveloperDir, CompilerPath, ArchiverPath, Logger);
 		}
 
 		public static DirectoryReference FindProductDirectory(FileReference? ProjectFile, DirectoryReference BinaryDir, string? NameIfProgram)
diff --git a/Engine/Source/Programs/UnrealBuildTool/ToolChain/AppleToolChain.cs b/Engine/Source/Programs/UnrealBuildTool/ToolChain/AppleToolChain.cs
--- a/Engine/Source/Programs/UnrealBuildTool/ToolChain/AppleToolChain.cs
+++ b/Engine/Source/Programs/UnrealBuildTool/ToolChain/AppleToolChain.cs
@@ -241,14 +241,43 @@ namespace UnrealBuildTool
 	{
 		protected class AppleToolChainInfo : ClangToolChainInfo
 		{
-			public AppleToolChainInfo(DirectoryReference DeveloperDir, FileReference Clang, FileReference Archiver, ILogger Logger)
+			Tuple<Version, Version>[] AppleVersionToLLVMVersion;
+
+			public AppleToolChainInfo(UnrealTargetPlatform Platform, DirectoryReference DeveloperDir, FileReference Clang, FileReference Archiver, ILogger Logger)
 				: base(DeveloperDir, Clang, Archiver, Logger)
 			{
+				// get the mapping from Apple_SDK.json (turning "version ranges" into a mapping list)
+				UEBuildPlatformSDK SDK = UEBuildPlatformSDK.GetSDKForPlatform(Platform.ToString())!;
+				
+				AppleVersionToLLVMVersion = SDK.GetVersionNumberRangeArrayFromConfig("AppleVersionToLLVMVersions").
+					Select(x => new Tuple<Version, Version>(new Version(x.Min.ToString()), new Version(x.Max.ToString()))).ToArray();
 			}
 
 			// libtool doesn't provide version, just use clang's version string
 			/// <inheritdoc/>
 			protected override string QueryArchiverVersionString() => ClangVersionString;
+
+			protected override Version QueryClangVersion()
+			{
+				// get the clang version from the clang -v
+				Version AppleVersion = base.QueryClangVersion();
+
+					// now look up this version in mappings
+				for (int MappingIndex = AppleVersionToLLVMVersion.Length - 1; MappingIndex >= 0; MappingIndex--)
+				{
+					if (AppleVersion >= AppleVersionToLLVMVersion[MappingIndex].Item1)
+					{
+						Version LLVMVersion = AppleVersionToLLVMVersion[MappingIndex].Item2;
+						Logger.LogDebug("Converted Apple version {AppleVersion} to LLVM version {LLVMVersion}", AppleVersion, LLVMVersion);
+						return LLVMVersion;
+					}
+				}
+
+				throw new BuildException($"Failed to find mapping of Apple clang version {AppleVersion} in Apple_SDK.json");
+			}
+
+			// get the actual Apple version, not the LLVM version (only Apple platform code should use this)
+			public VersionNumber AppleClangVersion => VersionNumber.Parse(base.QueryClangVersion().ToString());
 		}
 
 		public Lazy<AppleToolChainSettings> ToolChainSettings;
@@ -488,14 +517,6 @@ namespace UnrealBuildTool
 			Arguments.Add("-Wno-unknown-warning-option");
 			Arguments.Add("-Wno-range-loop-analysis");
 			Arguments.Add("-Wno-single-bit-bitfield-constant-conversion");
-
-			// Disable warnings for Xcode 16 for now
-			if (Info.ClangVersion.CompareTo(new Version(16, 0, 0)) >= 0)
-			{
-            	Arguments.Add("-Wno-shadow");
-				Arguments.Add("-Wno-invalid-unevaluated-string");
-				Arguments.Add("-Wno-deprecated-this-capture");
-			}
 		}
 
 		/// <inheritdoc/>
@@ -813,12 +834,12 @@ namespace UnrealBuildTool
 
 		protected virtual void GetLinkArguments_Global(LinkEnvironment LinkEnvironment, List<string> Arguments)
 		{
-			// The Apple's new linker in Xcode 15 beta 5 (clang version 1500.0.38.1) has issues with some templated classes and dynamic linking.
-			// Fall back to using the classic.
-			if (Info.ClangVersion.CompareTo(new Version(15, 0, 38)) >= 0 && Info.ClangVersion.CompareTo(new Version(16, 0, 0)) < 0)
-            {
-                Arguments.Add(" -ld_classic");
-            }
+			// Temp solution for UE-191350
+			VersionNumber AppleClangVersion = ((AppleToolChainInfo)GetToolChainInfo()).AppleClangVersion;
+			if (AppleClangVersion >= new VersionNumber(15) && AppleClangVersion < new VersionNumber(16))
+			{
+				Arguments.Add(" -ld_classic");
+			}
 		}
 
 		#region Stub Xcode Projects
