UnrealBuildTool: Always write .sarif files

#rnx

diff --git a/Engine/Source/Programs/UnrealBuildTool/Platform/Windows/UEBuildWindows.cs b/Engine/Source/Programs/UnrealBuildTool/Platform/Windows/UEBuildWindows.cs
--- a/Engine/Source/Programs/UnrealBuildTool/Platform/Windows/UEBuildWindows.cs
+++ b/Engine/Source/Programs/UnrealBuildTool/Platform/Windows/UEBuildWindows.cs
@@ -442,9 +442,8 @@ namespace UnrealBuildTool
 		/// <summary>
 		/// Whether .sarif files containing errors and warnings are written alongside each .obj, if supported
 		/// </summary>
-		[XmlConfigFile(Category = "BuildConfiguration", Name = "bWriteSarif")]
-		[CommandLine("-Sarif")]
-		public bool bWriteSarif = false;
+		[Obsolete("No longer necessary, .sarif files will always be written")]
+		public bool bWriteSarif = true;
 
 		/// <summary>
 		/// Whether we should export a file containing .obj to source file mappings.
@@ -881,8 +880,6 @@ namespace UnrealBuildTool
 
 		public bool bEnableLibFuzzer => Inner.bEnableLibFuzzer;
 
-		public bool bWriteSarif => Inner.bWriteSarif;
-
 		public string? ObjSrcMapFile => Inner.ObjSrcMapFile;
 
 		public string? LinkReproDir => Inner.LinkReproDir;
diff --git a/Engine/Source/Programs/UnrealBuildTool/Platform/Windows/VCToolChain.cs b/Engine/Source/Programs/UnrealBuildTool/Platform/Windows/VCToolChain.cs
--- a/Engine/Source/Programs/UnrealBuildTool/Platform/Windows/VCToolChain.cs
+++ b/Engine/Source/Programs/UnrealBuildTool/Platform/Windows/VCToolChain.cs
@@ -1947,8 +1947,8 @@ namespace UnrealBuildTool
 					CompileAction.bShowIncludes = Target.bShowIncludes;
 				}
 
-				// Write cl errors and warnings to a file if supported
-				if (Target.WindowsPlatform.Compiler.IsMSVC() && Target.WindowsPlatform.bWriteSarif)
+				// Write cl errors and warnings to a file
+				if (Target.WindowsPlatform.Compiler.IsMSVC())
 				{
 					if (Target.StaticAnalyzer == StaticAnalyzer.Default && !CompileEnvironment.bDisableStaticAnalysis)
 					{
@@ -2680,6 +2680,7 @@ namespace UnrealBuildTool
 
 			Arguments.AddRange(InputFileNames);
 
+
 			// Add the output file to the command-line.
 			Arguments.Add($"/OUT:\"{NormalizeCommandLinePath(OutputFile)}\"");
 
