PR #11012: Makefile generator is now platform agnostic

diff --git a/Engine/Source/Programs/UnrealBuildTool/ProjectFiles/Make/MakefileGenerator.cs b/Engine/Source/Programs/UnrealBuildTool/ProjectFiles/Make/MakefileGenerator.cs
--- a/Engine/Source/Programs/UnrealBuildTool/ProjectFiles/Make/MakefileGenerator.cs
+++ b/Engine/Source/Programs/UnrealBuildTool/ProjectFiles/Make/MakefileGenerator.cs
@@ -1,6 +1,7 @@
 // Copyright Epic Games, Inc. All Rights Reserved.
 
 using System;
+using System.Collections.Generic;
 using System.Linq;
 using System.Text;
 using EpicGames.Core;
@@ -38,6 +39,12 @@ namespace UnrealBuildTool
 			return bGenerateIntelliSenseData;
 		}
 
+		protected override void ConfigureProjectFileGeneration(string[] Arguments, ref bool IncludeAllPlatforms, ILogger Logger)
+		{
+			base.ConfigureProjectFileGeneration(Arguments, ref IncludeAllPlatforms, Logger);
+			bIncludeEnginePrograms = true;
+		}
+
 		/// File extension for project files we'll be generating (e.g. ".vcxproj")
 		public override string ProjectFileExtension => ".mk";
 
@@ -49,130 +56,98 @@ namespace UnrealBuildTool
 
 		private bool WriteMakefile(ILogger Logger)
 		{
-			string GameProjectFile = "";
-			string BuildCommand = "";
-			string ProjectBuildCommand = "";
-
-			string MakeGameProjectFile = "";
-
-			string UnrealRootPath = Unreal.RootDirectory.FullName;
-
 			string EditorTarget = "UnrealEditor";
+			string ScriptExtension = (OperatingSystem.IsWindows()) ? ".bat" : ".sh";
+			string CurrentPlatform = BuildHostPlatform.Current.Platform.ToString();
+
+			IEnumerable<ProjectTarget> AllProjects = GeneratedProjectFiles.SelectMany(x => x.ProjectTargets.OfType<ProjectTarget>()).Where(x => x.TargetFilePath != null);
 
 			if (!String.IsNullOrEmpty(GameProjectName))
 			{
-				GameProjectFile = OnlyGameProject!.FullName;
-				MakeGameProjectFile = "GAMEPROJECTFILE =" + GameProjectFile + "\n";
-				ProjectBuildCommand = $"PROJECTBUILD = \"$(UNREALROOTPATH)/Engine/{Unreal.RelativeDotnetDirectory}/dotnet\" \"$(UNREALROOTPATH)/Engine/Binaries/DotNET/UnrealBuildTool/UnrealBuildTool.dll\"\n";
-				EditorTarget = GameProjectName + "Editor";
+				ProjectTarget? EditorProjectTarget = AllProjects.FirstOrDefault(x => x.UnrealProjectFilePath == OnlyGameProject && x.TargetRules?.Type == TargetType.Editor);
+				if (EditorProjectTarget != null)
+				{
+					EditorTarget = EditorProjectTarget.Name;
+				}
 			}
 
-			BuildCommand = "BUILD = bash \"$(UNREALROOTPATH)/Engine/Build/BatchFiles/Linux/Build.sh\"\n";
-
+			string BuildCommand = $"BUILD = \"$(UNREALROOTPATH)/Engine/Build/BatchFiles/RunUBT{ScriptExtension}\"\n";
 			string FileName = "Makefile"; // PrimaryProjectName + ".mk";
 			StringBuilder MakefileContent = new StringBuilder();
-			MakefileContent.Append(
-				"# Makefile generated by MakefileGenerator.cs\n" +
-				"# *DO NOT EDIT*\n\n" +
-				"UNREALROOTPATH = " + UnrealRootPath + "\n" +
-				MakeGameProjectFile + "\n" +
-				"TARGETS ="
-			);
-			string MakeProjectCmdArg = "";
-			string MakeBuildCommand = "";
-			foreach (ProjectFile Project in GeneratedProjectFiles)
+			MakefileContent.AppendLine("# Makefile generated by MakefileGenerator.cs");
+			MakefileContent.AppendLine("# *DO NOT EDIT*");
+			MakefileContent.AppendLine();
+			MakefileContent.AppendLine($"UNREALROOTPATH = {Unreal.RootDirectory.FullName}");
+			MakefileContent.AppendLine();
+
+			MakefileContent.Append("TARGETS =");
+			foreach (ProjectTarget TargetFile in AllProjects)
 			{
-				foreach (ProjectTarget TargetFile in Project.ProjectTargets.OfType<ProjectTarget>())
-				{
-					if (TargetFile.TargetFilePath == null)
-					{
-						continue;
-					}
-
-					string TargetFileName = TargetFile.TargetFilePath.GetFileNameWithoutExtension();
-					string Basename = TargetFileName.Substring(0, TargetFileName.LastIndexOf(".Target", StringComparison.InvariantCultureIgnoreCase));
+				string TargetFileName = TargetFile.TargetFilePath.GetFileNameWithoutExtension();
+				string Basename = TargetFileName.Substring(0, TargetFileName.LastIndexOf(".Target", StringComparison.InvariantCultureIgnoreCase));
 
+				foreach (UnrealTargetPlatform Platform in TargetFile.SupportedPlatforms.Concat(TargetFile.ExtraSupportedPlatforms).Distinct().OrderBy(x => x.ToString()))
+				{
 					foreach (UnrealTargetConfiguration CurConfiguration in (UnrealTargetConfiguration[])Enum.GetValues(typeof(UnrealTargetConfiguration)))
 					{
-						if (CurConfiguration != UnrealTargetConfiguration.Unknown && CurConfiguration != UnrealTargetConfiguration.Development)
+						if (CurConfiguration != UnrealTargetConfiguration.Unknown)
 						{
-							if (InstalledPlatformInfo.IsValidConfiguration(CurConfiguration, EProjectType.Code))
+							if (InstalledPlatformInfo.IsValid(TargetFile.TargetRules!.Type, Platform, CurConfiguration, EProjectType.Code, InstalledPlatformState.Supported))
 							{
 								string Confname = Enum.GetName(typeof(UnrealTargetConfiguration), CurConfiguration)!;
-								MakefileContent.Append(String.Format(" \\\n\t{0}-Linux-{1} ", Basename, Confname));
+								MakefileContent.Append(String.Format(" \\\n\t{0}-{1}-{2} ", Basename, Platform, Confname));
 							}
 						}
 					}
-					MakefileContent.Append(" \\\n\t" + Basename);
 				}
+				MakefileContent.Append(" \\\n\t" + Basename);
 			}
 			MakefileContent.Append("\\\n\tconfigure");
 
-			MakefileContent.Append("\n\n" + BuildCommand + ProjectBuildCommand + "\n" +
+			MakefileContent.Append("\n\n" + BuildCommand + "\n" +
 				"all: StandardSet\n\n" +
-				"RequiredTools: CrashReportClient-Linux-Shipping CrashReportClientEditor-Linux-Shipping ShaderCompileWorker UnrealLightmass EpicWebHelper-Linux-Shipping\n\n" +
+				$"RequiredTools: CrashReportClient-{CurrentPlatform}-Shipping CrashReportClientEditor-{CurrentPlatform}-Shipping ShaderCompileWorker UnrealLightmass EpicWebHelper-{CurrentPlatform}-Shipping\n\n" +
 				$"StandardSet: RequiredTools UnrealFrontend {EditorTarget} UnrealInsights\n\n" +
-				$"DebugSet: RequiredTools UnrealFrontend-Linux-Debug {EditorTarget}-Linux-Debug\n\n"
+				$"DebugSet: RequiredTools UnrealFrontend-{CurrentPlatform}-Debug {EditorTarget}-{CurrentPlatform}-Debug\n\n"
 			);
 
-			foreach (ProjectFile Project in GeneratedProjectFiles)
+			string MakeBuildCommand = "$(BUILD)";
+			foreach (ProjectTarget TargetFile in AllProjects)
 			{
-				foreach (ProjectTarget TargetFile in Project.ProjectTargets.OfType<ProjectTarget>())
-				{
-					if (TargetFile.TargetFilePath == null)
-					{
-						continue;
-					}
+				string TargetFileName = TargetFile.TargetFilePath.GetFileNameWithoutExtension();
+				string Basename = TargetFileName.Substring(0, TargetFileName.LastIndexOf(".Target", StringComparison.InvariantCultureIgnoreCase));
+				string MakeProjectCmdArg = "";
 
-					string TargetFileName = TargetFile.TargetFilePath.GetFileNameWithoutExtension();
-					string Basename = TargetFileName.Substring(0, TargetFileName.LastIndexOf(".Target", StringComparison.InvariantCultureIgnoreCase));
-
-					if (Basename == GameProjectName || Basename == (GameProjectName + "Editor"))
-					{
-						MakeProjectCmdArg = " -project=\"$(GAMEPROJECTFILE)\"";
-						MakeBuildCommand = "$(PROJECTBUILD)";
-					}
-					else
-					{
-						MakeBuildCommand = "$(BUILD)";
-					}
+				if (TargetFile.UnrealProjectFilePath != null)
+				{
+					MakeProjectCmdArg = $" -Project=\"{TargetFile.UnrealProjectFilePath}\"";
+				}
 
+				foreach (UnrealTargetPlatform Platform in TargetFile.SupportedPlatforms.Concat(TargetFile.ExtraSupportedPlatforms).Distinct().OrderBy(x => x.ToString()))
+				{
 					foreach (UnrealTargetConfiguration CurConfiguration in (UnrealTargetConfiguration[])Enum.GetValues(typeof(UnrealTargetConfiguration)))
 					{
-						if (Basename == GameProjectName || Basename == (GameProjectName + "Editor"))
-						{
-							MakeProjectCmdArg = " -project=\"$(GAMEPROJECTFILE)\"";
-							MakeBuildCommand = "$(PROJECTBUILD)";
-						}
-						else
+						if (CurConfiguration != UnrealTargetConfiguration.Unknown)
 						{
-							MakeBuildCommand = "$(BUILD)";
-						}
-
-						if (CurConfiguration != UnrealTargetConfiguration.Unknown && CurConfiguration != UnrealTargetConfiguration.Development)
-						{
-							if (InstalledPlatformInfo.IsValidConfiguration(CurConfiguration, EProjectType.Code))
+							if (InstalledPlatformInfo.IsValid(TargetFile.TargetRules!.Type, Platform, CurConfiguration, EProjectType.Code, InstalledPlatformState.Supported))
 							{
 								string Confname = Enum.GetName(typeof(UnrealTargetConfiguration), CurConfiguration)!;
-								MakefileContent.Append(String.Format("\n{1}-Linux-{2}:\n\t {0} {1} Linux {2} {3} $(ARGS)\n", MakeBuildCommand, Basename, Confname, MakeProjectCmdArg));
+								MakefileContent.Append(String.Format("\n{1}-{4}-{2}:\n\t {0} {1} {4} {2} {3} $(ARGS)\n", MakeBuildCommand, Basename, Confname, MakeProjectCmdArg, Platform));
 							}
 						}
 					}
-					MakefileContent.Append(String.Format("\n{1}:\n\t {0} {1} Linux Development {2} $(ARGS)\n", MakeBuildCommand, Basename, MakeProjectCmdArg));
 				}
+				MakefileContent.Append(String.Format("\n{0}: {0}-{1}-Development\n", Basename, CurrentPlatform));
 			}
 
 			MakefileContent.Append("\nconfigure:\n");
-			if (!String.IsNullOrEmpty(GameProjectName))
+			if (OnlyGameProject != null)
 			{
-				// Make sure UBT is updated.
-				MakefileContent.Append("\txbuild /property:Configuration=Development /verbosity:quiet /nologo ");
-				MakefileContent.Append("\"$(UNREALROOTPATH)/Engine/Source/Programs/UnrealBuildTool/UnrealBuildTool.csproj\"\n");
-				MakefileContent.Append("\t$(PROJECTBUILD) -projectfiles -project=\"\\\"$(GAMEPROJECTFILE)\\\"\" -game -engine \n");
+				MakefileContent.Append($"\t{MakeBuildCommand} -ProjectFiles -Project=\"{OnlyGameProject.FullName}\" -Game \n");
 			}
 			else
 			{
-				MakefileContent.Append("\tbash \"$(UNREALROOTPATH)/GenerateProjectFiles.sh\" \n");
+				MakefileContent.Append("\t{MakeBuildCommand} -ProjectFiles \n");
 			}
 
 			MakefileContent.Append("\n.PHONY: $(TARGETS)\n");
diff --git a/Engine/Source/Programs/UnrealBuildTool/UnrealBuildTool.cs b/Engine/Source/Programs/UnrealBuildTool/UnrealBuildTool.cs
--- a/Engine/Source/Programs/UnrealBuildTool/UnrealBuildTool.cs
+++ b/Engine/Source/Programs/UnrealBuildTool/UnrealBuildTool.cs
@@ -251,7 +251,7 @@ namespace UnrealBuildTool
 
 			[CommandLine("-ProjectFiles", Value = "GenerateProjectFiles", Description = "Generate project files based on IDE preference. Equivalent to -Mode=GenerateProjectFiles")]
 			[CommandLine("-ProjectFileFormat=", Value = "GenerateProjectFiles", Description = "Generate project files in specified format. May be used multiple times.")]
-			[CommandLine("-Makefile", Value = "GenerateProjectFiles", Description = "Generate Linux Makefile")]
+			[CommandLine("-Makefile", Value = "GenerateProjectFiles", Description = "Generate Makefile")]
 			[CommandLine("-CMakefile", Value = "GenerateProjectFiles", Description = "Generate project files for CMake")]
 			[CommandLine("-QMakefile", Value = "GenerateProjectFiles", Description = "Generate project files for QMake")]
 			[CommandLine("-KDevelopfile", Value = "GenerateProjectFiles", Description = "Generate project files for KDevelop")]
