UnrealBuildTool: Workaround for 14.31 compiler crash

#rnx

diff --git a/Engine/Source/Programs/UnrealBuildTool/Platform/Windows/VCToolChain.cs b/Engine/Source/Programs/UnrealBuildTool/Platform/Windows/VCToolChain.cs
--- a/Engine/Source/Programs/UnrealBuildTool/Platform/Windows/VCToolChain.cs
+++ b/Engine/Source/Programs/UnrealBuildTool/Platform/Windows/VCToolChain.cs
@@ -303,6 +303,12 @@ namespace UnrealBuildTool
 
 		protected virtual void AppendCLArguments_Global(CppCompileEnvironment CompileEnvironment, List<string> Arguments)
 		{
+			// Workaround for MSVC 14.31 compiler crash
+			if (EnvVars.Compiler.IsMSVC() && EnvVars.ToolChainVersion >= new VersionNumber(14, 31))
+			{
+				Arguments.Add("/d2ssa-cfg-question-");
+			}
+
 			if (CompileEnvironment.bVcRemoveUnreferencedComdat)
 			{
 				// Suppress generation of object code for unreferenced inline functions. Enabling this option is more standards compliant, and causes a big reduction
