UnrealBuildTool: Remove some unnecessary MSVC version checks for compilers older than min version

#rnx

diff --git a/Engine/Source/Programs/UnrealBuildTool/ProjectFiles/VisualStudio/VCProject.cs b/Engine/Source/Programs/UnrealBuildTool/ProjectFiles/VisualStudio/VCProject.cs
--- a/Engine/Source/Programs/UnrealBuildTool/ProjectFiles/VisualStudio/VCProject.cs
+++ b/Engine/Source/Programs/UnrealBuildTool/ProjectFiles/VisualStudio/VCProject.cs
@@ -583,10 +583,6 @@ namespace UnrealBuildTool
 
 					Arguments.Add("/permissive-");
 					Arguments.Add("/Zc:strictStrings-"); // Have to disable strict const char* semantics due to Windows headers not being compliant.
-					if (CompilerVersion >= new VersionNumber(14, 32) && CompilerVersion < new VersionNumber(14, 33, 31629))
-					{
-						Arguments.Add("/Zc:lambda-");
-					}
 				}
 				else
 				{
@@ -608,7 +604,7 @@ namespace UnrealBuildTool
 					Arguments.Add("/Zc:preprocessor");
 				}
 
-				if (Target.WindowsPlatform.bStrictEnumTypesConformance && CompilerVersion >= new VersionNumber(14, 34, 31931))
+				if (Target.WindowsPlatform.bStrictEnumTypesConformance)
 				{
 					Arguments.Add("/Zc:enumTypes");
 				}
