UnrealBuildTool: Add sdk property for minimum Visual Studio version

#rnx

diff --git a/Engine/Config/Windows/Windows_SDK.json b/Engine/Config/Windows/Windows_SDK.json
--- a/Engine/Config/Windows/Windows_SDK.json
+++ b/Engine/Config/Windows/Windows_SDK.json
@@ -57,5 +57,7 @@
 	],
 
 	"VisualStudio2022SuggestedComponents": [
-	]
+	],
+
+	"MinimumVisualStudio2022Version": "17.4"
 }
diff --git a/Engine/Source/Programs/UnrealBuildTool/Platform/Windows/MicrosoftPlatformSDK.VisualStudio.cs b/Engine/Source/Programs/UnrealBuildTool/Platform/Windows/MicrosoftPlatformSDK.VisualStudio.cs
--- a/Engine/Source/Programs/UnrealBuildTool/Platform/Windows/MicrosoftPlatformSDK.VisualStudio.cs
+++ b/Engine/Source/Programs/UnrealBuildTool/Platform/Windows/MicrosoftPlatformSDK.VisualStudio.cs
@@ -18,11 +18,17 @@ namespace UnrealBuildTool
 		/// to support the Linux platform.
 		/// </summary>
 		static readonly string[] s_visualStudioSuggestedLinuxComponents = SDK.GetStringArrayFromConfig("VisualStudioSuggestedLinuxComponents");
+
 		/// <summary>
 		/// Additional set of components that should be suggested to be installed for Visual Studio 2022.
 		/// </summary>
 		static readonly string[] s_visualStudio2022SuggestedComponents = SDK.GetStringArrayFromConfig("VisualStudio2022SuggestedComponents");
 
+		/// <summary>
+		/// The minimum supported Visual Studio IDE version
+		/// </summary>
+		static VersionNumber MinimumVisualStudio2022Version => SDK.GetRequiredVersionNumberFromConfig("MinimumVisualStudio2022Version");
+
 		/// <summary>
 		/// Returns the list of suggested of components that should be suggested to be installed for Visual Studio.
 		/// Used to generate a .vsconfig file which will prompt Visual Studio to ask the user to install these components.
diff --git a/Engine/Source/Programs/UnrealBuildTool/Platform/Windows/MicrosoftPlatformSDK.cs b/Engine/Source/Programs/UnrealBuildTool/Platform/Windows/MicrosoftPlatformSDK.cs
--- a/Engine/Source/Programs/UnrealBuildTool/Platform/Windows/MicrosoftPlatformSDK.cs
+++ b/Engine/Source/Programs/UnrealBuildTool/Platform/Windows/MicrosoftPlatformSDK.cs
@@ -882,7 +882,7 @@ namespace UnrealBuildTool
 						int MajorVersion = Version.GetComponent(0);
 
 						WindowsCompiler Compiler;
-						if (MajorVersion >= 17) // Treat any newer versions as 2022, until we have an explicit enum for them
+						if (Version >= MicrosoftPlatformSDK.MinimumVisualStudio2022Version) // Treat any newer versions as 2022, until we have an explicit enum for them
 						{
 							Compiler = WindowsCompiler.VisualStudio2022;
 						}
