[Backout] - CL35431761
Original CL Desc
-----------------------------------------------------------------
UnrealBuildTool: Always write .sarif files

#rnx

diff --git a/Engine/Source/Programs/UnrealBuildTool/Platform/Windows/UEBuildWindows.cs b/Engine/Source/Programs/UnrealBuildTool/Platform/Windows/UEBuildWindows.cs
--- a/Engine/Source/Programs/UnrealBuildTool/Platform/Windows/UEBuildWindows.cs
+++ b/Engine/Source/Programs/UnrealBuildTool/Platform/Windows/UEBuildWindows.cs
@@ -442,8 +442,9 @@ namespace UnrealBuildTool
 		/// <summary>
 		/// Whether .sarif files containing errors and warnings are written alongside each .obj, if supported
 		/// </summary>
-		[Obsolete("No longer necessary, .sarif files will always be written")]
-		public bool bWriteSarif = true;
+		[XmlConfigFile(Category = "BuildConfiguration", Name = "bWriteSarif")]
+		[CommandLine("-Sarif")]
+		public bool bWriteSarif = false;
 
 		/// <summary>
 		/// Whether we should export a file containing .obj to source file mappings.
@@ -880,6 +881,8 @@ namespace UnrealBuildTool
 
 		public bool bEnableLibFuzzer => Inner.bEnableLibFuzzer;
 
+		public bool bWriteSarif => Inner.bWriteSarif;
+
 		public string? ObjSrcMapFile => Inner.ObjSrcMapFile;
 
 		public string? LinkReproDir => Inner.LinkReproDir;
diff --git a/Engine/Source/Programs/UnrealBuildTool/Platform/Windows/VCToolChain.cs b/Engine/Source/Programs/UnrealBuildTool/Platform/Windows/VCToolChain.cs
--- a/Engine/Source/Programs/UnrealBuildTool/Platform/Windows/VCToolChain.cs
+++ b/Engine/Source/Programs/UnrealBuildTool/Platform/Windows/VCToolChain.cs
@@ -1953,8 +1953,8 @@ namespace UnrealBuildTool
 					CompileAction.bShowIncludes = Target.bShowIncludes;
 				}
 
-				// Write cl errors and warnings to a file
-				if (Target.WindowsPlatform.Compiler.IsMSVC())
+				// Write cl errors and warnings to a file if supported
+				if (Target.WindowsPlatform.Compiler.IsMSVC() && Target.WindowsPlatform.bWriteSarif)
 				{
 					if (Target.StaticAnalyzer == StaticAnalyzer.Default && !CompileEnvironment.bDisableStaticAnalysis)
 					{
@@ -2686,7 +2686,6 @@ namespace UnrealBuildTool
 
 			Arguments.AddRange(InputFileNames);
 
-
 			// Add the output file to the command-line.
 			Arguments.Add($"/OUT:\"{NormalizeCommandLinePath(OutputFile)}\"");
 
