UnrealBuildTool: Clean up PVS recursive action

#rnx

diff --git a/Engine/Source/Programs/UnrealBuildTool/Platform/Windows/PVSToolChain.cs b/Engine/Source/Programs/UnrealBuildTool/Platform/Windows/PVSToolChain.cs
--- a/Engine/Source/Programs/UnrealBuildTool/Platform/Windows/PVSToolChain.cs
+++ b/Engine/Source/Programs/UnrealBuildTool/Platform/Windows/PVSToolChain.cs
@@ -948,7 +948,6 @@ namespace UnrealBuildTool
 				AnalyzeAction.bCanExecuteRemotely = true;
 				AnalyzeAction.bCanExecuteRemotelyWithSNDBS = false;
 				AnalyzeAction.bCanExecuteRemotelyWithXGE = false;
-				AnalyzeAction.bCanExecuteInUBA = true;
 
 				Result.ObjectFiles.AddRange(AnalyzeAction.ProducedItems);
 			}
@@ -986,13 +985,10 @@ namespace UnrealBuildTool
 			FileItem InputFileListItem = MakefileBuilder.CreateIntermediateTextFile(OutputFile.ChangeExtension(".input"), InputFiles.Select(x => x.FullName));
 			FileItem IgnoredFileListeItem = MakefileBuilder.CreateIntermediateTextFile(OutputFile.ChangeExtension(".ignored"), SystemIncludePaths.Select(x => x.FullName));
 
-			Action AnalyzeAction = MakefileBuilder.CreateAction(ActionType.Compile);
-			AnalyzeAction.ActionType = ActionType.PostBuildStep;
+			string Arguments = $"-Input=\"{InputFileListItem.Location}\" -Output=\"{OutputFile}\" -Ignored=\"{IgnoredFileListeItem.Location}\" -PrintLevel={Target.StaticAnalyzerPVSPrintLevel} -AnalyzerVersion={AnalyzerVersion}";
+
+			Action AnalyzeAction = MakefileBuilder.CreateRecursiveAction<PVSGatherMode>(ActionType.PostBuildStep, Arguments);
 			AnalyzeAction.CommandDescription = "Process PVS-Studio Results";
-			AnalyzeAction.CommandPath = Unreal.DotnetPath;
-			AnalyzeAction.CommandArguments = $"\"{Unreal.UnrealBuildToolDllPath}\" -Mode=PVSGather -Input=\"{InputFileListItem.Location}\" -Output=\"{OutputFile}\" -Ignored=\"{IgnoredFileListeItem.Location}\"" +
-																				$" -PrintLevel={Target.StaticAnalyzerPVSPrintLevel} -AnalyzerVersion={AnalyzerVersion}";
-			AnalyzeAction.WorkingDirectory = Unreal.EngineSourceDirectory;
 			AnalyzeAction.PrerequisiteItems.Add(InputFileListItem);
 			AnalyzeAction.PrerequisiteItems.Add(IgnoredFileListeItem);
 			AnalyzeAction.PrerequisiteItems.UnionWith(Makefile.OutputItems);
@@ -1000,7 +996,6 @@ namespace UnrealBuildTool
 			AnalyzeAction.ProducedItems.Add(FileItem.GetItemByFileReference(OutputFile));
 			AnalyzeAction.ProducedItems.Add(FileItem.GetItemByPath(OutputFile.FullName + "_does_not_exist")); // Force the gather step to always execute
 			AnalyzeAction.DeleteItems.UnionWith(AnalyzeAction.ProducedItems);
-			AnalyzeAction.bCanExecuteInUBA = false;
 
 			Makefile.OutputItems.AddRange(AnalyzeAction.ProducedItems);
 		}
