UnrealBuildTool: Use IoHash instead of MD5 when checking for updated action attributes

#rnx

diff --git a/Engine/Source/Programs/Shared/EpicGames.IoHash/IoHash.cs b/Engine/Source/Programs/Shared/EpicGames.IoHash/IoHash.cs
--- a/Engine/Source/Programs/Shared/EpicGames.IoHash/IoHash.cs
+++ b/Engine/Source/Programs/Shared/EpicGames.IoHash/IoHash.cs
@@ -373,21 +373,42 @@ namespace EpicGames.Core
 	/// </summary>
 	public static class IoHashExtensions
 	{
+		/// <summary>
+		/// Read an <see cref="IoHash"/> from a binary archive
+		/// </summary>
+		/// <param name="reader">Reader to serialize data from</param>
+		/// <returns>New IoHash instance</returns>
+		public static IoHash? ReadIoHash(this BinaryArchiveReader reader)
+		{
+			byte[]? data = reader.ReadByteArray();
+			return data == null ? null : new IoHash(data);
+		}
+
 		/// <summary>
 		/// Read an <see cref="IoHash"/> from a memory reader
 		/// </summary>
-		/// <param name="reader"></param>
-		/// <returns></returns>
+		/// <param name="reader">Reader to serialize data from</param>
+		/// <returns>New IoHash instance</returns>
 		public static IoHash ReadIoHash(this IMemoryReader reader)
 		{
 			return new IoHash(reader.ReadFixedLengthBytes(IoHash.NumBytes).Span);
 		}
 
+		/// <summary>
+		/// Write an <see cref="IoHash"/> to a binary archive
+		/// </summary>
+		/// <param name="writer">The writer to output data to</param>
+		/// <param name="hash">The IoHash to write</param>
+		public static void WriteIoHash(this BinaryArchiveWriter writer, IoHash? hash)
+		{
+			writer.WriteByteArray(hash?.ToByteArray());
+		}
+
 		/// <summary>
 		/// Write an <see cref="IoHash"/> to a memory writer
 		/// </summary>
-		/// <param name="writer"></param>
-		/// <param name="hash"></param>
+		/// <param name="writer">The writer to output data to</param>
+		/// <param name="hash">The IoHash to write</param>
 		public static void WriteIoHash(this IMemoryWriter writer, IoHash hash)
 		{
 			hash.CopyTo(writer.GetSpan(IoHash.NumBytes));
diff --git a/Engine/Source/Programs/UnrealBuildTool/System/ActionHistory.cs b/Engine/Source/Programs/UnrealBuildTool/System/ActionHistory.cs
--- a/Engine/Source/Programs/UnrealBuildTool/System/ActionHistory.cs
+++ b/Engine/Source/Programs/UnrealBuildTool/System/ActionHistory.cs
@@ -6,7 +6,6 @@ using System.Collections.Generic;
 using System.Diagnostics;
 using System.IO;
 using System.Linq;
-using System.Security.Cryptography;
 using System.Text;
 using EpicGames.Core;
 using Microsoft.Extensions.Logging;
@@ -23,12 +22,7 @@ namespace UnrealBuildTool
 		/// <summary>
 		/// Version number to check
 		/// </summary>
-		const int CurrentVersion = 2;
-
-		/// <summary>
-		/// Size of each hash value
-		/// </summary>
-		const int HashLength = 16;
+		const int CurrentVersion = 3;
 
 		/// <summary>
 		/// Path to store the cache data to.
@@ -41,7 +35,7 @@ namespace UnrealBuildTool
 		/// <summary>
 		/// The Attributes used to produce files, keyed by the absolute file paths.
 		/// </summary>
-		ConcurrentDictionary<FileItem, byte[]> OutputItemToAttributeHash = new ConcurrentDictionary<FileItem, byte[]>();
+		ConcurrentDictionary<FileItem, IoHash> OutputItemToAttributeHash = new ConcurrentDictionary<FileItem, IoHash>();
 
 		/// <summary>
 		/// Whether the dependency cache is dirty and needs to be saved.
@@ -79,7 +73,7 @@ namespace UnrealBuildTool
 						return;
 					}
 
-					OutputItemToAttributeHash = new ConcurrentDictionary<FileItem, byte[]>(Reader.ReadDictionary(() => Reader.ReadFileItem()!, () => Reader.ReadFixedSizeByteArray(HashLength))!);
+					OutputItemToAttributeHash = new ConcurrentDictionary<FileItem, IoHash>(Reader.ReadDictionary(() => Reader.ReadFileItem()!, () => Reader.ReadIoHash() ?? IoHash.Zero)!);
 				}
 			}
 			catch (Exception Ex)
@@ -100,41 +94,18 @@ namespace UnrealBuildTool
 				using (BinaryArchiveWriter Writer = new BinaryArchiveWriter(Location))
 				{
 					Writer.WriteInt(CurrentVersion);
-					Writer.WriteDictionary(OutputItemToAttributeHash, Key => Writer.WriteFileItem(Key), Value => Writer.WriteFixedSizeByteArray(Value));
+					Writer.WriteDictionary(OutputItemToAttributeHash, Key => Writer.WriteFileItem(Key), Value => Writer.WriteIoHash(Value));
 				}
 				bModified = false;
 			}
 		}
 
 		/// <summary>
-		/// Computes the case-invariant hash for a string
+		/// Computes the case-invariant IoHash for a string
 		/// </summary>
 		/// <param name="Text">The text to hash</param>
-		/// <returns>Hash of the string</returns>
-		static byte[] ComputeHash(string Text)
-		{
-			string InvariantText = Text.ToUpperInvariant();
-			byte[] InvariantBytes = Encoding.Unicode.GetBytes(InvariantText);
-			return MD5.Create().ComputeHash(InvariantBytes);
-		}
-
-		/// <summary>
-		/// Compares two hashes for equality
-		/// </summary>
-		/// <param name="A">The first hash value</param>
-		/// <param name="B">The second hash value</param>
-		/// <returns>True if the hashes are equal</returns>
-		static bool CompareHashes(byte[] A, byte[] B)
-		{
-			for (int Idx = 0; Idx < HashLength; Idx++)
-			{
-				if (A[Idx] != B[Idx])
-				{
-					return false;
-				}
-			}
-			return true;
-		}
+		/// <returns>IoHash of the string</returns>
+		static IoHash ComputeHash(string Text) => IoHash.Compute(Encoding.Unicode.GetBytes(Text.ToUpperInvariant()));
 
 		/// <summary>
 		/// Gets the producing attributes for the given file
@@ -144,7 +115,7 @@ namespace UnrealBuildTool
 		/// <returns>True if Attributes have changed and is updated, false otherwise</returns>
 		public bool UpdateProducingCommandLine(FileItem File, string Attributes)
 		{
-			byte[] NewHash = ComputeHash(Attributes);
+			IoHash NewHash = ComputeHash(Attributes);
 
 			for (; ; )
 			{
@@ -156,10 +127,9 @@ namespace UnrealBuildTool
 				}
 				else
 				{
-					byte[]? OldHash;
-					if (OutputItemToAttributeHash.TryGetValue(File, out OldHash))
+					if (OutputItemToAttributeHash.TryGetValue(File, out IoHash OldHash))
 					{
-						if (CompareHashes(NewHash, OldHash))
+						if (NewHash == OldHash)
 						{
 							// hashes are the same, no update needed
 							return false;
@@ -250,7 +220,7 @@ namespace UnrealBuildTool
 		/// <summary>
 		/// Used to ensure exclusive access to the layers list
 		/// </summary>
-		object LockObject = new object();
+		readonly object LockObject = new object();
 
 		/// <summary>
 		/// Map of filename to layer
@@ -267,7 +237,7 @@ namespace UnrealBuildTool
 		}
 
 		/// <summary>
-		/// Attempt to update the producing commandline for the given file
+		/// Attempt to update the producing command line for the given file
 		/// </summary>
 		/// <param name="File">The file to update</param>
 		/// <param name="Attributes">The new attributes</param>
