Win64 Clang PGO fixes
- Store the compiler name and version alongside the PGO data when linking the PGO profile executable
- Some minor logging and comments for clarity

diff --git a/Engine/Source/Programs/UnrealBuildTool/Platform/Windows/VCToolChain.cs b/Engine/Source/Programs/UnrealBuildTool/Platform/Windows/VCToolChain.cs
--- a/Engine/Source/Programs/UnrealBuildTool/Platform/Windows/VCToolChain.cs
+++ b/Engine/Source/Programs/UnrealBuildTool/Platform/Windows/VCToolChain.cs
@@ -748,13 +748,14 @@ namespace UnrealBuildTool
 				else if (CompileEnvironment.bPGOOptimize)
 				{
 					// Use a merged profdata file.
+					Log.TraceInformationOnce($"Using PGO profile data \"{Path.Combine(CompileEnvironment.PGODirectory!, CompileEnvironment.PGOFilenamePrefix!)}.profdata\"");
 					if (Target.WindowsPlatform.Compiler.IsIntel() && Target.WindowsPlatform.bSampleBasedPGO)
 					{
-						Arguments.Add($"-fprofile-sample-use=\"{Path.Combine(CompileEnvironment.PGODirectory!, CompileEnvironment.PGOFilenamePrefix!)}\".profdata");
+						Arguments.Add($"-fprofile-sample-use=\"{Path.Combine(CompileEnvironment.PGODirectory!, CompileEnvironment.PGOFilenamePrefix!)}.profdata\"");
 					}
 					else
 					{
-						Arguments.Add($"-fprofile-use=\"{Path.Combine(CompileEnvironment.PGODirectory!, CompileEnvironment.PGOFilenamePrefix!)}\".profdata");
+						Arguments.Add($"-fprofile-use=\"{Path.Combine(CompileEnvironment.PGODirectory!, CompileEnvironment.PGOFilenamePrefix!)}.profdata\"");
 					}
 				}
 			}
@@ -2845,7 +2846,7 @@ namespace UnrealBuildTool
 			return OutputFile;
 		}
 
-		protected bool PreparePGOFiles(LinkEnvironment LinkEnvironment)
+		protected bool PreparePGOFilesMsvc(LinkEnvironment LinkEnvironment)
 		{
 			if (LinkEnvironment.bPGOOptimize && LinkEnvironment.OutputFilePath.FullName.EndsWith(".exe"))
 			{
@@ -2960,7 +2961,16 @@ namespace UnrealBuildTool
 			bool bPGOOptimize = LinkEnvironment.bPGOOptimize;
 			bool bPGOProfile = LinkEnvironment.bPGOProfile;
 
-			if (!Target.WindowsPlatform.Compiler.IsClang())
+			// Write the compiler and version used to generate the PGO profile data
+			// We may want to use this to ensure compatibility
+			if (bPGOProfile)
+			{
+				string[] CompilerVersionLines = { EnvVars.Compiler.ToString(), EnvVars.CompilerVersion.ToString() };
+				string CompilerVersionFilename = Path.Combine(LinkEnvironment.PGODirectory!, "PGOProfileCompilerInfo.txt");
+				Utils.WriteFileIfChanged(new FileReference(CompilerVersionFilename), CompilerVersionLines, Logger);
+			}
+
+			if (Target.WindowsPlatform.Compiler.IsMSVC())
 			{
 				if (bPGOOptimize || bPGOProfile)
 				{
@@ -3017,7 +3027,7 @@ namespace UnrealBuildTool
 
 				if (bPGOOptimize)
 				{
-					if (PreparePGOFiles(LinkEnvironment))
+					if (PreparePGOFilesMsvc(LinkEnvironment))
 					{
 						Arguments.Add("/USEPROFILE");
 						Log.TraceInformationOnce("Enabling using Profile Guided Optimization (PGO). Linking will take a while.");
@@ -3050,7 +3060,7 @@ namespace UnrealBuildTool
 					}
 				}
 			}
-			else
+			else // Clang
 			{
 				if (LinkEnvironment.bAllowLTCG)
 				{
