Disable ASAN global variable check in vctoolchain when using clang compiler. This can lead to false / positive with ASan mixing types of different size in its shadow space
(for instance, null character being shared for char and wchar_t but having different sizes (1 for char and 2 for wchar_t))

#rnx

diff --git a/Engine/Source/Programs/UnrealBuildTool/Platform/Windows/VCToolChain.cs b/Engine/Source/Programs/UnrealBuildTool/Platform/Windows/VCToolChain.cs
--- a/Engine/Source/Programs/UnrealBuildTool/Platform/Windows/VCToolChain.cs
+++ b/Engine/Source/Programs/UnrealBuildTool/Platform/Windows/VCToolChain.cs
@@ -614,6 +614,13 @@ namespace UnrealBuildTool
 				// Works for clang too.
 				Arguments.Add("/fsanitize=address");
 
+				if (Target.WindowsPlatform.Compiler.IsClang())
+				{
+					// Don't check global variables when address sanitizer is enabled. This can lead to false / positive with ASan mixing types of different size in its shadow space
+					// (for instance, null character being shared for char and wchar_t but having different sizes (1 for char and 2 for wchar_t))
+					Arguments.Add("-mllvm -asan-globals=0");
+				}
+
 				// Use the CRT allocator so that ASan is able to hook into it for better error
 				// detection.
 				AddDefinition(Arguments, "FORCE_ANSI_ALLOCATOR=1");
