UnrealBuildTool: Add KnownLogEvents entry for RulesAssembly issues

#rnx

diff --git a/Engine/Source/Programs/Shared/EpicGames.Core/KnownLogEvents.cs b/Engine/Source/Programs/Shared/EpicGames.Core/KnownLogEvents.cs
--- a/Engine/Source/Programs/Shared/EpicGames.Core/KnownLogEvents.cs
+++ b/Engine/Source/Programs/Shared/EpicGames.Core/KnownLogEvents.cs
@@ -40,6 +40,11 @@ namespace EpicGames.Core
 		/// </summary>
 		public static EventId Compiler_Summary { get; } = new EventId(101);
 
+		/// <summary>
+		/// TargetRules or ModuleRules error
+		/// </summary>
+		public static EventId RulesAssembly { get; } = new EventId(102);
+
 		/// <summary>
 		/// UHT compiler message
 		/// </summary>
diff --git a/Engine/Source/Programs/UnrealBuildTool/Configuration/ModuleRules.cs b/Engine/Source/Programs/UnrealBuildTool/Configuration/ModuleRules.cs
--- a/Engine/Source/Programs/UnrealBuildTool/Configuration/ModuleRules.cs
+++ b/Engine/Source/Programs/UnrealBuildTool/Configuration/ModuleRules.cs
@@ -644,7 +644,7 @@ namespace UnrealBuildTool
 			get => overridePackageType ?? PackageOverrideType.None;
 			set => overridePackageType = !IsPlugin
 					? value
-					: throw new CompilationResultException(CompilationResult.RulesError, "Module '{ModuleName}' cannot override package type because it is part of a plugin!", Name);
+					: throw new CompilationResultException(CompilationResult.RulesError, KnownLogEvents.RulesAssembly, "Module '{ModuleName}' cannot override package type because it is part of a plugin!", Name);
 		}
 
 		private PackageOverrideType? overridePackageType;
@@ -1286,7 +1286,7 @@ namespace UnrealBuildTool
 		public string GetModuleDirectory(string moduleName)
 		{
 			FileReference? moduleFileReference = RulesAssembly.GetModuleFileName(moduleName)
-				?? throw new CompilationResultException(CompilationResult.RulesError, "Could not find a module named '{ModuleName}'.", moduleName);
+				?? throw new CompilationResultException(CompilationResult.RulesError, KnownLogEvents.RulesAssembly, "Could not find a module named '{ModuleName}'.", moduleName);
 			return moduleFileReference.Directory.FullName;
 		}
 
@@ -1472,7 +1472,7 @@ namespace UnrealBuildTool
 		/// Property for the directory containing this plugin. Useful for adding paths to third party dependencies.
 		/// </summary>
 		public string PluginDirectory => Plugin == null
-					? throw new CompilationResultException(CompilationResult.RulesError, "Module '{ModuleName}' does not belong to a plugin; PluginDirectory property is invalid.", Name)
+					? throw new CompilationResultException(CompilationResult.RulesError, KnownLogEvents.RulesAssembly, "Module '{ModuleName}' does not belong to a plugin; PluginDirectory property is invalid.", Name)
 					: Plugin.Directory.FullName;
 
 		/// <summary>
diff --git a/Engine/Source/Programs/UnrealBuildTool/Configuration/TargetRules.cs b/Engine/Source/Programs/UnrealBuildTool/Configuration/TargetRules.cs
--- a/Engine/Source/Programs/UnrealBuildTool/Configuration/TargetRules.cs
+++ b/Engine/Source/Programs/UnrealBuildTool/Configuration/TargetRules.cs
@@ -2987,7 +2987,7 @@ namespace UnrealBuildTool
 
 			// Find the constructor
 			ConstructorInfo? constructor = rulesType.GetConstructor(new Type[] { typeof(TargetInfo) })
-				?? throw new CompilationResultException(CompilationResult.RulesError, "No constructor found on {TargetName} which takes an argument of type TargetInfo.", rulesType.Name);
+				?? throw new CompilationResultException(CompilationResult.RulesError, KnownLogEvents.RulesAssembly, "No constructor found on {TargetName} which takes an argument of type TargetInfo.", rulesType.Name);
 
 			// Invoke the regular constructor
 			try
@@ -2996,7 +2996,7 @@ namespace UnrealBuildTool
 			}
 			catch (Exception ex)
 			{
-				throw new CompilationResultException(CompilationResult.RulesError, ex, "Unable to instantiate instance of '{TargetName}' object type from compiled assembly '{AssemblyPath}'.  Unreal Build Tool creates an instance of your module's 'Rules' object in order to find out about your module's requirements.  The CLR exception details may provide more information: {ExceptionMessage}",
+				throw new CompilationResultException(CompilationResult.RulesError, ex, KnownLogEvents.RulesAssembly, "Unable to instantiate instance of '{TargetName}' object type from compiled assembly '{AssemblyPath}'.  Unreal Build Tool creates an instance of your module's 'Rules' object in order to find out about your module's requirements.  The CLR exception details may provide more information: {ExceptionMessage}",
 					rulesType.Name, Path.GetFileNameWithoutExtension(rulesType.Assembly?.Location) ?? "Unknown Assembly", ex.ToString());
 			}
 
diff --git a/Engine/Source/Programs/UnrealBuildTool/System/BuildException.cs b/Engine/Source/Programs/UnrealBuildTool/System/BuildException.cs
--- a/Engine/Source/Programs/UnrealBuildTool/System/BuildException.cs
+++ b/Engine/Source/Programs/UnrealBuildTool/System/BuildException.cs
@@ -131,7 +131,7 @@ namespace UnrealBuildTool
 		public override void LogException(ILogger Logger)
 		{
 			Logger.Log(Event.Level, Event.Id, Event, this, (s, e) => s.ToString());
-			Logger.LogDebug(this, "{Ex}", ExceptionUtils.FormatExceptionDetails(this));
+			Logger.LogDebug(Event.Id, this, "{Ex}", ExceptionUtils.FormatExceptionDetails(this));
 		}
 	}
 
@@ -181,6 +181,19 @@ namespace UnrealBuildTool
 			this.Result = Result;
 		}
 
+		/// <summary>
+		/// Constructor
+		/// </summary>
+		/// <param name="Result">The resulting exit code</param>
+		/// <param name="EventId">Event id for the error</param>
+		/// <param name="Format">Formatting string for the error message</param>
+		/// <param name="Arguments">Arguments for the formatting string</param>
+		public CompilationResultException(CompilationResult Result, EventId EventId, string Format, params object[] Arguments)
+			: base(LogEvent.Create(LogLevel.Error, EventId, Format, Arguments))
+		{
+			this.Result = Result;
+		}
+
 		/// <summary>
 		/// Constructor
 		/// </summary>
@@ -227,7 +240,7 @@ namespace UnrealBuildTool
 			{
 				Logger.Log(Event.Level, Event.Id, Event, this, (s, e) => s.ToString());
 			}
-			Logger.LogDebug(this, "{Ex}", ExceptionUtils.FormatExceptionDetails(this));
+			Logger.LogDebug(Event.Id, this, "{Ex}", ExceptionUtils.FormatExceptionDetails(this));
 		}
 	}
 
diff --git a/Engine/Source/Programs/UnrealBuildTool/System/RulesAssembly.cs b/Engine/Source/Programs/UnrealBuildTool/System/RulesAssembly.cs
--- a/Engine/Source/Programs/UnrealBuildTool/System/RulesAssembly.cs
+++ b/Engine/Source/Programs/UnrealBuildTool/System/RulesAssembly.cs
@@ -459,7 +459,7 @@ namespace UnrealBuildTool
 			{
 				if (Parent == null)
 				{
-					throw new CompilationResultException(CompilationResult.RulesError, "Could not find definition for module '{ModuleTypeName}', (referenced via {ReferenceChain})", ModuleTypeName, ReferenceChain);
+					throw new CompilationResultException(CompilationResult.RulesError, KnownLogEvents.RulesAssembly, "Could not find definition for module '{ModuleTypeName}', (referenced via {ReferenceChain})", ModuleTypeName, ReferenceChain);
 				}
 				else
 				{
@@ -482,7 +482,7 @@ namespace UnrealBuildTool
 					// we expect only one platform group to be found in the extensions
 					if (GroupRulesObjectType != null && PlatformRulesObjectType != null)
 					{
-						throw new CompilationResultException(CompilationResult.RulesError, "Found multiple platform group overrides ({GroupRulesName} and {PlatformRulesName}) for module {ModuleName} without a platform specific override. Create a platform override with the class hierarchy as needed.",
+						throw new CompilationResultException(CompilationResult.RulesError, KnownLogEvents.RulesAssembly, "Found multiple platform group overrides ({GroupRulesName} and {PlatformRulesName}) for module {ModuleName} without a platform specific override. Create a platform override with the class hierarchy as needed.",
 							GroupRulesObjectType.Name, PlatformRulesObjectType.Name, ModuleName);
 					}
 					// remember the platform group if we found it, but keep searching to verify there isn't more than one
@@ -503,7 +503,7 @@ namespace UnrealBuildTool
 				if (ModuleFileName.ContainsAnyNames(DisallowedPlatformsAndGroups, Unreal.EngineDirectory) ||
 					(Target.ProjectFile != null && ModuleFileName.ContainsAnyNames(DisallowedPlatformsAndGroups, Target.ProjectFile.Directory)))
 				{
-					throw new CompilationResultException(CompilationResult.RulesError, "Platform module file {ModuleFileName} is not allowed (only platforms '{Platforms}', and their groups, are allowed. This indicates a module reference not being checked with something like IsPlatformAvailableForTarget()).",
+					throw new CompilationResultException(CompilationResult.RulesError, KnownLogEvents.RulesAssembly, "Platform module file {ModuleFileName} is not allowed (only platforms '{Platforms}', and their groups, are allowed. This indicates a module reference not being checked with something like IsPlatformAvailableForTarget()).",
 						ModuleFileName, String.Join(",", Target.OptedInModulePlatforms));
 				}
 			}
@@ -512,7 +512,7 @@ namespace UnrealBuildTool
 			Type? RulesObjectType = PlatformRulesObjectType ?? BaseRulesObjectType;
 			if (RulesObjectType == null)
 			{
-				throw new CompilationResultException(CompilationResult.RulesError, "Expecting to find a type to be declared in a module rules named '{ModuleTypeName}' in '{AssemblyName}'.  This type must derive from the 'ModuleRules' type defined by UnrealBuildTool.",
+				throw new CompilationResultException(CompilationResult.RulesError, KnownLogEvents.RulesAssembly, "Expecting to find a type to be declared in a module rules named '{ModuleTypeName}' in '{AssemblyName}'.  This type must derive from the 'ModuleRules' type defined by UnrealBuildTool.",
 					ModuleTypeName, CompiledAssembly?.FullName ?? "Unknown Assembly");
 			}
 
@@ -556,7 +556,7 @@ namespace UnrealBuildTool
 						}
 						if (SubType.BaseType == null)
 						{
-							throw new CompilationResultException(CompilationResult.RulesError, "{TypeName} is not derived from {BaseTypeName}", RulesObjectType.Name, BaseRulesObjectType.Name);
+							throw new CompilationResultException(CompilationResult.RulesError, KnownLogEvents.RulesAssembly, "{TypeName} is not derived from {BaseTypeName}", RulesObjectType.Name, BaseRulesObjectType.Name);
 						}
 						SubType = SubType.BaseType;
 					}
@@ -566,7 +566,7 @@ namespace UnrealBuildTool
 				ConstructorInfo? Constructor = RulesObjectType.GetConstructor(new Type[] { typeof(ReadOnlyTargetRules) });
 				if (Constructor == null)
 				{
-					throw new CompilationResultException(CompilationResult.RulesError, "No valid constructor found for {ModuleName}.", ModuleName);
+					throw new CompilationResultException(CompilationResult.RulesError, KnownLogEvents.RulesAssembly, "No valid constructor found for {ModuleName}.", ModuleName);
 				}
 
 				// Add the parent assemblies to the assembly cache so the types in them can be used when the constructor is called
@@ -596,7 +596,7 @@ namespace UnrealBuildTool
 			catch (Exception Ex)
 			{
 				Exception MessageEx = (Ex is TargetInvocationException && Ex.InnerException != null) ? Ex.InnerException : Ex;
-				throw new CompilationResultException(CompilationResult.RulesError, Ex, "Unable to instantiate module '{ModuleName}': {ExceptionMessage}\n(referenced via {ReferenceChain})", ModuleName, MessageEx.ToString(), ReferenceChain);
+				throw new CompilationResultException(CompilationResult.RulesError, Ex, KnownLogEvents.RulesAssembly, "Unable to instantiate module '{ModuleName}': {ExceptionMessage}\n(referenced via {ReferenceChain})", ModuleName, MessageEx.ToString(), ReferenceChain);
 			}
 		}
 
@@ -633,7 +633,7 @@ namespace UnrealBuildTool
 					// we expect only one platform group to be found in the extensions
 					if (GroupRulesObjectType != null && PlatformRulesType != null)
 					{
-						throw new CompilationResultException(CompilationResult.RulesError, "Found multiple platform group overrides ({GroupRulesName} and {PlatformRulesName}) for rules {TypeName} without a platform specific override. Create a platform override with the class hierarchy as needed.",
+						throw new CompilationResultException(CompilationResult.RulesError, KnownLogEvents.RulesAssembly, "Found multiple platform group overrides ({GroupRulesName} and {PlatformRulesName}) for rules {TypeName} without a platform specific override. Create a platform override with the class hierarchy as needed.",
 							GroupRulesObjectType.Name, PlatformRulesType.Name, TypeName);
 					}
 					// remember the platform group if we found it, but keep searching to verify there isn't more than one
@@ -646,7 +646,7 @@ namespace UnrealBuildTool
 			}
 			if (PlatformRulesType != null && !PlatformRulesType.IsSubclassOf(BaseRulesType))
 			{
-				throw new CompilationResultException(CompilationResult.RulesError, "Expecting {PlatformRulesType} to be a specialization of {BaseRulesType}", PlatformRulesType, BaseRulesType);
+				throw new CompilationResultException(CompilationResult.RulesError, KnownLogEvents.RulesAssembly, "Expecting {PlatformRulesType} to be a specialization of {BaseRulesType}", PlatformRulesType, BaseRulesType);
 			}
 
 			// Create an instance of the module's rules object, and set some defaults before calling the constructor.
@@ -663,7 +663,7 @@ namespace UnrealBuildTool
 			// Set the final value for the link type in the target rules
 			if (Rules.LinkType == TargetLinkType.Default)
 			{
-				throw new CompilationResultException(CompilationResult.RulesError, "TargetRules.LinkType should be inferred from TargetType");
+				throw new CompilationResultException(CompilationResult.RulesError, KnownLogEvents.RulesAssembly, "TargetRules.LinkType should be inferred from TargetType");
 			}
 
 			if (bValidateTarget)
@@ -674,7 +674,7 @@ namespace UnrealBuildTool
 				// Set the default value for whether to use the shared build environment
 				if (Rules.BuildEnvironment == TargetBuildEnvironment.Unique && Unreal.IsEngineInstalled())
 				{
-					throw new CompilationResultException(CompilationResult.RulesError, "Targets with a unique build environment cannot be built with an installed engine.");
+					throw new CompilationResultException(CompilationResult.RulesError, KnownLogEvents.RulesAssembly, "Targets with a unique build environment cannot be built with an installed engine.");
 				}
 			}
 
@@ -777,7 +777,7 @@ namespace UnrealBuildTool
 			// Some platforms may *require* monolithic compilation...
 			if (Rules.LinkType != TargetLinkType.Monolithic && UEBuildPlatform.PlatformRequiresMonolithicBuilds(Rules.Platform, Rules.Configuration))
 			{
-				throw new CompilationResultException(CompilationResult.RulesError, String.Format("{RulesName}: {RulesPlatform} does not support modular builds", Rules.Name, Rules.Platform));
+				throw new CompilationResultException(CompilationResult.RulesError, KnownLogEvents.RulesAssembly, "{RulesName}: {RulesPlatform} does not support modular builds", Rules.Name, Rules.Platform);
 			}
 
 			if (IsTestTarget)
@@ -810,7 +810,7 @@ namespace UnrealBuildTool
 					// if the target doesn't allow for an override at all, error
 					if (!Rules.AllowsPerProjectSDKVersion())
 					{
-						throw new CompilationResultException(CompilationResult.RulesError, "Target {RulesName} is being built with a overridden {Platform} SDK version to '{SdkVersion}', but this target is not allowed - likely due to a modular build using a Shared BuildEnvironment",
+						throw new CompilationResultException(CompilationResult.RulesError, KnownLogEvents.RulesAssembly, "Target {RulesName} is being built with a overridden {Platform} SDK version to '{SdkVersion}', but this target is not allowed - likely due to a modular build using a Shared BuildEnvironment",
 							Rules.Name, Platform, SDK.GetMainVersion());
 					}
 
@@ -818,7 +818,7 @@ namespace UnrealBuildTool
 					if (Rules.ProjectFile == null || !SDK.ProjectsThatOverrodeSDK.Contains(Rules.ProjectFile))
 					{
 						string OverrideProject = SDK.ProjectsThatOverrodeSDK[0].GetFileNameWithoutAnyExtensions();
-						throw new CompilationResultException(CompilationResult.RulesError, "Target {RulesName} is using default {Platform} SDK version, but another target (probably {OverrideProject}) has overridden the SDK version to '{SdkVersion}'. If this target doesn't care about SDK versions, set 'bAreTargetSDKVersionsRelevantOverride = false' in your Target.cs file",
+						throw new CompilationResultException(CompilationResult.RulesError, KnownLogEvents.RulesAssembly, "Target {RulesName} is using default {Platform} SDK version, but another target (probably {OverrideProject}) has overridden the SDK version to '{SdkVersion}'. If this target doesn't care about SDK versions, set 'bAreTargetSDKVersionsRelevantOverride = false' in your Target.cs file",
 							Rules.Name, Platform, OverrideProject, SDK.GetMainVersion());
 					}
 				}
@@ -864,7 +864,7 @@ namespace UnrealBuildTool
 						}
 					}
 
-					throw new CompilationResultException(CompilationResult.RulesError, ExceptionMessage,
+					throw new CompilationResultException(CompilationResult.RulesError, KnownLogEvents.RulesAssembly, ExceptionMessage,
 						TargetName, CompiledAssembly?.FullName ?? "Unknown Assembly", CompiledAssembly?.Location ?? "Unknown Location");
 				}
 				else
@@ -881,7 +881,7 @@ namespace UnrealBuildTool
 
 			if (TargetRules == null)
 			{
-				throw new CompilationResultException(CompilationResult.RulesError, "Expecting to find a type to be declared in a target rules named '{TargetTypeName}'.  This type must derive from the 'TargetRules' type defined by UnrealBuildTool.", TargetTypeName);
+				throw new CompilationResultException(CompilationResult.RulesError, KnownLogEvents.RulesAssembly, "Expecting to find a type to be declared in a target rules named '{TargetTypeName}'.  This type must derive from the 'TargetRules' type defined by UnrealBuildTool.", TargetTypeName);
 			}
 
 			return TargetRules;
@@ -914,7 +914,7 @@ namespace UnrealBuildTool
 			{
 				if (Parent == null)
 				{
-					throw new CompilationResultException(CompilationResult.RulesError, "Unable to find target of type '{Type}' for project '{1}'", Type, ProjectFile?.FullName ?? "NoProject");
+					throw new CompilationResultException(CompilationResult.RulesError, KnownLogEvents.RulesAssembly, "Unable to find target of type '{Type}' for project '{1}'", Type, ProjectFile?.FullName ?? "NoProject");
 				}
 				else
 				{
@@ -942,7 +942,7 @@ namespace UnrealBuildTool
 					}
 				}
 
-				throw new CompilationResultException(CompilationResult.RulesError, "Found multiple targets with TargetType={Type}: {Matches}.\nSpecify a default with a {KeyName} entry in [/Script/BuildSettings.BuildSettings] section of your DefaultEngine.ini", Type, String.Join(", ", Matches), KeyName);
+				throw new CompilationResultException(CompilationResult.RulesError, KnownLogEvents.RulesAssembly, "Found multiple targets with TargetType={Type}: {Matches}.\nSpecify a default with a {KeyName} entry in [/Script/BuildSettings.BuildSettings] section of your DefaultEngine.ini", Type, String.Join(", ", Matches), KeyName);
 			}
 		}
 
