- Added a cache for a repeated version lookup

diff --git a/Engine/Source/Programs/UnrealBuildTool/ToolChain/AppleToolChain.cs b/Engine/Source/Programs/UnrealBuildTool/ToolChain/AppleToolChain.cs
--- a/Engine/Source/Programs/UnrealBuildTool/ToolChain/AppleToolChain.cs
+++ b/Engine/Source/Programs/UnrealBuildTool/ToolChain/AppleToolChain.cs
@@ -832,10 +832,15 @@ namespace UnrealBuildTool
 			return OutputFiles;
 		}
 
+		private static VersionNumber? AppleClangVersion = null;
+
 		protected virtual void GetLinkArguments_Global(LinkEnvironment LinkEnvironment, List<string> Arguments)
 		{
+			if (AppleClangVersion == null)
+			{
+				AppleClangVersion = ((AppleToolChainInfo)GetToolChainInfo()).AppleClangVersion;
+			}
 			// Temp solution for UE-191350
-			VersionNumber AppleClangVersion = ((AppleToolChainInfo)GetToolChainInfo()).AppleClangVersion;
 			if (AppleClangVersion >= new VersionNumber(15) && AppleClangVersion < new VersionNumber(16))
 			{
 				Arguments.Add(" -ld_classic");
