[UBA]
* Changed UBT to default to listen to 0.0.0.0 instead of traversing network adapters.

Thanks ZeroErrors[REDACTED_HANDLE] for reporting the issue with the network adapter traversal.

diff --git a/Engine/Source/Programs/UnrealBuildTool/Executors/UnrealBuildAccelerator/UnrealBuildAcceleratorConfig.cs b/Engine/Source/Programs/UnrealBuildTool/Executors/UnrealBuildAccelerator/UnrealBuildAcceleratorConfig.cs
--- a/Engine/Source/Programs/UnrealBuildTool/Executors/UnrealBuildAccelerator/UnrealBuildAcceleratorConfig.cs
+++ b/Engine/Source/Programs/UnrealBuildTool/Executors/UnrealBuildAccelerator/UnrealBuildAcceleratorConfig.cs
@@ -86,7 +86,7 @@ namespace UnrealBuildTool
 		/// </summary>
 		[XmlConfigFile(Category = "UnrealBuildAccelerator")]
 		[CommandLine("-UBAHost")]
-		public string Host { get; set; } = String.Empty;
+		public string Host { get; set; } = "0.0.0.0";
 
 		/// <summary>
 		/// Which port UBA server should listen to for connections.
