ubt: added build options for enabling the autortfm verifier

diff --git a/Engine/Source/Programs/UnrealBuildTool/Configuration/ReadOnlyTargetRules.cs b/Engine/Source/Programs/UnrealBuildTool/Configuration/ReadOnlyTargetRules.cs
--- a/Engine/Source/Programs/UnrealBuildTool/Configuration/ReadOnlyTargetRules.cs
+++ b/Engine/Source/Programs/UnrealBuildTool/Configuration/ReadOnlyTargetRules.cs
@@ -135,6 +135,7 @@ namespace UnrealBuildTool
 		public bool bUseVerseVM => !bUseVerseBPVM || Inner.GlobalDefinitions.Contains("WITH_VERSE_VM=1") || Inner.GlobalDefinitions.Contains("WITH_VERSE_VM=WITH_COREUOBJECT");
 
 		public bool bUseAutoRTFMCompiler => Inner.bUseAutoRTFMCompiler;
+		public bool bUseAutoRTFMVerifier => Inner.bUseAutoRTFMVerifier;
 
 		[Obsolete("Deprecated in UE5.1 - No longer used as Chaos is always enabled.")]
 		public bool bCompileChaos => Inner.bCompileChaos;
diff --git a/Engine/Source/Programs/UnrealBuildTool/Configuration/TargetRules.cs b/Engine/Source/Programs/UnrealBuildTool/Configuration/TargetRules.cs
--- a/Engine/Source/Programs/UnrealBuildTool/Configuration/TargetRules.cs
+++ b/Engine/Source/Programs/UnrealBuildTool/Configuration/TargetRules.cs
@@ -888,6 +888,14 @@ namespace UnrealBuildTool
 		[XmlConfigFile(Category = "BuildConfiguration")]
 		public bool bForceNoAutoRTFMCompiler { get; set; } = false;
 
+		/// <summary>
+		/// Whether to enable emitting AutoRTFM verification metadata
+		/// </summary>
+		[RequiresUniqueBuildEnvironment]
+		[CommandLine("-UseAutoRTFMVerifier", Value = "true")]
+		[XmlConfigFile(Category = "BuildConfiguration")]
+		public bool bUseAutoRTFMVerifier { get; set; } = false;
+
 		/// <summary>
 		/// Whether to compile the Chaos physics plugin.
 		/// </summary>
diff --git a/Engine/Source/Programs/UnrealBuildTool/Configuration/UEBuildModuleCPP.cs b/Engine/Source/Programs/UnrealBuildTool/Configuration/UEBuildModuleCPP.cs
--- a/Engine/Source/Programs/UnrealBuildTool/Configuration/UEBuildModuleCPP.cs
+++ b/Engine/Source/Programs/UnrealBuildTool/Configuration/UEBuildModuleCPP.cs
@@ -1565,6 +1565,7 @@ namespace UnrealBuildTool
 			CompileEnvironment.IncludeOrderVersion = ModuleCompileEnvironment.IncludeOrderVersion;
 			CompileEnvironment.bUseAutoRTFMCompiler = ModuleCompileEnvironment.bUseAutoRTFMCompiler;
 			CompileEnvironment.bDisableAutoRTFMInstrumentation = ModuleCompileEnvironment.bDisableAutoRTFMInstrumentation;
+			CompileEnvironment.bEnableAutoRTFMVerification = ModuleCompileEnvironment.bEnableAutoRTFMVerification;
 			CompileEnvironment.bValidateFormatStrings = ModuleCompileEnvironment.bValidateFormatStrings;
 			CompileEnvironment.bValidateInternalApi = ModuleCompileEnvironment.bValidateInternalApi;
 		}
@@ -2163,10 +2164,11 @@ namespace UnrealBuildTool
 
 			CompileEnvironmentDebugInfoSettings(Target, Result);
 
-			// Only copy the AutoRTFM flag if we are using the AutoRTFM compiler
+			// Only copy the AutoRTFM flags if we are using the AutoRTFM compiler
 			if (Result.bUseAutoRTFMCompiler)
 			{
 				Result.bDisableAutoRTFMInstrumentation = Rules.bDisableAutoRTFMInstrumentation;
+				Result.bEnableAutoRTFMVerification = Target.bUseAutoRTFMVerifier;
 			}
 
 			if (Result.OptimizationLevel != Rules.OptimizationLevel)
diff --git a/Engine/Source/Programs/UnrealBuildTool/Platform/Windows/VCToolChain.cs b/Engine/Source/Programs/UnrealBuildTool/Platform/Windows/VCToolChain.cs
--- a/Engine/Source/Programs/UnrealBuildTool/Platform/Windows/VCToolChain.cs
+++ b/Engine/Source/Programs/UnrealBuildTool/Platform/Windows/VCToolChain.cs
@@ -1176,6 +1176,11 @@ namespace UnrealBuildTool
 				if (CompileEnvironment.bEnableAutoRTFMInstrumentation)
 				{
 					Arguments.Add("-fautortfm");
+
+					if (CompileEnvironment.bEnableAutoRTFMVerification)
+					{
+						Arguments.Add("-fautortfm-verify");
+					}
 				}
 			}
 
@@ -1217,6 +1222,11 @@ namespace UnrealBuildTool
 				if (CompileEnvironment.bEnableAutoRTFMInstrumentation)
 				{
 					Arguments.Add("-fautortfm");
+
+					if (CompileEnvironment.bEnableAutoRTFMVerification)
+					{
+						Arguments.Add("-fautortfm-verify");
+					}
 				}
 			}
 		}
diff --git a/Engine/Source/Programs/UnrealBuildTool/System/CppCompileEnvironment.cs b/Engine/Source/Programs/UnrealBuildTool/System/CppCompileEnvironment.cs
--- a/Engine/Source/Programs/UnrealBuildTool/System/CppCompileEnvironment.cs
+++ b/Engine/Source/Programs/UnrealBuildTool/System/CppCompileEnvironment.cs
@@ -292,6 +292,11 @@ namespace UnrealBuildTool
 		/// </summary>
 		public bool bEnableAutoRTFMInstrumentation => bUseAutoRTFMCompiler && !bDisableAutoRTFMInstrumentation;
 
+		/// <summary>
+		/// Whether AutoRTFM verification is enabled or not.
+		/// </summary>
+		public bool bEnableAutoRTFMVerification = false;
+
 		/// <summary>
 		/// If unity builds are enabled this can be used to override if this specific module will build using Unity.
 		/// This is set using the per module configurations in BuildConfiguration.
@@ -781,6 +786,7 @@ namespace UnrealBuildTool
 			bOptimizeCode = Other.bOptimizeCode;
 			bUseAutoRTFMCompiler = Other.bUseAutoRTFMCompiler;
 			bDisableAutoRTFMInstrumentation = Other.bDisableAutoRTFMInstrumentation;
+			bEnableAutoRTFMVerification = Other.bEnableAutoRTFMVerification;
 			bCodeCoverage = Other.bCodeCoverage;
 			OptimizationLevel = Other.OptimizationLevel;
 			FPSemantics = Other.FPSemantics;
diff --git a/Engine/Source/Programs/UnrealBuildTool/ToolChain/ClangToolChain.cs b/Engine/Source/Programs/UnrealBuildTool/ToolChain/ClangToolChain.cs
--- a/Engine/Source/Programs/UnrealBuildTool/ToolChain/ClangToolChain.cs
+++ b/Engine/Source/Programs/UnrealBuildTool/ToolChain/ClangToolChain.cs
@@ -390,10 +390,15 @@ namespace UnrealBuildTool
 				// Validate PCH inputs by content if mtime check fails
 				Arguments.Add("-fpch-validate-input-files-content");
 			}
-			
+
 			if (CompileEnvironment.bEnableAutoRTFMInstrumentation)
 			{
 				Arguments.Add("-fautortfm");
+
+				if (CompileEnvironment.bEnableAutoRTFMVerification)
+				{
+					Arguments.Add("-fautortfm-verify");
+				}
 			}
 		}
 
@@ -425,6 +430,11 @@ namespace UnrealBuildTool
 			if (!CompileEnvironment.bDisableAutoRTFMInstrumentation && CompileEnvironment.bUseAutoRTFMCompiler)
 			{
 				Arguments.Add("-fautortfm");
+
+				if (CompileEnvironment.bEnableAutoRTFMVerification)
+				{
+					Arguments.Add("-fautortfm-verify");
+				}
 			}
 		}
 
