Clang PGO - Handle falling back to an alternative ProfData file if the exact match for the current config doesn't exist (similar to what we do with PGD files with MSVC).

diff --git a/Engine/Source/Programs/UnrealBuildTool/Platform/Windows/VCToolChain.cs b/Engine/Source/Programs/UnrealBuildTool/Platform/Windows/VCToolChain.cs
--- a/Engine/Source/Programs/UnrealBuildTool/Platform/Windows/VCToolChain.cs
+++ b/Engine/Source/Programs/UnrealBuildTool/Platform/Windows/VCToolChain.cs
@@ -301,6 +301,26 @@ namespace UnrealBuildTool
 			Arguments.Add($"/clang:-MD /clang:-MF\"{SourceDependsFileString}\"");
 		}
 
+		protected static string GetClangProfDataFilename(CppCompileEnvironment CompileEnvironment)
+		{
+			string ProfDataFilename = Path.Combine(CompileEnvironment.PGODirectory!, CompileEnvironment.PGOFilenamePrefix!);
+			if (File.Exists(ProfDataFilename)) 
+			{
+				return ProfDataFilename;
+			}
+			// If an exact match doesn't exist, fall back to an alternative. This is supported for building Test with Shipping PGO data for example
+			string[] ProfDataFiles = Directory.GetFiles(CompileEnvironment.PGODirectory!, "*.profdata");
+			if (ProfDataFiles.Length > 1)
+			{
+				throw new BuildException("More than one .profdata file found in \"{0}\" and \"{1}\" not found ", CompileEnvironment.PGODirectory, ProfDataFilename);
+			}
+			if (ProfDataFiles.Length == 0)
+			{
+				throw new BuildException("No .profdata files found in \"{PgoDir}\".", CompileEnvironment.PGODirectory);
+			}
+			return ProfDataFiles.First();
+		}	
+
 		protected virtual void AppendCLArguments_Global(CppCompileEnvironment CompileEnvironment, List<string> Arguments)
 		{
 			// Workaround for MSVC 14.31 compiler crash
@@ -755,14 +775,15 @@ namespace UnrealBuildTool
 				else if (CompileEnvironment.bPGOOptimize)
 				{
 					// Use a merged profdata file.
-					Log.TraceInformationOnce($"Using PGO profile data \"{Path.Combine(CompileEnvironment.PGODirectory!, CompileEnvironment.PGOFilenamePrefix!)}.profdata\"");
+					string ProfDataFilename = GetClangProfDataFilename(CompileEnvironment);
+					Log.TraceInformationOnce($"Using PGO profile data \"{ProfDataFilename}\"");
 					if (Target.WindowsPlatform.Compiler.IsIntel() && Target.WindowsPlatform.bSampleBasedPGO)
 					{
-						Arguments.Add($"-fprofile-sample-use=\"{Path.Combine(CompileEnvironment.PGODirectory!, CompileEnvironment.PGOFilenamePrefix!)}.profdata\"");
+						Arguments.Add($"-fprofile-sample-use=\"{ProfDataFilename}\"");
 					}
 					else
 					{
-						Arguments.Add($"-fprofile-use=\"{Path.Combine(CompileEnvironment.PGODirectory!, CompileEnvironment.PGOFilenamePrefix!)}.profdata\"");
+						Arguments.Add($"-fprofile-use=\"{ProfDataFilename}\"");
 					}
 				}
 			}
