Add verse module information to the UHT manifest

diff --git a/Engine/Source/Programs/Shared/EpicGames.Core/UHTTypes.cs b/Engine/Source/Programs/Shared/EpicGames.Core/UHTTypes.cs
--- a/Engine/Source/Programs/Shared/EpicGames.Core/UHTTypes.cs
+++ b/Engine/Source/Programs/Shared/EpicGames.Core/UHTTypes.cs
@@ -72,6 +72,32 @@ namespace EpicGames.Core
 		GameThirdParty,
 	}
 
+	/// <summary>
+	/// Describes the origin and visibility of Verse code
+	/// </summary>
+	public enum UHTVerseScope
+	{
+		/// <summary>
+		/// Created by Epic and only public definitions will be visible to public users
+		/// </summary>
+		PublicAPI,
+
+		/// <summary>
+		/// Created by Epic and is entirely hidden from public users
+		/// </summary>
+		InternalAPI,
+
+		/// <summary>
+		/// Created by a public user
+		/// </summary>
+		PublicUser,
+
+		/// <summary>
+		/// Created by an Epic internal user
+		/// </summary>
+		InternalUser
+	}
+
 	public class UHTManifest
 	{
 		public class Module
@@ -94,6 +120,10 @@ namespace EpicGames.Core
 			[JsonConverter(typeof(JsonStringEnumConverter))]
 			[JsonPropertyName("UHTGeneratedCodeVersion")]
 			public EGeneratedCodeVersion GeneratedCodeVersion { get; set; } = EGeneratedCodeVersion.None;
+			public string VersePath { get; set; } = "";
+			[JsonConverter(typeof(JsonStringEnumConverter))]
+			public UHTVerseScope VerseScope { get; set; } = UHTVerseScope.PublicAPI;
+			public bool HasVerse { get; set; } = false;
 
 			public override string ToString()
 			{
diff --git a/Engine/Source/Programs/UnrealBuildTool/System/ExternalExecution.cs b/Engine/Source/Programs/UnrealBuildTool/System/ExternalExecution.cs
--- a/Engine/Source/Programs/UnrealBuildTool/System/ExternalExecution.cs
+++ b/Engine/Source/Programs/UnrealBuildTool/System/ExternalExecution.cs
@@ -101,7 +101,7 @@ namespace UnrealBuildTool
 		/// <summary>
 		/// The include search path for generated headers to include other headers
 		/// </summary>
-		public DirectoryReference ModuleIncudeBase;
+		public DirectoryReference ModuleIncludeBase;
 
 		/// <summary>
 		/// Module type
@@ -144,7 +144,7 @@ namespace UnrealBuildTool
 		public List<string> PublicDefines;
 
 		/// <summary>
-		/// Base (i.e. extensionless) path+filename of the .gen files
+		/// Path and base filename, without extension, of the .gen files
 		/// </summary>
 		public string? GeneratedCPPFilenameBase;
 
@@ -158,12 +158,39 @@ namespace UnrealBuildTool
 		/// </summary>
 		public bool bIsReadOnly;
 
-		public UHTModuleInfo(string ModuleName, FileReference ModuleRulesFile, DirectoryReference[] ModuleDirectories, DirectoryReference ModuleIncudeBase, UHTModuleType ModuleType, DirectoryItem GeneratedCodeDirectory, EGeneratedCodeVersion GeneratedCodeVersion, bool bIsReadOnly, ModuleRules.PackageOverrideType OverrideType)
+		/// <summary>
+		/// Path of the verse generated types 
+		/// </summary>
+		public string VersePath;
+
+		/// <summary>
+		/// Scope of the verse definitions
+		/// </summary>
+		public VerseScope VerseScope;
+
+		/// <summary>
+		/// If true, verse scripts were found
+		/// </summary>
+		public bool HasVerse;
+
+		public UHTModuleInfo(
+			string ModuleName, 
+			FileReference ModuleRulesFile, 
+			DirectoryReference[] ModuleDirectories, 
+			DirectoryReference ModuleIncludeBase, 
+			UHTModuleType ModuleType, 
+			DirectoryItem GeneratedCodeDirectory, 
+			EGeneratedCodeVersion GeneratedCodeVersion, 
+			bool bIsReadOnly, 
+			ModuleRules.PackageOverrideType OverrideType,
+			string VersePath,
+			VerseScope VerseScope,
+			bool HasVerse)
 		{
 			this.ModuleName = ModuleName;
 			this.ModuleRulesFile = ModuleRulesFile;
 			this.ModuleDirectories = ModuleDirectories;
-			this.ModuleIncudeBase = ModuleIncudeBase;
+			this.ModuleIncludeBase = ModuleIncludeBase;
 			this.ModuleType = ModuleType.ToString();
 			OverrideModuleType = OverrideType.ToString();
 			PublicUObjectClassesHeaders = new List<FileItem>();
@@ -174,6 +201,9 @@ namespace UnrealBuildTool
 			this.GeneratedCodeDirectory = GeneratedCodeDirectory;
 			this.GeneratedCodeVersion = GeneratedCodeVersion;
 			this.bIsReadOnly = bIsReadOnly;
+			this.VersePath = VersePath;
+			this.VerseScope = VerseScope;
+			this.HasVerse = HasVerse;
 		}
 
 		public UHTModuleInfo(BinaryArchiveReader Reader)
@@ -181,7 +211,7 @@ namespace UnrealBuildTool
 			ModuleName = Reader.ReadString()!;
 			ModuleRulesFile = Reader.ReadFileReference();
 			ModuleDirectories = Reader.ReadArray<DirectoryReference>(Reader.ReadDirectoryReferenceNotNull)!;
-			ModuleIncudeBase = Reader.ReadDirectoryReference()!;
+			ModuleIncludeBase = Reader.ReadDirectoryReference()!;
 			ModuleType = Reader.ReadString()!;
 			OverrideModuleType = Reader.ReadString()!;
 			PublicUObjectClassesHeaders = Reader.ReadList(() => Reader.ReadFileItem())!;
@@ -193,6 +223,9 @@ namespace UnrealBuildTool
 			GeneratedCodeVersion = (EGeneratedCodeVersion)Reader.ReadInt();
 			bIsReadOnly = Reader.ReadBool();
 			PublicDefines = Reader.ReadList(() => Reader.ReadString())!;
+			VersePath = Reader.ReadString()!;
+			VerseScope = (VerseScope)Reader.ReadInt();
+			HasVerse = Reader.ReadBool();
 		}
 
 		public void Write(BinaryArchiveWriter Writer)
@@ -200,7 +233,7 @@ namespace UnrealBuildTool
 			Writer.WriteString(ModuleName);
 			Writer.WriteFileReference(ModuleRulesFile);
 			Writer.WriteArray<DirectoryReference>(ModuleDirectories, Writer.WriteDirectoryReference);
-			Writer.WriteDirectoryReference(ModuleIncudeBase);
+			Writer.WriteDirectoryReference(ModuleIncludeBase);
 			Writer.WriteString(ModuleType);
 			Writer.WriteString(OverrideModuleType);
 			Writer.WriteList(PublicUObjectClassesHeaders, Item => Writer.WriteFileItem(Item));
@@ -212,6 +245,9 @@ namespace UnrealBuildTool
 			Writer.WriteInt((int)GeneratedCodeVersion);
 			Writer.WriteBool(bIsReadOnly);
 			Writer.WriteList(PublicDefines, Item => Writer.WriteString(Item));
+			Writer.WriteString(VersePath);
+			Writer.WriteInt((int)VerseScope);
+			Writer.WriteBool(HasVerse);
 		}
 
 		public override string ToString()
@@ -384,7 +420,15 @@ namespace UnrealBuildTool
 			return Unreal.EngineSourceDirectory;
 		}
 
-		public static void SetupUObjectModules(IEnumerable<UEBuildModuleCPP> ModulesToGenerateHeadersFor, UnrealTargetPlatform Platform, ProjectDescriptor? ProjectDescriptor, List<UHTModuleInfo> UObjectModules, List<UHTModuleHeaderInfo> UObjectModuleHeaders, EGeneratedCodeVersion GeneratedCodeVersion, SourceFileMetadataCache MetadataCache, ILogger Logger)
+		public static void SetupUObjectModules(
+			IEnumerable<UEBuildModuleCPP> ModulesToGenerateHeadersFor,
+			UnrealTargetPlatform Platform, 
+			ProjectDescriptor? ProjectDescriptor, 
+			List<UHTModuleInfo> UObjectModules,
+			List<UHTModuleHeaderInfo> UObjectModuleHeaders, 
+			EGeneratedCodeVersion GeneratedCodeVersion, 
+			SourceFileMetadataCache MetadataCache, 
+			ILogger Logger)
 		{
 			// Find the type of each module
 			Dictionary<UEBuildModuleCPP, UHTModuleType> ModuleToType = new Dictionary<UEBuildModuleCPP, UHTModuleType>();
@@ -409,7 +453,19 @@ namespace UnrealBuildTool
 					DirectoryItem GeneratedCodeDirectory = DirectoryItem.GetItemByDirectoryReference(Module.GeneratedCodeDirectoryUHT!);
 					DirectoryReference ModuleIncludeBase = FindIncludeBase(Module, Logger);
 
-					UHTModuleInfo Info = new UHTModuleInfo(Module.Name, Module.RulesFile, Module.ModuleDirectories, ModuleIncludeBase, ModuleToType[Module], GeneratedCodeDirectory, GeneratedCodeVersion, Module.Rules.bUsePrecompiled, Module.Rules.OverridePackageType);
+					UHTModuleInfo Info = new UHTModuleInfo(
+						Module.Name, 
+						Module.RulesFile, 
+						Module.ModuleDirectories, 
+						ModuleIncludeBase, 
+						ModuleToType[Module], 
+						GeneratedCodeDirectory, 
+						GeneratedCodeVersion, 
+						Module.Rules.bUsePrecompiled, 
+						Module.Rules.OverridePackageType,
+						Module.Rules.VersePath ?? "",
+						Module.Rules.VerseScope,
+						Module.bHasVerse);
 					ModuleInfoArray[Idx] = Info;
 
 					Queue.Enqueue(() => SetupUObjectModule(Info, ExcludedFolders, MetadataCache, Queue));
@@ -970,7 +1026,7 @@ namespace UnrealBuildTool
 						ModuleType = (UHTModuleType)Enum.Parse(typeof(UHTModuleType), UObjectModule.ModuleType),
 						OverrideModuleType = (EPackageOverrideType)Enum.Parse(typeof(EPackageOverrideType), UObjectModule.OverrideModuleType),
 						BaseDirectory = UObjectModule.ModuleDirectories[0].FullName,
-						IncludeBase = UObjectModule.ModuleIncudeBase.FullName,
+						IncludeBase = UObjectModule.ModuleIncludeBase.FullName,
 						OutputDirectory = Path.GetDirectoryName(UObjectModule.GeneratedCPPFilenameBase)!,
 						ClassesHeaders = UObjectModule.PublicUObjectClassesHeaders.Select((Header) => Header.AbsolutePath).ToList(),
 						PublicHeaders = UObjectModule.PublicUObjectHeaders.Select((Header) => Header.AbsolutePath).ToList(),
@@ -980,6 +1036,9 @@ namespace UnrealBuildTool
 						GeneratedCPPFilenameBase = UObjectModule.GeneratedCPPFilenameBase,
 						SaveExportedHeaders = !UObjectModule.bIsReadOnly,
 						GeneratedCodeVersion = UObjectModule.GeneratedCodeVersion,
+						VersePath = UObjectModule.VersePath,
+						VerseScope = (UHTVerseScope)Enum.Parse(typeof(UHTVerseScope), UObjectModule.VerseScope.ToString()),
+						HasVerse = UObjectModule.HasVerse,
 					});
 			}
 
