[UBT]
* Fixed merged modules path for targets where ispc byte code files are used

diff --git a/Engine/Source/Programs/UnrealBuildTool/ToolChain/ISPCToolChain.cs b/Engine/Source/Programs/UnrealBuildTool/ToolChain/ISPCToolChain.cs
--- a/Engine/Source/Programs/UnrealBuildTool/ToolChain/ISPCToolChain.cs
+++ b/Engine/Source/Programs/UnrealBuildTool/ToolChain/ISPCToolChain.cs
@@ -769,12 +769,11 @@ namespace UnrealBuildTool
 
 				string AdditionalArguments = "";
 				// Must be added after response file is created just to make sure it ends up on the command line and not in the response file
-				if (bMergeModules)
+				if (!bByteCodeOutput && bMergeModules)
 				{
 					// EXTRACTEXPORTS can only be interpreted by UBA.. so this action won't build outside uba
 					AdditionalArguments = " /EXTRACTEXPORTS";
-					FileItem SymFile = FileItem.GetItemByFileReference(FileReference.Combine(OutputDir, FileName + ".exi"));
-					CompileAction.ProducedItems.Add(SymFile);
+					CompileAction.ProducedItems.Add(FileItem.GetItemByFileReference(FileReference.Combine(OutputDir, FileName + ".exi")));
 				}
 
 				CompileAction.CommandArguments = $"@\"{ResponseFileName}\"{AdditionalArguments}";
@@ -794,10 +793,11 @@ namespace UnrealBuildTool
 						List<FileItem> FinalObjectFiles = new List<FileItem>();
 						foreach (FileItem CompiledBytecodeObjFile in CompiledISPCObjFiles)
 						{
+							string FileNameWithoutExtension = Path.GetFileNameWithoutExtension(CompiledBytecodeObjFile.AbsolutePath);
 							FileItem FinalCompiledISPCObjFile = FileItem.GetItemByFileReference(
 								FileReference.Combine(
 									OutputDir,
-									Path.GetFileNameWithoutExtension(CompiledBytecodeObjFile.AbsolutePath) + GetISPCObjectFileSuffix(CompileEnvironment.Platform)
+									FileNameWithoutExtension + GetISPCObjectFileSuffix(CompileEnvironment.Platform)
 									)
 								);
 
@@ -809,10 +809,17 @@ namespace UnrealBuildTool
 							PostCompileArgs.AddRange(CommonArgs);
 							PostCompileArgs.Add($"-o \"{NormalizeCommandLinePath(FinalCompiledISPCObjFile)}\"");
 
+							if (bMergeModules)
+							{
+								// EXTRACTEXPORTS can only be interpreted by UBA.. so this action won't build outside uba
+								AdditionalArguments = " /EXTRACTEXPORTS";
+								PostCompileAction.ProducedItems.Add(FileItem.GetItemByFileReference(FileReference.Combine(OutputDir, FileNameWithoutExtension + ".exi")));
+							}
+
 							// Write the args to a response file
 							FileReference PostCompileResponseFileName = GetResponseFileName(CompileEnvironment, FinalCompiledISPCObjFile);
 							FileItem PostCompileResponseFileItem = Graph.CreateIntermediateTextFile(PostCompileResponseFileName, PostCompileArgs.Select(x => Utils.ExpandVariables(x)));
-							PostCompileAction.CommandArguments = $"@\"{PostCompileResponseFileName}\"";
+							PostCompileAction.CommandArguments = $"@\"{PostCompileResponseFileName}\"{AdditionalArguments}";
 							PostCompileAction.PrerequisiteItems.Add(PostCompileResponseFileItem);
 
 							PostCompileAction.PrerequisiteItems.Add(CompiledBytecodeObjFile);
