UnrealBuildTool: Fix incorrect version for ispc string when building multiple platforms with different ispc executables

#rnx

diff --git a/Engine/Source/Programs/UnrealBuildTool/ToolChain/ISPCToolChain.cs b/Engine/Source/Programs/UnrealBuildTool/ToolChain/ISPCToolChain.cs
--- a/Engine/Source/Programs/UnrealBuildTool/ToolChain/ISPCToolChain.cs
+++ b/Engine/Source/Programs/UnrealBuildTool/ToolChain/ISPCToolChain.cs
@@ -215,31 +215,31 @@ namespace UnrealBuildTool
 			return null;
 		}
 
-		static Dictionary<UnrealTargetPlatform, string> ISPCCompilerVersions = new Dictionary<UnrealTargetPlatform, string>();
+		static readonly Dictionary<string, string> s_ISPCCompilerVersions = new Dictionary<string, string>();
 
 		/// <summary>
 		/// Returns the version of the ISPC compiler for the specified platform. If GetISPCHostCompilerPath() doesn't return a valid path
 		/// this will return a -1 version.
 		/// </summary>
-		/// <param name="Platform">Which OS build platform is running on.</param>
-		/// <returns>Version reported by the ISPC compilerr</returns>
-		public virtual string GetISPCHostCompilerVersion(UnrealTargetPlatform Platform)
+		/// <param name="platform">Which OS build platform is running on.</param>
+		/// <returns>Version reported by the ISPC compiler</returns>
+		public virtual string GetISPCHostCompilerVersion(UnrealTargetPlatform platform)
 		{
-			if (!ISPCCompilerVersions.ContainsKey(Platform))
+			string compilerPath = GetISPCHostCompilerPath(platform);
+			if (!s_ISPCCompilerVersions.ContainsKey(compilerPath))
 			{
-				Version? CompilerVersion = null;
-				string CompilerPath = GetISPCHostCompilerPath(Platform);
-
-				if (!File.Exists(CompilerPath))
+				if (File.Exists(compilerPath))
 				{
-					Logger.LogWarning("No ISPC compiler at {CompilerPath}", CompilerPath);
-					CompilerVersion = new Version(-1, -1);
+					s_ISPCCompilerVersions[compilerPath] = RunToolAndCaptureOutput(new FileReference(compilerPath), "--version", "(.*)")!;
+				}
+				else
+				{
+					Logger.LogWarning("No ISPC compiler at {CompilerPath}", compilerPath);
+					s_ISPCCompilerVersions[compilerPath] = "-1";
 				}
-
-				ISPCCompilerVersions[Platform] = RunToolAndCaptureOutput(new FileReference(CompilerPath), "--version", "(.*)")!;
 			}
 
-			return ISPCCompilerVersions[Platform];
+			return s_ISPCCompilerVersions[compilerPath];
 		}
 
 		/// <summary>
