UnrealBuildTool: Ensure extension directory targets & modules write intermediates to the correct path

#rnx

diff --git a/Engine/Source/Programs/UnrealBuildTool/Configuration/UEBuildTarget.cs b/Engine/Source/Programs/UnrealBuildTool/Configuration/UEBuildTarget.cs
--- a/Engine/Source/Programs/UnrealBuildTool/Configuration/UEBuildTarget.cs
+++ b/Engine/Source/Programs/UnrealBuildTool/Configuration/UEBuildTarget.cs
@@ -1659,6 +1659,11 @@ namespace UnrealBuildTool
 		/// </summary>
 		public string PlatformIntermediateFolderNoArch;
 
+		/// <summary>
+		/// Base directory for output files. For an agnostic editor/game executable, this will be under the engine directory, otherwise typically underneath ProjectDirectory.
+		/// </summary>
+		public DirectoryReference BaseOutputDirectory;
+
 		/// <summary>
 		/// Root directory for the active project. Typically contains the .uproject file, or the engine root.
 		/// </summary>
@@ -1817,29 +1822,16 @@ namespace UnrealBuildTool
 				}
 			}
 
+			// Figure out what the project directory is. If we have a uproject file, use that. Otherwise use the engine directory.
+			ProjectDirectory = ProjectFile?.Directory ?? Unreal.EngineDirectory;
+
 			// track where intermediates and receipt files will be saved under (Programs with .uproject files still want to put the receipt
 			// under Engine/Binaries, not Engine/Programs/Foo/Binaries
-			DirectoryReference OutputRootDirectory = Unreal.EngineDirectory;
-
-			// Figure out what the project directory is. If we have a uproject file, use that. Otherwise use the engine directory.
-			if (ProjectFile != null)
-			{
-				ProjectDirectory = ProjectFile.Directory;
-				// programs are often compiled with another project passed in for context, so if it's using a _different_ uproject, use that uproject
-				// for the output location
-				if (Rules.Type != TargetType.Program || !String.Equals(ProjectFile.GetFileNameWithoutAnyExtensions(), Rules.Name, StringComparison.CurrentCultureIgnoreCase))
-				{
-					OutputRootDirectory = ProjectDirectory;
-				}
-			}
-			else
-			{
-				ProjectDirectory = Unreal.EngineDirectory;
-			}
+			BaseOutputDirectory = GetBaseOutputDirectory(Logger);
 
 			// Build the project intermediate directory
-			ProjectIntermediateDirectory = DirectoryReference.Combine(OutputRootDirectory, PlatformIntermediateFolder, GetTargetIntermediateFolderName(TargetName, IntermediateEnvironment), Configuration.ToString());
-			ProjectIntermediateDirectoryNoArch = DirectoryReference.Combine(OutputRootDirectory, PlatformIntermediateFolderNoArch, GetTargetIntermediateFolderName(TargetName, IntermediateEnvironment), Configuration.ToString());
+			ProjectIntermediateDirectory = DirectoryReference.Combine(BaseOutputDirectory, PlatformIntermediateFolder, GetTargetIntermediateFolderName(TargetName, IntermediateEnvironment), Configuration.ToString());
+			ProjectIntermediateDirectoryNoArch = DirectoryReference.Combine(BaseOutputDirectory, PlatformIntermediateFolderNoArch, GetTargetIntermediateFolderName(TargetName, IntermediateEnvironment), Configuration.ToString());
 
 			// Build the engine intermediate directory. If we're building agnostic engine binaries, we can use the engine intermediates folder. Otherwise we need to use the project intermediates directory.
 			if (!bUseSharedBuildEnvironment)
@@ -1856,7 +1848,7 @@ namespace UnrealBuildTool
 			}
 
 			// Get the receipt path for this target
-			ReceiptFileName = TargetReceipt.GetDefaultPath(OutputRootDirectory, TargetName, Platform, Configuration, Architectures);
+			ReceiptFileName = TargetReceipt.GetDefaultPath(BaseOutputDirectory, TargetName, Platform, Configuration, Architectures);
 
 			// Read the project descriptor
 			if (ProjectFile != null)
@@ -4298,6 +4290,29 @@ namespace UnrealBuildTool
 			}
 		}
 
+		DirectoryReference GetBaseOutputDirectory(ILogger Logger)
+		{
+			// if we are building a program with a project file, and we are making a unique build environment, output the program to that project's binaries
+			bool bIsUniqueBuildProgramWithExternalProject =
+					Rules.Type == TargetType.Program && Rules.BuildEnvironment == TargetBuildEnvironment.Unique &&
+					Rules.File.IsUnderDirectory(Unreal.EngineDirectory) &&
+					ProjectFile != null && !ProjectFile.IsUnderDirectory(Unreal.EngineDirectory);
+
+			if (bIsUniqueBuildProgramWithExternalProject)
+			{
+				// engine programs generally won't be under the directories searched in GetOutputDirectoryForExecutable, so it will return the EngineDir, even tho unique programs
+				// built with a ProjectFile want to be under the Project dir, so just force it into the ProjectDir if one wasn't found otherwise
+				Logger.LogInformation("Overriding output directory for unique build program {Name} from {Output} to {Project}, because it it is being built with external project {ProjectFile}", Rules.Name, GetOutputDirectoryForExecutable(ProjectDirectory, Rules.File), ProjectDirectory, ProjectFile);
+				return ProjectDirectory;
+			}
+			else if (ProjectFile != null && (bCompileMonolithic || !bUseSharedBuildEnvironment) && (Rules.File.IsUnderDirectory(ProjectDirectory)))
+			{
+				return GetOutputDirectoryForExecutable(ProjectDirectory, Rules.File);
+			}
+
+			return GetOutputDirectoryForExecutable(Unreal.EngineDirectory, Rules.File);
+		}
+
 		/// <summary>
 		/// Gets the output directory for a target
 		/// </summary>
@@ -4353,7 +4368,7 @@ namespace UnrealBuildTool
 		private DirectoryReference GetModuleIntermediateDirectory(ModuleRules ModuleRules, UnrealArchitectures? Architectures)
 		{
 			// Get the root output directory and base name (target name/app name) for this binary
-			DirectoryReference BaseOutputDirectory = GetBaseOutputDirectory(ModuleRules);
+			DirectoryReference BaseOutputDirectory = GetOutputDirectoryForExecutable(GetBaseOutputDirectory(ModuleRules), ModuleRules.File);
 
 			// Get the configuration that this module will be built in. Engine modules compiled in DebugGame will use Development.
 			UnrealTargetConfiguration ModuleConfiguration = Configuration;
@@ -5053,40 +5068,9 @@ namespace UnrealBuildTool
 			UEBuildModuleCPP LaunchModule = FindOrCreateCppModuleByName(Rules.LaunchModuleName, TargetRulesFile.GetFileName(), Logger);
 
 			// Get the intermediate directory for the launch module directory. This can differ from the standard engine intermediate directory because it is always configuration-specific.
-			DirectoryReference IntermediateDirectory;
-			if (LaunchModule.RulesFile.IsUnderDirectory(Unreal.EngineDirectory) && !ShouldCompileMonolithic())
-			{
-				IntermediateDirectory = DirectoryReference.Combine(Unreal.EngineDirectory, PlatformIntermediateFolder, UEBuildTarget.GetTargetIntermediateFolderName(AppName, IntermediateEnvironment), Configuration.ToString());
-			}
-			else
-			{
-				IntermediateDirectory = ProjectIntermediateDirectory;
-			}
-
-			// Construct the output paths for this target's executable
-			DirectoryReference OutputDirectory;
-			// if we are building a program with a project file, and we are making a unique build environment, output the program to that project's binaries
-			bool bIsUniqueBuildProgramWithExternalProject =
-					Rules.Type == TargetType.Program && Rules.BuildEnvironment == TargetBuildEnvironment.Unique &&
-					Rules.File.IsUnderDirectory(Unreal.EngineDirectory) &&
-					ProjectFile != null && !ProjectFile.IsUnderDirectory(Unreal.EngineDirectory);
-
-			if (bIsUniqueBuildProgramWithExternalProject)
-			{
-				// engine programs generally won't be under the directories searched in GetOutputDirectoryForExecutable, so it will return the EngineDir, even tho unique programs
-				// built with a ProjectFile want to be under the Project dir, so just force it into the ProjectDir if one wasn't found otherwise
-				Logger.LogInformation("Overriding output directory for unique build program {Name} from {Output} to {Project}, because it it is being built with external project {ProjectFile}", Rules.Name, GetOutputDirectoryForExecutable(ProjectDirectory, Rules.File), ProjectDirectory, ProjectFile);
-				OutputDirectory = ProjectDirectory;
-			}
-			else if (ProjectFile != null && (bCompileMonolithic || !bUseSharedBuildEnvironment) && (Rules.File.IsUnderDirectory(ProjectDirectory)))
-			{
-				OutputDirectory = GetOutputDirectoryForExecutable(ProjectDirectory, Rules.File);
-
-			}
-			else
-			{
-				OutputDirectory = GetOutputDirectoryForExecutable(Unreal.EngineDirectory, Rules.File);
-			}
+			DirectoryReference IntermediateDirectory = LaunchModule.RulesFile.IsUnderDirectory(Unreal.EngineDirectory) && !ShouldCompileMonolithic()
+				? DirectoryReference.Combine(EngineIntermediateDirectory.ParentDirectory!, Configuration.ToString())
+				: ProjectIntermediateDirectory;
 
 			// We only want these defines in the launch module to not create non-deterministic Definitions.h on other modules
 			if (Rules.bWithLiveCoding && Rules.LinkType == TargetLinkType.Monolithic)
@@ -5099,7 +5083,7 @@ namespace UnrealBuildTool
 			}
 
 			bool bCompileAsDLL = Rules.bShouldCompileAsDLL && bCompileMonolithic;
-			List<FileReference> OutputPaths = MakeBinaryPaths(OutputDirectory, bCompileMonolithic ? TargetName : AppName, Platform, Configuration, bCompileAsDLL ? UEBuildBinaryType.DynamicLinkLibrary : UEBuildBinaryType.Executable, Rules.Architectures, Rules.UndecoratedConfiguration, bCompileMonolithic && ProjectFile != null, Rules.ExeBinariesSubFolder, ProjectFile, Rules);
+			List<FileReference> OutputPaths = MakeBinaryPaths(BaseOutputDirectory, bCompileMonolithic ? TargetName : AppName, Platform, Configuration, bCompileAsDLL ? UEBuildBinaryType.DynamicLinkLibrary : UEBuildBinaryType.Executable, Rules.Architectures, Rules.UndecoratedConfiguration, bCompileMonolithic && ProjectFile != null, Rules.ExeBinariesSubFolder, ProjectFile, Rules);
 
 			// Create the binary
 			UEBuildBinary Binary = new UEBuildBinary(
