- Don't replace single quotes with double quotes on Mac/Linux with RunLocalProcessAndReturnStdOut

diff --git a/Engine/Source/Programs/UnrealBuildTool/System/Utils.cs b/Engine/Source/Programs/UnrealBuildTool/System/Utils.cs
--- a/Engine/Source/Programs/UnrealBuildTool/System/Utils.cs
+++ b/Engine/Source/Programs/UnrealBuildTool/System/Utils.cs
@@ -446,7 +446,14 @@ namespace UnrealBuildTool
 			// see [REDACTED_PROFILE_URL]/runtime/issues/29857
 			// also see UE-102580
 			// for rules see https://docs.microsoft.com/en-us/cpp/cpp/main-function-command-line-args
-			Args = Args?.Replace('\'', '\"') ?? String.Empty;
+			if (OperatingSystem.IsWindows())
+			{
+				Args = Args?.Replace('\'', '\"') ?? String.Empty;
+			}
+			else if (Args == null)
+			{
+				Args = String.Empty;
+			}
 
 			ProcessStartInfo StartInfo = new ProcessStartInfo(Command, Args);
 			StartInfo.UseShellExecute = false;
