UnrealBuildTool: Fix typo in clang arg

#rnx

diff --git a/Engine/Source/Programs/UnrealBuildTool/Platform/Windows/VCToolChain.cs b/Engine/Source/Programs/UnrealBuildTool/Platform/Windows/VCToolChain.cs
--- a/Engine/Source/Programs/UnrealBuildTool/Platform/Windows/VCToolChain.cs
+++ b/Engine/Source/Programs/UnrealBuildTool/Platform/Windows/VCToolChain.cs
@@ -885,7 +885,7 @@ namespace UnrealBuildTool
 					// https://clang.llvm.org/docs/UsersManual.html#cmdoption-feliminate-unused-debug-types
 					if (ClangVersion >= new VersionNumber(19))
 					{
-						Arguments.Add("-fno-eliminate-unused-debug-type");
+						Arguments.Add("-fno-eliminate-unused-debug-types");
 					}
 				}
 			}
