UBT: Fixing invalid string format

#rnx

diff --git a/Engine/Source/Programs/UnrealBuildTool/Platform/Windows/VCToolChain.cs b/Engine/Source/Programs/UnrealBuildTool/Platform/Windows/VCToolChain.cs
--- a/Engine/Source/Programs/UnrealBuildTool/Platform/Windows/VCToolChain.cs
+++ b/Engine/Source/Programs/UnrealBuildTool/Platform/Windows/VCToolChain.cs
@@ -312,11 +312,11 @@ namespace UnrealBuildTool
 			string[] ProfDataFiles = Directory.GetFiles(CompileEnvironment.PGODirectory!, "*.profdata");
 			if (ProfDataFiles.Length > 1)
 			{
-				throw new BuildException("More than one .profdata file found in \"{0}\" and \"{1}\" not found ", CompileEnvironment.PGODirectory, ProfDataFilename);
+				throw new BuildException($"More than one .profdata file found in \"{CompileEnvironment.PGODirectory}\" and \"{ProfDataFilename}\" not found ");
 			}
 			if (ProfDataFiles.Length == 0)
 			{
-				throw new BuildException("No .profdata files found in \"{PgoDir}\".", CompileEnvironment.PGODirectory);
+				throw new BuildException($"No .profdata files found in \"{CompileEnvironment.PGODirectory}\".");
 			}
 			return ProfDataFiles.First();
 		}	
