UnrealBuildTool: Add connection mode to telemetry

#rnx

diff --git a/Engine/Source/Programs/UnrealBuildTool/Executors/UnrealBuildAccelerator/UBAAgentCoordinatorHorde.cs b/Engine/Source/Programs/UnrealBuildTool/Executors/UnrealBuildAccelerator/UBAAgentCoordinatorHorde.cs
--- a/Engine/Source/Programs/UnrealBuildTool/Executors/UnrealBuildAccelerator/UBAAgentCoordinatorHorde.cs
+++ b/Engine/Source/Programs/UnrealBuildTool/Executors/UnrealBuildAccelerator/UBAAgentCoordinatorHorde.cs
@@ -645,6 +645,7 @@ namespace UnrealBuildTool
 	{
 		public string? Server => HordeConfig.HordeServer;
 		public bool Enabled => !HordeConfig.bDisableHorde;
+		public string ConnectionModeString => (Enum.TryParse(HordeConfig.HordeConnectionMode, true, out ConnectionMode cm) ? cm : ConnectionMode.Direct).ToString();
 
 		public UBAAgentCoordinatorHorde(ILogger logger, UnrealBuildAcceleratorConfig ubaConfig, CommandLineArguments? additionalArguments = null, DirectoryReference? projectDir = null)
 		{
diff --git a/Engine/Source/Programs/UnrealBuildTool/Executors/UnrealBuildAccelerator/UBAExecutor.cs b/Engine/Source/Programs/UnrealBuildTool/Executors/UnrealBuildAccelerator/UBAExecutor.cs
--- a/Engine/Source/Programs/UnrealBuildTool/Executors/UnrealBuildAccelerator/UBAExecutor.cs
+++ b/Engine/Source/Programs/UnrealBuildTool/Executors/UnrealBuildAccelerator/UBAExecutor.cs
@@ -56,6 +56,9 @@ namespace UnrealBuildTool
 		// Tracking for all actions processed remotely
 		int _remoteProcessedActions = 0;
 
+		// Tracking for remote connection mode
+		string _remoteConnectionMode = "Local";
+
 		protected override void Dispose(bool disposing)
 		{
 			if (disposing)
@@ -122,6 +125,7 @@ namespace UnrealBuildTool
 				hordeAgentCoordinators.Add(new UBAAgentCoordinatorHorde(logger, UBAConfig, targetDescriptor.AdditionalArguments, targetDescriptor.ProjectFile?.Directory));
 			}
 			hordeAgentCoordinators.RemoveAll(x => !x.Enabled);
+			_remoteConnectionMode = hordeAgentCoordinators.FirstOrDefault()?.ConnectionModeString ?? "Local";
 			_agentCoordinators.AddRange(hordeAgentCoordinators.DistinctBy(x => x.Server));
 
 			_threadedLogger = new ThreadedLogger(logger);
@@ -579,6 +583,7 @@ namespace UnrealBuildTool
 					_localRetryActions.Count, _forcedRetryActions.Count,
 					!UBAConfig.bDisableRemote ? _agentCoordinators.Count : 0, !UBAConfig.bDisableRemote ? _successfulCoordinatorConnections : 0, !UBAConfig.bDisableRemote ? _failedCoordinatorConnections : 0,
 					!UBAConfig.bDisableRemote ? _ubaDurationWaitingForRemote : TimeSpan.Zero,
+					!UBAConfig.bDisableRemote || _agentCoordinators.Count == 0 ? "Local" : _remoteConnectionMode,
 					DateTime.UtcNow);
 
 				return res;
diff --git a/Engine/Source/Programs/UnrealBuildTool/System/Telemetry.cs b/Engine/Source/Programs/UnrealBuildTool/System/Telemetry.cs
--- a/Engine/Source/Programs/UnrealBuildTool/System/Telemetry.cs
+++ b/Engine/Source/Programs/UnrealBuildTool/System/Telemetry.cs
@@ -158,6 +158,8 @@ namespace UnrealBuildTool
 		[JsonConverter(typeof(DurationDataRouterEventConverter))]
 		public TimeSpan DurationWaitingForRemote { get; }
 
+		public string ConnectionMode { get; }
+
 		public double RetriedLocalRate => RetriedLocalActions / (double)TotalActions;
 		public double RetriedDisabledRate => RetriedDisabledActions / (double)TotalActions;
 
@@ -168,7 +170,7 @@ namespace UnrealBuildTool
 			int localActions, int remoteActions,
 			int retriedLocalActions, int retriedDisabledActions,
 			int totalCoordinators, int succeededCoordinators, int failedCoordinators,
-			TimeSpan durationWaitingForRemote,
+			TimeSpan durationWaitingForRemote, string connectionMode,
 			DateTime timestamp)
 			: base(executor, startUTC, result, totalActions, succeededActions, failedActions, cacheHitActions, cacheMissActions, timestamp)
 		{
@@ -180,6 +182,7 @@ namespace UnrealBuildTool
 			SucceededCoordinators = succeededCoordinators;
 			FailedCoordinators = failedCoordinators;
 			DurationWaitingForRemote = durationWaitingForRemote;
+			ConnectionMode = connectionMode;
 		}
 	}
 
