UnrealBuildTool: Fix sn-dbs msvc paths when other compilers are in use

#rnx

diff --git a/Engine/Source/Programs/UnrealBuildTool/Executors/SNDBS.cs b/Engine/Source/Programs/UnrealBuildTool/Executors/SNDBS.cs
--- a/Engine/Source/Programs/UnrealBuildTool/Executors/SNDBS.cs
+++ b/Engine/Source/Programs/UnrealBuildTool/Executors/SNDBS.cs
@@ -462,26 +462,26 @@ error_report_mode=prompt
 response_file_content_pattern=\s--\s"".*?cl\.exe""\s(.*)
 
 [group]
-server={VC_COMPILER_DIR}\mspdbsrv.exe
+server={VC_TOOLCHAIN_DIR}\mspdbsrv.exe
 
 [files]
 main=cl-filter.exe
-file01={VC_COMPILER_DIR}\c1.dll
-file01={VC_COMPILER_DIR}\c1ui.dll
-file02={VC_COMPILER_DIR}\c1xx.dll
-file03={VC_COMPILER_DIR}\c2.dll
-file04={VC_COMPILER_DIR}\mspdb140.dll
-file05={VC_COMPILER_DIR}\mspdbcore.dll
-file06={VC_COMPILER_DIR}\mspdbsrv.exe
-file07={VC_COMPILER_DIR}\mspft140.dll
-file08={VC_COMPILER_DIR}\vcmeta.dll
-file09={VC_COMPILER_DIR}\*\clui.dll
-file10={VC_COMPILER_DIR}\*\mspft140ui.dll
-file11={VC_COMPILER_DIR}\localespc.dll
-file12={VC_COMPILER_DIR}\cppcorecheck.dll
-file13={VC_COMPILER_DIR}\experimentalcppcorecheck.dll
-file14={VC_COMPILER_DIR}\espxengine.dll
-file15={VC_COMPILER_DIR}\c1.exe
+file01={VC_TOOLCHAIN_DIR}\c1.dll
+file01={VC_TOOLCHAIN_DIR}\c1ui.dll
+file02={VC_TOOLCHAIN_DIR}\c1xx.dll
+file03={VC_TOOLCHAIN_DIR}\c2.dll
+file04={VC_TOOLCHAIN_DIR}\mspdb140.dll
+file05={VC_TOOLCHAIN_DIR}\mspdbcore.dll
+file06={VC_TOOLCHAIN_DIR}\mspdbsrv.exe
+file07={VC_TOOLCHAIN_DIR}\mspft140.dll
+file08={VC_TOOLCHAIN_DIR}\vcmeta.dll
+file09={VC_TOOLCHAIN_DIR}\*\clui.dll
+file10={VC_TOOLCHAIN_DIR}\*\mspft140ui.dll
+file11={VC_TOOLCHAIN_DIR}\localespc.dll
+file12={VC_TOOLCHAIN_DIR}\cppcorecheck.dll
+file13={VC_TOOLCHAIN_DIR}\experimentalcppcorecheck.dll
+file14={VC_TOOLCHAIN_DIR}\espxengine.dll
+file15={VC_TOOLCHAIN_DIR}\c1.exe
 
 [output-file-patterns]
 outputfile01=\s*""([^ "",]+\.cpp\.txt)\""
@@ -514,25 +514,25 @@ force_synchronous_pdb_writes=true
 error_report_mode=prompt
 
 [group]
-server={VC_COMPILER_DIR}\mspdbsrv.exe
+server={VC_TOOLCHAIN_DIR}\mspdbsrv.exe
 
 [files]
-main={VC_COMPILER_DIR}\cl.exe
-file01={VC_COMPILER_DIR}\c1.dll
-file01={VC_COMPILER_DIR}\c1ui.dll
-file02={VC_COMPILER_DIR}\c1xx.dll
-file03={VC_COMPILER_DIR}\c2.dll
-file04={VC_COMPILER_DIR}\mspdb140.dll
-file05={VC_COMPILER_DIR}\mspdbcore.dll
-file06={VC_COMPILER_DIR}\mspdbsrv.exe
-file07={VC_COMPILER_DIR}\mspft140.dll
-file08={VC_COMPILER_DIR}\vcmeta.dll
-file09={VC_COMPILER_DIR}\*\clui.dll
-file10={VC_COMPILER_DIR}\*\mspft140ui.dll
-file11={VC_COMPILER_DIR}\localespc.dll
-file12={VC_COMPILER_DIR}\cppcorecheck.dll
-file13={VC_COMPILER_DIR}\experimentalcppcorecheck.dll
-file14={VC_COMPILER_DIR}\espxengine.dll
+main={VC_TOOLCHAIN_DIR}\cl.exe
+file01={VC_TOOLCHAIN_DIR}\c1.dll
+file01={VC_TOOLCHAIN_DIR}\c1ui.dll
+file02={VC_TOOLCHAIN_DIR}\c1xx.dll
+file03={VC_TOOLCHAIN_DIR}\c2.dll
+file04={VC_TOOLCHAIN_DIR}\mspdb140.dll
+file05={VC_TOOLCHAIN_DIR}\mspdbcore.dll
+file06={VC_TOOLCHAIN_DIR}\mspdbsrv.exe
+file07={VC_TOOLCHAIN_DIR}\mspft140.dll
+file08={VC_TOOLCHAIN_DIR}\vcmeta.dll
+file09={VC_TOOLCHAIN_DIR}\*\clui.dll
+file10={VC_TOOLCHAIN_DIR}\*\mspft140ui.dll
+file11={VC_TOOLCHAIN_DIR}\localespc.dll
+file12={VC_TOOLCHAIN_DIR}\cppcorecheck.dll
+file13={VC_TOOLCHAIN_DIR}\experimentalcppcorecheck.dll
+file14={VC_TOOLCHAIN_DIR}\espxengine.dll
 
 [output-file-patterns]
 outputfile01=\s*""([^ "",]+\.cpp\.txt\.json)\""
@@ -561,7 +561,7 @@ filter06=concrt140*.dll",
 use_cache=no
 
 [files]
-main={VC_COMPILER_DIR}\mspdbsrv.exe
+main={VC_TOOLCHAIN_DIR}\mspdbsrv.exe
 
 [openmp]
 omp=true
diff --git a/Engine/Source/Programs/UnrealBuildTool/Platform/Windows/VCEnvironment.cs b/Engine/Source/Programs/UnrealBuildTool/Platform/Windows/VCEnvironment.cs
--- a/Engine/Source/Programs/UnrealBuildTool/Platform/Windows/VCEnvironment.cs
+++ b/Engine/Source/Programs/UnrealBuildTool/Platform/Windows/VCEnvironment.cs
@@ -152,6 +152,7 @@ namespace UnrealBuildTool
 			// Add the compiler path and directory as environment variables for the process so they may be used elsewhere.
 			Environment.SetEnvironmentVariable("VC_COMPILER_PATH", CompilerPath.FullName, EnvironmentVariableTarget.Process);
 			Environment.SetEnvironmentVariable("VC_COMPILER_DIR", CompilerPath.Directory.FullName, EnvironmentVariableTarget.Process);
+			Environment.SetEnvironmentVariable("VC_TOOLCHAIN_DIR", ToolchainCompilerPath.Directory.FullName, EnvironmentVariableTarget.Process);
 
 			DirectoryReference.AddDirectoryToPath(GetVCToolPath(ToolChainDir, Architecture));
 			if (Architecture == UnrealArch.Arm64)
