[UBA]
* Changed name of maxcores, startupdelay and allowwine in uba horde xml config. Now they should be under <Horde> and be called <MaxCores>, <StartupDelay> and <AllowWine>

diff --git a/Engine/Source/Programs/UnrealBuildTool/Executors/UnrealBuildAccelerator/UnrealBuildAcceleratorHordeConfig.cs b/Engine/Source/Programs/UnrealBuildTool/Executors/UnrealBuildAccelerator/UnrealBuildAcceleratorHordeConfig.cs
--- a/Engine/Source/Programs/UnrealBuildTool/Executors/UnrealBuildAccelerator/UnrealBuildAcceleratorHordeConfig.cs
+++ b/Engine/Source/Programs/UnrealBuildTool/Executors/UnrealBuildAccelerator/UnrealBuildAcceleratorHordeConfig.cs
@@ -138,8 +138,8 @@ namespace UnrealBuildTool
 		/// <summary>
 		/// Max cores allowed to be used by build session
 		/// </summary>
-		[ConfigFile(ConfigHierarchyType.Engine, "Horde", "UbaMaxCores")]
-		[XmlConfigFile(Category = "UnrealBuildAccelerator")]
+		[ConfigFile(ConfigHierarchyType.Engine, "Horde", "MaxCores")]
+		[XmlConfigFile(Category = "Horde", Name = "MaxCores")]
 		[CommandLine("-UBAHordeMaxCores")]
 		public int HordeMaxCores { get; set; } = 576;
 
@@ -147,7 +147,7 @@ namespace UnrealBuildTool
 		/// How long UBT should wait to ask for help. Useful in build configs where machine can delay remote work and still get same wall time results (pch dependencies etc)
 		/// </summary>
 		[ConfigFile(ConfigHierarchyType.Engine, "Horde", "UbaDelay")]
-		[XmlConfigFile(Category = "UnrealBuildAccelerator")]
+		[XmlConfigFile(Category = "Horde", Name = "StartupDelay")]
 		[CommandLine("-UBAHordeDelay")]
 		public int HordeDelay { get; set; } = 0;
 
@@ -155,7 +155,7 @@ namespace UnrealBuildTool
 		/// Allow use of Wine. Only applicable to Horde agents running Linux. Can still be ignored if Wine executable is not set on agent.
 		/// </summary>
 		[ConfigFile(ConfigHierarchyType.Engine, "Horde", "UbaAllowWine")]
-		[XmlConfigFile(Category = "UnrealBuildAccelerator")]
+		[XmlConfigFile(Category = "Horde", Name = "AllowWine")]
 		[CommandLine("-UBAHordeAllowWine", Value = "true")]
 		public bool bHordeAllowWine { get; set; } = true;
 
