UnrealBuildTool: Ensure extension directory targets & modules write intermediates to the correct path

#rnx

diff --git a/Engine/Source/Programs/UnrealBuildTool/Configuration/UEBuildTarget.cs b/Engine/Source/Programs/UnrealBuildTool/Configuration/UEBuildTarget.cs
--- a/Engine/Source/Programs/UnrealBuildTool/Configuration/UEBuildTarget.cs
+++ b/Engine/Source/Programs/UnrealBuildTool/Configuration/UEBuildTarget.cs
@@ -1819,7 +1819,7 @@ namespace UnrealBuildTool
 
 			// track where intermediates and receipt files will be saved under (Programs with .uproject files still want to put the receipt
 			// under Engine/Binaries, not Engine/Programs/Foo/Binaries
-			DirectoryReference OutputRootDirectory = Unreal.EngineDirectory;
+			DirectoryReference OutputRootDirectory = GetOutputDirectoryForExecutable(Unreal.EngineDirectory, Rules.File);
 
 			// Figure out what the project directory is. If we have a uproject file, use that. Otherwise use the engine directory.
 			if (ProjectFile != null)
@@ -4354,6 +4354,10 @@ namespace UnrealBuildTool
 		{
 			// Get the root output directory and base name (target name/app name) for this binary
 			DirectoryReference BaseOutputDirectory = GetBaseOutputDirectory(ModuleRules);
+			if (BaseOutputDirectory == Unreal.EngineDirectory)
+			{
+				BaseOutputDirectory = GetOutputDirectoryForExecutable(BaseOutputDirectory, ModuleRules.File);
+			}
 
 			// Get the configuration that this module will be built in. Engine modules compiled in DebugGame will use Development.
 			UnrealTargetConfiguration ModuleConfiguration = Configuration;
@@ -5053,15 +5057,9 @@ namespace UnrealBuildTool
 			UEBuildModuleCPP LaunchModule = FindOrCreateCppModuleByName(Rules.LaunchModuleName, TargetRulesFile.GetFileName(), Logger);
 
 			// Get the intermediate directory for the launch module directory. This can differ from the standard engine intermediate directory because it is always configuration-specific.
-			DirectoryReference IntermediateDirectory;
-			if (LaunchModule.RulesFile.IsUnderDirectory(Unreal.EngineDirectory) && !ShouldCompileMonolithic())
-			{
-				IntermediateDirectory = DirectoryReference.Combine(Unreal.EngineDirectory, PlatformIntermediateFolder, UEBuildTarget.GetTargetIntermediateFolderName(AppName, IntermediateEnvironment), Configuration.ToString());
-			}
-			else
-			{
-				IntermediateDirectory = ProjectIntermediateDirectory;
-			}
+			DirectoryReference IntermediateDirectory = LaunchModule.RulesFile.IsUnderDirectory(Unreal.EngineDirectory) && !ShouldCompileMonolithic()
+				? DirectoryReference.Combine(EngineIntermediateDirectory.ParentDirectory!, Configuration.ToString())
+				: ProjectIntermediateDirectory;
 
 			// Construct the output paths for this target's executable
 			DirectoryReference OutputDirectory;
