ModifyConfig action in buildgraph is not preserving quotations on original value.

diff --git a/Engine/Source/Programs/UnrealBuildTool.Tests/ConfigFileTests.cs b/Engine/Source/Programs/UnrealBuildTool.Tests/ConfigFileTests.cs
new file mode 100644
--- /dev/null
+++ b/Engine/Source/Programs/UnrealBuildTool.Tests/ConfigFileTests.cs
@@ -0,0 +1,82 @@
+// Copyright Epic Games, Inc. All Rights Reserved.
+
+using System;
+using System.IO;
+using EpicGames.Core;
+using Microsoft.VisualStudio.TestTools.UnitTesting;
+using UnrealBuildTool;
+
+#nullable enable
+
+namespace UnrealBuildToolTests
+{
+
+	[TestClass]
+	public class ConfigFileTests
+	{
+		#region -- Test Context --
+
+		private const string TEST_FILE_NAME = "ConfigFileTest.ini";
+
+		private class TestConfigFile : IDisposable
+		{
+			internal string? TemporaryFile { get; private set; }
+
+			internal TestConfigFile(string FileName, string Contents)
+			{
+				try
+				{
+					TemporaryFile = Path.Combine(System.IO.Path.GetTempPath(), Guid.NewGuid().ToString() + "_" + FileName);
+					using (StreamWriter writer = new StreamWriter(TemporaryFile))
+					{
+						writer.Write(Contents);
+					}
+				}
+				catch (Exception)
+				{
+					TemporaryFile = null;
+				}
+			}
+
+			public void Dispose()
+			{
+				if (!String.IsNullOrEmpty(TemporaryFile) && File.Exists(TemporaryFile))
+				{
+					File.Delete(TemporaryFile);
+				}
+			}
+		}
+
+		#endregion
+
+		[TestMethod]
+		public void ConfigFileTestsQuotations()
+		{
+			const string TestSection = "MySection";
+			const string TestValueKey = "Name";
+			const string TestValue = "\"https://www.unrealengine.com/\"";
+
+			string TestContents = String.Format("[{0}]\r\n{1}={2}", TestSection, TestValueKey, TestValue);
+
+			using (TestConfigFile TemporaryConfigFile = new TestConfigFile(TEST_FILE_NAME, TestContents))
+			{
+				string? TestFilePath = TemporaryConfigFile.TemporaryFile;
+
+				Assert.IsFalse(String.IsNullOrEmpty(TestFilePath), "Unable to generate a test file.");
+
+				ConfigFile ConfigFile;
+				FileReference ConfigFileLocation = new FileReference(TestFilePath);
+
+				Assert.IsTrue(FileReference.Exists(ConfigFileLocation), "Unable to acquire a test file.");
+				
+				ConfigFile = new ConfigFile(ConfigFileLocation);
+
+				FileReference.MakeWriteable(ConfigFileLocation);
+				ConfigFile.Write(ConfigFileLocation);
+
+				string FileContents = File.ReadAllText(TemporaryConfigFile.TemporaryFile);
+				Assert.AreEqual(TestContents.Trim(), FileContents.Trim(), "The file contents after writing do not match the expected contents.");
+			}
+		}
+	}
+}
diff --git a/Engine/Source/Programs/UnrealBuildTool/System/ConfigFile.cs b/Engine/Source/Programs/UnrealBuildTool/System/ConfigFile.cs
--- a/Engine/Source/Programs/UnrealBuildTool/System/ConfigFile.cs
+++ b/Engine/Source/Programs/UnrealBuildTool/System/ConfigFile.cs
@@ -56,6 +56,11 @@ namespace UnrealBuildTool
 		/// </summary>
 		public string Value;
 
+		/// <summary>
+		/// Whether to preserve enclosing quotations around <see cref="Value"/>
+		/// </summary>
+		internal bool ShouldPreserveValueQuotations = false;
+
 		/// <summary>
 		/// Constructor.
 		/// </summary>
@@ -69,6 +74,18 @@ namespace UnrealBuildTool
 			this.Value = Value;
 		}
 
+		/// <summary>
+		/// Constructor.
+		/// </summary>
+		/// <param name="Action">Action to take when merging this key/value pair with an existing value</param>
+		/// <param name="Key">Name of the key to modify</param>
+		/// <param name="Value">Value to assign</param>
+		/// <param name="ShouldPreserveValueQuotations">Whether to preserve enclosing quotations around <see cref="Value"/></param>
+		public ConfigLine(ConfigLineAction Action, string Key, string Value, bool ShouldPreserveValueQuotations) : this(Action, Key, Value)
+		{
+			this.ShouldPreserveValueQuotations = ShouldPreserveValueQuotations;
+		}
+
 		/// <summary>
 		/// Formats this object for the debugger
 		/// </summary>
@@ -416,10 +433,13 @@ namespace UnrealBuildTool
 				}
 			}
 
+			bool HasStrippedValueQuotes = false;
+
 			// Strip quotes around the value if present
 			int ValueEndIdx = EndIdx;
 			if (ValueEndIdx >= ValueStartIdx + 2 && Line[ValueStartIdx] == '"' && Line[ValueEndIdx - 1] == '"')
 			{
+				HasStrippedValueQuotes = true;
 				ValueStartIdx++;
 				ValueEndIdx--;
 			}
@@ -443,12 +463,13 @@ namespace UnrealBuildTool
 				}
 
 				string KeyName = NewKey.Substring(NewKey.IndexOf(':') + 1);
-				CurrentSection.Lines.Add(new ConfigLine(Action, KeyName, Value));
+
+				CurrentSection.Lines.Add(new ConfigLine(Action, KeyName, Value, HasStrippedValueQuotes));
 
 				return true;
 			}
 
-			Section.Lines.Add(new ConfigLine(Action, NewKey, Value));
+			Section.Lines.Add(new ConfigLine(Action, NewKey, Value, HasStrippedValueQuotes));
 			return true;
 		}
 
@@ -497,6 +518,10 @@ namespace UnrealBuildTool
 					Writer.WriteLine("[{0}]", Section.Name);
 					foreach (ConfigLine Line in Section.Lines)
 					{
+						if (Line.ShouldPreserveValueQuotations)
+						{
+							Line.Value = $"\"{Line.Value}\"";
+						}
 						Writer.WriteLine("{0}", Line.ToString());
 					}
 					Writer.WriteLine();
