bEnableNoInlinedWarnings=true causes crash under specific module fileset

diff --git a/Engine/Source/Programs/UnrealBuildTool/Configuration/UEBuildModuleCPP.cs b/Engine/Source/Programs/UnrealBuildTool/Configuration/UEBuildModuleCPP.cs
--- a/Engine/Source/Programs/UnrealBuildTool/Configuration/UEBuildModuleCPP.cs
+++ b/Engine/Source/Programs/UnrealBuildTool/Configuration/UEBuildModuleCPP.cs
@@ -575,18 +575,21 @@ namespace UnrealBuildTool
 						Dictionary<string, FileItem> CPPFilesLookup = new Dictionary<string, FileItem>();
 						foreach (FileItem CPPFile in CPPFiles)
 						{
-							CPPFilesLookup.Add(Utils.GetFilenameWithoutAnyExtensions(CPPFile.Name), CPPFile);
+							CPPFilesLookup.Add(CPPFile.Name, CPPFile);
 						}
+
 						foreach (string Name in GeneratedFiles.Keys)
 						{
 							if (!Name.StartsWith("UHT/"))
 							{
 								continue;
 							}
+
 							string NameWithoutPrefix = Name.Substring(4);
-							if (CPPFilesLookup.TryGetValue(NameWithoutPrefix, out FileItem? Item))
+							string NameWithoutPrefixWithSuffix = NameWithoutPrefix + ".cpp";
+							if (CPPFilesLookup.TryGetValue(NameWithoutPrefixWithSuffix, out FileItem? Item))
 							{
-								Logger.LogWarning("'{0}' .gen.cpp not inlined. Add '#include UE_INLINE_GENERATED_CPP_BY_NAME({1})'", Item.Name, NameWithoutPrefix);
+								Logger.LogWarning("{0}(1): .gen.cpp not inlined. Add '#include UE_INLINE_GENERATED_CPP_BY_NAME({1})' after header includes.", Item.AbsolutePath, NameWithoutPrefix);
 							}
 						}
 					}
