Add support for -fsanitize-recover for ASAN which allows builds to continue on error

#tests Tested ASAN Dedicated Server builds on Linux
[RN] Linux ASAN builds will now include fsanitize-recover to allow continue on error

diff --git a/Engine/Source/Programs/UnrealBuildTool/Platform/Linux/LinuxToolChain.cs b/Engine/Source/Programs/UnrealBuildTool/Platform/Linux/LinuxToolChain.cs
--- a/Engine/Source/Programs/UnrealBuildTool/Platform/Linux/LinuxToolChain.cs
+++ b/Engine/Source/Programs/UnrealBuildTool/Platform/Linux/LinuxToolChain.cs
@@ -721,6 +721,7 @@ namespace UnrealBuildTool
 				if (Options.HasFlag(ClangToolChainOptions.EnableAddressSanitizer))
 				{
 					Arguments.Add("-fsanitize=address");
+					Arguments.Add("-fsanitize-recover=address"); 
 				}
 				else if (Options.HasFlag(ClangToolChainOptions.EnableThreadSanitizer))
 				{
diff --git a/Engine/Source/Programs/UnrealBuildTool/ToolChain/ClangToolChain.cs b/Engine/Source/Programs/UnrealBuildTool/ToolChain/ClangToolChain.cs
--- a/Engine/Source/Programs/UnrealBuildTool/ToolChain/ClangToolChain.cs
+++ b/Engine/Source/Programs/UnrealBuildTool/ToolChain/ClangToolChain.cs
@@ -728,6 +728,7 @@ namespace UnrealBuildTool
 			if (Options.HasFlag(ClangToolChainOptions.EnableAddressSanitizer))
 			{
 				Arguments.Add("-fsanitize=address");
+				Arguments.Add("-fsanitize-recover=address");
 			}
 
 			// TSan
