[UBT]
* Removed use of fat objects for merged-modules ltcg builds now when UbaObjTool handles LLVM IR Stream format

diff --git a/Engine/Source/Programs/UnrealBuildTool/ToolChain/ClangToolChain.cs b/Engine/Source/Programs/UnrealBuildTool/ToolChain/ClangToolChain.cs
--- a/Engine/Source/Programs/UnrealBuildTool/ToolChain/ClangToolChain.cs
+++ b/Engine/Source/Programs/UnrealBuildTool/ToolChain/ClangToolChain.cs
@@ -675,11 +675,6 @@ namespace UnrealBuildTool
 				Arguments.Add("-fprofile-generate");
 			}
 
-			if (bMergeModules && (Options.HasFlag(ClangToolChainOptions.EnableLinkTimeOptimization) || CompileEnvironment.bAllowLTCG || CompileEnvironment.bPGOOptimize || CompileEnvironment.bPGOProfile))
-			{
-				Arguments.Add("-ffat-lto-objects");
-			}
-
 			if (!CompileEnvironment.bUseInlining)
 			{
 				Arguments.Add("-fno-inline-functions");
