[UBT]
* Added /T:<platform> when running UbaObjTool to generate extra obj file.. this is needed to know if tool should spit out coff or elf file (when running with clang and ltcg the obj files are of LLVM IR stream format regardless of platform which makes it impossible to figure out automatically what format the extra obj file needs to be)

diff --git a/Engine/Source/Programs/UnrealBuildTool/Configuration/UEBuildTarget.cs b/Engine/Source/Programs/UnrealBuildTool/Configuration/UEBuildTarget.cs
--- a/Engine/Source/Programs/UnrealBuildTool/Configuration/UEBuildTarget.cs
+++ b/Engine/Source/Programs/UnrealBuildTool/Configuration/UEBuildTarget.cs
@@ -3050,6 +3050,7 @@ namespace UnrealBuildTool
 
 				FileReference extraObj = FileReference.Combine(IntermediateDirectory, $"{Name}.extra.{(CompileEnvironment.Platform.IsInGroup(UnrealPlatformGroup.Microsoft) ? "obj" : "o")}");
 				arguments.Add($"/O:{extraObj}");
+				arguments.Add($"/T:{Platform}");
 
 				MakefileBuilder.CreateIntermediateTextFile(stripRsp, arguments);
 
