[UBT]
* Ignore tracking the .exp file when using clang linker to prevent relinks happening when building on certain microsoft platforms

diff --git a/Engine/Source/Programs/UnrealBuildTool/Platform/Windows/VCToolChain.cs b/Engine/Source/Programs/UnrealBuildTool/Platform/Windows/VCToolChain.cs
--- a/Engine/Source/Programs/UnrealBuildTool/Platform/Windows/VCToolChain.cs
+++ b/Engine/Source/Programs/UnrealBuildTool/Platform/Windows/VCToolChain.cs
@@ -2731,7 +2731,7 @@ namespace UnrealBuildTool
 			// For import libraries and exports generated by cross-referenced builds, we don't track output files. VS 15.3+ doesn't touch timestamps for libs
 			// and exp files with no modifications, breaking our dependency checking, but incremental linking will fall back to a full link if we delete it.
 			// Since all DLLs are typically marked as cross referenced now anyway, we can just ignore this file to allow incremental linking to work.
-			if (LinkEnvironment.bHasExports && !LinkEnvironment.bIsBuildingLibrary && !LinkEnvironment.bIsCrossReferenced)
+			if (LinkEnvironment.bHasExports && !LinkEnvironment.bIsBuildingLibrary && !LinkEnvironment.bIsCrossReferenced && !Target.WindowsPlatform.bAllowClangLinker)
 			{
 				FileReference ExportFilePath = ImportLibraryFilePath!.ChangeExtension(".exp");
 				FileItem ExportFile = FileItem.GetItemByFileReference(ExportFilePath);
