[UBT]
* Changed so when building and linking with clang we use UbaObjTool to generate import libs instead of using msvc link.exe. This is needed because obj files can be of LLVM IR Stream format if a target has ltcg enabled (link.exe can't read those obj files)

diff --git a/Engine/Source/Programs/UnrealBuildTool/Platform/Windows/VCEnvironment.cs b/Engine/Source/Programs/UnrealBuildTool/Platform/Windows/VCEnvironment.cs
--- a/Engine/Source/Programs/UnrealBuildTool/Platform/Windows/VCEnvironment.cs
+++ b/Engine/Source/Programs/UnrealBuildTool/Platform/Windows/VCEnvironment.cs
@@ -236,16 +236,11 @@ namespace UnrealBuildTool
 		/// </summary>
 		FileReference GetLibraryLinkerToolPath(WindowsCompiler Compiler, UnrealArch Architecture, DirectoryReference CompilerDir, DirectoryReference ToochainDir)
 		{
-			if ((Compiler == WindowsCompiler.Clang || Compiler == WindowsCompiler.ClangRTFM) && bAllowClangLinker)
-			{
-				// @todo: lld-link is not currently working for building .lib
-				//return FileReference.Combine(CompilerDir, "bin", "lld-link.exe");
-			}
-			else if (Compiler == WindowsCompiler.Intel && bAllowClangLinker)
+			if (bAllowClangLinker && Compiler.IsClang())
 			{
-				// @todo: lld-link is not currently working for building .lib
-				//return FileReference.Combine(CompilerDir, "bin", "compiler", "lld-link.exe");
-				//return FileReference.Combine(CompilerDir, "bin", "xilib.exe");
+				// Since obj files could be LLVM IR Stream file format (when building with ltcg) we can't use link.exe. UbaObjTool support all formats and produce identical lib files as link.exe when running non-ltcg
+				DirectoryReference Dir = DirectoryReference.Combine(UnrealBuildBase.Unreal.EngineDirectory, "Binaries", "Win64", "UnrealBuildAccelerator", Architecture.ToString());
+				return FileReference.Combine(Dir, "UbaObjTool.exe");
 			}
 			return FileReference.Combine(GetVCToolPath(ToochainDir, Architecture), "link.exe"); // We add /LIB to cmd line so we can use link.exe directly instead of going via lib.exe
 		}
