UnrealBuildTool: dotnet fixes from VSCode generated solution

* Set environment variables to force the bundled dotnet
* Run dotnet apps directly when debugging
* Compile automation scripts before debugging

#rnx

diff --git a/Engine/Source/Programs/UnrealBuildTool/ProjectFiles/VisualStudioCode/VSCodeProjectFileGenerator.cs b/Engine/Source/Programs/UnrealBuildTool/ProjectFiles/VisualStudioCode/VSCodeProjectFileGenerator.cs
--- a/Engine/Source/Programs/UnrealBuildTool/ProjectFiles/VisualStudioCode/VSCodeProjectFileGenerator.cs
+++ b/Engine/Source/Programs/UnrealBuildTool/ProjectFiles/VisualStudioCode/VSCodeProjectFileGenerator.cs
@@ -1019,14 +1019,7 @@ namespace UnrealBuildTool
 						{
 							OutFile.AddField("label", TaskName);
 							OutFile.AddField("group", "build");
-							if (!RuntimePlatform.IsWindows)
-							{
-								OutFile.AddField("command", MakePathString(FileReference.Combine(ProjectRoot, "Engine", "Build", "BatchFiles", "RunDotnet.sh")));
-							}
-							else
-							{
-								OutFile.AddField("command", "dotnet");
-							}
+							OutFile.AddField("command", MakePathString(Unreal.DotnetPath));
 							OutFile.BeginArray("args");
 							{
 								OutFile.AddUnnamedField(Command.ToLower());
@@ -1036,16 +1029,78 @@ namespace UnrealBuildTool
 								OutFile.AddUnnamedField(MakeUnquotedPathString(BuildProduct.CSharpInfo!.ProjectPath, EPathType.Absolute));
 							}
 							OutFile.EndArray();
+
+							OutFile.AddField("problemMatcher", "$msCompile");
+							OutFile.AddField("type", "shell");
+
+							OutFile.BeginObject("options");
+							{
+								OutFile.AddField("cwd", MakeUnquotedPathString(ProjectRoot, EPathType.Absolute));
+								OutFile.BeginObject("env");
+								{
+									OutFile.AddField("PATH", $"{Unreal.DotnetDirectory}{Path.PathSeparator}${{env:PATH}}");;
+									OutFile.AddField("DOTNET_MULTILEVEL_LOOKUP", "0");
+									OutFile.AddField("DOTNET_ROLL_FORWARD", "LatestMajor");
+								}
+								OutFile.EndObject();
+							}
+							OutFile.EndObject();
 						}
-						OutFile.AddField("problemMatcher", "$msCompile");
-						OutFile.AddField("type", "shell");
+						OutFile.EndObject();
+					}
+				}
+			}
+
+			if (InProject.Name.Equals("AutomationTool", StringComparison.OrdinalIgnoreCase)) {
+				WriteAutomationScriptTask(InProject, OutFile);
+			}
+		}
+
+		private void WriteAutomationScriptTask(ProjectData.Project InProject, JsonFile OutFile)
+		{
+			string[] Commands = { "Build" };
+
+			foreach (ProjectData.Target Target in InProject.Targets)
+			{
+				foreach (ProjectData.BuildProduct BuildProduct in Target.BuildProducts.Where(x => x.OutputType == ProjectData.EOutputType.Exe))
+				{
+					foreach (string Command in Commands)
+					{
+						string TaskName = String.Format("{0}+Scripts {1} {2} {3}", Target.Name, BuildProduct.Platform, BuildProduct.Config, Command);
+						string DependsOnTaskName = String.Format("{0} {1} {2} {3}", Target.Name, BuildProduct.Platform, BuildProduct.Config, Command);
 
-						OutFile.BeginObject("options");
+						OutFile.BeginObject();
 						{
-							OutFile.AddField("cwd", MakeUnquotedPathString(ProjectRoot, EPathType.Absolute));
-						}
+							OutFile.AddField("label", TaskName);
+							OutFile.AddField("group", "build");
+							OutFile.AddField("command", MakePathString(BuildProduct.OutputFile));
+							OutFile.BeginArray("args");
+							{
+								OutFile.AddUnnamedField("-list");
+							}
+							OutFile.EndArray();
+							OutFile.BeginArray("dependsOn");
+							{
+								OutFile.AddUnnamedField(DependsOnTaskName);
+							}
+							OutFile.EndArray();
 
-						OutFile.EndObject();
+							OutFile.AddField("problemMatcher", "$msCompile");
+							OutFile.AddField("type", "shell");
+
+							OutFile.BeginObject("options");
+							{
+								OutFile.AddField("cwd", MakeUnquotedPathString(ProjectRoot, EPathType.Absolute));
+								OutFile.BeginObject("env");
+								{
+									OutFile.AddField("PATH", $"{Unreal.DotnetDirectory}{Path.PathSeparator}${{env:PATH}}");
+									OutFile.AddField("DOTNET_MULTILEVEL_LOOKUP", "0");
+									OutFile.AddField("DOTNET_ROLL_FORWARD", "LatestMajor");
+								}
+								OutFile.EndObject();
+							}
+							OutFile.EndObject();
+						}
 						OutFile.EndObject();
 					}
 				}
@@ -1410,6 +1465,15 @@ namespace UnrealBuildTool
 					}
 				}
 				OutFile.EndArray();
+				OutFile.BeginObject("env");
+				{
+					OutFile.AddField("PATH", $"{Unreal.DotnetDirectory}{Path.PathSeparator}${{env:PATH}}");
+					OutFile.AddField("DOTNET_ROOT", Unreal.DotnetDirectory.FullName);
+					OutFile.AddField("DOTNET_HOST_PATH", Unreal.DotnetPath.FullName);
+					OutFile.AddField("DOTNET_MULTILEVEL_LOOKUP", "0");
+					OutFile.AddField("DOTNET_ROLL_FORWARD", "LatestMajor");
+				}
+				OutFile.EndObject();
 				if (HostPlatform == UnrealTargetPlatform.Win64)
 				{
 					OutFile.AddField("console", "integratedTerminal");
@@ -1437,8 +1501,13 @@ namespace UnrealBuildTool
 					{
 						string TaskName = String.Format("{0} ({1})", Target.Name, BuildProduct.Config);
 						string BuildTaskName = String.Format("{0} {1} {2} Build", Target.Name, HostPlatform, BuildProduct.Config);
+						string[]? Args = null;
+						if (Target.Name.Equals("AutomationTool", StringComparison.OrdinalIgnoreCase)) {
+							BuildTaskName = String.Format("{0}+Scripts {1} {2} Build", Target.Name, HostPlatform, BuildProduct.Config);
+							Args = new string[] {"-NoCompile"};
+						}
 
-						WriteSingleCSharpLaunchConfig(OutFile, TaskName, BuildTaskName, BuildProduct.OutputFile, null);
+						WriteSingleCSharpLaunchConfig(OutFile, TaskName, BuildTaskName, BuildProduct.OutputFile, Args);
 					}
 				}
 			}
@@ -1624,6 +1693,16 @@ namespace UnrealBuildTool
 				WorkspaceFile.AddField("typescript.tsc.autoDetect", "off");
 				// disable npm script autodetect to avoid lag populating tasks list 
 				WorkspaceFile.AddField("npm.autoDetect", "off");
+				// set environment variables so correct dotnet is used
+				WorkspaceFile.BeginObject($"terminal.integrated.env.{(OperatingSystem.IsWindows() ? "windows" : OperatingSystem.IsMacOS() ? "mac" : "linux")}");
+				{
+					WorkspaceFile.AddField("PATH", $"{Unreal.DotnetDirectory}{Path.PathSeparator}${{env:PATH}}");
+					WorkspaceFile.AddField("DOTNET_ROOT", Unreal.DotnetDirectory.FullName);
+					WorkspaceFile.AddField("DOTNET_HOST_PATH", Unreal.DotnetPath.FullName);
+					WorkspaceFile.AddField("DOTNET_MULTILEVEL_LOOKUP", "0");
+					WorkspaceFile.AddField("DOTNET_ROLL_FORWARD", "LatestMajor");
+				}
+				WorkspaceFile.EndObject();
 
 				if (bUseVSCodeExtension)
 				{
