[UBT]
* Enabled uba for linux lto builds. It was disabled half a year ago because we saw some long link times. Lots of fixes have gone in

diff --git a/Engine/Source/Programs/UnrealBuildTool/Platform/Linux/LinuxToolChain.cs b/Engine/Source/Programs/UnrealBuildTool/Platform/Linux/LinuxToolChain.cs
--- a/Engine/Source/Programs/UnrealBuildTool/Platform/Linux/LinuxToolChain.cs
+++ b/Engine/Source/Programs/UnrealBuildTool/Platform/Linux/LinuxToolChain.cs
@@ -1174,9 +1174,6 @@ namespace UnrealBuildTool
 				LinkAction.CommandDescription = "Link";
 			}
 
-			// Saw a 6 hour link time potentially caused by box. Will disable for now and revisit later
-			LinkAction.bCanExecuteInUBA = !LinkEnvironment.bPGOProfile && !LinkEnvironment.bPGOOptimize && !LinkEnvironment.bAllowLTCG;
-
 			if (!OperatingSystem.IsWindows())
 			{
 				LinkAction.bCanExecuteInUBA = false; // Linker on native linux uses vfork/exec which is not handled in uba right now
