Nastepfilters re-added back into Visual Studio project file generation.

diff --git a/Engine/Source/Programs/UnrealBuildTool/ProjectFiles/VisualStudio/VCProjectFileGenerator.cs b/Engine/Source/Programs/UnrealBuildTool/ProjectFiles/VisualStudio/VCProjectFileGenerator.cs
--- a/Engine/Source/Programs/UnrealBuildTool/ProjectFiles/VisualStudio/VCProjectFileGenerator.cs
+++ b/Engine/Source/Programs/UnrealBuildTool/ProjectFiles/VisualStudio/VCProjectFileGenerator.cs
@@ -793,11 +793,13 @@ namespace UnrealBuildTool
 
 				// Get the path to the visualizers file. Try to make it relative to the solution directory, but fall back to a full path if it's a foreign project.
 				FileReference VisualizersFile = FileReference.Combine(Unreal.EngineDirectory, "Extras", "VisualStudioDebugging", "Unreal.natvis");
+				FileReference VisualizersStepFile = FileReference.Combine(Unreal.EngineDirectory, "Extras", "VisualStudioDebugging", "Unreal.natstepfilter");
 
 				// Add the visualizers at the solution level. Doesn't seem to be picked up from a makefile project in VS2017 15.8.5.
 				VCSolutionFileContent.AppendLine(String.Format("Project(\"{0}\") = \"Visualizers\", \"Visualizers\", \"{{1CCEC849-CC72-4C59-8C36-2F7C38706D4C}}\"", SolutionFolderEntryGUID));
 				VCSolutionFileContent.AppendLine("\tProjectSection(SolutionItems) = preProject");
 				VCSolutionFileContent.AppendLine("\t\t{0} = {0}", VisualizersFile.MakeRelativeTo(PrimaryProjectPath));
+				VCSolutionFileContent.AppendLine("\t\t{0} = {0}", VisualizersStepFile.MakeRelativeTo(PrimaryProjectPath));
 				VCSolutionFileContent.AppendLine("\tEndProjectSection");
 				VCSolutionFileContent.AppendLine("EndProject");
 			}
