UnrealBuildTool: Fix typo causing an exception when multiple single filenames are the same

#rnx

diff --git a/Engine/Source/Programs/UnrealBuildTool/ToolChain/ClangSpecificFileAction.cs b/Engine/Source/Programs/UnrealBuildTool/ToolChain/ClangSpecificFileAction.cs
--- a/Engine/Source/Programs/UnrealBuildTool/ToolChain/ClangSpecificFileAction.cs
+++ b/Engine/Source/Programs/UnrealBuildTool/ToolChain/ClangSpecificFileAction.cs
@@ -60,7 +60,7 @@ namespace UnrealBuildTool
 			}
 			else
 			{
-				Filename = $"SingleFile/{Path.GetFileNameWithoutExtension(SourceFile.Name)}{SingleFiles[Filename].Count}{Path.GetExtension(SourceFile.Name)}";
+				Filename = $"SingleFile/{Path.GetFileNameWithoutExtension(SourceFile.Name)}{SingleFiles[SourceFile.Name].Count}{Path.GetExtension(SourceFile.Name)}";
 			}
 			SingleFiles[SourceFile.Name].Add(SourceFile);
 
