[UBT]
* Fixed so merged modules work on various console platforms

diff --git a/Engine/Source/Programs/UnrealBuildTool/Configuration/UEBuildBinary.cs b/Engine/Source/Programs/UnrealBuildTool/Configuration/UEBuildBinary.cs
--- a/Engine/Source/Programs/UnrealBuildTool/Configuration/UEBuildBinary.cs
+++ b/Engine/Source/Programs/UnrealBuildTool/Configuration/UEBuildBinary.cs
@@ -878,7 +878,13 @@ namespace UnrealBuildTool
 				bStripUnusedExports = true;
 
 				string Name = OutputFilePaths.First().GetFileNameWithoutExtension();
-				FileReference extraObj = FileReference.Combine(IntermediateDirectory!, $"{Name}.extra.{(Target.Platform.IsInGroup(UnrealPlatformGroup.Microsoft) ? "obj" : "o")}");
+				string Ext = "obj";
+				if (CompileEnvironment.Platform.IsInGroup(UnrealPlatformGroup.Nintendo) || CompileEnvironment.Platform.IsInGroup(UnrealPlatformGroup.Sony))
+					Ext = "ldscript";
+				else if (!CompileEnvironment.Platform.IsInGroup(UnrealPlatformGroup.Microsoft))
+					Ext = "o";
+
+				FileReference extraObj = FileReference.Combine(IntermediateDirectory!, $"{Name}.extra.{Ext}");
 				BinaryLinkEnvironment.InputFiles.Add(FileItem.GetItemByFileReference(extraObj));
 			}
 
diff --git a/Engine/Source/Programs/UnrealBuildTool/Configuration/UEBuildTarget.cs b/Engine/Source/Programs/UnrealBuildTool/Configuration/UEBuildTarget.cs
--- a/Engine/Source/Programs/UnrealBuildTool/Configuration/UEBuildTarget.cs
+++ b/Engine/Source/Programs/UnrealBuildTool/Configuration/UEBuildTarget.cs
@@ -3048,7 +3048,12 @@ namespace UnrealBuildTool
 				arguments.AddRange(toStrip.Select(path => $"/S:{path.Location.ChangeExtension(".exi")}").OrderBy(x => x));
 				arguments.AddRange(otherObjects.Select(path => $"/D:{path.Location.ChangeExtension(".exi")}").OrderBy(x => x));
 
-				FileReference extraObj = FileReference.Combine(IntermediateDirectory, $"{Name}.extra.{(CompileEnvironment.Platform.IsInGroup(UnrealPlatformGroup.Microsoft) ? "obj" : "o")}");
+				string Ext = "obj";
+				if (CompileEnvironment.Platform.IsInGroup(UnrealPlatformGroup.Nintendo) || CompileEnvironment.Platform.IsInGroup(UnrealPlatformGroup.Sony))
+					Ext = "ldscript";
+				else if (!CompileEnvironment.Platform.IsInGroup(UnrealPlatformGroup.Microsoft))
+					Ext = "o";
+				FileReference extraObj = FileReference.Combine(IntermediateDirectory, $"{Name}.extra.{Ext}");
 				arguments.Add($"/O:{extraObj}");
 				arguments.Add($"/T:{Platform}");
 
