[Backout] - CL35794431
Original CL Desc
-----------------------------------------------------------------
[UBT]
* Fixed so merged modules work on various console platforms

diff --git a/Engine/Source/Programs/UnrealBuildTool/Configuration/UEBuildBinary.cs b/Engine/Source/Programs/UnrealBuildTool/Configuration/UEBuildBinary.cs
--- a/Engine/Source/Programs/UnrealBuildTool/Configuration/UEBuildBinary.cs
+++ b/Engine/Source/Programs/UnrealBuildTool/Configuration/UEBuildBinary.cs
@@ -878,13 +878,7 @@ namespace UnrealBuildTool
 				bStripUnusedExports = true;
 
 				string Name = OutputFilePaths.First().GetFileNameWithoutExtension();
-				string Ext = "obj";
-				if (CompileEnvironment.Platform.IsInGroup(UnrealPlatformGroup.Nintendo) || CompileEnvironment.Platform.IsInGroup(UnrealPlatformGroup.Sony))
-					Ext = "ldscript";
-				else if (!CompileEnvironment.Platform.IsInGroup(UnrealPlatformGroup.Microsoft))
-					Ext = "o";
-
-				FileReference extraObj = FileReference.Combine(IntermediateDirectory!, $"{Name}.extra.{Ext}");
+				FileReference extraObj = FileReference.Combine(IntermediateDirectory!, $"{Name}.extra.{(Target.Platform.IsInGroup(UnrealPlatformGroup.Microsoft) ? "obj" : "o")}");
 				BinaryLinkEnvironment.InputFiles.Add(FileItem.GetItemByFileReference(extraObj));
 			}
 
diff --git a/Engine/Source/Programs/UnrealBuildTool/Configuration/UEBuildTarget.cs b/Engine/Source/Programs/UnrealBuildTool/Configuration/UEBuildTarget.cs
--- a/Engine/Source/Programs/UnrealBuildTool/Configuration/UEBuildTarget.cs
+++ b/Engine/Source/Programs/UnrealBuildTool/Configuration/UEBuildTarget.cs
@@ -3048,12 +3048,7 @@ namespace UnrealBuildTool
 				arguments.AddRange(toStrip.Select(path => $"/S:{path.Location.ChangeExtension(".exi")}").OrderBy(x => x));
 				arguments.AddRange(otherObjects.Select(path => $"/D:{path.Location.ChangeExtension(".exi")}").OrderBy(x => x));
 
-				string Ext = "obj";
-				if (CompileEnvironment.Platform.IsInGroup(UnrealPlatformGroup.Nintendo) || CompileEnvironment.Platform.IsInGroup(UnrealPlatformGroup.Sony))
-					Ext = "ldscript";
-				else if (!CompileEnvironment.Platform.IsInGroup(UnrealPlatformGroup.Microsoft))
-					Ext = "o";
-				FileReference extraObj = FileReference.Combine(IntermediateDirectory, $"{Name}.extra.{Ext}");
+				FileReference extraObj = FileReference.Combine(IntermediateDirectory, $"{Name}.extra.{(CompileEnvironment.Platform.IsInGroup(UnrealPlatformGroup.Microsoft) ? "obj" : "o")}");
 				arguments.Add($"/O:{extraObj}");
 				arguments.Add($"/T:{Platform}");
 
