[UBA]
* Fixed so obj file compression is disabled if hotreloadmode is enabled (live coding)

diff --git a/Engine/Source/Programs/UnrealBuildTool/Executors/UnrealBuildAccelerator/UBAExecutor.cs b/Engine/Source/Programs/UnrealBuildTool/Executors/UnrealBuildAccelerator/UBAExecutor.cs
--- a/Engine/Source/Programs/UnrealBuildTool/Executors/UnrealBuildAccelerator/UBAExecutor.cs
+++ b/Engine/Source/Programs/UnrealBuildTool/Executors/UnrealBuildAccelerator/UBAExecutor.cs
@@ -120,6 +120,8 @@ namespace UnrealBuildTool
 			List<UBAAgentCoordinatorHorde> hordeAgentCoordinators = new();
 			foreach (TargetDescriptor targetDescriptor in targetDescriptors)
 			{
+				if (targetDescriptor.HotReloadMode != HotReloadMode.Disabled)
+					UBAConfig.bStoreObjFilesCompressed = false;
 				targetDescriptor.AdditionalArguments.ApplyTo(this);
 				targetDescriptor.AdditionalArguments.ApplyTo(UBAConfig);
 				hordeAgentCoordinators.Add(new UBAAgentCoordinatorHorde(logger, UBAConfig, targetDescriptor.AdditionalArguments, targetDescriptor.ProjectFile?.Directory));
