[UBT]
* Changed -Wno-dllexport-explicit-instantiation-decl to be set on non-dll builds too since it seems like bool is not always true when building non-monolithic builds

diff --git a/Engine/Source/Programs/UnrealBuildTool/ToolChain/ClangWarnings.cs b/Engine/Source/Programs/UnrealBuildTool/ToolChain/ClangWarnings.cs
--- a/Engine/Source/Programs/UnrealBuildTool/ToolChain/ClangWarnings.cs
+++ b/Engine/Source/Programs/UnrealBuildTool/ToolChain/ClangWarnings.cs
@@ -159,11 +159,8 @@ namespace UnrealBuildTool
 				}
 			}
 
-			if (CompileEnvironment.bIsBuildingDLL)
-			{
-				// The code base contains lots of places where we do "extern template X_API class ..." and we want to keep doing that to reduce compile times.
-				Arguments.Add("-Wno-dllexport-explicit-instantiation-decl");
-			}
+			// The code base contains lots of places where we do "extern template X_API class ..." and we want to keep doing that to reduce compile times.
+			Arguments.Add("-Wno-dllexport-explicit-instantiation-decl");
 		}
 
 		internal static void GetHeaderDisabledWarnings(List<string> Arguments)
