Add global mutex for IOS plist generation

diff --git a/Engine/Source/Programs/UnrealBuildTool/Platform/IOS/UEDeployIOS.cs b/Engine/Source/Programs/UnrealBuildTool/Platform/IOS/UEDeployIOS.cs
--- a/Engine/Source/Programs/UnrealBuildTool/Platform/IOS/UEDeployIOS.cs
+++ b/Engine/Source/Programs/UnrealBuildTool/Platform/IOS/UEDeployIOS.cs
@@ -536,13 +536,18 @@ namespace UnrealBuildTool
 
 			if (BuildHostPlatform.Current.Platform == UnrealTargetPlatform.Mac && !bBuildAsFramework)
 			{
-				FileReference FinalPlistFile;
-				FinalPlistFile = new FileReference($"{ProjectDirectory}/Build/IOS/UBTGenerated/Info.Template.plist");
+				// need to make sure touching this plist is in the same lock as FinalizeAppWithXcode()
+				string MutexName = SingleInstanceMutex.GetUniqueMutexForPath("UnrealBuildTool_XcodeBuild", Unreal.RootDirectory.FullName);
+				using (new SingleInstanceMutex(MutexName, true))
+				{
+					FileReference FinalPlistFile;
+					FinalPlistFile = new FileReference($"{ProjectDirectory}/Build/IOS/UBTGenerated/Info.Template.plist");
 
-				DirectoryReference.CreateDirectory(FinalPlistFile.Directory);
-				// @todo: writeifdifferent is better
-				FileReference.Delete(FinalPlistFile);
-				File.Copy(PListFile, FinalPlistFile.FullName);
+					DirectoryReference.CreateDirectory(FinalPlistFile.Directory);
+					// @todo: writeifdifferent is better
+					FileReference.Delete(FinalPlistFile);
+					File.Copy(PListFile, FinalPlistFile.FullName);
+				}
 			}
 
 			return true;
