[UBT]
* Added IS_MERGEDMODULES define that can be used to check if build is merged modules. Hopefully we can remove this once we have figured out all the dll export template instantiations but because of compiler limitations (bugs?) on our console platforms we have to work around this

diff --git a/Engine/Source/Programs/UnrealBuildTool/Configuration/UEBuildTarget.cs b/Engine/Source/Programs/UnrealBuildTool/Configuration/UEBuildTarget.cs
--- a/Engine/Source/Programs/UnrealBuildTool/Configuration/UEBuildTarget.cs
+++ b/Engine/Source/Programs/UnrealBuildTool/Configuration/UEBuildTarget.cs
@@ -5306,6 +5306,7 @@ namespace UnrealBuildTool
 			GlobalCompileEnvironment.Definitions.Add("__UNREAL__");
 
 			GlobalCompileEnvironment.Definitions.Add(String.Format("IS_MONOLITHIC={0}", ShouldCompileMonolithic() ? "1" : "0"));
+			GlobalCompileEnvironment.Definitions.Add(String.Format("IS_MERGEDMODULES={0}", Rules.bMergeModules ? "1" : "0"));
 
 			GlobalCompileEnvironment.Definitions.Add(String.Format("WITH_ENGINE={0}", Rules.bCompileAgainstEngine ? "1" : "0"));
 			GlobalCompileEnvironment.Definitions.Add(String.Format("WITH_UNREAL_DEVELOPER_TOOLS={0}", Rules.bBuildDeveloperTools ? "1" : "0"));
