[UBT]
* Moved extra link file extension to be fetched from UEToolChain instead and overridden per platform
* Fixed so merged modules used overridden PER_MODULE_INLINE_FILE if exists

diff --git a/Engine/Source/Programs/UnrealBuildTool/Configuration/UEBuildBinary.cs b/Engine/Source/Programs/UnrealBuildTool/Configuration/UEBuildBinary.cs
--- a/Engine/Source/Programs/UnrealBuildTool/Configuration/UEBuildBinary.cs
+++ b/Engine/Source/Programs/UnrealBuildTool/Configuration/UEBuildBinary.cs
@@ -778,6 +778,7 @@ namespace UnrealBuildTool
 
 			HashSet<FileItem> InputFilesLookup = new();
 
+			int KeepLevel = 0;
 			FileItem? PerModuleItemToKeep = null;
 
 			foreach (UEBuildModule Module in UEBuildModule.StableTopologicalSort(Modules))
@@ -789,16 +790,51 @@ namespace UnrealBuildTool
 					List<FileItem> LinkInputFiles = Module.Compile(Target, ToolChain, BinaryCompileEnvironment, WorkingSet, Graph, Logger);
 
 					// If modules are merged we will have multiple PerModuleInline which will cause duplicated symbols. So let's only keep one of those
-					// Also, some modules might have empty PerModuleInline.gen.cpp therefore it is important that we use Core if it exists
+					// The used PerModuleInline file should be picked in order "Any", "Core", "If owning module is in project", "If module has define set"
+					// Ideally it should be "Any", "Core", "If module has define set" but unfortunately define is not set for depending modules at this point
+					// so we can have merged binaries that has 10 modules where all depends on the module that has the define... and in that case any of them would work to pick the PerModuleInline file
+					// This code is very inefficient and need to be improved.
 					if (Target.bMergeModules)
 					{
 						FileItem? PerModuleItem = LinkInputFiles.FirstOrDefault(x => x!.Name.Contains("PerModuleInline.gen.cpp"), null);
 						if (PerModuleItem != null)
 						{
-							if (PerModuleItemToKeep == null || Module.Name == "Core")
+							// If we don't have a PerModuleItem, use any
+							if (KeepLevel == 0)
+							{
+								if (PerModuleItemToKeep == null)
+								{
+									PerModuleItemToKeep = PerModuleItem;
+									KeepLevel = 1;
+								}
+							}
+							else if (KeepLevel == 1 && Module.Name == "Core") // If we have one we look for Core
 							{
 								PerModuleItemToKeep = PerModuleItem;
+								KeepLevel = 2;
+
 							}
+							else if (KeepLevel <= 2) 
+							{
+								if (Target.ProjectFile != null && Module.RulesFile.IsUnderDirectory(Target.ProjectFile.Directory))  // TODO: THIS IS WRONG AND CAN EASILY GO WRONG
+								{
+									PerModuleItemToKeep = PerModuleItem;
+									KeepLevel = 3;
+								}
+							}
+							else if (KeepLevel <= 3)
+							{
+								foreach (string Def in Module.PublicDefinitions)
+								{
+									if (Def.StartsWith("PER_MODULE_INLINE_FILE"))
+									{
+										PerModuleItemToKeep = PerModuleItem;
+										KeepLevel = 4;
+										break;
+									}
+								}
+							}
+
 							LinkInputFiles.Remove(PerModuleItem);
 						}
 					}
@@ -878,10 +914,7 @@ namespace UnrealBuildTool
 				bStripUnusedExports = true;
 
 				string Name = OutputFilePaths.First().GetFileNameWithoutExtension();
-				string Ext = "obj";
-				if (!CompileEnvironment.Platform.IsInGroup(UnrealPlatformGroup.Microsoft))
-					Ext = "ldscript";
-
+				string Ext = ToolChain.GetExtraLinkFileExtension();
 				FileReference extraObj = FileReference.Combine(IntermediateDirectory!, $"{Name}.extra.{Ext}");
 				BinaryLinkEnvironment.InputFiles.Add(FileItem.GetItemByFileReference(extraObj));
 			}
diff --git a/Engine/Source/Programs/UnrealBuildTool/Configuration/UEBuildTarget.cs b/Engine/Source/Programs/UnrealBuildTool/Configuration/UEBuildTarget.cs
--- a/Engine/Source/Programs/UnrealBuildTool/Configuration/UEBuildTarget.cs
+++ b/Engine/Source/Programs/UnrealBuildTool/Configuration/UEBuildTarget.cs
@@ -2666,7 +2666,7 @@ namespace UnrealBuildTool
 
 			if (Rules.bStripExports)
 			{
-				CreateStripExportsActions(MakefileBuilder, Makefile, GlobalCompileEnvironment);
+				CreateStripExportsActions(TargetToolChain, MakefileBuilder, Makefile, GlobalCompileEnvironment);
 			}
 
 			// Cache inline gen cpp data
@@ -3005,7 +3005,7 @@ namespace UnrealBuildTool
 			return Makefile;
 		}
 
-		void CreateStripExportsActions(TargetMakefileBuilder MakefileBuilder, TargetMakefile Makefile, CppCompileEnvironment CompileEnvironment)
+		void CreateStripExportsActions(UEToolChain TargetToolChain, TargetMakefileBuilder MakefileBuilder, TargetMakefile Makefile, CppCompileEnvironment CompileEnvironment)
 		{
 			FileReference commandPath;
 			if (OperatingSystem.IsWindows())
@@ -3048,13 +3048,11 @@ namespace UnrealBuildTool
 				arguments.AddRange(toStrip.Select(path => $"/S:{path.Location.ChangeExtension(".exi")}").OrderBy(x => x));
 				arguments.AddRange(otherObjects.Select(path => $"/D:{path.Location.ChangeExtension(".exi")}").OrderBy(x => x));
 
-				string Ext = "obj";
-				if (!CompileEnvironment.Platform.IsInGroup(UnrealPlatformGroup.Microsoft))
-					Ext = "ldscript";
-
+				string Ext = TargetToolChain.GetExtraLinkFileExtension();
 				FileReference extraObj = FileReference.Combine(IntermediateDirectory, $"{Name}.extra.{Ext}");
 				arguments.Add($"/O:{extraObj}");
 				arguments.Add($"/T:{Platform}");
+				arguments.Add($"/M:{Name}");
 
 				MakefileBuilder.CreateIntermediateTextFile(stripRsp, arguments);
 
diff --git a/Engine/Source/Programs/UnrealBuildTool/Platform/Windows/VCToolChain.cs b/Engine/Source/Programs/UnrealBuildTool/Platform/Windows/VCToolChain.cs
--- a/Engine/Source/Programs/UnrealBuildTool/Platform/Windows/VCToolChain.cs
+++ b/Engine/Source/Programs/UnrealBuildTool/Platform/Windows/VCToolChain.cs
@@ -3197,6 +3197,11 @@ namespace UnrealBuildTool
 			return IncludePaths.ToString();
 		}
 
+		public override string GetExtraLinkFileExtension()
+		{
+			return Target.WindowsPlatform.Compiler.IsClang() ? "dynlist" : "obj";
+		}
+
 		public override void ModifyBuildProducts(ReadOnlyTargetRules Target, UEBuildBinary Binary, IEnumerable<string> Libraries, IEnumerable<UEBuildBundleResource> BundleResources, Dictionary<FileReference, BuildProductType> BuildProducts)
 		{
 			if (Binary.Type == UEBuildBinaryType.DynamicLinkLibrary)
diff --git a/Engine/Source/Programs/UnrealBuildTool/ToolChain/UEToolChain.cs b/Engine/Source/Programs/UnrealBuildTool/ToolChain/UEToolChain.cs
--- a/Engine/Source/Programs/UnrealBuildTool/ToolChain/UEToolChain.cs
+++ b/Engine/Source/Programs/UnrealBuildTool/ToolChain/UEToolChain.cs
@@ -266,6 +266,15 @@ namespace UnrealBuildTool
 			return "Not Applicable";
 		}
 
+		/// <summary>
+		/// This is the extension of the file that is added as extra link object. It can be an obj file but also a dynamic list.. depends on the platform
+		/// </summary>
+		/// <returns>The extension (without dot)</returns>
+		public virtual string GetExtraLinkFileExtension()
+		{
+			throw new NotSupportedException("This platform does not support merged modules.");
+		}
+
 		/// <summary>
 		/// Runs the provided tool and argument. Returns the output, using a rexex capture if one is provided
 		/// </summary>
