- Get default target windows version from Windows_SDK.json
- Has a separate version for Arm64

diff --git a/Engine/Config/Windows/Windows_SDK.json b/Engine/Config/Windows/Windows_SDK.json
--- a/Engine/Config/Windows/Windows_SDK.json
+++ b/Engine/Config/Windows/Windows_SDK.json
@@ -59,5 +59,9 @@
 	"VisualStudio2022SuggestedComponents": [
 	],
 
-	"MinimumVisualStudio2022Version": "17.4"
+	"MinimumVisualStudio2022Version": "17.4",
+
+	"MinimumWindowsX64TargetVersion": "0x601", // Win7
+	"MinimumWindowsArm64TargetVersion":  "0xA00" // Win10
+
 }
diff --git a/Engine/Source/Programs/UnrealBuildTool/Platform/Windows/UEBuildWindows.cs b/Engine/Source/Programs/UnrealBuildTool/Platform/Windows/UEBuildWindows.cs
--- a/Engine/Source/Programs/UnrealBuildTool/Platform/Windows/UEBuildWindows.cs
+++ b/Engine/Source/Programs/UnrealBuildTool/Platform/Windows/UEBuildWindows.cs
@@ -373,7 +373,28 @@ namespace UnrealBuildTool
 		/// </summary>
 		[RequiresUniqueBuildEnvironment]
 		[ConfigFile(ConfigHierarchyType.Engine, "/Script/WindowsTargetPlatform.WindowsTargetSettings", "TargetWindowsVersion")]
-		public int TargetWindowsVersion = 0x601;
+		public int TargetWindowsVersion
+		{
+			get => TargetWindowsVersionPrivate ?? GetTargetVersionFromPlatformSDK();
+			set => TargetWindowsVersionPrivate = value;
+		}
+		private int? TargetWindowsVersionPrivate = null;
+		private int GetTargetVersionFromPlatformSDK()
+		{
+			UnrealTargetPlatform SDKPlatform = Target.Platform;
+			// due to some reflection property walking, this can actually be called with a non-Windows platform
+			// in which case, just use Win64's version. This likely will not matter at all
+			if (!Target.Platform.IsInGroup(UnrealPlatformGroup.Windows))
+			{
+				SDKPlatform = UnrealTargetPlatform.Win64;
+			}
+
+			// @todo some MS platforms are setting this hardcoded, they could move to their SDK.json
+			string Key = Architecture.bIsX64 ? "MinimumWindowsX64TargetVersion" : "MinimumWindowsArm64TargetVersion";
+			UEBuildPlatformSDK SDK = UEBuildPlatformSDK.GetSDKForPlatform(SDKPlatform.ToString())!;
+			// the string in the .json will be eg. 0x601, so convert the string from hex
+			return Convert.ToInt32(SDK.GetRequiredVersionFromConfig(Key), 16);
+		}
 
 		/// <summary>
 		/// Value for the NTDDI_VERSION macro, defining the minimum supported Windows version.
