UnrealBuildTool: Show a warning if a module is marked with the obsolete attribute and prevent automatically including in a -AllModules build

#rnx

diff --git a/Engine/Source/Programs/UnrealBuildTool/Configuration/ModuleRules.cs b/Engine/Source/Programs/UnrealBuildTool/Configuration/ModuleRules.cs
--- a/Engine/Source/Programs/UnrealBuildTool/Configuration/ModuleRules.cs
+++ b/Engine/Source/Programs/UnrealBuildTool/Configuration/ModuleRules.cs
@@ -1725,6 +1725,13 @@ namespace UnrealBuildTool
 				}
 			}
 
+			IEnumerable<ObsoleteAttribute> obsoleteAttributes = moduleType.GetCustomAttributes<ObsoleteAttribute>();
+			if (obsoleteAttributes.Any())
+			{
+				invalidReason = $"Obsolete '{(obsoleteAttributes.First().Message ?? "unspecified")}'";
+				return false;
+			}
+
 			invalidReason = null;
 			return true;
 		}
