Properly implement the OptimizationLevel param for Clang so -Oz is used for OptimizeForSize etc. This includes PGO specific overrides where necessary.
Remove platform specific Clang PGO compilation overrides since this is now handled by the cross platform base class.

diff --git a/Engine/Source/Programs/UnrealBuildTool/Platform/Windows/VCToolChain.cs b/Engine/Source/Programs/UnrealBuildTool/Platform/Windows/VCToolChain.cs
--- a/Engine/Source/Programs/UnrealBuildTool/Platform/Windows/VCToolChain.cs
+++ b/Engine/Source/Programs/UnrealBuildTool/Platform/Windows/VCToolChain.cs
@@ -707,17 +707,63 @@ namespace UnrealBuildTool
 				}
 				else
 				{
-					// Maximum optimizations.
-					Arguments.Add("/Ox");
 
-					if (CompileEnvironment.OptimizationLevel != OptimizationMode.Speed)
+					if (Target.WindowsPlatform.Compiler.IsClang())
 					{
-						Arguments.Add("/Os");
+						switch (CompileEnvironment.OptimizationLevel)
+						{
+							case OptimizationMode.Size:
+								{
+									// We use Clang's -Oz here because it produces a smaller binary than /O1
+									Arguments.Add("-Xclang -Oz");
+								}
+								break;
+							case OptimizationMode.SizeAndSpeed:
+								{
+									// Typically the Cl compatible /Ox /Os args result in a smaller binary than just -XClang -Os (by 30MB or so in a Dev build)
+									// However, when using PGO, -XClang -Os is the same size but marginally faster
+									if ( CompileEnvironment.bPGOProfile || CompileEnvironment.bPGOOptimize )
+									{
+										Arguments.Add("-Xclang -Os");
+									}
+									else
+									{
+										Arguments.Add("/Ox");
+										Arguments.Add("/Os");
+									}
+								}
+								break;
+							case OptimizationMode.Speed:
+								{
+									if ( CompileEnvironment.bPGOProfile || CompileEnvironment.bPGOOptimize )
+									{
+										// This is optimal both for speed and size when PGO is enabled
+										Arguments.Add("-Xclang -Os");
+									}
+									else
+									{
+										// Maximum optimizations. We just use the MSVC flags and let the Clang-Cl driver translate
+										Arguments.Add("/Ox");
+										Arguments.Add("/Ot");
+									}
+								}
+								break;
+						}
 					}
 					else
 					{
-						// Favor code speed.
-						Arguments.Add("/Ot");
+						// Maximum optimizations.
+						Arguments.Add("/Ox");
+
+						if (CompileEnvironment.OptimizationLevel != OptimizationMode.Speed)
+						{
+							Arguments.Add("/Os");
+						}
+						else
+						{
+							// Favor code speed.
+							Arguments.Add("/Ot");
+						}
 					}
 
 					// Coalesce duplicate strings
