- Added ability to read plugins into a config hierarchy, controlled by passing optional IncludePluginsForTargetType param to ConfigCache.ReadHierarchy(), with a TargetType (like Game or Editor).
- Use this in CopyBuildToStagingDIrectory when reading the Packaging settings when staging. This allows plugins to add files/directories to stage without needing to modify every project.
- Added a ConfigTester Turnkey command that can help test C# inis programmatically. Currently need to pass in these 3 params on commandline, it doesn't ask for them: -Branch=Engine -Section=SomeSection -Key=SomeKey

diff --git a/Engine/Source/Programs/AutomationTool/Scripts/CopyBuildToStagingDirectory.Automation.cs b/Engine/Source/Programs/AutomationTool/Scripts/CopyBuildToStagingDirectory.Automation.cs
--- a/Engine/Source/Programs/AutomationTool/Scripts/CopyBuildToStagingDirectory.Automation.cs
+++ b/Engine/Source/Programs/AutomationTool/Scripts/CopyBuildToStagingDirectory.Automation.cs
@@ -1144,7 +1144,7 @@ namespace AutomationScripts
 				return;
 			}
 
-			ConfigHierarchy PlatformGameConfig = ConfigCache.ReadHierarchy(ConfigHierarchyType.Game, DirectoryReference.FromFile(Params.RawProjectPath), SC.StageTargetPlatform.IniPlatformType, SC.CustomConfig);
+			ConfigHierarchy PlatformGameConfig = ConfigCache.ReadHierarchy(ConfigHierarchyType.Game, DirectoryReference.FromFile(Params.RawProjectPath), SC.StageTargetPlatform.IniPlatformType, SC.CustomConfig, IncludePluginsForTargetType:TargetType.Game);
 
 			if (Params.HasIterateSharedCookedBuild)
 			{
diff --git a/Engine/Source/Programs/AutomationTool/Turnkey/Commands/ConfigTester.cs b/Engine/Source/Programs/AutomationTool/Turnkey/Commands/ConfigTester.cs
new file mode 100644
--- /dev/null
+++ b/Engine/Source/Programs/AutomationTool/Turnkey/Commands/ConfigTester.cs
@@ -0,0 +1,40 @@
+// Copyright Epic Games, Inc. All Rights Reserved.
+
+using System;
+using System.Collections.Generic;
+using System.Linq;
+using System.Text;
+using System.Threading.Tasks;
+using EpicGames.Core;
+using UnrealBuildTool;
+using UnrealBuildBase;
+
+namespace Turnkey.Commands
+{
+	class ConfigTester : TurnkeyCommand
+	{
+		protected override CommandGroup Group => CommandGroup.Misc;
+
+		protected override void Execute(string[] CommandOptions)
+		{
+			FileReference ProjectFile = TurnkeyUtils.GetProjectFromCommandLineOrUser(CommandOptions);
+			List<UnrealTargetPlatform> Platforms = TurnkeyUtils.GetPlatformsFromCommandLineOrUser(CommandOptions, null);
+			string BranchName = TurnkeyUtils.ParseParamValue("Branch", null, CommandOptions);
+			string Section = TurnkeyUtils.ParseParamValue("Section", null, CommandOptions);
+			string Key = TurnkeyUtils.ParseParamValue("Key", null, CommandOptions);
+
+			ConfigHierarchyType Branch = Enum.Parse<ConfigHierarchyType>(BranchName);
+
+			// get a list of builds from config
+			foreach (UnrealTargetPlatform Platform in Platforms)
+			{
+				ConfigHierarchy Config = ConfigCache.ReadHierarchy(Branch, ProjectFile == null ? null : ProjectFile.Directory, Platform, IncludePluginsForTargetType:TargetType.Game);
+
+				string Value;
+				Config.GetString(Section, Key, out Value);
+
+				TurnkeyUtils.Log($"Value[{Platform}]: {Value}");
+			}
+		}
+	}
+}
diff --git a/Engine/Source/Programs/UnrealBuildTool/System/ConfigCache.cs b/Engine/Source/Programs/UnrealBuildTool/System/ConfigCache.cs
--- a/Engine/Source/Programs/UnrealBuildTool/System/ConfigCache.cs
+++ b/Engine/Source/Programs/UnrealBuildTool/System/ConfigCache.cs
@@ -179,6 +179,11 @@ namespace UnrealBuildTool
 			/// </summary>
 			public DirectoryReference? HotfixDir;
 
+			/// <summary>
+			/// If specified, this hierarchy will include plugins enabled for the given target type
+			/// </summary>
+			public TargetType? IncludePluginsForTargetType;
+
 			/// <summary>
 			/// Constructor
 			/// </summary>
@@ -188,7 +193,8 @@ namespace UnrealBuildTool
 			/// <param name="CustomConfig">Custom config subdirectory to load</param>
 			/// <param name="OverrideStrings">Custom override strings</param>
 			/// <param name="HotfixDir">A hotfix directory where modifications/additions to config files can be read from</param>
-			public ConfigHierarchyKey(ConfigHierarchyType Type, DirectoryReference? ProjectDir, UnrealTargetPlatform Platform, string CustomConfig, List<string> OverrideStrings, DirectoryReference? HotfixDir)
+			/// <param name="IncludePluginsForTargetType">If specified, this hierarchy will unclude plugins enabled for the given target type</param>
+			public ConfigHierarchyKey(ConfigHierarchyType Type, DirectoryReference? ProjectDir, UnrealTargetPlatform Platform, string CustomConfig, List<string> OverrideStrings, DirectoryReference? HotfixDir, TargetType? IncludePluginsForTargetType)
 			{
 				this.Type = Type;
 				this.ProjectDir = ProjectDir;
@@ -196,6 +202,7 @@ namespace UnrealBuildTool
 				this.CustomConfig = CustomConfig;
 				this.OverrideStrings = OverrideStrings;
 				this.HotfixDir = HotfixDir;
+				this.IncludePluginsForTargetType = IncludePluginsForTargetType;
 			}
 
 			/// <summary>
@@ -212,7 +219,8 @@ namespace UnrealBuildTool
 					OtherKey.Platform == Platform &&
 					OtherKey.CustomConfig == CustomConfig &&
 					OtherKey.OverrideStrings.SequenceEqual(OverrideStrings) &&
-					OtherKey.HotfixDir == HotfixDir;
+					OtherKey.HotfixDir == HotfixDir &&
+					OtherKey.IncludePluginsForTargetType == IncludePluginsForTargetType;
 			}
 
 			/// <summary>
@@ -231,6 +239,7 @@ namespace UnrealBuildTool
 					Hash = (Hash * 31) + OverrideString.GetHashCode();
 				}
 				Hash = (Hash * 31) + ((HotfixDir == null) ? 0 : HotfixDir.GetHashCode());
+				Hash = (Hash * 31) + ((IncludePluginsForTargetType == null) ? 0 : IncludePluginsForTargetType.GetHashCode());
 
 				return Hash;
 			}
@@ -281,8 +290,9 @@ namespace UnrealBuildTool
 		/// <param name="CustomConfig">Optional override config directory to search, for support of multiple target types</param>
 		/// <param name="CustomArgs">Optional list of command line arguments added to the existing command line arguments</param>
 		/// <param name="HotfixDir">A hotfix directory where modifications/additions to config files can be read from</param>
+		/// <param name="IncludePluginsForTargetType">If specified, this hierarchy will unclude plugins enabled for the given target type</param>
 		/// <returns>The requested config hierarchy</returns>
-		public static ConfigHierarchy ReadHierarchy(ConfigHierarchyType Type, DirectoryReference? ProjectDir, UnrealTargetPlatform Platform, string CustomConfig = "", string[]? CustomArgs = null, DirectoryReference? HotfixDir = null)
+		public static ConfigHierarchy ReadHierarchy(ConfigHierarchyType Type, DirectoryReference? ProjectDir, UnrealTargetPlatform Platform, string CustomConfig = "", string[]? CustomArgs = null, DirectoryReference? HotfixDir = null, TargetType? IncludePluginsForTargetType = null)
 		{
 			CommandLineArguments CombinedArgs;
 			if (CustomArgs != null)
@@ -306,7 +316,7 @@ namespace UnrealBuildTool
 				CombinedArgs = new CommandLineArguments(Environment.GetCommandLineArgs());
 			}
 
-			return ReadHierarchy(Type, ProjectDir, Platform, CustomConfig, CombinedArgs, HotfixDir);
+			return ReadHierarchy(Type, ProjectDir, Platform, CustomConfig, CombinedArgs, HotfixDir, IncludePluginsForTargetType);
 		}
 
 		/// <summary>
@@ -318,9 +328,10 @@ namespace UnrealBuildTool
 		/// <param name="CustomConfig">Optional override config directory to search, for support of multiple target types</param>
 		/// <param name="CmdLineArgs">The command line arguments to parse</param>
 		/// <param name="HotfixDir">A hotfix directory where modifications/additions to config files can be read from</param>
+		/// <param name="IncludePluginsForTargetType">If specified, this hierarchy will unclude plugins enabled for the given target type</param>
 		/// <returns>The requested config hierarchy</returns>
 		public static ConfigHierarchy ReadHierarchy(ConfigHierarchyType Type, DirectoryReference? ProjectDir, UnrealTargetPlatform Platform,
-			string CustomConfig, CommandLineArguments CmdLineArgs, DirectoryReference? HotfixDir = null)
+			string CustomConfig, CommandLineArguments CmdLineArgs, DirectoryReference? HotfixDir = null, TargetType? IncludePluginsForTargetType = null)
 		{
 			// Handle command line overrides
 			List<string> OverrideStrings = new List<string>();
@@ -347,7 +358,7 @@ namespace UnrealBuildTool
 			CustomConfig ??= String.Empty;
 
 			// Get the key to use for the cache. It cannot be null, so we use the engine directory if a project directory is not given.
-			ConfigHierarchyKey Key = new ConfigHierarchyKey(Type, ProjectDir, Platform, CustomConfig, OverrideStrings, HotfixDir);
+			ConfigHierarchyKey Key = new ConfigHierarchyKey(Type, ProjectDir, Platform, CustomConfig, OverrideStrings, HotfixDir, IncludePluginsForTargetType);
 
 			ILogger Logger = NullLogger.Instance;
 
@@ -362,7 +373,7 @@ namespace UnrealBuildTool
 			{
 				// Find all the input files
 				List<ConfigFile> Files = new List<ConfigFile>();
-				foreach (FileReference IniFileName in ConfigHierarchy.EnumerateConfigFileLocations(Type, ProjectDir, Platform, CustomConfig))
+				foreach (FileReference IniFileName in ConfigHierarchy.EnumerateConfigFileLocations(Type, ProjectDir, Platform, CustomConfig, IncludePluginsForTargetType))
 				{
 					ConfigFile? File;
 					Logger.LogInformation($"Trying to load: {IniFileName}");
diff --git a/Engine/Source/Programs/UnrealBuildTool/System/ConfigHierarchy.cs b/Engine/Source/Programs/UnrealBuildTool/System/ConfigHierarchy.cs
--- a/Engine/Source/Programs/UnrealBuildTool/System/ConfigHierarchy.cs
+++ b/Engine/Source/Programs/UnrealBuildTool/System/ConfigHierarchy.cs
@@ -1140,10 +1140,44 @@ namespace UnrealBuildTool
 		/// <summary>
 		/// Returns a list of INI filenames for the given project
 		/// </summary>
-		public static IEnumerable<FileReference> EnumerateConfigFileLocations(ConfigHierarchyType Type, DirectoryReference? ProjectDir, UnrealTargetPlatform Platform, string CustomConfig = "")
+		public static IEnumerable<FileReference> EnumerateConfigFileLocations(ConfigHierarchyType Type, DirectoryReference? ProjectDir, UnrealTargetPlatform Platform, string CustomConfig, TargetType? IncludePluginsForTargetType)
 		{
 			string BaseIniName = Enum.GetName(typeof(ConfigHierarchyType), Type) ?? String.Empty;
 			string PlatformName = GetIniPlatformName(Platform);
+			List<DirectoryReference> PluginDirs = new();
+			List<DirectoryReference> PluginPlatformExtDirs = new();
+
+			if (IncludePluginsForTargetType.HasValue)
+			{
+				ProjectDescriptor? Project = null;
+				if (ProjectDir != null)
+				{
+					Project = ProjectDescriptor.FromFile(FileReference.Combine(ProjectDir, ProjectDir.GetDirectoryName() + ".uproject"));
+				}
+				List<PluginInfo> AllPLugins = Plugins.ReadAvailablePlugins(Unreal.EngineDirectory, ProjectDir, null);
+				foreach (PluginInfo Plugin in AllPLugins)
+				{
+					if (Plugins.IsPluginEnabledForTarget(Plugin, Project, Platform, UnrealTargetConfiguration.Development, IncludePluginsForTargetType.Value))
+					{
+						PluginDirs.Add(Plugin.Directory);
+						foreach (FileReference ChildPlugin in Plugin.ChildFiles)
+						{
+							// only look at a child plugin for a platform we are enumerating (ie, it is in a /PlatformName/ subdir), but always look in Restricted folders
+							if (ChildPlugin.ContainsName("Platforms", 0))
+							{
+								if (ChildPlugin.ContainsName(PlatformName, 0))
+								{
+									PluginPlatformExtDirs.Add(ChildPlugin.Directory);
+								}
+							}
+							else
+							{
+								PluginDirs.Add(ChildPlugin.Directory);
+							}
+						}
+					}
+				}
+			}
 
 			foreach (string Layer in ConfigLayers)
 			{
@@ -1208,6 +1242,16 @@ namespace UnrealBuildTool
 				}
 			}
 
+			foreach (DirectoryReference PluginDir in PluginDirs)
+			{
+				yield return FileReference.Combine(PluginDir, "Config", $"{Type}.ini");
+				yield return FileReference.Combine(PluginDir, "Config", PlatformName, $"{PlatformName}{Type}.ini");
+			}
+			foreach (DirectoryReference PluginDir in PluginPlatformExtDirs)
+			{
+				yield return FileReference.Combine(PluginDir, "Config", $"{PlatformName}{Type}.ini");
+			}
+
 			// Find all the generated config files
 			foreach (FileReference GeneratedConfigFile in EnumerateGeneratedConfigFileLocations(Type, ProjectDir, Platform))
 			{
diff --git a/Engine/Source/Programs/UnrealBuildTool/System/Plugins.cs b/Engine/Source/Programs/UnrealBuildTool/System/Plugins.cs
--- a/Engine/Source/Programs/UnrealBuildTool/System/Plugins.cs
+++ b/Engine/Source/Programs/UnrealBuildTool/System/Plugins.cs
@@ -625,7 +625,7 @@ namespace UnrealBuildTool
 		/// <param name="Configuration">The target configuration</param>
 		/// <param name="TargetType">The type of target being built</param>
 		/// <returns>True if the plugin should be enabled for this project</returns>
-		public static bool IsPluginEnabledForTarget(PluginInfo Plugin, ProjectDescriptor Project, UnrealTargetPlatform Platform, UnrealTargetConfiguration Configuration, TargetType TargetType)
+		public static bool IsPluginEnabledForTarget(PluginInfo Plugin, ProjectDescriptor? Project, UnrealTargetPlatform Platform, UnrealTargetConfiguration Configuration, TargetType TargetType)
 		{
 			if (!Plugin.Descriptor.SupportsTargetPlatform(Platform))
 			{
