Horde: Log Horde server and agent version for a UBA session

diff --git a/Engine/Source/Programs/Shared/EpicGames.Horde/Compute/Clients/ServerComputeClient.cs b/Engine/Source/Programs/Shared/EpicGames.Horde/Compute/Clients/ServerComputeClient.cs
--- a/Engine/Source/Programs/Shared/EpicGames.Horde/Compute/Clients/ServerComputeClient.cs
+++ b/Engine/Source/Programs/Shared/EpicGames.Horde/Compute/Clients/ServerComputeClient.cs
@@ -388,7 +388,7 @@ namespace EpicGames.Horde.Compute.Clients
 			// Send the nonce
 			byte[] nonce = StringUtils.ParseHexString(response.Nonce);
 			await socket.SendMessageAsync(nonce, SocketFlags.None, cancellationToken);
-			workerLogger.LogInformation("Connected to {AgentId} ({Ip}) under lease {LeaseId}", response.AgentId, response.Ip, response.LeaseId);
+			workerLogger.LogInformation("Connected to {AgentId} ({Ip}) under lease {LeaseId} (agent version: {AgentVersion})", response.AgentId, response.Ip, response.LeaseId, response.AgentVersion ?? "unknown");
 
 			await using ComputeTransport transport = await CreateTransportAsync(socket, response, cancellationToken);
 			await using RemoteComputeSocket computeSocket = new(transport, (ComputeProtocol)response.Protocol, workerLogger);
diff --git a/Engine/Source/Programs/Shared/EpicGames.Horde/Compute/ComputeMessages.cs b/Engine/Source/Programs/Shared/EpicGames.Horde/Compute/ComputeMessages.cs
--- a/Engine/Source/Programs/Shared/EpicGames.Horde/Compute/ComputeMessages.cs
+++ b/Engine/Source/Programs/Shared/EpicGames.Horde/Compute/ComputeMessages.cs
@@ -183,6 +183,11 @@ namespace EpicGames.Horde.Compute
 		/// Identifier for the remote machine
 		/// </summary>
 		public AgentId AgentId { get; set; }
+		
+		/// <summary>
+		/// Agent version for the remote machine
+		/// </summary>
+		public string? AgentVersion { get; set; }
 
 		/// <summary>
 		/// Identifier for the new lease on the remote machine
diff --git a/Engine/Source/Programs/UnrealBuildTool/Executors/UnrealBuildAccelerator/UBAAgentCoordinatorHorde.cs b/Engine/Source/Programs/UnrealBuildTool/Executors/UnrealBuildAccelerator/UBAAgentCoordinatorHorde.cs
--- a/Engine/Source/Programs/UnrealBuildTool/Executors/UnrealBuildAccelerator/UBAAgentCoordinatorHorde.cs
+++ b/Engine/Source/Programs/UnrealBuildTool/Executors/UnrealBuildAccelerator/UBAAgentCoordinatorHorde.cs
@@ -14,6 +14,7 @@ using EpicGames.Horde.Agents;
 using EpicGames.Horde.Common;
 using EpicGames.Horde.Compute;
 using EpicGames.Horde.Compute.Clients;
+using EpicGames.Horde.Server;
 using EpicGames.Horde.Storage;
 using EpicGames.Horde.Storage.Bundles;
 using EpicGames.Horde.Storage.Nodes;
@@ -57,6 +58,7 @@ namespace UnrealBuildTool
 		readonly ServiceProvider _serviceProvider;
 		readonly string _crypto;
 		IComputeClient? _client;
+		HordeHttpClient? _hordeHttpClient;
 
 		public class Worker
 		{
@@ -138,8 +140,10 @@ namespace UnrealBuildTool
 			}
 
 			_client = _serviceProvider.GetRequiredService<IHordeClient>().CreateComputeClient();
+			_hordeHttpClient = _serviceProvider.GetRequiredService<IHordeClient>().CreateHttpClient();
+			GetServerInfoResponse serverInfo = await _hordeHttpClient.GetServerInfoAsync(cancellationToken);
+			_logger.LogInformation("Horde server: {ServerVersion}, agent: {AgentVersion}", serverInfo.ServerVersion, serverInfo.AgentVersion);
 
-			_logger.LogInformation("Creating tool bundle...");
 			DirectoryReference ubaDir = UBAExecutor.UbaBinariesDir;
 			List<string> agentFiles = new();
 			if (OperatingSystem.IsWindows())
