Loosen the restriction on unknown platforms from some plugin folder locations

 - allow game-specific plugins to reference certain platform extensions without using a plugin extension (currently just WinGDK)

#rnx

diff --git a/Engine/Source/Programs/UnrealBuildTool/System/ModuleDescriptor.cs b/Engine/Source/Programs/UnrealBuildTool/System/ModuleDescriptor.cs
--- a/Engine/Source/Programs/UnrealBuildTool/System/ModuleDescriptor.cs
+++ b/Engine/Source/Programs/UnrealBuildTool/System/ModuleDescriptor.cs
@@ -275,7 +275,7 @@ namespace UnrealBuildTool
 						{
 							Module.PlatformAllowList.Add(Platform);
 						}
-						else
+						else if ( !PluginDescriptor.IsAllowableMissingPlatform(TargetPlatformName, JsonFilePath.Directory) )
 						{
 							Log.TraceWarningTask(JsonFilePath, $"Unknown platform {TargetPlatformName} while parsing allow list for module descriptor {Module.Name}");
 						}
@@ -293,7 +293,7 @@ namespace UnrealBuildTool
 						{
 							Module.PlatformDenyList.Add(Platform);
 						}
-						else
+						else if ( !PluginDescriptor.IsAllowableMissingPlatform(TargetPlatformName, JsonFilePath.Directory) )
 						{
 							Log.TraceWarningTask(JsonFilePath, $"Unknown platform {TargetPlatformName} while parsing deny list for module descriptor {Module.Name}");
 						}
diff --git a/Engine/Source/Programs/UnrealBuildTool/System/PluginDescriptor.cs b/Engine/Source/Programs/UnrealBuildTool/System/PluginDescriptor.cs
--- a/Engine/Source/Programs/UnrealBuildTool/System/PluginDescriptor.cs
+++ b/Engine/Source/Programs/UnrealBuildTool/System/PluginDescriptor.cs
@@ -304,7 +304,7 @@ namespace UnrealBuildTool
 					{
 						SupportedTargetPlatforms.Add(Platform);
 					}
-					else
+					else if ( !IsAllowableMissingPlatform(TargetPlatformName, PluginPath.Directory) )
 					{
 						Log.TraceWarningTask(PluginPath, $"Unknown platform {TargetPlatformName} listed in plugin with FriendlyName \"{FriendlyName}\"");
 					}
@@ -706,5 +706,33 @@ namespace UnrealBuildTool
 		{
 			return SupportedTargetPlatforms?.Select(P => P.ToString()).ToArray();
 		}
+
+
+		/// <summary>
+		/// Indicates whether it is acceptable for the given platform to be missing
+		/// Typically for platform extensions that are not covered by NDA in third party plugins
+		/// </summary>
+		/// <param name="Platform"></param>
+		/// <param name="PluginFolder"></param>
+		internal static bool IsAllowableMissingPlatform(string Platform, DirectoryReference PluginFolder)
+		{
+			// certain directories should always show "best practice" and enforce plugin extensions
+			DirectoryReference[] EnforcedDirectories = new DirectoryReference[]
+			{
+				DirectoryReference.Combine(UnrealBuildBase.Unreal.EngineDirectory, "Plugins"),
+				DirectoryReference.Combine(UnrealBuildBase.Unreal.RootDirectory, "Samples"),
+			};
+			if (EnforcedDirectories.Any( X => PluginFolder.IsUnderDirectory(X) ) )
+			{
+				return false;
+			}
+
+			// all other plugin locations allow some unknown platforms (internal code, game projects etc)
+			string[] AllowedMissingPlatforms = new string[] 
+			{ 
+				"WinGDK",
+			};
+			return AllowedMissingPlatforms.Contains(Platform);
+		}
 	}
 }
