UnrealBuildTool: Fix intermediate directory for merged modules

#rnx

diff --git a/Engine/Source/Programs/UnrealBuildTool/Configuration/UEBuildTarget.cs b/Engine/Source/Programs/UnrealBuildTool/Configuration/UEBuildTarget.cs
--- a/Engine/Source/Programs/UnrealBuildTool/Configuration/UEBuildTarget.cs
+++ b/Engine/Source/Programs/UnrealBuildTool/Configuration/UEBuildTarget.cs
@@ -4432,7 +4432,7 @@ namespace UnrealBuildTool
 		/// <returns>The new binary. This has not been added to the target.</returns>
 		private UEBuildBinary CreateDynamicLibraryForModules(string MergedName, IEnumerable<UEBuildModuleCPP> Modules)
 		{
-			DirectoryReference IntermediateDirectory = DirectoryReference.Combine(ProjectIntermediateDirectory, MergedName);
+			DirectoryReference MergedIntermediateDirectory = DirectoryReference.Combine(ProjectIntermediateDirectory, MergedName);
 
 			// Get the root output directory and base name (target name/app name) for this binary
 			DirectoryReference OutputDirectory;
@@ -4458,7 +4458,7 @@ namespace UnrealBuildTool
 			UEBuildBinary Binary = new UEBuildBinary(
 				Type: UEBuildBinaryType.DynamicLinkLibrary,
 				OutputFilePaths: OutputPaths,
-				IntermediateDirectory: ProjectIntermediateDirectory,
+				IntermediateDirectory: MergedIntermediateDirectory,
 				bAllowExports: true,
 				bBuildAdditionalConsoleApp: false,
 				bBuildConsoleAppOnly: false,
