Enable ISPC for Android. Remove deprecated physics macros from 5.1. These were generating very long lines that were confusing the android build system.

diff --git a/Engine/Source/Programs/UnrealBuildTool/Configuration/ModuleRules.cs b/Engine/Source/Programs/UnrealBuildTool/Configuration/ModuleRules.cs
--- a/Engine/Source/Programs/UnrealBuildTool/Configuration/ModuleRules.cs
+++ b/Engine/Source/Programs/UnrealBuildTool/Configuration/ModuleRules.cs
@@ -1665,28 +1665,6 @@ namespace UnrealBuildTool
 			PublicDefinitions.Add("WITH_CLOTH_COLLISION_DETECTION=1");
 
 			SetupModuleChaosVisualDebuggerSupport(target);
-
-			// Modules may still be relying on appropriate definitions for physics.
-			// Nothing in engine should use these anymore as they were all deprecated and 
-			// assumed to be in the following configuration from 5.1, this will cause
-			// deprecation warning to fire in any module still relying on these macros
-			static string GetDeprecatedPhysicsMacro(string macro, string value, string version) => $"{macro}=UE_DEPRECATED_MACRO({version}, \"{macro} is deprecated and should always be considered {value}.\") {value}";
-
-			PublicDefinitions.AddRange(
-				new string[]{
-					GetDeprecatedPhysicsMacro("INCLUDE_CHAOS", "1", "5.1"),
-					GetDeprecatedPhysicsMacro("WITH_CHAOS", "1", "5.1"),
-					GetDeprecatedPhysicsMacro("WITH_CHAOS_CLOTHING", "1", "5.1"),
-					GetDeprecatedPhysicsMacro("WITH_CHAOS_NEEDS_TO_BE_FIXED", "1", "5.1"),
-					GetDeprecatedPhysicsMacro("WITH_PHYSX", "1", "5.1"),
-					GetDeprecatedPhysicsMacro("WITH_PHYSX_COOKING", "0", "5.1"),
-					GetDeprecatedPhysicsMacro("PHYSICS_INTERFACE_PHYSX", "0", "5.1"),
-					GetDeprecatedPhysicsMacro("WITH_APEX", "0", "5.1"),
-					GetDeprecatedPhysicsMacro("WITH_APEX_CLOTHING", "0", "5.1"),
-					GetDeprecatedPhysicsMacro("WITH_NVCLOTH", "0", "5.1"),
-					GetDeprecatedPhysicsMacro("WITH_IMMEDIATE_PHYSX", "0", "5.1"),
-					GetDeprecatedPhysicsMacro("WITH_CUSTOM_SQ_STRUCTURE", "0", "5.1")
-				});
 		}
 
 		/// <summary>
diff --git a/Engine/Source/Programs/UnrealBuildTool/Platform/Android/UEBuildAndroid.cs b/Engine/Source/Programs/UnrealBuildTool/Platform/Android/UEBuildAndroid.cs
--- a/Engine/Source/Programs/UnrealBuildTool/Platform/Android/UEBuildAndroid.cs
+++ b/Engine/Source/Programs/UnrealBuildTool/Platform/Android/UEBuildAndroid.cs
@@ -213,7 +213,7 @@ namespace UnrealBuildTool
 			}
 
 			Target.bCompileRecast = true;
-			Target.bCompileISPC = false;
+			Target.bCompileISPC = true;
 
 			// disable plugins by architecture (if we are compiling for multiple architectures, we still need to disable the plugin for all architectures)
 			if (Target.Architectures.Contains(UnrealArch.Arm64) && Target.Name != "UnrealHeaderTool")
