- FIx for loading plugin configs when uproject name doesn't match directory name
- Added ProjectDescriptor.FromDirectory(), which will attempt to find and load a .uproject file from a project direectory (with .uproject in the root of that dir)

diff --git a/Engine/Source/Programs/UnrealBuildTool/System/ConfigHierarchy.cs b/Engine/Source/Programs/UnrealBuildTool/System/ConfigHierarchy.cs
--- a/Engine/Source/Programs/UnrealBuildTool/System/ConfigHierarchy.cs
+++ b/Engine/Source/Programs/UnrealBuildTool/System/ConfigHierarchy.cs
@@ -1152,7 +1152,7 @@ namespace UnrealBuildTool
 				ProjectDescriptor? Project = null;
 				if (ProjectDir != null)
 				{
-					Project = ProjectDescriptor.FromFile(FileReference.Combine(ProjectDir, ProjectDir.GetDirectoryName() + ".uproject"));
+					Project = ProjectDescriptor.FromDirectory(ProjectDir);
 				}
 				List<PluginInfo> AllPLugins = Plugins.ReadAvailablePlugins(Unreal.EngineDirectory, ProjectDir, null);
 				foreach (PluginInfo Plugin in AllPLugins)
diff --git a/Engine/Source/Programs/UnrealBuildTool/System/ProjectDescriptor.cs b/Engine/Source/Programs/UnrealBuildTool/System/ProjectDescriptor.cs
--- a/Engine/Source/Programs/UnrealBuildTool/System/ProjectDescriptor.cs
+++ b/Engine/Source/Programs/UnrealBuildTool/System/ProjectDescriptor.cs
@@ -232,6 +232,32 @@ namespace UnrealBuildTool
 			}
 		}
 
+		/// <summary>
+		/// Creates a plugin descriptor based on a directory, searching for best uproject
+		/// </summary>
+		/// <param name="DirectoryName">The directory to search for a uproject</param>
+		/// <returns>New plugin descriptor</returns>
+		public static ProjectDescriptor FromDirectory(DirectoryReference DirectoryName)
+		{
+			FileReference ProjectFile = FileReference.Combine(DirectoryName, DirectoryName.GetDirectoryName() + ".uproject");
+			if (FileReference.Exists(ProjectFile))
+			{
+				return ProjectDescriptor.FromFile(ProjectFile);
+			}
+
+			// find any uproject file, in case the name doesn't match
+			IEnumerable<FileReference> FoundProjects = DirectoryReference.EnumerateFiles(DirectoryName, "*.uproject", SearchOption.TopDirectoryOnly);
+			if (FoundProjects.Count() == 0)
+			{
+				throw new FileNotFoundException($"Unable to find a .uproject file in {DirectoryName}");
+			}
+			if (FoundProjects.Count() > 1)
+			{
+				Log.TraceWarningOnce("Found multiple uproject files in {0}, choosing {1}", DirectoryName, FoundProjects.First().GetFileName());
+			}
+			return ProjectDescriptor.FromFile(FoundProjects.First());
+		}
+
 		/// <summary>
 		/// If the descriptor has either additional plugin directories or additional root directories
 		/// then it is considered to have additional paths.  The additional paths will be relative
