UnrealBuildTool: Support ISPC multiple architectures in a single build

* Update GenerateISPCHeaders & CompileISPCFiles to work the same as CompileCppFiles, handling multiple arch
* GenerateISPCHeaders & CompileISPCFiles is now protected
* Simplify platform and arch logic in ISPCToolchain
* Remove need to override ISPC compile functions for Android

diff --git a/Engine/Source/Programs/UnrealBuildTool/Configuration/UEBuildModuleCPP.cs b/Engine/Source/Programs/UnrealBuildTool/Configuration/UEBuildModuleCPP.cs
--- a/Engine/Source/Programs/UnrealBuildTool/Configuration/UEBuildModuleCPP.cs
+++ b/Engine/Source/Programs/UnrealBuildTool/Configuration/UEBuildModuleCPP.cs
@@ -779,7 +779,7 @@ namespace UnrealBuildTool
 			// Compile ISPC files directly
 			if (InputFiles.ISPCFiles.Count > 0)
 			{
-				LinkInputFiles.AddRange(ToolChain.CompileISPCFiles(ModuleCompileEnvironment, InputFiles.ISPCFiles, IntermediateDirectory, Graph).ObjectFiles);
+				LinkInputFiles.AddRange(ToolChain.CompileAllISPCFiles(ModuleCompileEnvironment, InputFiles.ISPCFiles, IntermediateDirectory, Graph).ObjectFiles);
 			}
 
 			// Compile C files directly. Do not use a PCH here, because a C++ PCH is not compatible with C source files.
@@ -1872,7 +1872,7 @@ namespace UnrealBuildTool
 		/// <param name="Graph">List of actions to be executed. Additional actions will be added to this list.</param>
 		static void CreateHeadersForISPC(UEToolChain ToolChain, CppCompileEnvironment CompileEnvironment, List<FileItem> InputFiles, DirectoryReference IntermediateDirectory, IActionGraphBuilder Graph)
 		{
-			CPPOutput Output = ToolChain.GenerateISPCHeaders(CompileEnvironment, InputFiles, IntermediateDirectory, Graph);
+			CPPOutput Output = ToolChain.GenerateAllISPCHeaders(CompileEnvironment, InputFiles, IntermediateDirectory, Graph);
 
 			CompileEnvironment.AdditionalPrerequisites.AddRange(Output.GeneratedHeaderFiles);
 			CompileEnvironment.UserIncludePaths.Add(IntermediateDirectory);
diff --git a/Engine/Source/Programs/UnrealBuildTool/Platform/Android/AndroidToolChain.cs b/Engine/Source/Programs/UnrealBuildTool/Platform/Android/AndroidToolChain.cs
--- a/Engine/Source/Programs/UnrealBuildTool/Platform/Android/AndroidToolChain.cs
+++ b/Engine/Source/Programs/UnrealBuildTool/Platform/Android/AndroidToolChain.cs
@@ -1072,332 +1072,6 @@ namespace UnrealBuildTool
 			return DirectoryReference.Combine(PathRef, "include", Arch.ToString());
 		}
 
-		public override CPPOutput GenerateISPCHeaders(CppCompileEnvironment CompileEnvironment, IEnumerable<FileItem> InputFiles, DirectoryReference OutputDir, IActionGraphBuilder Graph)
-		{
-			CPPOutput Result = new CPPOutput();
-
-			if (!CompileEnvironment.bCompileISPC)
-			{
-				return Result;
-			}
-
-			List<string> CompileTargets = GetISPCCompileTargets(CompileEnvironment.Platform, CompileEnvironment.Architecture);
-
-			CompileEnvironment.UserIncludePaths.Add(InlineArchIncludeFolder(OutputDir, CompileEnvironment.Architecture));
-
-			foreach (FileItem ISPCFile in InputFiles)
-			{
-				Action CompileAction = Graph.CreateAction(ActionType.Compile);
-				CompileAction.CommandDescription = "Compile";
-				CompileAction.WorkingDirectory = Unreal.EngineSourceDirectory;
-				CompileAction.CommandPath = new FileReference(GetISPCHostCompilerPath(BuildHostPlatform.Current.Platform));
-				CompileAction.StatusDescription = Path.GetFileName(ISPCFile.AbsolutePath);
-				CompileAction.CommandVersion = GetISPCHostCompilerVersion(BuildHostPlatform.Current.Platform).ToString();
-				if (bAllowUbaCompression)
-				{
-					CompileAction.CommandVersion = $"{CompileAction.CommandVersion} Compressed";
-				}
-
-				// Disable remote execution to workaround mismatched case on XGE
-				CompileAction.bCanExecuteRemotely = false;
-
-				List<string> Arguments = new List<string>();
-
-				// Add the ISPC obj file as a prerequisite of the action.
-				Arguments.Add(String.Format(" \"{0}\"", ISPCFile.AbsolutePath));
-
-				// Add the ISPC h file to the produced item list.
-				FileItem ISPCIncludeHeaderFile = FileItem.GetItemByFileReference(
-					FileReference.Combine(
-						InlineArchIncludeFolder(OutputDir, CompileEnvironment.Architecture),
-						Path.GetFileName(ISPCFile.AbsolutePath) + ".generated.dummy.h"
-						)
-					);
-
-				// Add the ISPC file to be compiled.
-				Arguments.Add(String.Format("-h \"{0}\"", ISPCIncludeHeaderFile));
-
-				// Build target string. No comma on last
-				string TargetString = "";
-				foreach (string Target in CompileTargets)
-				{
-					if (Target == CompileTargets.Last())
-					{
-						TargetString += Target;
-					}
-					else
-					{
-						TargetString += Target + ",";
-					}
-				}
-
-				// Build target triplet
-				Arguments.Add(String.Format("--target-os=\"{0}\"", GetISPCOSTarget(CompileEnvironment.Platform)));
-				Arguments.Add(String.Format("--arch=\"{0}\"", GetISPCArchTarget(CompileEnvironment.Platform, CompileEnvironment.Architecture)));
-				Arguments.Add(String.Format("--target=\"{0}\"", TargetString));
-
-				Arguments.Add("--pic");
-
-				// Include paths. Don't use AddIncludePath() here, since it uses the full path and exceeds the max command line length.
-				foreach (DirectoryReference IncludePath in CompileEnvironment.UserIncludePaths)
-				{
-					Arguments.Add(String.Format("-I\"{0}\"", IncludePath));
-				}
-
-				// System include paths.
-				foreach (DirectoryReference SystemIncludePath in CompileEnvironment.SystemIncludePaths)
-				{
-					Arguments.Add(String.Format("-I\"{0}\"", SystemIncludePath));
-				}
-
-				// Generate the included header dependency list
-				FileItem DependencyListFile = FileItem.GetItemByFileReference(FileReference.Combine(OutputDir, InlineArchName(Path.GetFileName(ISPCFile.AbsolutePath) + ".d", CompileEnvironment.Architecture, true)));
-				Arguments.Add(String.Format("-M -MF \"{0}\"", DependencyListFile.AbsolutePath.Replace('\\', '/')));
-				CompileAction.DependencyListFile = DependencyListFile;
-				CompileAction.ProducedItems.Add(DependencyListFile);
-
-				CompileAction.ProducedItems.Add(ISPCIncludeHeaderFile);
-
-				CompileAction.CommandArguments = String.Join(" ", Arguments);
-
-				// Add the source file and its included files to the prerequisite item list.
-				CompileAction.PrerequisiteItems.Add(ISPCFile);
-				CompileAction.StatusDescription = String.Format("[{0}] {1}", CompileEnvironment.Architecture, Path.GetFileName(ISPCFile.AbsolutePath));
-
-				FileItem ISPCFinalHeaderFile = FileItem.GetItemByFileReference(
-					FileReference.Combine(
-						InlineArchIncludeFolder(OutputDir, CompileEnvironment.Architecture),
-						Path.GetFileName(ISPCFile.AbsolutePath) + ".generated.h"
-						)
-					);
-
-				// Fix interrupted build issue by copying header after generation completes
-				FileReference SourceFile = ISPCIncludeHeaderFile.Location;
-				FileReference TargetFile = ISPCFinalHeaderFile.Location;
-
-				FileItem SourceFileItem = FileItem.GetItemByFileReference(SourceFile);
-				FileItem TargetFileItem = FileItem.GetItemByFileReference(TargetFile);
-
-				Action CopyAction = Graph.CreateAction(ActionType.BuildProject);
-				CopyAction.CommandDescription = "Copy";
-				CopyAction.CommandPath = BuildHostPlatform.Current.Shell;
-				if (BuildHostPlatform.Current.ShellType == ShellType.Cmd)
-				{
-					CopyAction.CommandArguments = String.Format("/C \"copy /Y \"{0}\" \"{1}\" 1>nul\"", SourceFile, TargetFile);
-				}
-				else
-				{
-					CopyAction.CommandArguments = String.Format("-c 'cp -f \"{0}\" \"{1}\"'", SourceFile.FullName, TargetFile.FullName);
-				}
-				CopyAction.WorkingDirectory = Unreal.EngineSourceDirectory;
-				CopyAction.PrerequisiteItems.Add(SourceFileItem);
-				CopyAction.ProducedItems.Add(TargetFileItem);
-				CopyAction.StatusDescription = TargetFileItem.Location.GetFileName();
-				CopyAction.bCanExecuteRemotely = false;
-				CopyAction.bShouldOutputStatusDescription = false;
-
-				Result.GeneratedHeaderFiles.Add(TargetFileItem);
-
-				Logger.LogDebug("   ISPC Generating Header {StatusDescription}: \"{CommandPath}\" {CommandArguments}", CompileAction.StatusDescription, CompileAction.CommandPath, CompileAction.CommandArguments);
-			}
-
-			return Result;
-		}
-
-		public override CPPOutput CompileISPCFiles(CppCompileEnvironment CompileEnvironment, IEnumerable<FileItem> InputFiles, DirectoryReference OutputDir, IActionGraphBuilder Graph)
-		{
-			CPPOutput Result = new CPPOutput();
-
-			if (!CompileEnvironment.bCompileISPC)
-			{
-				return Result;
-			}
-
-			List<string> CompileTargets = GetISPCCompileTargets(CompileEnvironment.Platform, CompileEnvironment.Architecture);
-
-			foreach (FileItem ISPCFile in InputFiles)
-			{
-				Action CompileAction = Graph.CreateAction(ActionType.Compile);
-				CompileAction.CommandDescription = "Compile";
-				CompileAction.WorkingDirectory = Unreal.EngineSourceDirectory;
-				CompileAction.CommandPath = new FileReference(GetISPCHostCompilerPath(BuildHostPlatform.Current.Platform));
-				CompileAction.StatusDescription = Path.GetFileName(ISPCFile.AbsolutePath);
-
-				// Disable remote execution to workaround mismatched case on XGE
-				CompileAction.bCanExecuteRemotely = false;
-
-				List<string> Arguments = new List<string>();
-
-				// Add the ISPC file to be compiled.
-				Arguments.Add(String.Format(" \"{0}\"", ISPCFile.AbsolutePath));
-
-				List<FileItem> CompiledISPCObjFiles = new List<FileItem>();
-				List<FileItem> FinalISPCObjFiles = new List<FileItem>();
-				string TargetString = "";
-
-				foreach (string Target in CompileTargets)
-				{
-					string ObjTarget = Target;
-
-					if (Target.Contains('-'))
-					{
-						// Remove lane width and gang size from obj file name
-						ObjTarget = Target.Split('-')[0];
-					}
-
-					FileItem CompiledISPCObjFile;
-					FileItem FinalISPCObjFile;
-
-					if (CompileTargets.Count > 1)
-					{
-						CompiledISPCObjFile = FileItem.GetItemByFileReference(
-						FileReference.Combine(
-							OutputDir,
-							Path.GetFileNameWithoutExtension(InlineArchName(Path.GetFileName(ISPCFile.AbsolutePath) + ".o", CompileEnvironment.Architecture, true)) + "_" + ObjTarget + ".o"
-							)
-						);
-
-						FinalISPCObjFile = FileItem.GetItemByFileReference(
-						FileReference.Combine(
-							OutputDir,
-							Path.GetFileName(ISPCFile.AbsolutePath) + "_" + ObjTarget + InlineArchName(".o", CompileEnvironment.Architecture, true)
-							)
-						);
-					}
-					else
-					{
-						CompiledISPCObjFile = FileItem.GetItemByFileReference(
-							FileReference.Combine(
-								OutputDir,
-								InlineArchName(Path.GetFileName(ISPCFile.AbsolutePath) + ".o", CompileEnvironment.Architecture, true)
-								)
-							);
-
-						FinalISPCObjFile = CompiledISPCObjFile;
-					}
-
-					// Add the ISA specific ISPC obj files to the produced item list.
-					CompiledISPCObjFiles.Add(CompiledISPCObjFile);
-					FinalISPCObjFiles.Add(FinalISPCObjFile);
-
-					// Build target string. No comma on last
-					if (Target == CompileTargets.Last())
-					{
-						TargetString += Target;
-					}
-					else
-					{
-						TargetString += Target + ",";
-					}
-				}
-
-				// Add the common ISPC obj file to the produced item list.
-				FileItem CompiledISPCObjFileNoISA = FileItem.GetItemByFileReference(
-					FileReference.Combine(
-						OutputDir,
-						InlineArchName(Path.GetFileName(ISPCFile.AbsolutePath) + ".o", CompileEnvironment.Architecture, true)
-						)
-					);
-
-				CompiledISPCObjFiles.Add(CompiledISPCObjFileNoISA);
-				FinalISPCObjFiles.Add(CompiledISPCObjFileNoISA);
-
-				// Add the output ISPC obj file
-				Arguments.Add(String.Format("-o \"{0}\"", CompiledISPCObjFileNoISA));
-
-				// Build target triplet
-				Arguments.Add(String.Format("--target-os=\"{0}\"", GetISPCOSTarget(CompileEnvironment.Platform)));
-				Arguments.Add(String.Format("--arch=\"{0}\"", GetISPCArchTarget(CompileEnvironment.Platform, CompileEnvironment.Architecture)));
-				Arguments.Add(String.Format("--target=\"{0}\"", TargetString));
-
-				if (CompileEnvironment.Configuration == CppConfiguration.Debug)
-				{
-					Arguments.Add("-g -O0");
-				}
-				else
-				{
-					Arguments.Add("-O2");
-				}
-
-				Arguments.Add("--pic");
-
-				// Add include paths to the argument list (filtered by architecture)
-				foreach (DirectoryReference IncludePath in CompileEnvironment.SystemIncludePaths)
-				{
-					if (IsDirectoryForArch(IncludePath.FullName, CompileEnvironment.Architecture))
-					{
-						Arguments.Add(String.Format(" -I\"{0}\"", IncludePath));
-					}
-				}
-				foreach (DirectoryReference IncludePath in CompileEnvironment.UserIncludePaths)
-				{
-					if (IsDirectoryForArch(IncludePath.FullName, CompileEnvironment.Architecture))
-					{
-						Arguments.Add(String.Format(" -I\"{0}\"", IncludePath));
-					}
-				}
-
-				// Preprocessor definitions.
-				foreach (string Definition in CompileEnvironment.Definitions)
-				{
-					Arguments.Add(String.Format(" -D\"{0}\"", Definition));
-				}
-
-				// Consume the included header dependency list
-				FileItem DependencyListFile = FileItem.GetItemByFileReference(FileReference.Combine(OutputDir, InlineArchName(Path.GetFileName(ISPCFile.AbsolutePath) + ".d", CompileEnvironment.Architecture, true)));
-				CompileAction.DependencyListFile = DependencyListFile;
-				CompileAction.PrerequisiteItems.Add(DependencyListFile);
-
-				CompileAction.ProducedItems.UnionWith(CompiledISPCObjFiles);
-
-				CompileAction.CommandArguments = String.Join(" ", Arguments);
-
-				// Add the source file and its included files to the prerequisite item list.
-				CompileAction.PrerequisiteItems.Add(ISPCFile);
-
-				CompileAction.StatusDescription = String.Format("[{0}] [{1}]", CompileEnvironment.Architecture, Path.GetFileName(ISPCFile.AbsolutePath));
-
-				for (int i = 0; i < CompiledISPCObjFiles.Count; i++)
-				{
-					// ISPC compiler can't add suffix on the end of the arch, so copy to put into what linker expects
-					FileReference SourceFile = CompiledISPCObjFiles[i].Location;
-					FileReference TargetFile = FinalISPCObjFiles[i].Location;
-
-					if (SourceFile.Equals(TargetFile))
-					{
-						continue;
-					}
-
-					FileItem SourceFileItem = FileItem.GetItemByFileReference(SourceFile);
-					FileItem TargetFileItem = FileItem.GetItemByFileReference(TargetFile);
-
-					Action CopyAction = Graph.CreateAction(ActionType.BuildProject);
-					CopyAction.CommandDescription = "Copy";
-					CopyAction.CommandPath = BuildHostPlatform.Current.Shell;
-					if (BuildHostPlatform.Current.ShellType == ShellType.Cmd)
-					{
-						CopyAction.CommandArguments = String.Format("/C \"copy /Y \"{0}\" \"{1}\" 1>nul\"", SourceFile, TargetFile);
-					}
-					else
-					{
-						CopyAction.CommandArguments = String.Format("-c 'cp -f \"{0}\" \"{1}\"'", SourceFile.FullName, TargetFile.FullName);
-					}
-					CopyAction.WorkingDirectory = Unreal.EngineSourceDirectory;
-					CopyAction.PrerequisiteItems.Add(SourceFileItem);
-					CopyAction.ProducedItems.Add(TargetFileItem);
-					CopyAction.StatusDescription = TargetFileItem.Location.GetFileName();
-					CopyAction.bCanExecuteRemotely = false;
-					CopyAction.bShouldOutputStatusDescription = false;
-				}
-
-				Result.ObjectFiles.AddRange(FinalISPCObjFiles);
-
-				Logger.LogDebug("   ISPC Compiling {StatusDescription}: \"{CommandPath}\" {CommandArguments}", CompileAction.StatusDescription, CompileAction.CommandPath, CompileAction.CommandArguments);
-			}
-
-			return Result;
-		}
-
 		public override FileItem? LinkFiles(LinkEnvironment LinkEnvironment, bool bBuildImportLibraryOnly, IActionGraphBuilder Graph)
 		{
 			if (!LinkEnvironment.bIsBuildingLibrary)
diff --git a/Engine/Source/Programs/UnrealBuildTool/Platform/Linux/UEBuildLinux.cs b/Engine/Source/Programs/UnrealBuildTool/Platform/Linux/UEBuildLinux.cs
--- a/Engine/Source/Programs/UnrealBuildTool/Platform/Linux/UEBuildLinux.cs
+++ b/Engine/Source/Programs/UnrealBuildTool/Platform/Linux/UEBuildLinux.cs
@@ -24,7 +24,7 @@ namespace UnrealBuildTool
 		{
 			get
 			{
-				if (AppleToolchainArchitectures.ContainsKey(this))
+				if (LinuxToolchainArchitectures.ContainsKey(this))
 				{
 					return LinuxToolchainArchitectures[this];
 				}
diff --git a/Engine/Source/Programs/UnrealBuildTool/ToolChain/ISPCToolChain.cs b/Engine/Source/Programs/UnrealBuildTool/ToolChain/ISPCToolChain.cs
--- a/Engine/Source/Programs/UnrealBuildTool/ToolChain/ISPCToolChain.cs
+++ b/Engine/Source/Programs/UnrealBuildTool/ToolChain/ISPCToolChain.cs
@@ -35,20 +35,22 @@ namespace UnrealBuildTool
 		/// <param name="Platform">Which OS platform to target.</param>
 		/// <param name="Arch">Which architecture inside an OS platform to target. Only used for Android currently.</param>
 		/// <returns>List of instruction set targets passed to ISPC compiler</returns>
-		public virtual List<string> GetISPCCompileTargets(UnrealTargetPlatform Platform, UnrealArch? Arch)
+		public virtual List<string> GetISPCCompileTargets(UnrealTargetPlatform Platform, UnrealArch Arch)
 		{
 			List<string> ISPCTargets = new List<string>();
 
-			// @todo this could be simplified for the arm case - but sse has more options
 			if (UEBuildPlatform.IsPlatformInGroup(Platform, UnrealPlatformGroup.Windows) ||
-				(UEBuildPlatform.IsPlatformInGroup(Platform, UnrealPlatformGroup.Unix) && Platform != UnrealTargetPlatform.LinuxArm64) ||
-				Platform == UnrealTargetPlatform.Mac)
+				UEBuildPlatform.IsPlatformInGroup(Platform, UnrealPlatformGroup.Unix) ||
+				UEBuildPlatform.IsPlatformInGroup(Platform, UnrealPlatformGroup.Apple))
 			{
-				ISPCTargets.AddRange(new string[] { "avx512skx-i32x8", "avx2", "avx", "sse4" });
-			}
-			else if (Platform == UnrealTargetPlatform.LinuxArm64)
-			{
-				ISPCTargets.AddRange(new string[] { "neon" });
+				if (Arch.bIsX64)
+				{
+					ISPCTargets.AddRange(new string[] { "avx512skx-i32x8", "avx2", "avx", "sse4" });
+				}
+				else
+				{
+					ISPCTargets.Add("neon");
+				}
 			}
 			else if (Platform == UnrealTargetPlatform.Android)
 			{
@@ -65,10 +67,6 @@ namespace UnrealBuildTool
 					Logger.LogWarning("Invalid Android architecture for ISPC. At least one architecture (arm64, x64) needs to be selected in the project settings to build");
 				}
 			}
-			else if (Platform == UnrealTargetPlatform.IOS)
-			{
-				ISPCTargets.Add("neon");
-			}
 			else
 			{
 				Logger.LogWarning("Unsupported ISPC platform target!");
@@ -84,34 +82,29 @@ namespace UnrealBuildTool
 		/// <returns>OS string passed to ISPC compiler</returns>
 		public virtual string GetISPCOSTarget(UnrealTargetPlatform Platform)
 		{
-			string ISPCOS = "";
-
 			if (UEBuildPlatform.IsPlatformInGroup(Platform, UnrealPlatformGroup.Windows))
 			{
-				ISPCOS += "windows";
+				return "windows";
 			}
 			else if (UEBuildPlatform.IsPlatformInGroup(Platform, UnrealPlatformGroup.Unix))
 			{
-				ISPCOS += "linux";
+				return "linux";
 			}
-			else if (Platform == UnrealTargetPlatform.Android)
+			else if (UEBuildPlatform.IsPlatformInGroup(Platform, UnrealPlatformGroup.Android))
 			{
-				ISPCOS += "android";
+				return "android";
 			}
-			else if (Platform == UnrealTargetPlatform.IOS)
+			else if (UEBuildPlatform.IsPlatformInGroup(Platform, UnrealPlatformGroup.IOS))
 			{
-				ISPCOS += "ios";
+				return "ios";
 			}
 			else if (Platform == UnrealTargetPlatform.Mac)
 			{
-				ISPCOS += "macos";
-			}
-			else
-			{
-				Logger.LogWarning("Unsupported ISPC platform target!");
+				return "macos";
 			}
-
-			return ISPCOS;
+			
+			Logger.LogWarning("Unsupported ISPC platform target!");
+			return String.Empty;
 		}
 
 		/// <summary>
@@ -120,45 +113,22 @@ namespace UnrealBuildTool
 		/// <param name="Platform">Which OS platform to target.</param>
 		/// <param name="Arch">Which architecture inside an OS platform to target. Only used for Android currently.</param>
 		/// <returns>Arch string passed to ISPC compiler</returns>
-		public virtual string GetISPCArchTarget(UnrealTargetPlatform Platform, UnrealArch? Arch)
+		public virtual string GetISPCArchTarget(UnrealTargetPlatform Platform, UnrealArch Arch)
 		{
-			string ISPCArch = "";
-
 			if (UEBuildPlatform.IsPlatformInGroup(Platform, UnrealPlatformGroup.Windows) ||
-				(UEBuildPlatform.IsPlatformInGroup(Platform, UnrealPlatformGroup.Unix) && Platform != UnrealTargetPlatform.LinuxArm64) ||
-				Platform == UnrealTargetPlatform.Mac)
-			{
-				ISPCArch += "x86-64";
-			}
-			else if (Platform == UnrealTargetPlatform.LinuxArm64)
+				UEBuildPlatform.IsPlatformInGroup(Platform, UnrealPlatformGroup.Unix) ||
+				UEBuildPlatform.IsPlatformInGroup(Platform, UnrealPlatformGroup.Apple) ||
+				UEBuildPlatform.IsPlatformInGroup(Platform, UnrealPlatformGroup.Android))
 			{
-				ISPCArch += "aarch64";
-			}
-			else if (Platform == UnrealTargetPlatform.Android)
-			{
-				if (Arch == UnrealArch.Arm64)
+				if (Arch.bIsX64)
 				{
-					ISPCArch += "aarch64";
-				}
-				else if (Arch == UnrealArch.X64)
-				{
-					ISPCArch += "x86-64";
-				}
-				else
-				{
-					Logger.LogWarning("Invalid Android architecture for ISPC. At least one architecture (arm64, x64) needs to be selected in the project settings to build");
+					return "x86-64";
 				}
+				return "aarch64";
 			}
-			else if (Platform == UnrealTargetPlatform.IOS)
-			{
-				ISPCArch += "aarch64";
-			}
-			else
-			{
-				Logger.LogWarning("Unsupported ISPC platform target!");
-			}
-
-			return ISPCArch;
+			
+			Logger.LogWarning("Unsupported ISPC platform target!");
+			return String.Empty;
 		}
 
 		/// <summary>
@@ -249,25 +219,16 @@ namespace UnrealBuildTool
 		/// <returns>Object file suffix</returns>
 		public virtual string GetISPCObjectFileFormat(UnrealTargetPlatform Platform)
 		{
-			string Format = "";
-
-			if (UEBuildPlatform.IsPlatformInGroup(Platform, UnrealPlatformGroup.Windows))
-			{
-				Format += "obj";
-			}
-			else if (UEBuildPlatform.IsPlatformInGroup(Platform, UnrealPlatformGroup.Unix) ||
-					Platform == UnrealTargetPlatform.Mac ||
-					Platform == UnrealTargetPlatform.IOS ||
-					Platform == UnrealTargetPlatform.Android)
-			{
-				Format += "obj";
-			}
-			else
+			if (UEBuildPlatform.IsPlatformInGroup(Platform, UnrealPlatformGroup.Windows) ||
+				UEBuildPlatform.IsPlatformInGroup(Platform, UnrealPlatformGroup.Unix) ||
+				UEBuildPlatform.IsPlatformInGroup(Platform, UnrealPlatformGroup.Apple) ||
+				UEBuildPlatform.IsPlatformInGroup(Platform, UnrealPlatformGroup.Android))
 			{
-				Logger.LogWarning("Unsupported ISPC platform target!");
+				return "obj";
 			}
 
-			return Format;
+			Logger.LogWarning("Unsupported ISPC platform target!");
+			return String.Empty;
 		}
 
 		/// <summary>
@@ -277,25 +238,19 @@ namespace UnrealBuildTool
 		/// <returns>Object file suffix</returns>
 		public virtual string GetISPCObjectFileSuffix(UnrealTargetPlatform Platform)
 		{
-			string Suffix = "";
-
 			if (UEBuildPlatform.IsPlatformInGroup(Platform, UnrealPlatformGroup.Windows))
 			{
-				Suffix += ".obj";
+				return ".obj";
 			}
 			else if (UEBuildPlatform.IsPlatformInGroup(Platform, UnrealPlatformGroup.Unix) ||
-					Platform == UnrealTargetPlatform.Mac ||
-					Platform == UnrealTargetPlatform.IOS ||
-					Platform == UnrealTargetPlatform.Android)
-			{
-				Suffix += ".o";
-			}
-			else
+				UEBuildPlatform.IsPlatformInGroup(Platform, UnrealPlatformGroup.Apple) ||
+				UEBuildPlatform.IsPlatformInGroup(Platform, UnrealPlatformGroup.Android))
 			{
-				Logger.LogWarning("Unsupported ISPC platform target!");
+				return ".o";
 			}
 
-			return Suffix;
+			Logger.LogWarning("Unsupported ISPC platform target!");
+			return String.Empty;
 		}
 
 		private string EscapeDefinitionForISPC(string Definition)
@@ -424,7 +379,7 @@ namespace UnrealBuildTool
 			yield return DirectoryItem.GetItemByDirectoryReference(Unreal.RootDirectory);
 		}
 
-		public override CPPOutput GenerateISPCHeaders(CppCompileEnvironment CompileEnvironment, IEnumerable<FileItem> InputFiles, DirectoryReference OutputDir, IActionGraphBuilder Graph)
+		protected override CPPOutput GenerateISPCHeaders(CppCompileEnvironment CompileEnvironment, IEnumerable<FileItem> InputFiles, DirectoryReference OutputDir, IActionGraphBuilder Graph)
 		{
 			CPPOutput Result = new CPPOutput();
 
@@ -433,14 +388,14 @@ namespace UnrealBuildTool
 				return Result;
 			}
 
-			List<string> CompileTargets = GetISPCCompileTargets(CompileEnvironment.Platform, null);
+			List<string> CompileTargets = GetISPCCompileTargets(CompileEnvironment.Platform, CompileEnvironment.Architecture);
 
 			List<string> GlobalArguments = new List<string>();
 
 			// Build target string. No comma on last
 			string TargetString = String.Join(',', CompileTargets);
 
-			string ISPCArch = GetISPCArchTarget(CompileEnvironment.Platform, null);
+			string ISPCArch = GetISPCArchTarget(CompileEnvironment.Platform, CompileEnvironment.Architecture);
 
 			// Build target triplet
 			GlobalArguments.Add($"--target-os={GetISPCOSTarget(CompileEnvironment.Platform)}");
@@ -454,10 +409,11 @@ namespace UnrealBuildTool
 				GlobalArguments.Add($"--cpu={CpuTarget}");
 			}
 
-			// PIC is needed for modular builds except on Microsoft platforms
+			// PIC is needed for modular builds except on Microsoft platforms, and for android
 			if ((CompileEnvironment.bIsBuildingDLL ||
 				CompileEnvironment.bIsBuildingLibrary) &&
-				!UEBuildPlatform.IsPlatformInGroup(CompileEnvironment.Platform, UnrealPlatformGroup.Microsoft))
+				!UEBuildPlatform.IsPlatformInGroup(CompileEnvironment.Platform, UnrealPlatformGroup.Microsoft) ||
+				UEBuildPlatform.IsPlatformInGroup(CompileEnvironment.Platform, UnrealPlatformGroup.Android))
 			{
 				GlobalArguments.Add("--pic");
 			}
@@ -563,7 +519,7 @@ namespace UnrealBuildTool
 			return Result;
 		}
 
-		public override CPPOutput CompileISPCFiles(CppCompileEnvironment CompileEnvironment, IEnumerable<FileItem> InputFiles, DirectoryReference OutputDir, IActionGraphBuilder Graph)
+		protected override CPPOutput CompileISPCFiles(CppCompileEnvironment CompileEnvironment, IEnumerable<FileItem> InputFiles, DirectoryReference OutputDir, IActionGraphBuilder Graph)
 		{
 			CPPOutput Result = new CPPOutput();
 
@@ -572,7 +528,7 @@ namespace UnrealBuildTool
 				return Result;
 			}
 
-			List<string> CompileTargets = GetISPCCompileTargets(CompileEnvironment.Platform, null);
+			List<string> CompileTargets = GetISPCCompileTargets(CompileEnvironment.Platform, CompileEnvironment.Architecture);
 
 			List<string> GlobalArguments = new List<string>();
 
@@ -592,7 +548,7 @@ namespace UnrealBuildTool
 
 			// Build target triplet
 			string PlatformObjectFileFormat = GetISPCObjectFileFormat(CompileEnvironment.Platform);
-			string ISPCArch = GetISPCArchTarget(CompileEnvironment.Platform, null);
+			string ISPCArch = GetISPCArchTarget(CompileEnvironment.Platform, CompileEnvironment.Architecture);
 
 			GlobalArguments.Add($"--target-os={GetISPCOSTarget(CompileEnvironment.Platform)}");
 			GlobalArguments.Add($"--arch={ISPCArch}");
@@ -630,10 +586,11 @@ namespace UnrealBuildTool
 			}
 			GlobalArguments.AddRange(CommonArgs);
 
-			// PIC is needed for modular builds except on Microsoft platforms
+			// PIC is needed for modular builds except on Microsoft platforms, and for android
 			if ((CompileEnvironment.bIsBuildingDLL ||
 				CompileEnvironment.bIsBuildingLibrary) &&
-				!UEBuildPlatform.IsPlatformInGroup(CompileEnvironment.Platform, UnrealPlatformGroup.Microsoft))
+				!UEBuildPlatform.IsPlatformInGroup(CompileEnvironment.Platform, UnrealPlatformGroup.Microsoft) ||
+				UEBuildPlatform.IsPlatformInGroup(CompileEnvironment.Platform, UnrealPlatformGroup.Android))
 			{
 				GlobalArguments.Add("--pic");
 			}
diff --git a/Engine/Source/Programs/UnrealBuildTool/ToolChain/IWYUToolChain.cs b/Engine/Source/Programs/UnrealBuildTool/ToolChain/IWYUToolChain.cs
--- a/Engine/Source/Programs/UnrealBuildTool/ToolChain/IWYUToolChain.cs
+++ b/Engine/Source/Programs/UnrealBuildTool/ToolChain/IWYUToolChain.cs
@@ -199,12 +199,12 @@ namespace UnrealBuildTool
 		}
 
 		// Skip ISPC headers
-		public override CPPOutput GenerateISPCHeaders(CppCompileEnvironment CompileEnvironment, IEnumerable<FileItem> InputFiles, DirectoryReference OutputDir, IActionGraphBuilder Graph)
+		protected override CPPOutput GenerateISPCHeaders(CppCompileEnvironment CompileEnvironment, IEnumerable<FileItem> InputFiles, DirectoryReference OutputDir, IActionGraphBuilder Graph)
 		{
 			return new CPPOutput();
 		}
 
-		public override CPPOutput CompileISPCFiles(CppCompileEnvironment CompileEnvironment, IEnumerable<FileItem> InputFiles, DirectoryReference OutputDir, IActionGraphBuilder Graph)
+		protected override CPPOutput CompileISPCFiles(CppCompileEnvironment CompileEnvironment, IEnumerable<FileItem> InputFiles, DirectoryReference OutputDir, IActionGraphBuilder Graph)
 		{
 			return new CPPOutput();
 		}
diff --git a/Engine/Source/Programs/UnrealBuildTool/ToolChain/UEToolChain.cs b/Engine/Source/Programs/UnrealBuildTool/ToolChain/UEToolChain.cs
--- a/Engine/Source/Programs/UnrealBuildTool/ToolChain/UEToolChain.cs
+++ b/Engine/Source/Programs/UnrealBuildTool/ToolChain/UEToolChain.cs
@@ -77,6 +77,11 @@ namespace UnrealBuildTool
 					DirectoryReference ArchOutputDir = new(OutputDir.FullName.Replace(PlatformArchitecturesString, ArchConfig.GetFolderNameForArchitecture(Arch)));
 
 					CppCompileEnvironment ArchEnvironment = new(CompileEnvironment, Arch);
+					if (ArchEnvironment.UserIncludePaths.Contains(OutputDir))
+					{
+						ArchEnvironment.UserIncludePaths.Remove(OutputDir);
+						ArchEnvironment.UserIncludePaths.Add(ArchOutputDir);
+					}
 					CPPOutput ArchResult = CompileCPPFiles(ArchEnvironment, InputFiles, ArchOutputDir, ModuleName, Graph);
 					Result.Merge(ArchResult, Arch);
 				}
@@ -95,15 +100,61 @@ namespace UnrealBuildTool
 			return Result;
 		}
 
-		public virtual CPPOutput CompileISPCFiles(CppCompileEnvironment CompileEnvironment, IEnumerable<FileItem> InputFiles, DirectoryReference OutputDir, IActionGraphBuilder Graph)
+		protected abstract CPPOutput CompileISPCFiles(CppCompileEnvironment CompileEnvironment, IEnumerable<FileItem> InputFiles, DirectoryReference OutputDir, IActionGraphBuilder Graph);
+
+		public CPPOutput CompileAllISPCFiles(CppCompileEnvironment CompileEnvironment, IEnumerable<FileItem> InputFiles, DirectoryReference OutputDir, IActionGraphBuilder Graph)
 		{
-			CPPOutput Result = new CPPOutput();
+			CPPOutput Result;
+
+			UnrealArchitectureConfig ArchConfig = UnrealArchitectureConfig.ForPlatform(CompileEnvironment.Platform);
+			// compile architectures separately if needed
+			if (ArchConfig.Mode == UnrealArchitectureMode.SingleTargetCompileSeparately || ArchConfig.Mode == UnrealArchitectureMode.SingleTargetLinkSeparately)
+			{
+				Result = new CPPOutput();
+				foreach (UnrealArch Arch in CompileEnvironment.Architectures.Architectures)
+				{
+					string PlatformArchitecturesString = ArchConfig.GetFolderNameForArchitectures(CompileEnvironment.Architectures);
+					DirectoryReference ArchOutputDir = new(OutputDir.FullName.Replace(PlatformArchitecturesString, ArchConfig.GetFolderNameForArchitecture(Arch)));
+
+					CppCompileEnvironment ArchEnvironment = new(CompileEnvironment, Arch);
+					CPPOutput ArchResult = CompileISPCFiles(ArchEnvironment, InputFiles, ArchOutputDir, Graph);
+					Result.Merge(ArchResult, Arch);
+				}
+			}
+			else
+			{
+				Result = CompileISPCFiles(CompileEnvironment, InputFiles, OutputDir, Graph);
+			}
+
 			return Result;
 		}
 
-		public virtual CPPOutput GenerateISPCHeaders(CppCompileEnvironment CompileEnvironment, IEnumerable<FileItem> InputFiles, DirectoryReference OutputDir, IActionGraphBuilder Graph)
+		protected abstract CPPOutput GenerateISPCHeaders(CppCompileEnvironment CompileEnvironment, IEnumerable<FileItem> InputFiles, DirectoryReference OutputDir, IActionGraphBuilder Graph);
+
+		public CPPOutput GenerateAllISPCHeaders(CppCompileEnvironment CompileEnvironment, IEnumerable<FileItem> InputFiles, DirectoryReference OutputDir, IActionGraphBuilder Graph)
 		{
-			CPPOutput Result = new CPPOutput();
+			CPPOutput Result;
+
+			UnrealArchitectureConfig ArchConfig = UnrealArchitectureConfig.ForPlatform(CompileEnvironment.Platform);
+			// compile architectures separately if needed
+			if (ArchConfig.Mode == UnrealArchitectureMode.SingleTargetCompileSeparately || ArchConfig.Mode == UnrealArchitectureMode.SingleTargetLinkSeparately)
+			{
+				Result = new CPPOutput();
+				foreach (UnrealArch Arch in CompileEnvironment.Architectures.Architectures)
+				{
+					string PlatformArchitecturesString = ArchConfig.GetFolderNameForArchitectures(CompileEnvironment.Architectures);
+					DirectoryReference ArchOutputDir = new(OutputDir.FullName.Replace(PlatformArchitecturesString, ArchConfig.GetFolderNameForArchitecture(Arch)));
+
+					CppCompileEnvironment ArchEnvironment = new(CompileEnvironment, Arch);
+					CPPOutput ArchResult = GenerateISPCHeaders(ArchEnvironment, InputFiles, ArchOutputDir, Graph);
+					Result.Merge(ArchResult, Arch);
+				}
+			}
+			else
+			{
+				Result = GenerateISPCHeaders(CompileEnvironment, InputFiles, OutputDir, Graph);
+			}
+
 			return Result;
 		}
 
