UnrealBuildTool: Allow ISPC for Win64 arm64

#rnx

diff --git a/Engine/Source/Programs/UnrealBuildTool/Platform/Windows/UEBuildWindows.cs b/Engine/Source/Programs/UnrealBuildTool/Platform/Windows/UEBuildWindows.cs
--- a/Engine/Source/Programs/UnrealBuildTool/Platform/Windows/UEBuildWindows.cs
+++ b/Engine/Source/Programs/UnrealBuildTool/Platform/Windows/UEBuildWindows.cs
@@ -1294,11 +1294,6 @@ namespace UnrealBuildTool
 
 			Target.bCompileISPC = true;
 
-			if (Platform == UnrealTargetPlatform.Win64 && !Target.Architecture.bIsX64)
-			{
-				Target.bCompileISPC = false; // The version of ISPC we currently use does not support Windows Aarch64
-			}
-
 			if (OperatingSystem.IsWindows())
 			{
 				// Initialize the VC environment for the target, and set all the version numbers to the concrete values we chose
