- Added some missing quotes around engine directory location for swift compilation

diff --git a/Engine/Source/Programs/UnrealBuildTool/ToolChain/AppleToolChain.cs b/Engine/Source/Programs/UnrealBuildTool/ToolChain/AppleToolChain.cs
--- a/Engine/Source/Programs/UnrealBuildTool/ToolChain/AppleToolChain.cs
+++ b/Engine/Source/Programs/UnrealBuildTool/ToolChain/AppleToolChain.cs
@@ -672,7 +672,7 @@ namespace UnrealBuildTool
 
 			Arguments.Add("-enable-objc-interop");
 			Arguments.Add("-cxx-interoperability-mode=default");
-			Arguments.Add($"-import-objc-header {Unreal.EngineDirectory}/Source/Runtime/Launch/Private/IOS/UECppToSwift.h");
+			Arguments.Add($"-import-objc-header \"{Unreal.EngineDirectory}/Source/Runtime/Launch/Private/IOS/UECppToSwift.h\"");
 			Arguments.Add($"-module-name Launch");
 
 			if (bUseSwiftUIMain)
@@ -713,7 +713,7 @@ namespace UnrealBuildTool
 
 				//Arguments.Add("-enable-objc-interop");
 				//Arguments.Add("-cxx-interoperability-mode=default");
-				Arguments.Add($"-import-objc-header {Unreal.EngineDirectory}/Source/Runtime/Launch/Private/IOS/UECppToSwift.h");
+				Arguments.Add($"-import-objc-header \"{Unreal.EngineDirectory}/Source/Runtime/Launch/Private/IOS/UECppToSwift.h\"");
 
 				// platform settings
 				Arguments.Add($"-target {ToolChainSettings.Value.GetTargetTuple(CompileEnvironment.Architecture)}");
