[UBT]
* Fixed so linux work with module merging

diff --git a/Engine/Source/Programs/UnrealBuildTool/Platform/Linux/LinuxToolChain.cs b/Engine/Source/Programs/UnrealBuildTool/Platform/Linux/LinuxToolChain.cs
--- a/Engine/Source/Programs/UnrealBuildTool/Platform/Linux/LinuxToolChain.cs
+++ b/Engine/Source/Programs/UnrealBuildTool/Platform/Linux/LinuxToolChain.cs
@@ -1191,7 +1191,14 @@ namespace UnrealBuildTool
 			List<string> ResponseLines = new List<string>();
 			foreach (FileItem InputFile in LinkEnvironment.InputFiles)
 			{
-				ResponseLines.Add(String.Format("\"{0}\"", InputFile.AbsolutePath.Replace("\\", "/")));
+				if (InputFile.HasExtension(".dynlist"))
+				{
+					ResponseLines.Add($"--dynamic-list=\"{NormalizeCommandLinePath(InputFile)}\"");
+				}
+				else
+				{
+					ResponseLines.Add($"\"{NormalizeCommandLinePath(InputFile)}\"");
+				}
 				LinkAction.PrerequisiteItems.Add(InputFile);
 			}
 
@@ -1749,5 +1756,10 @@ namespace UnrealBuildTool
 				ExtraArguments.Add(String.Format("-isystem {0}", System.IO.Path.Combine(InternalSdkPath, "usr", "include").Replace("\\", "/")));
 			}
 		}
+
+		public override string GetExtraLinkFileExtension()
+		{
+			return "dynlist";
+		}
 	}
 }
