UnrealBuildTool: Fix support for modular build

#rnx

diff --git a/Engine/Source/Programs/UnrealBuildTool/Configuration/ReadOnlyTargetRules.cs b/Engine/Source/Programs/UnrealBuildTool/Configuration/ReadOnlyTargetRules.cs
--- a/Engine/Source/Programs/UnrealBuildTool/Configuration/ReadOnlyTargetRules.cs
+++ b/Engine/Source/Programs/UnrealBuildTool/Configuration/ReadOnlyTargetRules.cs
@@ -549,6 +549,8 @@ namespace UnrealBuildTool
 
 		public IReadOnlyDictionary<string, IEnumerable<string>> MergePluginsShared => Inner.MergePluginsShared.ToDictionary(k => k.Key, v => v.Value.AsEnumerable());
 
+		public IReadOnlyList<string> MergePluginsLaunch => Inner.MergePluginsLaunch;
+
 		public IReadOnlyList<string> GlobalDefinitions => Inner.GlobalDefinitions.AsReadOnly();
 
 		public IReadOnlyList<string> ProjectDefinitions => Inner.ProjectDefinitions.AsReadOnly();
diff --git a/Engine/Source/Programs/UnrealBuildTool/Configuration/TargetRules.cs b/Engine/Source/Programs/UnrealBuildTool/Configuration/TargetRules.cs
--- a/Engine/Source/Programs/UnrealBuildTool/Configuration/TargetRules.cs
+++ b/Engine/Source/Programs/UnrealBuildTool/Configuration/TargetRules.cs
@@ -2531,10 +2531,15 @@ namespace UnrealBuildTool
 		public List<string> MergePlugins { get; } = new();
 
 		/// <summary>
-		/// Experimental: List of merged plugins to move common dependenceis to a shared library, requires MergePlugins. Can be chained
+		/// Experimental: List of merged plugins to move common dependencies to a shared library, requires MergePlugins. Can be chained
 		/// </summary>
 		public Dictionary<string, HashSet<string>> MergePluginsShared { get; } = new();
 
+		/// <summary>
+		/// Experimental: List of merged plugins to move to the main executable, requires MergePlugins. "Engine" for all shared engine modules and "Common" for shared project modules.
+		/// </summary>
+		public List<string> MergePluginsLaunch { get; } = new();
+
 		/// <summary>
 		/// Backing storage for the LinkType property.
 		/// </summary>
diff --git a/Engine/Source/Programs/UnrealBuildTool/Configuration/UEBuildTarget.cs b/Engine/Source/Programs/UnrealBuildTool/Configuration/UEBuildTarget.cs
--- a/Engine/Source/Programs/UnrealBuildTool/Configuration/UEBuildTarget.cs
+++ b/Engine/Source/Programs/UnrealBuildTool/Configuration/UEBuildTarget.cs
@@ -3480,13 +3480,14 @@ namespace UnrealBuildTool
 		/// <param name="Logger"></param>
 		protected void SetupMergeModules(ILogger Logger)
 		{
-			void MergeModules(string name, IEnumerable<UEBuildModuleCPP> modules)
+			void MergeModulesToDLL(string name, IEnumerable<UEBuildModuleCPP> modules)
 			{
 				if (!modules.Any())
 				{
 					return;
 				}
 
+				List<UEBuildModuleCPP> modulesList = modules.OrderBy(m => m.Name).ToList();
 				UEBuildBinary mergedBinary = CreateDynamicLibraryForModules(name, modules.OrderBy(x => x.Name));
 				foreach (UEBuildModuleCPP module in modules)
 				{
@@ -3496,6 +3497,59 @@ namespace UnrealBuildTool
 				Binaries.Add(mergedBinary);
 			}
 
+			void MergeModulesToLaunch(IEnumerable<UEBuildModuleCPP> modules)
+			{
+				if (!modules.Any())
+				{
+					return;
+				}
+				foreach (UEBuildModuleCPP module in modules)
+				{
+					Binaries[0].AddModule(module);
+					module.Binary = Binaries[0];
+				}
+			}
+
+			void ConsumeModules(string name, IEnumerable<UEBuildModuleCPP> modules)
+			{
+				if (Rules.MergePluginsLaunch.Contains(name))
+				{
+					MergeModulesToLaunch(modules);
+				}
+				else
+				{
+					MergeModulesToDLL(name, modules);
+				}
+			}
+		
+			void GetDependenciesRecursively(IEnumerable<UEBuildModuleCPP> Modules, ref HashSet<UEBuildModuleCPP> Needed)
+			{
+				IEnumerable<UEBuildModuleCPP> CPPModules = Modules.OfType<UEBuildModuleCPP>();
+
+				foreach (UEBuildModuleCPP item in CPPModules)
+				{
+					bool bAddDependencies = false;
+					if (!Needed.Contains(item))
+					{
+						Needed.Add(item);
+						bAddDependencies = true;
+					}
+
+					if (bAddDependencies)
+					{
+						if (item.PrivateDependencyModules != null && item.PrivateDependencyModules.Count > 0)
+						{
+							GetDependenciesRecursively(item.PrivateDependencyModules.OfType<UEBuildModuleCPP>(), ref Needed);
+						}
+
+						if (item.PublicDependencyModules != null && item.PublicDependencyModules.Count > 0)
+						{
+							GetDependenciesRecursively(item.PublicDependencyModules.OfType<UEBuildModuleCPP>(), ref Needed);
+						}
+					}
+				}
+			}
+
 			UEBuildModuleCPP launchModule = Binaries[0].PrimaryModule;
 			HashSet<UEBuildModuleCPP> allModules = new(Binaries.SelectMany(x => x.Modules.OfType<UEBuildModuleCPP>().Where(x => x != launchModule)));
 
@@ -3558,30 +3612,49 @@ namespace UnrealBuildTool
 				unmergablePlugins.ExceptWith(pluginGroups.SelectMany(x => x.Value));
 				unmergablePlugins = new(unmergablePlugins.Union(unmergablePlugins.SelectMany(x => x.Dependencies ?? new())));
 
-				// Create the merged plugin modules
+				HashSet<UEBuildModuleCPP> allPluginModules = new HashSet<UEBuildModuleCPP>();
+
+				// Remove unmergable plugin modules from plugin modules
+				// and create a set of all plugin modules
 				foreach (KeyValuePair<string, HashSet<UEBuildPlugin>> item in pluginGroups)
 				{
 					// Find all modules for the unique plugins to be merged, except those that were unmergable due to being a dependency of an unmerged plugin
-					IEnumerable<UEBuildModuleCPP> pluginModules = item.Value.Except(unmergablePlugins).SelectMany(x => x.Modules).Distinct();
+					item.Value.ExceptWith(unmergablePlugins);
+					IEnumerable<UEBuildModuleCPP> pluginModules = item.Value.SelectMany(x => x.Modules).Distinct();
 
-					// Merge the plugin modules and remove those modules from remaining set of modules to merge 
-					MergeModules(item.Key, pluginModules);
-					allModules.ExceptWith(pluginModules);
+					allPluginModules.UnionWith(pluginModules);
 				}
-			}
 
-			IEnumerable<UEBuildModuleCPP> engineModules = allModules.Where(x => x.RulesFile.IsUnderDirectory(Unreal.EngineDirectory) && !x.Rules.IsPlugin);
-			IEnumerable<UEBuildModuleCPP> enginePluginModules = allModules.Where(x => x.RulesFile.IsUnderDirectory(Unreal.EngineDirectory) && x.Rules.IsPlugin);
-			IEnumerable<UEBuildModuleCPP> projectModules = allModules.Except(engineModules).Except(enginePluginModules).Where(x => !x.Rules.IsPlugin);
-			IEnumerable<UEBuildModuleCPP> projectPluginModules = allModules.Except(engineModules).Except(enginePluginModules).Where(x => x.Rules.IsPlugin);
+				// create a set of leftover modules from all modules after removing plugin modules
+				HashSet<UEBuildModuleCPP> allLeftoverModules = new HashSet<UEBuildModuleCPP>(allModules);
+				allLeftoverModules.ExceptWith(allPluginModules);
+
+				// create a set of all dependencies of leftover modules
+				HashSet<UEBuildModuleCPP> allLeftoverDependenies = new HashSet<UEBuildModuleCPP>();
+				GetDependenciesRecursively(allLeftoverModules.OfType<UEBuildModuleCPP>(), ref allLeftoverDependenies);
+
+				// Remove all leftover dependencies from plugin modules
+				// then remove remaining plugin modules from all modules set
+				// then merge the module into a binary and update plugin modules in all modules to that binary
+				foreach (KeyValuePair<string, HashSet<UEBuildPlugin>> item in pluginGroups)
+				{
+					// Find all modules for the unique plugins to be merged, except those that were unmergable due to being a dependency of an unmerged plugin
+					HashSet<UEBuildModuleCPP> pluginModules = item.Value.SelectMany(x => x.Modules).ToHashSet();
 
+					pluginModules.ExceptWith(allLeftoverDependenies);
+					allModules.ExceptWith(pluginModules);
+
+					ConsumeModules(item.Key, pluginModules);
+				}
+			}
+			
 			// Merge remaining engine modules
-			MergeModules("CommonEngine", engineModules);
-			MergeModules("CommonEnginePlugins", enginePluginModules);
+			IEnumerable<UEBuildModuleCPP> engineModules = allModules.Where(x => x.RulesFile.IsUnderDirectory(Unreal.EngineDirectory));
+			ConsumeModules("Engine", engineModules);
 
-			// Merge remaining project modules
-			MergeModules("Common", projectModules);
-			MergeModules("CommonPlugins", projectPluginModules);
+			// Merge remaining modules to common set
+			IEnumerable<UEBuildModuleCPP> commonModules = allModules.Except(engineModules);
+			ConsumeModules("Common", commonModules);
 
 			// Remove any binaries that no longer have modules due to merging
 			Binaries.ForEach(x => x.Modules.RemoveAll(y => y.Binary != x));
@@ -4452,7 +4525,7 @@ namespace UnrealBuildTool
 				ModuleConfiguration = UnrealTargetConfiguration.Development;
 			}
 
-			List<FileReference> OutputPaths = MakeBinaryPaths(OutputDirectory, AppName + "-" + MergedName, Platform, Configuration, UEBuildBinaryType.DynamicLinkLibrary, Rules.Architectures, Rules.UndecoratedConfiguration, bCompileMonolithic && ProjectFile != null, Rules.ExeBinariesSubFolder, ProjectFile, Rules);
+			List<FileReference> OutputPaths = MakeBinaryPaths(OutputDirectory, AppName + "_" + MergedName, Platform, Configuration, UEBuildBinaryType.DynamicLinkLibrary, Rules.Architectures, Rules.UndecoratedConfiguration, bCompileMonolithic && ProjectFile != null, Rules.ExeBinariesSubFolder, ProjectFile, Rules);
 
 			// Create the binary
 			UEBuildBinary Binary = new UEBuildBinary(
