[UBT]
* Added so executable allow symbol exporting in merged modules. This is needed since dlls might depend on the executable for symbol imports. With merged modules unused exports are stripped out anyway so if noone depends on the binary it will have zero exports

diff --git a/Engine/Source/Programs/UnrealBuildTool/Configuration/UEBuildTarget.cs b/Engine/Source/Programs/UnrealBuildTool/Configuration/UEBuildTarget.cs
--- a/Engine/Source/Programs/UnrealBuildTool/Configuration/UEBuildTarget.cs
+++ b/Engine/Source/Programs/UnrealBuildTool/Configuration/UEBuildTarget.cs
@@ -3550,7 +3550,10 @@ namespace UnrealBuildTool
 				}
 			}
 
-			UEBuildModuleCPP launchModule = Binaries[0].PrimaryModule;
+			UEBuildBinary executableBinary = Binaries[0];
+			executableBinary.bAllowExports = true; // Need to export symbols because other binaries might depend on this. (They will be stripped anyway if not used)
+
+			UEBuildModuleCPP launchModule = executableBinary.PrimaryModule;
 			HashSet<UEBuildModuleCPP> allModules = new(Binaries.SelectMany(x => x.Modules.OfType<UEBuildModuleCPP>().Where(x => x != launchModule)));
 
 			if (Rules.MergePlugins.Any() && BuildPlugins != null)
