UnrealBuildTool: Add a config setting for the binary filename separator

#rnx

diff --git a/Engine/Source/Programs/AutomationTool/Apple/ApplePlatform.Automation.cs b/Engine/Source/Programs/AutomationTool/Apple/ApplePlatform.Automation.cs
--- a/Engine/Source/Programs/AutomationTool/Apple/ApplePlatform.Automation.cs
+++ b/Engine/Source/Programs/AutomationTool/Apple/ApplePlatform.Automation.cs
@@ -169,7 +169,7 @@ public abstract class ApplePlatform : Platform
 
 		// find the most recent archive for this target (based on name of target, this ignores Development vs Shipping, but 
 		// since Distribution is meant only for Shipping it's ok
-		string ArchiveName = AppleExports.MakeBinaryFileName(TargetName, Target.Receipt.Platform, UnrealTargetConfiguration.Development,
+		string ArchiveName = AppleExports.MakeBinaryFileName(TargetName, "-", Target.Receipt.Platform, UnrealTargetConfiguration.Development,
 		 	Target.Receipt.Architectures, UnrealTargetConfiguration.Development, null);
 
 		if (ArchiveName != TargetName)
@@ -186,7 +186,7 @@ public abstract class ApplePlatform : Platform
 		// content only projects have to not use the executable, but make a path to a .app 
 		if (!SC.IsCodeBasedProject)
 		{
-			string AppBundleName = AppleExports.MakeBinaryFileName(SC.ShortProjectName, Target.Receipt.Platform, Target.Receipt.Configuration, 
+			string AppBundleName = AppleExports.MakeBinaryFileName(SC.ShortProjectName, "-", Target.Receipt.Platform, Target.Receipt.Configuration, 
 				Target.Receipt.Architectures, UnrealTargetConfiguration.Development, ".app");
 
 			return DirectoryReference.Combine(SC.ProjectRoot, "Binaries", Target.Receipt.Platform.ToString(), AppBundleName);
diff --git a/Engine/Source/Programs/AutomationTool/IOS/IOSPlatform.Automation.cs b/Engine/Source/Programs/AutomationTool/IOS/IOSPlatform.Automation.cs
--- a/Engine/Source/Programs/AutomationTool/IOS/IOSPlatform.Automation.cs
+++ b/Engine/Source/Programs/AutomationTool/IOS/IOSPlatform.Automation.cs
@@ -1870,7 +1870,7 @@ public class IOSPlatform : ApplePlatform
 			if (!SC.IsCodeBasedProject)
 			{
 				TargetReceipt Target = SC.StageTargets[0].Receipt;
-				AppBaseName = AppleExports.MakeBinaryFileName(SC.ShortProjectName, Target.Platform, Target.Configuration, Target.Architectures, UnrealTargetConfiguration.Development, null);
+				AppBaseName = AppleExports.MakeBinaryFileName(SC.ShortProjectName, "-", Target.Platform, Target.Configuration, Target.Architectures, UnrealTargetConfiguration.Development, null);
 			}
 			AppToDeploy = FileReference.Combine(SC.StageDirectory, AppBaseName + ".app").FullName;
 
diff --git a/Engine/Source/Programs/UnrealBuildTool/Configuration/ReadOnlyTargetRules.cs b/Engine/Source/Programs/UnrealBuildTool/Configuration/ReadOnlyTargetRules.cs
--- a/Engine/Source/Programs/UnrealBuildTool/Configuration/ReadOnlyTargetRules.cs
+++ b/Engine/Source/Programs/UnrealBuildTool/Configuration/ReadOnlyTargetRules.cs
@@ -96,6 +96,8 @@ namespace UnrealBuildTool
 
 		public UnrealTargetConfiguration UndecoratedConfiguration => Inner.UndecoratedConfiguration;
 
+		public string DecoratedSeparator => Inner.DecoratedSeparator;
+
 		public bool bAllowHotReload => Inner.bAllowHotReload;
 
 		public bool bBuildAllModules => Inner.bBuildAllModules;
diff --git a/Engine/Source/Programs/UnrealBuildTool/Configuration/TargetRules.cs b/Engine/Source/Programs/UnrealBuildTool/Configuration/TargetRules.cs
--- a/Engine/Source/Programs/UnrealBuildTool/Configuration/TargetRules.cs
+++ b/Engine/Source/Programs/UnrealBuildTool/Configuration/TargetRules.cs
@@ -728,6 +728,12 @@ namespace UnrealBuildTool
 		[RequiresUniqueBuildEnvironment]
 		public UnrealTargetConfiguration UndecoratedConfiguration { get; set; } = UnrealTargetConfiguration.Development;
 
+		/// <summary>
+		/// Specifies the separator charcter for binary filenames.
+		/// </summary>
+		[RequiresUniqueBuildEnvironment]
+		public string DecoratedSeparator { get; set; } = "-";
+
 		/// <summary>
 		/// Whether this target supports hot reload
 		/// </summary>
diff --git a/Engine/Source/Programs/UnrealBuildTool/Configuration/UEBuildTarget.cs b/Engine/Source/Programs/UnrealBuildTool/Configuration/UEBuildTarget.cs
--- a/Engine/Source/Programs/UnrealBuildTool/Configuration/UEBuildTarget.cs
+++ b/Engine/Source/Programs/UnrealBuildTool/Configuration/UEBuildTarget.cs
@@ -4484,7 +4484,7 @@ namespace UnrealBuildTool
 			}
 
 			// Get the output filenames
-			FileReference BaseBinaryPath = FileReference.Combine(OutputDirectory, MakeBinaryFileName(AppName + "-" + Module.Name, Platform, ModuleConfiguration, Architectures, Rules.UndecoratedConfiguration, UEBuildBinaryType.DynamicLinkLibrary));
+			FileReference BaseBinaryPath = FileReference.Combine(OutputDirectory, MakeBinaryFileName($"{AppName}{Rules.DecoratedSeparator}{Module.Name}", Rules.DecoratedSeparator, Platform, ModuleConfiguration, Architectures, Rules.UndecoratedConfiguration, UEBuildBinaryType.DynamicLinkLibrary));
 			List<FileReference> OutputFilePaths = UEBuildPlatform.GetBuildPlatform(Platform).FinalizeBinaryPaths(BaseBinaryPath, ProjectFile, Rules);
 
 			// Create the binary
@@ -4550,13 +4550,14 @@ namespace UnrealBuildTool
 		/// Makes a filename (without path) for a compiled binary (e.g. "Core-Win64-Debug.lib") */
 		/// </summary>
 		/// <param name="BinaryName">The name of this binary</param>
+		/// <param name="Separator">The separator for a decorated filename</param>
 		/// <param name="Platform">The platform being built for</param>
 		/// <param name="Configuration">The configuration being built</param>
 		/// <param name="Architectures">The target architectures being built</param>
 		/// <param name="UndecoratedConfiguration">The target configuration which doesn't require a platform and configuration suffix. Development by default.</param>
 		/// <param name="BinaryType">Type of binary</param>
 		/// <returns>Name of the binary</returns>
-		public static string MakeBinaryFileName(string BinaryName, UnrealTargetPlatform Platform, UnrealTargetConfiguration Configuration, UnrealArchitectures Architectures, UnrealTargetConfiguration UndecoratedConfiguration, UEBuildBinaryType BinaryType)
+		public static string MakeBinaryFileName(string BinaryName, string Separator, UnrealTargetPlatform Platform, UnrealTargetConfiguration Configuration, UnrealArchitectures Architectures, UnrealTargetConfiguration UndecoratedConfiguration, UEBuildBinaryType BinaryType)
 		{
 			StringBuilder Result = new StringBuilder();
 
@@ -4569,7 +4570,10 @@ namespace UnrealBuildTool
 
 			if (Configuration != UndecoratedConfiguration)
 			{
-				Result.AppendFormat("-{0}-{1}", Platform.ToString(), Configuration.ToString());
+				Result.Append(Separator);
+				Result.Append(Platform);
+				Result.Append(Separator);
+				Result.Append(Configuration);
 			}
 
 			UEBuildPlatform BuildPlatform = UEBuildPlatform.GetBuildPlatform(Platform);
@@ -4616,7 +4620,7 @@ namespace UnrealBuildTool
 				{
 					BinaryDirectory = DirectoryReference.Combine(BinaryDirectory, ExeSubFolder);
 				}
-				BinaryFile = FileReference.Combine(BinaryDirectory, MakeBinaryFileName(BinaryName, Platform, ExeConfiguration, Architectures, UndecoratedConfiguration, BinaryType));
+				BinaryFile = FileReference.Combine(BinaryDirectory, MakeBinaryFileName(BinaryName, Rules.DecoratedSeparator, Platform, ExeConfiguration, Architectures, UndecoratedConfiguration, BinaryType));
 			}
 
 			// Allow the platform to customize the output path (and output several executables at once if necessary)
diff --git a/Engine/Source/Programs/UnrealBuildTool/Platform/Apple/AppleExports.cs b/Engine/Source/Programs/UnrealBuildTool/Platform/Apple/AppleExports.cs
--- a/Engine/Source/Programs/UnrealBuildTool/Platform/Apple/AppleExports.cs
+++ b/Engine/Source/Programs/UnrealBuildTool/Platform/Apple/AppleExports.cs
@@ -280,15 +280,16 @@ namespace UnrealBuildTool
 		/// Pass along the call to UEBuildTarget.MakeBinaryFileName, but taking an extension instead of a binary type, since the type is hidden
 		/// </summary>
 		/// <param name="BinaryName"></param>
+		/// <param name="Separator"></param>
 		/// <param name="Platform"></param>
 		/// <param name="Configuration"></param>
 		/// <param name="Architectures"></param>
 		/// <param name="UndecoratedConfiguration"></param>
 		/// <param name="Extension"></param>
 		/// <returns></returns>
-		public static string MakeBinaryFileName(string BinaryName, UnrealTargetPlatform Platform, UnrealTargetConfiguration Configuration, UnrealArchitectures Architectures, UnrealTargetConfiguration UndecoratedConfiguration, string? Extension)
+		public static string MakeBinaryFileName(string BinaryName, string Separator, UnrealTargetPlatform Platform, UnrealTargetConfiguration Configuration, UnrealArchitectures Architectures, UnrealTargetConfiguration UndecoratedConfiguration, string? Extension)
 		{
-			string StandardBinaryName = UEBuildTarget.MakeBinaryFileName(BinaryName, Platform, Configuration, Architectures, UndecoratedConfiguration, UEBuildBinaryType.Executable);
+			string StandardBinaryName = UEBuildTarget.MakeBinaryFileName(BinaryName, Separator, Platform, Configuration, Architectures, UndecoratedConfiguration, UEBuildBinaryType.Executable);
 			return System.IO.Path.ChangeExtension(StandardBinaryName, Extension);
 		}
 
diff --git a/Engine/Source/Programs/UnrealBuildTool/ProjectFiles/VisualStudioCode/VSCodeProjectFileGenerator.cs b/Engine/Source/Programs/UnrealBuildTool/ProjectFiles/VisualStudioCode/VSCodeProjectFileGenerator.cs
--- a/Engine/Source/Programs/UnrealBuildTool/ProjectFiles/VisualStudioCode/VSCodeProjectFileGenerator.cs
+++ b/Engine/Source/Programs/UnrealBuildTool/ProjectFiles/VisualStudioCode/VSCodeProjectFileGenerator.cs
@@ -1189,7 +1189,7 @@ namespace UnrealBuildTool
 			}
 
 			// Make the output file path
-			string BinaryFileName = UEBuildTarget.MakeBinaryFileName(BinaryName, Platform, Configuration, TargetRulesObject!.Architectures, TargetRulesObject.UndecoratedConfiguration, UEBuildBinaryType.Executable);
+			string BinaryFileName = UEBuildTarget.MakeBinaryFileName(BinaryName, TargetRulesObject!.DecoratedSeparator, Platform, Configuration, TargetRulesObject!.Architectures, TargetRulesObject.UndecoratedConfiguration, UEBuildBinaryType.Executable);
 			string ExecutableFilename = FileReference.Combine(OutputDirectory, BinaryFileName).FullName;
 
 			// Include the path to the actual executable for a Mac app bundle
diff --git a/Engine/Source/Programs/UnrealBuildTool/ProjectFiles/Xcode/XcodeProject.cs b/Engine/Source/Programs/UnrealBuildTool/ProjectFiles/Xcode/XcodeProject.cs
--- a/Engine/Source/Programs/UnrealBuildTool/ProjectFiles/Xcode/XcodeProject.cs
+++ b/Engine/Source/Programs/UnrealBuildTool/ProjectFiles/Xcode/XcodeProject.cs
@@ -2073,7 +2073,7 @@ namespace UnrealBuildTool.XcodeProjectXcconfig
 				// hook up the Buildconfig that matches this info to this xcconfig file
 				XcconfigFile ConfigXcconfig = MatchedConfig.Xcconfig!;
 
-				string ExecutableName = AppleExports.MakeBinaryFileName(Config.ExeName, Platform, Config.BuildConfig, TargetRules.Architectures, TargetRules.UndecoratedConfiguration, null);
+				string ExecutableName = AppleExports.MakeBinaryFileName(Config.ExeName, TargetRules.DecoratedSeparator, Platform, Config.BuildConfig, TargetRules.Architectures, TargetRules.UndecoratedConfiguration, null);
 				string ExetuableSubPath = FileReference.Combine(ConfigBuildDir, ExecutableName).MakeRelativeTo(ConfigBuildDir);
 				string ExecutableKey = $"UE_{Platform.ToString().ToUpper()}_EXECUTABLE_NAME";
 
@@ -2087,7 +2087,7 @@ namespace UnrealBuildTool.XcodeProjectXcconfig
 				// content only projects don't want UnrealGame, etc as the ProductName
 				if (UnrealData.bIsContentOnlyProject && TargetRules.Type != TargetType.Editor)
 				{
-					ProductName = AppleExports.MakeBinaryFileName(PerTargetTypeProductName, Platform, Config.BuildConfig, TargetRules.Architectures, TargetRules.UndecoratedConfiguration, null);
+					ProductName = AppleExports.MakeBinaryFileName(PerTargetTypeProductName, TargetRules.DecoratedSeparator, Platform, Config.BuildConfig, TargetRules.Architectures, TargetRules.UndecoratedConfiguration, null);
 				}
 
 				MetadataItem? EntitlementsMetadata = UnrealData.Metadata!.EntitlementsFiles[MetadataPlatform];
