[UBT]
* Fixed so code creating dynamic libraries for merged modules use the new DecoratedSeparator variable instead of hard coded "_"

diff --git a/Engine/Source/Programs/UnrealBuildTool/Configuration/UEBuildTarget.cs b/Engine/Source/Programs/UnrealBuildTool/Configuration/UEBuildTarget.cs
--- a/Engine/Source/Programs/UnrealBuildTool/Configuration/UEBuildTarget.cs
+++ b/Engine/Source/Programs/UnrealBuildTool/Configuration/UEBuildTarget.cs
@@ -4528,7 +4528,7 @@ namespace UnrealBuildTool
 				ModuleConfiguration = UnrealTargetConfiguration.Development;
 			}
 
-			List<FileReference> OutputPaths = MakeBinaryPaths(OutputDirectory, AppName + "_" + MergedName, Platform, Configuration, UEBuildBinaryType.DynamicLinkLibrary, Rules.Architectures, Rules.UndecoratedConfiguration, bCompileMonolithic && ProjectFile != null, Rules.ExeBinariesSubFolder, ProjectFile, Rules);
+			List<FileReference> OutputPaths = MakeBinaryPaths(OutputDirectory, AppName + Rules.DecoratedSeparator + MergedName, Platform, Configuration, UEBuildBinaryType.DynamicLinkLibrary, Rules.Architectures, Rules.UndecoratedConfiguration, bCompileMonolithic && ProjectFile != null, Rules.ExeBinariesSubFolder, ProjectFile, Rules);
 
 			// Create the binary
 			UEBuildBinary Binary = new UEBuildBinary(
