Fix invalid decorated executable paths when using merged modular build

#rnx

diff --git a/Engine/Source/Programs/UnrealBuildTool/ProjectFiles/VisualStudio/VCProject.cs b/Engine/Source/Programs/UnrealBuildTool/ProjectFiles/VisualStudio/VCProject.cs
--- a/Engine/Source/Programs/UnrealBuildTool/ProjectFiles/VisualStudio/VCProject.cs
+++ b/Engine/Source/Programs/UnrealBuildTool/ProjectFiles/VisualStudio/VCProject.cs
@@ -2391,7 +2391,7 @@ namespace UnrealBuildTool
 					NMakePath = FileReference.Combine(OutputDirectory, BaseExeName);
 					if (Configuration != TargetRulesObject.UndecoratedConfiguration)
 					{
-						NMakePath += "-" + UBTPlatformName + "-" + UBTConfigurationName;
+						NMakePath += TargetRulesObject.DecoratedSeparator + UBTPlatformName + TargetRulesObject.DecoratedSeparator + UBTConfigurationName;
 					}
 					if (UnrealArchitectureConfig.ForPlatform(Platform).RequiresArchitectureFilenames(TargetRulesObject.Architectures))
 					{
