Use ICF for Linux for Shipping/Test configs
This may cause some odd callstack lines for our runtime symbolication but ends up saving a ton in debug/binary size
For BreakpadSymbolEncoder, if we have a segment with a zero line number, move it to the previous line number to help reduce how many segments we have.
Zero is incorrect from dump_syms so picking a different incorrect number while reducing the size will at least be more helpful

diff --git a/Engine/Source/Programs/BreakpadSymbolEncoder/Private/BreakpadSymbolEncoder.cpp b/Engine/Source/Programs/BreakpadSymbolEncoder/Private/BreakpadSymbolEncoder.cpp
--- a/Engine/Source/Programs/BreakpadSymbolEncoder/Private/BreakpadSymbolEncoder.cpp
+++ b/Engine/Source/Programs/BreakpadSymbolEncoder/Private/BreakpadSymbolEncoder.cpp
@@ -246,6 +246,15 @@ bool ParseSymbolFile(const std::string& SymbolFile, std::vector<Record>& out_Rec
 			else
 			{
 				Record LastRecord = out_Records.back();
+
+				// We at times have a line number of zero from dump_syms, this can mess up compressing chunks of the same PC range
+				// just to save on size lets just assume if its a zero, use the previous line number. While not correct, not any less
+				// incorrect then 0, and helps reduce the sym file size by a good amount
+				if (Out.LineNumber == 0)
+				{
+					Out.LineNumber = LastRecord.LineNumber;
+				}
+
 				if (LastRecord.LineNumber != Out.LineNumber || LastRecord.FileRelativeOffset != Out.FileRelativeOffset)
 				{
 					// Same as the FileAbsoluteOffset but for the Symbol name. Need to covert at the end when all the memory we are going to write is in the structure
diff --git a/Engine/Source/Programs/UnrealBuildTool/Platform/Linux/LinuxToolChain.cs b/Engine/Source/Programs/UnrealBuildTool/Platform/Linux/LinuxToolChain.cs
--- a/Engine/Source/Programs/UnrealBuildTool/Platform/Linux/LinuxToolChain.cs
+++ b/Engine/Source/Programs/UnrealBuildTool/Platform/Linux/LinuxToolChain.cs
@@ -433,8 +433,10 @@ namespace UnrealBuildTool
 				|| Options.HasFlag(ClangToolChainOptions.EnableAddressSanitizer) || Options.HasFlag(ClangToolChainOptions.EnableMemorySanitizer) || Options.HasFlag(ClangToolChainOptions.EnableThreadSanitizer)
 				|| CompileEnvironment.Configuration == CppConfiguration.Debug;
 
+
 			if (CompileEnvironment.Configuration == CppConfiguration.Shipping)
 			{
+
 				if (!bRetainFramePointers)
 				{
 					Arguments.Add("-fomit-frame-pointer");
@@ -866,6 +868,11 @@ namespace UnrealBuildTool
 					Arguments.Add($"-L\"{NormalizeCommandLinePath(DirectoryReference.Combine(SysRootPath, "usr", "lib64"))}\"");
 				}
 			}
+
+			if (LinkEnvironment.Configuration == CppConfiguration.Shipping)
+			{
+				Arguments.Add("-Wl,--icf=all"); // Enables ICF (Identical Code Folding). [all, safe] safe == fold functions that can be proven not to have their address taken.
+			}
 		}
 
 		string GetArchiveArguments(LinkEnvironment LinkEnvironment)
