UnrealBuildTool: clang-cl flags for /arch

#rnx

diff --git a/Engine/Source/Programs/UnrealBuildTool/Platform/Windows/VCToolChain.cs b/Engine/Source/Programs/UnrealBuildTool/Platform/Windows/VCToolChain.cs
--- a/Engine/Source/Programs/UnrealBuildTool/Platform/Windows/VCToolChain.cs
+++ b/Engine/Source/Programs/UnrealBuildTool/Platform/Windows/VCToolChain.cs
@@ -840,7 +840,33 @@ namespace UnrealBuildTool
 			if (CompileEnvironment.Architecture == UnrealArch.X64 && CompileEnvironment.MinCpuArchX64 != MinimumCpuArchitectureX64.None)
 			{
 				// Define /arch:AVX[2,512] for the current compilation unit.  Machines without AVX support will crash on any SSE/AVX instructions if they run this compilation unit.
-				Arguments.Add($"/arch:{CompileEnvironment.MinCpuArchX64}");
+				if (Target.WindowsPlatform.Compiler.IsMSVC())
+				{
+					Arguments.Add($"/arch:{CompileEnvironment.MinCpuArchX64}");
+				}
+				else if (Target.WindowsPlatform.Compiler.IsClang())
+				{
+					if (CompileEnvironment.MinCpuArchX64 >= MinimumCpuArchitectureX64.AVX)
+					{
+						// Apparently MSVC enables (a subset?) of BMI (bit manipulation instructions) when /arch:AVX is set. Some code relies on this, so mirror it by enabling BMI1
+						Arguments.Add("-mavx");
+						Arguments.Add("-mbmi");
+					}
+					if (CompileEnvironment.MinCpuArchX64 >= MinimumCpuArchitectureX64.AVX2)
+					{
+						Arguments.Add("-mavx2");
+					}
+					if (CompileEnvironment.MinCpuArchX64 >= MinimumCpuArchitectureX64.AVX512)
+					{
+						// Match MSVC which says (https://learn.microsoft.com/en-us/cpp/build/reference/arch-x64?view=msvc-170):
+						// > The __AVX512F__, __AVX512CD__, __AVX512BW__, __AVX512DQ__ and __AVX512VL__ preprocessor symbols are defined when the /arch:AVX512 compiler option is specified
+						Arguments.Add("-mavx512f");
+						Arguments.Add("-mavx512cd");
+						Arguments.Add("-mavx512bw");
+						Arguments.Add("-mavx512dq");
+						Arguments.Add("-mavx512vl");
+					}
+				}
 
 				// AVX available implies sse4 and sse2 available.
 				// Inform Unreal code that we have sse2, sse4, and AVX, both available to compile and available to run
