UnrealBuildTool: Move min arch support up to ClangToolChain so it can be applied to any clang platform

#rnx

diff --git a/Engine/Source/Programs/UnrealBuildTool/Configuration/TargetRules.cs b/Engine/Source/Programs/UnrealBuildTool/Configuration/TargetRules.cs
--- a/Engine/Source/Programs/UnrealBuildTool/Configuration/TargetRules.cs
+++ b/Engine/Source/Programs/UnrealBuildTool/Configuration/TargetRules.cs
@@ -2485,8 +2485,8 @@ namespace UnrealBuildTool
 		public CStandardVersion CStandard { get; set; } = CStandardVersion.Default;
 
 		/// <summary>
-		/// Direct the compiler to generate AVX instructions wherever SSE or AVX intrinsics are used, on the x64 platforms that support it.
-		/// Note that by enabling this you are changing the minspec for the PC platform, and the resultant executable will crash on machines without AVX support.
+		/// Direct the compiler to generate AVX instructions wherever SSE or AVX intrinsics are used, on the x64 platforms that support it. Ignored for arm64.
+		/// Note that by enabling this you are changing the minspec for the target platform, and the resultant executable will crash on machines without AVX support.
 		/// </summary>
 		[RequiresUniqueBuildEnvironment]
 		[CommandLine("-MinCpuArchX64")]
diff --git a/Engine/Source/Programs/UnrealBuildTool/Platform/Linux/LinuxToolChain.cs b/Engine/Source/Programs/UnrealBuildTool/Platform/Linux/LinuxToolChain.cs
--- a/Engine/Source/Programs/UnrealBuildTool/Platform/Linux/LinuxToolChain.cs
+++ b/Engine/Source/Programs/UnrealBuildTool/Platform/Linux/LinuxToolChain.cs
@@ -341,43 +341,6 @@ namespace UnrealBuildTool
 
 			AddCompilerLTOFlags(Arguments);
 
-			// architecture (all but None are AVX)
-			if (CompileEnvironment.Architecture == UnrealArch.X64 && CompileEnvironment.MinCpuArchX64 != MinimumCpuArchitectureX64.None)
-			{
-				// The binary created will be targeting AVX instructions. Machines without AVX support will crash on any AVX instructions if they run this compilation unit.
-
-				// AVX available implies sse4 and sse2 available.
-				Arguments.Add("-DPLATFORM_ENABLE_VECTORINTRINSICS=1");
-
-				if (CompileEnvironment.MinCpuArchX64 >= MinimumCpuArchitectureX64.AVX)
-				{
-					// Apparently MSVC enables (a subset?) of BMI (bit manipulation instructions) when /arch:AVX is set. Some code relies on this, so mirror it by enabling BMI1
-					Arguments.Add("-mavx -mbmi");
-					// Inform Unreal code that we have sse2, sse4, and AVX, both available to compile and available to run
-					Arguments.Add("-DPLATFORM_MAYBE_HAS_AVX=1");
-					// By setting the ALWAYS_HAS defines, we we direct Unreal code to skip cpuid checks to verify that the running hardware supports sse/avx.
-					Arguments.Add("-DPLATFORM_ALWAYS_HAS_AVX=1");
-				}
-
-				if (CompileEnvironment.MinCpuArchX64 >= MinimumCpuArchitectureX64.AVX2)
-				{
-					Arguments.Add("-mavx2");
-					Arguments.Add("-DPLATFORM_ALWAYS_HAS_AVX_2=1");
-				}
-
-				if (CompileEnvironment.MinCpuArchX64 >= MinimumCpuArchitectureX64.AVX512)
-				{
-					// Match MSVC which says (https://learn.microsoft.com/en-us/cpp/build/reference/arch-x64?view=msvc-170):
-					// > The __AVX512F__, __AVX512CD__, __AVX512BW__, __AVX512DQ__ and __AVX512VL__ preprocessor symbols are defined when the /arch:AVX512 compiler option is specified
-					Arguments.Add("-mavx512f");
-					Arguments.Add("-mavx512cd");
-					Arguments.Add("-mavx512bw");
-					Arguments.Add("-mavx512dq");
-					Arguments.Add("-mavx512vl");
-					Arguments.Add("-DPLATFORM_ALWAYS_HAS_AVX_512=1");
-				}
-			}
-
 			if (CompileEnvironment.bCodeCoverage)
 			{
 				Arguments.Add("-O0");
diff --git a/Engine/Source/Programs/UnrealBuildTool/ToolChain/ClangToolChain.cs b/Engine/Source/Programs/UnrealBuildTool/ToolChain/ClangToolChain.cs
--- a/Engine/Source/Programs/UnrealBuildTool/ToolChain/ClangToolChain.cs
+++ b/Engine/Source/Programs/UnrealBuildTool/ToolChain/ClangToolChain.cs
@@ -685,6 +685,46 @@ namespace UnrealBuildTool
 				// We have 'this' vs nullptr comparisons that get optimized away for newer versions of Clang, which is undesirable until we refactor these checks.
 				Arguments.Add("-fno-delete-null-pointer-checks");
 			}
+
+			// architecture (all but None are AVX)
+			if (CompileEnvironment.Architecture == UnrealArch.X64 && CompileEnvironment.MinCpuArchX64 != MinimumCpuArchitectureX64.None)
+			{
+				// The binary created will be targeting AVX instructions. Machines without AVX support will crash on any AVX instructions if they run this compilation unit.
+
+				// AVX available implies sse4 and sse2 available.
+				// Inform Unreal code that we have sse2, sse4, and AVX, both available to compile and available to run
+				// By setting the ALWAYS_HAS defines, we we direct Unreal code to skip cpuid checks to verify that the running hardware supports sse/avx.
+				Arguments.Add("-DPLATFORM_ENABLE_VECTORINTRINSICS=1");
+
+				if (CompileEnvironment.MinCpuArchX64 >= MinimumCpuArchitectureX64.AVX)
+				{
+					// Apparently MSVC enables (a subset?) of BMI (bit manipulation instructions) when /arch:AVX is set. Some code relies on this, so mirror it by enabling BMI1
+					Arguments.Add("-mavx");
+					Arguments.Add("-mbmi");
+					// Inform Unreal code that we have sse2, sse4, and AVX, both available to compile and available to run
+					Arguments.Add("-DPLATFORM_MAYBE_HAS_AVX=1");
+					// By setting the ALWAYS_HAS defines, we we direct Unreal code to skip cpuid checks to verify that the running hardware supports sse/avx.
+					Arguments.Add("-DPLATFORM_ALWAYS_HAS_AVX=1");
+				}
+
+				if (CompileEnvironment.MinCpuArchX64 >= MinimumCpuArchitectureX64.AVX2)
+				{
+					Arguments.Add("-mavx2");
+					Arguments.Add("-DPLATFORM_ALWAYS_HAS_AVX_2=1");
+				}
+
+				if (CompileEnvironment.MinCpuArchX64 >= MinimumCpuArchitectureX64.AVX512)
+				{
+					// Match MSVC which says (https://learn.microsoft.com/en-us/cpp/build/reference/arch-x64?view=msvc-170):
+					// > The __AVX512F__, __AVX512CD__, __AVX512BW__, __AVX512DQ__ and __AVX512VL__ preprocessor symbols are defined when the /arch:AVX512 compiler option is specified
+					Arguments.Add("-mavx512f");
+					Arguments.Add("-mavx512cd");
+					Arguments.Add("-mavx512bw");
+					Arguments.Add("-mavx512dq");
+					Arguments.Add("-mavx512vl");
+					Arguments.Add("-DPLATFORM_ALWAYS_HAS_AVX_512=1");
+				}
+			}
 		}
 
 		/// <summary>
