UnrealBuildTool: Option to enforce Standard C++ ODR violations under /Gw with /Zc:checkGwOdr when strict conformance is enabled

#rnx

diff --git a/Engine/Source/Programs/UnrealBuildTool/Platform/Windows/UEBuildWindows.cs b/Engine/Source/Programs/UnrealBuildTool/Platform/Windows/UEBuildWindows.cs
--- a/Engine/Source/Programs/UnrealBuildTool/Platform/Windows/UEBuildWindows.cs
+++ b/Engine/Source/Programs/UnrealBuildTool/Platform/Windows/UEBuildWindows.cs
@@ -540,6 +540,14 @@ namespace UnrealBuildTool
 		[CommandLine("-StrictEnumTypes")]
 		public bool bStrictEnumTypesConformance = false;
 
+		/// <summary>
+		/// Enables enforcing standard C++ ODR violations (/Zc:checkGwOdr) in VS2022 17.5 Preview 2.0+
+		/// </summary>
+		[RequiresUniqueBuildEnvironment]
+		[XmlConfigFile(Category = "WindowsPlatform")]
+		[CommandLine("-StrictODR")]
+		public bool bStrictODRViolationConformance = false;
+
 		/// <summary>
 		/// Volatile Metadata is enabled by default and improves x64 emulation on arm64, but may come at a small perfomance cost (/volatileMetadata-).
 		/// </summary>
@@ -922,6 +930,8 @@ namespace UnrealBuildTool
 
 		public bool bStrictEnumTypesConformance => Inner.bStrictEnumTypesConformance;
 
+		public bool bStrictODRViolationConformance => Inner.bStrictODRViolationConformance;
+
 		public bool bDisableVolatileMetadata => Inner.bDisableVolatileMetadata;
 
 		public bool bStripPrivateSymbols => Inner.bStripPrivateSymbols;
diff --git a/Engine/Source/Programs/UnrealBuildTool/Platform/Windows/VCToolChain.cs b/Engine/Source/Programs/UnrealBuildTool/Platform/Windows/VCToolChain.cs
--- a/Engine/Source/Programs/UnrealBuildTool/Platform/Windows/VCToolChain.cs
+++ b/Engine/Source/Programs/UnrealBuildTool/Platform/Windows/VCToolChain.cs
@@ -553,6 +553,11 @@ namespace UnrealBuildTool
 				{
 					Arguments.Add("/Zc:enumTypes");
 				}
+
+				if (Target.WindowsPlatform.bStrictODRViolationConformance && Target.WindowsPlatform.bOptimizeGlobalData)
+				{
+					Arguments.Add("/Zc:checkGwOdr");
+				}
 			}
 
 			// @todo HoloLens: UE is non-compliant when it comes to use of %s and %S
