UnrealBuildTool: Allow extensionless files when scanning all include paths to find dependencies

#rnx

diff --git a/Engine/Source/Programs/UnrealBuildTool/Configuration/UEBuildTarget.cs b/Engine/Source/Programs/UnrealBuildTool/Configuration/UEBuildTarget.cs
--- a/Engine/Source/Programs/UnrealBuildTool/Configuration/UEBuildTarget.cs
+++ b/Engine/Source/Programs/UnrealBuildTool/Configuration/UEBuildTarget.cs
@@ -2067,7 +2067,7 @@ namespace UnrealBuildTool
 						foreach (string IncludeFileName in Directory.EnumerateFiles(IncludePath, "*", SearchOption.AllDirectories))
 						{
 							string Extension = Path.GetExtension(IncludeFileName).ToLower();
-							if (Extension == ".h" || Extension == ".inl" || Extension == ".hpp" || Extension == ".ipp")
+							if (Extension == ".h" || Extension == ".inl" || Extension == ".hpp" || Extension == ".ipp" || String.IsNullOrEmpty(Extension))
 							{
 								Files.Add(new FileReference(IncludeFileName));
 							}
