UnrealBuildTool: Add visual studio and msvc min version to warning when no valid default compiler is found

#rnx

diff --git a/Engine/Source/Programs/UnrealBuildTool/Platform/Windows/UEBuildWindows.cs b/Engine/Source/Programs/UnrealBuildTool/Platform/Windows/UEBuildWindows.cs
--- a/Engine/Source/Programs/UnrealBuildTool/Platform/Windows/UEBuildWindows.cs
+++ b/Engine/Source/Programs/UnrealBuildTool/Platform/Windows/UEBuildWindows.cs
@@ -1416,17 +1416,26 @@ namespace UnrealBuildTool
 
 			if (!bSkipWarning)
 			{
+				UEBuildPlatformSDK? SDK = GetSDK(UnrealTargetPlatform.Win64);
+				string ToolSetWarning = Architecture == UnrealArch.X64 ?
+					"MSVC v143 - VS 2022 C++ x64/x86 build tools (Latest)" :
+					"MSVC v143 - VS 2022 C++ ARM64 build tools (Latest)";
+				string MinMsvcVersion = SDK?.GetVersionNumberFromConfig("MinimumVisualCppVersion")?.ToString() ?? "Latest";
+				string MinVSVersion = SDK?.GetVersionNumberFromConfig("MinimumVisualStudio2022Version")?.ToString() ?? "Latest";
 				// If we do have a Visual Studio installation, but we're missing just the C++ parts, warn about that.
 				if (TryGetVSInstallDirs(WindowsCompiler.VisualStudio2022, Logger) != null)
 				{
-					string ToolSetWarning = Architecture == UnrealArch.X64 ?
-						"MSVC v143 - VS 2022 C++ x64/x86 build tools (Latest)" :
-						"MSVC v143 - VS 2022 C++ ARM64 build tools (Latest)";
-					Logger.LogWarning("Visual Studio 2022 is installed, but is missing the C++ toolchain. Please verify that the \"{Component}\" component is selected in the Visual Studio 2022 installation options.", ToolSetWarning);
+					Logger.LogWarning("Visual Studio 2022 is installed, but is out of date or missing a valid C++ toolchain (minimum version {MinVersion}). Please update Visual Studio 2022 to {MinVSVersion} or later and verify that the \"{Component}\" component is selected in the Visual Studio 2022 installation options.",
+						MinMsvcVersion,
+						MinVSVersion,
+						ToolSetWarning);
 				}
 				else
 				{
-					Logger.LogWarning("No Visual C++ installation was found. Please download and install Visual Studio 2022 with C++ components.");
+					Logger.LogWarning("No valid Visual C++ toolchain was found (minimum version {MinVersion}). Please download and install Visual Studio 2022 {MinVSVersion} or later and verify that the \"{Component}\" component is selected in the Visual Studio 2022 installation options.",
+						MinMsvcVersion,
+						MinVSVersion,
+						ToolSetWarning);
 				}
 			}
 
