UnrealBuildTool: Disable core.BitwiseShift clang checker

#rnx

diff --git a/Engine/Source/Programs/UnrealBuildTool/Configuration/ModuleRules.cs b/Engine/Source/Programs/UnrealBuildTool/Configuration/ModuleRules.cs
--- a/Engine/Source/Programs/UnrealBuildTool/Configuration/ModuleRules.cs
+++ b/Engine/Source/Programs/UnrealBuildTool/Configuration/ModuleRules.cs
@@ -1047,6 +1047,7 @@ namespace UnrealBuildTool
 		/// </summary>
 		public HashSet<string> StaticAnalyzerDisabledCheckers { get; } = new() {
 			// Work in progress, remove disabled checkers as all warnings are fixed
+			"core.BitwiseShift",
 			"core.CallAndMessage",
 			"core.DivideZero",
 			"core.NonNullParamChecker",
