Assorted changes and fixes to Verse support in UHT

#rnx

diff --git a/Engine/Source/Programs/Shared/EpicGames.Core/UHTTypes.cs b/Engine/Source/Programs/Shared/EpicGames.Core/UHTTypes.cs
--- a/Engine/Source/Programs/Shared/EpicGames.Core/UHTTypes.cs
+++ b/Engine/Source/Programs/Shared/EpicGames.Core/UHTTypes.cs
@@ -124,7 +124,7 @@ namespace EpicGames.Core
 			[JsonConverter(typeof(JsonStringEnumConverter))]
 			public UHTVerseScope VerseScope { get; set; } = UHTVerseScope.PublicAPI;
 			public bool HasVerse { get; set; } = false;
-			public string VersePluginName { get; set; } = "";
+			public string VerseMountPoint { get; set; } = "";
 
 			public override string ToString()
 			{
diff --git a/Engine/Source/Programs/Shared/EpicGames.UHT/Exporters/CodeGen/UhtHeaderCodeGeneratorCppFile.cs b/Engine/Source/Programs/Shared/EpicGames.UHT/Exporters/CodeGen/UhtHeaderCodeGeneratorCppFile.cs
--- a/Engine/Source/Programs/Shared/EpicGames.UHT/Exporters/CodeGen/UhtHeaderCodeGeneratorCppFile.cs
+++ b/Engine/Source/Programs/Shared/EpicGames.UHT/Exporters/CodeGen/UhtHeaderCodeGeneratorCppFile.cs
@@ -109,17 +109,9 @@ namespace EpicGames.UHT.Exporters.CodeGen
 				{
 					includesToAdd.Add("Net/Serialization/FastArraySerializerImplementation.h");
 				}
-				if (headerInfo.NeedsVerseClass)
+				if (headerInfo.NeedsVerseClass || headerInfo.NeedsVerseEnum || headerInfo.NeedsVerseStruct)
 				{
-					includesToAdd.Add("VerseVM/VVMVerseClass.h");
-				}
-				if (headerInfo.NeedsVerseStruct)
-				{
-					includesToAdd.Add("VerseVM/VVMVerseStruct.h");
-				}
-				if (headerInfo.NeedsVerseEnum)
-				{
-					includesToAdd.Add("VerseVM/VVMVerseEnum.h");
+					includesToAdd.Add("VerseVM/VVMUECodeGen.h");
 				}
 
 				foreach (UhtType type in HeaderFile.Children)
@@ -282,7 +274,7 @@ namespace EpicGames.UHT.Exporters.CodeGen
 						continue;
 					}
 
-					string name = $"Z_CompiledInDeferFile_{headerInfo.FileId}";
+					string name = $"Z_CompiledInDeferFile_{headerInfo.FileId}_{PackageInfos[package.PackageTypeIndex].StrippedName}";
 					string staticsName = $"{name}_Statics";
 
 					uint combinedHash = UInt32.MaxValue;
@@ -775,7 +767,7 @@ namespace EpicGames.UHT.Exporters.CodeGen
 				builder.Append("\t{\r\n");
 				if (scriptStruct.IsVerseField)
 				{
-					builder.Append("\t\tUECodeGen_Private::ConstructUScriptStruct(").Append(innerSingletonName).Append(", ").Append(staticsName).Append("::StructParams, Verse::Private::ConstructUScriptStructFunc);\r\n");
+					builder.Append("\t\tVerse::CodeGen::Private::ConstructUVerseStruct(").Append(innerSingletonName).Append(", ").Append(staticsName).Append("::StructParams);\r\n");
 				}
 				else 
 				{
diff --git a/Engine/Source/Programs/Shared/EpicGames.UHT/Exporters/CodeGen/UhtHeaderCodeGeneratorHFile.cs b/Engine/Source/Programs/Shared/EpicGames.UHT/Exporters/CodeGen/UhtHeaderCodeGeneratorHFile.cs
--- a/Engine/Source/Programs/Shared/EpicGames.UHT/Exporters/CodeGen/UhtHeaderCodeGeneratorHFile.cs
+++ b/Engine/Source/Programs/Shared/EpicGames.UHT/Exporters/CodeGen/UhtHeaderCodeGeneratorHFile.cs
@@ -427,7 +427,9 @@ namespace EpicGames.UHT.Exporters.CodeGen
 				// Verse declarations
 				if (scriptStruct.IsVerseField)
 				{
+					builder.Append("#if WITH_VERSE_BPVM\r\n");
 					builder.Append("VUHT_DECLARE_STRUCT_RUNTIME_TRAITS(struct ").Append(scriptStruct.SourceName).Append(")\r\n");
+					builder.Append("#endif // WITH_VERSE_BPVM\r\n");
 					builder.Append("\r\n");
 				}
 			}
diff --git a/Engine/Source/Programs/Shared/EpicGames.UHT/Types/UhtField.cs b/Engine/Source/Programs/Shared/EpicGames.UHT/Types/UhtField.cs
--- a/Engine/Source/Programs/Shared/EpicGames.UHT/Types/UhtField.cs
+++ b/Engine/Source/Programs/Shared/EpicGames.UHT/Types/UhtField.cs
@@ -59,73 +59,7 @@ namespace EpicGames.UHT.Types
 				throw new UhtException(fieldObj, "Attempt to write the Verse VNI package name on a field that isn't part of Verse");
 			}
 			UhtModule module = fieldObj.Module;
-			return builder.Append('/').Append(module.Module.VersePluginName).Append("/_Verse/VNI/").Append(module.Module.Name);
-		}
-
-		/// <summary>
-		/// Append the Verse UE name 
-		/// </summary>
-		/// <param name="builder">Destination builder</param>
-		/// <param name="fieldObj">Field to serialize</param>
-		/// <returns>Builder</returns>
-		public static StringBuilder AppendVerseUEName(this StringBuilder builder, UhtField fieldObj)
-		{
-			if (!fieldObj.IsVerseField)
-			{
-				throw new UhtException(fieldObj, "Attempt to write the Verse name on a field that isn't part of Verse");
-			}
-			return builder.Append(fieldObj.VerseName!);
-		}
-
-		/// <summary>
-		/// Append the Verse UE name 
-		/// </summary>
-		/// <param name="builder">Destination builder</param>
-		/// <param name="fieldObj">Field to serialize</param>
-		/// <returns>Builder</returns>
-		public static StringBuilder AppendVerseUEModuleAndName(this StringBuilder builder, UhtField fieldObj)
-		{
-			if (fieldObj.VerseModule == null)
-			{
-				throw new UhtException(fieldObj, "Attempt to write the Verse module on a field doesn't have the module name set");
-			}
-			return builder.Append(fieldObj.VerseModule!).Append('_').AppendVerseUEName(fieldObj);
-		}
-
-		/// <summary>
-		/// Append the Verse package name
-		/// </summary>
-		/// <param name="builder">Destination builder</param>
-		/// <param name="fieldObj">Field to serialize</param>
-		/// <returns>Builder</returns>
-		public static StringBuilder AppendVersePackageName(this StringBuilder builder, UhtField fieldObj)
-		{
-			if (!fieldObj.IsVerseField)
-			{
-				throw new UhtException(fieldObj, "Attempt to write the Verse package name on a field that isn't part of Verse");
-			}
-			UhtModule module = fieldObj.Module;
-			return builder.Append('/').Append(module.Module.VersePluginName).Append("/_Verse/VNI/").Append(module.Module.Name);
-		}
-
-		/// <summary>
-		/// Append the Verse module name
-		/// </summary>
-		/// <param name="builder">Destination builder</param>
-		/// <param name="fieldObj">Field to serialize</param>
-		/// <returns>Builder</returns>
-		public static StringBuilder AppendVerseModuleName(this StringBuilder builder, UhtField fieldObj)
-		{
-			if (!fieldObj.IsVerseField)
-			{
-				throw new UhtException(fieldObj, "Attempt to write the Verse module name on a field that isn't part of Verse");
-			}
-			builder.Append(fieldObj.Package.Module.Module.VersePath);
-			if (fieldObj.VerseModule != null)
-			{
-				builder.Append('/').Append(fieldObj.VerseModule!);
-			}
-			return builder;
+			return builder.Append('/').Append(module.Module.VerseMountPoint).Append("/_Verse/VNI/").Append(module.Module.Name);
 		}
 	}
 }
diff --git a/Engine/Source/Programs/UnrealBuildTool/Configuration/UEBuildTarget.cs b/Engine/Source/Programs/UnrealBuildTool/Configuration/UEBuildTarget.cs
--- a/Engine/Source/Programs/UnrealBuildTool/Configuration/UEBuildTarget.cs
+++ b/Engine/Source/Programs/UnrealBuildTool/Configuration/UEBuildTarget.cs
@@ -2562,7 +2562,7 @@ namespace UnrealBuildTool
 			// Prepare cached data for UHT header generation
 			using (GlobalTracer.Instance.BuildSpan("ExternalExecution.SetupUObjectModules()").StartActive())
 			{
-				ExternalExecution.SetupUObjectModules(ModulesToGenerateHeadersFor, Rules.Platform, ProjectDescriptor, Makefile.UObjectModules, Makefile.UObjectModuleHeaders, Rules.GeneratedCodeVersion, MetadataCache, Logger);
+				ExternalExecution.SetupUObjectModules(ModulesToGenerateHeadersFor, RulesAssembly, Rules.Platform, ProjectDescriptor, Makefile.UObjectModules, Makefile.UObjectModuleHeaders, Rules.GeneratedCodeVersion, MetadataCache, Logger);
 			}
 
 			if (Rules.NativePointerMemberBehaviorOverride != null)
diff --git a/Engine/Source/Programs/UnrealBuildTool/System/ExternalExecution.cs b/Engine/Source/Programs/UnrealBuildTool/System/ExternalExecution.cs
--- a/Engine/Source/Programs/UnrealBuildTool/System/ExternalExecution.cs
+++ b/Engine/Source/Programs/UnrealBuildTool/System/ExternalExecution.cs
@@ -169,14 +169,14 @@ namespace UnrealBuildTool
 		public VerseScope VerseScope;
 
 		/// <summary>
-		/// If true, verse scripts were found
+		/// Mount point for UE definitions
 		/// </summary>
-		public bool HasVerse;
+		public string VerseMountPoint;
 
 		/// <summary>
-		/// Verse plugin name
+		/// If true, verse scripts were found
 		/// </summary>
-		public string VersePluginName;
+		public bool HasVerse;
 
 		public UHTModuleInfo(
 			string ModuleName, 
@@ -191,7 +191,7 @@ namespace UnrealBuildTool
 			string VersePath,
 			VerseScope VerseScope,
 			bool HasVerse,
-			string VersePluginName)
+			string VerseMountPoint)
 		{
 			this.ModuleName = ModuleName;
 			this.ModuleRulesFile = ModuleRulesFile;
@@ -210,7 +210,7 @@ namespace UnrealBuildTool
 			this.VersePath = VersePath;
 			this.VerseScope = VerseScope;
 			this.HasVerse = HasVerse;
-			this.VersePluginName = VersePluginName;
+			this.VerseMountPoint = VerseMountPoint;
 		}
 
 		public UHTModuleInfo(BinaryArchiveReader Reader)
@@ -233,7 +233,7 @@ namespace UnrealBuildTool
 			VersePath = Reader.ReadString()!;
 			VerseScope = (VerseScope)Reader.ReadInt();
 			HasVerse = Reader.ReadBool();
-			VersePluginName = Reader.ReadString()!;
+			VerseMountPoint = Reader.ReadString()!;
 		}
 
 		public void Write(BinaryArchiveWriter Writer)
@@ -256,7 +256,7 @@ namespace UnrealBuildTool
 			Writer.WriteString(VersePath);
 			Writer.WriteInt((int)VerseScope);
 			Writer.WriteBool(HasVerse);
-			Writer.WriteString(VersePluginName);
+			Writer.WriteString(VerseMountPoint);
 		}
 
 		public override string ToString()
@@ -431,6 +431,7 @@ namespace UnrealBuildTool
 
 		public static void SetupUObjectModules(
 			IEnumerable<UEBuildModuleCPP> ModulesToGenerateHeadersFor,
+			RulesAssembly RulesAssembly, 
 			UnrealTargetPlatform Platform, 
 			ProjectDescriptor? ProjectDescriptor, 
 			List<UHTModuleInfo> UObjectModules,
@@ -463,19 +464,19 @@ namespace UnrealBuildTool
 					DirectoryReference ModuleIncludeBase = FindIncludeBase(Module, Logger);
 
 					UHTModuleInfo Info = new UHTModuleInfo(
-						Module.Name, 
-						Module.RulesFile, 
-						Module.ModuleDirectories, 
-						ModuleIncludeBase, 
-						ModuleToType[Module], 
-						GeneratedCodeDirectory, 
-						GeneratedCodeVersion, 
-						Module.Rules.bUsePrecompiled, 
+						Module.Name,
+						Module.RulesFile,
+						Module.ModuleDirectories,
+						ModuleIncludeBase,
+						ModuleToType[Module],
+						GeneratedCodeDirectory,
+						GeneratedCodeVersion,
+						Module.Rules.bUsePrecompiled,
 						Module.Rules.OverridePackageType,
 						Module.Rules.VersePath ?? "",
 						Module.Rules.VerseScope,
 						Module.bHasVerse,
-						Module.bHasVerse ? (Module.Rules.Plugin != null ? Module.Rules.Plugin.Name : Module.Name) : string.Empty);
+						Module.bHasVerse ? VNIExecution.GetMountPointForModule(Module, RulesAssembly) : string.Empty);
 					ModuleInfoArray[Idx] = Info;
 
 					Queue.Enqueue(() => SetupUObjectModule(Info, ExcludedFolders, MetadataCache, Queue));
@@ -1049,8 +1050,8 @@ namespace UnrealBuildTool
 						VersePath = UObjectModule.VersePath,
 						VerseScope = (UHTVerseScope)Enum.Parse(typeof(UHTVerseScope), UObjectModule.VerseScope.ToString()),
 						HasVerse = UObjectModule.HasVerse,
-						VersePluginName = UObjectModule.VersePluginName,
-					});
+						VerseMountPoint = UObjectModule.VerseMountPoint,
+					}); ;
 			}
 
 			List<string> UhtPlugins = new List<string>();
