UnrealBuildTool: Default UBALocalExecutor to be enabled by default

#rnx

diff --git a/Engine/Source/Programs/UnrealBuildTool/Configuration/BuildConfiguration.cs b/Engine/Source/Programs/UnrealBuildTool/Configuration/BuildConfiguration.cs
--- a/Engine/Source/Programs/UnrealBuildTool/Configuration/BuildConfiguration.cs
+++ b/Engine/Source/Programs/UnrealBuildTool/Configuration/BuildConfiguration.cs
@@ -65,7 +65,7 @@ namespace UnrealBuildTool
 			get => bAllowUBAExecutor || _bAllowUBALocalExecutorPrivate;
 			set => _bAllowUBALocalExecutorPrivate = value;
 		}
-		private bool _bAllowUBALocalExecutorPrivate = false;
+		private bool _bAllowUBALocalExecutorPrivate = true;
 
 		/// <summary>
 		/// Whether XGE may be used if available, default is true.
