UnrealBuildTool: Handling for newer clang checkers

#rnx

diff --git a/Engine/Source/Programs/UnrealBuildTool/Platform/Windows/VCToolChain.cs b/Engine/Source/Programs/UnrealBuildTool/Platform/Windows/VCToolChain.cs
--- a/Engine/Source/Programs/UnrealBuildTool/Platform/Windows/VCToolChain.cs
+++ b/Engine/Source/Programs/UnrealBuildTool/Platform/Windows/VCToolChain.cs
@@ -384,7 +384,7 @@ namespace UnrealBuildTool
 						Arguments.Add("--analyzer-no-default-checks");
 
 						// Only enable specific checks.
-						foreach (string Checker in CompileEnvironment.StaticAnalyzerCheckers)
+						foreach (string Checker in CompileEnvironment.StaticAnalyzerCheckers.Where(x => ClangWarnings.IsAvailableAnalyzerChecker(x, EnvVars.CompilerVersion)))
 						{
 							Arguments.Add($"-Xclang -analyzer-checker -Xclang {Checker}");
 						}
@@ -392,12 +392,12 @@ namespace UnrealBuildTool
 					else
 					{
 						// Disable default checks.
-						foreach (string Checker in CompileEnvironment.StaticAnalyzerDisabledCheckers)
+						foreach (string Checker in CompileEnvironment.StaticAnalyzerDisabledCheckers.Where(x => ClangWarnings.IsAvailableAnalyzerChecker(x, EnvVars.CompilerVersion)))
 						{
 							Arguments.Add($"-Xclang -analyzer-disable-checker -Xclang {Checker}");
 						}
 						// Enable additional non-default checks.
-						foreach (string Checker in CompileEnvironment.StaticAnalyzerAdditionalCheckers)
+						foreach (string Checker in CompileEnvironment.StaticAnalyzerAdditionalCheckers.Where(x => ClangWarnings.IsAvailableAnalyzerChecker(x, EnvVars.CompilerVersion)))
 						{
 							Arguments.Add($"-Xclang -analyzer-checker -Xclang {Checker}");
 						}
diff --git a/Engine/Source/Programs/UnrealBuildTool/ToolChain/ClangToolChain.cs b/Engine/Source/Programs/UnrealBuildTool/ToolChain/ClangToolChain.cs
--- a/Engine/Source/Programs/UnrealBuildTool/ToolChain/ClangToolChain.cs
+++ b/Engine/Source/Programs/UnrealBuildTool/ToolChain/ClangToolChain.cs
@@ -840,13 +840,15 @@ namespace UnrealBuildTool
 				Arguments.Add("-Xclang -analyzer-config -Xclang mode=shallow");
 			}
 
+			VersionNumber ClangVersion = new VersionNumber(Info.ClangVersion.Major, Info.ClangVersion.Minor, Info.ClangVersion.Build);
+
 			if (CompileEnvironment.StaticAnalyzerCheckers.Count > 0)
 			{
 				// Disable all default checks
 				Arguments.Add("--analyzer-no-default-checks");
 
 				// Only enable specific checkers.
-				foreach (string Checker in CompileEnvironment.StaticAnalyzerCheckers)
+				foreach (string Checker in CompileEnvironment.StaticAnalyzerCheckers.Where(x => ClangWarnings.IsAvailableAnalyzerChecker(x, ClangVersion)))
 				{
 					Arguments.Add($"-Xclang -analyzer-checker -Xclang {Checker}");
 				}
@@ -854,12 +856,12 @@ namespace UnrealBuildTool
 			else
 			{
 				// Disable default checks.
-				foreach (string Checker in CompileEnvironment.StaticAnalyzerDisabledCheckers)
+				foreach (string Checker in CompileEnvironment.StaticAnalyzerDisabledCheckers.Where(x => ClangWarnings.IsAvailableAnalyzerChecker(x, ClangVersion)))
 				{
 					Arguments.Add($"-Xclang -analyzer-disable-checker -Xclang {Checker}");
 				}
 				// Enable additional non-default checks.
-				foreach (string Checker in CompileEnvironment.StaticAnalyzerAdditionalCheckers)
+				foreach (string Checker in CompileEnvironment.StaticAnalyzerAdditionalCheckers.Where(x => ClangWarnings.IsAvailableAnalyzerChecker(x, ClangVersion)))
 				{
 					Arguments.Add($"-Xclang -analyzer-checker -Xclang {Checker}");
 				}
diff --git a/Engine/Source/Programs/UnrealBuildTool/ToolChain/ClangWarnings.cs b/Engine/Source/Programs/UnrealBuildTool/ToolChain/ClangWarnings.cs
--- a/Engine/Source/Programs/UnrealBuildTool/ToolChain/ClangWarnings.cs
+++ b/Engine/Source/Programs/UnrealBuildTool/ToolChain/ClangWarnings.cs
@@ -1,5 +1,6 @@
 ﻿// Copyright Epic Games, Inc. All Rights Reserved.
 
+using System;
 using System.Collections.Generic;
 using EpicGames.Core;
 
@@ -241,5 +242,24 @@ namespace UnrealBuildTool
 			Arguments.Add("-Wno-c++20-extensions");
 			Arguments.Add("-Wno-deprecated-declarations");
 		}
+
+		static Lazy<Dictionary<string, int>> CheckerAddedVersion = new Lazy<Dictionary<string, int>>(() => new()
+			{
+				{ "core.BitwiseShift", 18 },
+				{ "optin.core.EnumCastOutOfRange", 18 },
+				{ "security.cert.env.InvalidPtr", 18 },
+				{ "unix.Errno", 18 },
+				{ "unix.StdCLibraryFunctions", 18 },
+				{ "cplusplus.ArrayDelete", 19 },
+				{ "cplusplus.Move", 19 },
+				{ "optin.taint.TaintedAlloc", 19 },
+				{ "security.MmapWriteExec", 19 },
+				{ "security.PutenvStackArray", 19 },
+				{ "security.SetgidSetuidOrder", 19 },
+				{ "unix.BlockInCriticalSection", 19 },
+				{ "unix.Stream", 19 },
+			});
+
+		internal static bool IsAvailableAnalyzerChecker(string Checker, VersionNumber ClangVersion) => !CheckerAddedVersion.Value.ContainsKey(Checker) || ClangVersion.Components[0] >= CheckerAddedVersion.Value[Checker];
 	}
 }
