[UBT]
* Changed root paths to be sorted with engine/project/unreal first and additional paths after.. by keeping this sorting consistent between all actions we will get fewer entries in the path table

diff --git a/Engine/Source/Programs/UnrealBuildTool/Platform/Windows/VCToolChain.cs b/Engine/Source/Programs/UnrealBuildTool/Platform/Windows/VCToolChain.cs
--- a/Engine/Source/Programs/UnrealBuildTool/Platform/Windows/VCToolChain.cs
+++ b/Engine/Source/Programs/UnrealBuildTool/Platform/Windows/VCToolChain.cs
@@ -1645,34 +1645,34 @@ namespace UnrealBuildTool
 
 		protected override IEnumerable<DirectoryItem> GetEnvironmentBasePaths(CppCompileEnvironment CompileEnvironment)
 		{
-			yield return DirectoryItem.GetItemByDirectoryReference(EnvVars.CompilerDir);
-			if (EnvVars.CompilerDir != EnvVars.ToolChainDir)
-			{
-				yield return DirectoryItem.GetItemByDirectoryReference(EnvVars.ToolChainDir);
-			}
-			yield return DirectoryItem.GetItemByDirectoryReference(EnvVars.WindowsSdkDir);
 			yield return DirectoryItem.GetItemByDirectoryReference(Unreal.EngineDirectory);
 			if (Target.ProjectFile != null && (!CompileEnvironment.bUseSharedBuildEnvironment || CompileEnvironment.AllIncludePath.Any(x => x.IsUnderDirectory(Target.ProjectFile.Directory))))
 			{
 				yield return DirectoryItem.GetItemByDirectoryReference(Target.ProjectFile.Directory);
 			}
 			yield return DirectoryItem.GetItemByDirectoryReference(Unreal.RootDirectory);
-		}
-
-		protected override IEnumerable<DirectoryItem> GetEnvironmentBasePaths(LinkEnvironment LinkEnvironment)
-		{
 			yield return DirectoryItem.GetItemByDirectoryReference(EnvVars.CompilerDir);
 			if (EnvVars.CompilerDir != EnvVars.ToolChainDir)
 			{
 				yield return DirectoryItem.GetItemByDirectoryReference(EnvVars.ToolChainDir);
 			}
 			yield return DirectoryItem.GetItemByDirectoryReference(EnvVars.WindowsSdkDir);
+		}
+
+		protected override IEnumerable<DirectoryItem> GetEnvironmentBasePaths(LinkEnvironment LinkEnvironment)
+		{
 			yield return DirectoryItem.GetItemByDirectoryReference(Unreal.EngineDirectory);
 			if (Target.ProjectFile != null && LinkEnvironment.InputFiles.Any(x => x.Location.IsUnderDirectory(Target.ProjectFile.Directory)))
 			{
 				yield return DirectoryItem.GetItemByDirectoryReference(Target.ProjectFile.Directory);
 			}
 			yield return DirectoryItem.GetItemByDirectoryReference(Unreal.RootDirectory);
+			yield return DirectoryItem.GetItemByDirectoryReference(EnvVars.CompilerDir);
+			if (EnvVars.CompilerDir != EnvVars.ToolChainDir)
+			{
+				yield return DirectoryItem.GetItemByDirectoryReference(EnvVars.ToolChainDir);
+			}
+			yield return DirectoryItem.GetItemByDirectoryReference(EnvVars.WindowsSdkDir);
 		}
 
 		private VCCompileAction CreateBaseCompileAction(CppCompileEnvironment CompileEnvironment)
diff --git a/Engine/Source/Programs/UnrealBuildTool/ToolChain/ClangToolChain.cs b/Engine/Source/Programs/UnrealBuildTool/ToolChain/ClangToolChain.cs
--- a/Engine/Source/Programs/UnrealBuildTool/ToolChain/ClangToolChain.cs
+++ b/Engine/Source/Programs/UnrealBuildTool/ToolChain/ClangToolChain.cs
@@ -1175,30 +1175,30 @@ namespace UnrealBuildTool
 
 		protected override IEnumerable<DirectoryItem> GetEnvironmentBasePaths(CppCompileEnvironment CompileEnvironment)
 		{
-			if (GetToolChainInfo().BasePath != null)
-			{
-				yield return DirectoryItem.GetItemByDirectoryReference(GetToolChainInfo().BasePath!);
-			}
 			yield return DirectoryItem.GetItemByDirectoryReference(Unreal.EngineDirectory);
 			if (ProjectFile != null && (!CompileEnvironment.bUseSharedBuildEnvironment || CompileEnvironment.AllIncludePath.Any(x => x.IsUnderDirectory(ProjectFile.Directory))))
 			{
 				yield return DirectoryItem.GetItemByDirectoryReference(ProjectFile.Directory);
 			}
 			yield return DirectoryItem.GetItemByDirectoryReference(Unreal.RootDirectory);
-		}
-
-		protected override IEnumerable<DirectoryItem> GetEnvironmentBasePaths(LinkEnvironment LinkEnvironment)
-		{
 			if (GetToolChainInfo().BasePath != null)
 			{
 				yield return DirectoryItem.GetItemByDirectoryReference(GetToolChainInfo().BasePath!);
 			}
+		}
+
+		protected override IEnumerable<DirectoryItem> GetEnvironmentBasePaths(LinkEnvironment LinkEnvironment)
+		{
 			yield return DirectoryItem.GetItemByDirectoryReference(Unreal.EngineDirectory);
 			if (ProjectFile != null && LinkEnvironment.InputFiles.Any(x => x.Location.IsUnderDirectory(ProjectFile.Directory)))
 			{
 				yield return DirectoryItem.GetItemByDirectoryReference(ProjectFile.Directory);
 			}
 			yield return DirectoryItem.GetItemByDirectoryReference(Unreal.RootDirectory);
+			if (GetToolChainInfo().BasePath != null)
+			{
+				yield return DirectoryItem.GetItemByDirectoryReference(GetToolChainInfo().BasePath!);
+			}
 		}
 
 		protected virtual Action CompileCPPFile(CppCompileEnvironment CompileEnvironment, FileItem SourceFile, DirectoryReference OutputDir, string ModuleName, IActionGraphBuilder Graph, IReadOnlyCollection<string> GlobalArguments, CPPOutput Result)
