UBA: Forward CPU count and multiplier from Horde agent to UbaAgent process

Makes UBA respect these settings as configured by the Horde agent. For example when you want to limit the number of CPUs used when offering your workstation as a UBA agent.

diff --git a/Engine/Source/Programs/UnrealBuildTool/Executors/UnrealBuildAccelerator/UBAAgentCoordinatorHorde.cs b/Engine/Source/Programs/UnrealBuildTool/Executors/UnrealBuildAccelerator/UBAAgentCoordinatorHorde.cs
--- a/Engine/Source/Programs/UnrealBuildTool/Executors/UnrealBuildAccelerator/UBAAgentCoordinatorHorde.cs
+++ b/Engine/Source/Programs/UnrealBuildTool/Executors/UnrealBuildAccelerator/UBAAgentCoordinatorHorde.cs
@@ -544,6 +544,8 @@ namespace UnrealBuildTool
 
 						arguments.Add("-Dir=%UE_HORDE_SHARED_DIR%\\Uba");
 						arguments.Add("-Eventfile=%UE_HORDE_TERMINATION_SIGNAL_FILE%");
+						arguments.Add("-MaxCpu=%UE_HORDE_CPU_COUNT%");
+						arguments.Add("-MulCpu=%UE_HORDE_CPU_MULTIPLIER%");
 						arguments.Add("-MaxIdle=15");
 
 						if (_owner.UBAConfig.bLogEnabled)
