UnrealBuildTool: Fix clang checker added version

#rnx

diff --git a/Engine/Source/Programs/UnrealBuildTool/ToolChain/ClangWarnings.cs b/Engine/Source/Programs/UnrealBuildTool/ToolChain/ClangWarnings.cs
--- a/Engine/Source/Programs/UnrealBuildTool/ToolChain/ClangWarnings.cs
+++ b/Engine/Source/Programs/UnrealBuildTool/ToolChain/ClangWarnings.cs
@@ -243,6 +243,7 @@ namespace UnrealBuildTool
 			Arguments.Add("-Wno-deprecated-declarations");
 		}
 
+		// TODO: Get valid checkers by calling clang -cc1 -analyzer-checker-help
 		static Lazy<Dictionary<string, int>> CheckerAddedVersion = new Lazy<Dictionary<string, int>>(() => new()
 			{
 				{ "core.BitwiseShift", 18 },
@@ -251,9 +252,7 @@ namespace UnrealBuildTool
 				{ "unix.Errno", 18 },
 				{ "unix.StdCLibraryFunctions", 18 },
 				{ "cplusplus.ArrayDelete", 19 },
-				{ "cplusplus.Move", 19 },
 				{ "optin.taint.TaintedAlloc", 19 },
-				{ "security.MmapWriteExec", 19 },
 				{ "security.PutenvStackArray", 19 },
 				{ "security.SetgidSetuidOrder", 19 },
 				{ "unix.BlockInCriticalSection", 19 },
