- Fixed an issue with Marketplace plugins breaking project packaging

diff --git a/Engine/Source/Programs/UnrealBuildTool/Modes/QueryTargetsMode.cs b/Engine/Source/Programs/UnrealBuildTool/Modes/QueryTargetsMode.cs
--- a/Engine/Source/Programs/UnrealBuildTool/Modes/QueryTargetsMode.cs
+++ b/Engine/Source/Programs/UnrealBuildTool/Modes/QueryTargetsMode.cs
@@ -67,7 +67,7 @@ namespace UnrealBuildTool
 			}
 
 			// Write information about these targets
-			WriteTargetInfo(ProjectFile, Assembly, OutputFile, Arguments, Logger, bIncludeAllTargets);
+			WriteTargetInfo(ProjectFile, Assembly, OutputFile, Arguments, Logger, bIncludeAllTargets, bIncludeParentAssembly: true);
 			Logger.LogInformation("Written {OutputFile}", OutputFile);
 			return Task.FromResult(0);
 		}
@@ -98,11 +98,12 @@ namespace UnrealBuildTool
 		/// <param name="Arguments"></param>
 		/// <param name="Logger">Logger for output</param>
 		/// <param name="bIncludeAllTargets">Include all targets even if a default target is specified for a given target type.</param>
-		public static void WriteTargetInfo(FileReference? ProjectFile, RulesAssembly Assembly, FileReference OutputFile, CommandLineArguments Arguments, ILogger Logger, bool bIncludeAllTargets = true)
+		/// <param name="bIncludeParentAssembly">Include all targets from parent assemblies (useful for EngineRules which can be chained)</param>
+		public static void WriteTargetInfo(FileReference? ProjectFile, RulesAssembly Assembly, FileReference OutputFile, CommandLineArguments Arguments, ILogger Logger, bool bIncludeAllTargets = true, bool bIncludeParentAssembly = false)
 		{
 			// Construct all the targets in this assembly
 			List<string> TargetNames = new List<string>();
-			Assembly.GetAllTargetNames(TargetNames, false);
+			Assembly.GetAllTargetNames(TargetNames, bIncludeParentAssembly);
 
 			// Write the output file
 			DirectoryReference.CreateDirectory(OutputFile.Directory);
