UnrealBuildTool: Intel ICX 2024.2 integration

diff --git a/Engine/Config/Windows/Windows_SDK.json b/Engine/Config/Windows/Windows_SDK.json
--- a/Engine/Config/Windows/Windows_SDK.json
+++ b/Engine/Config/Windows/Windows_SDK.json
@@ -30,9 +30,9 @@
 	"MinimumVisualCppVersion": "14.38.33130",
 
 	"PreferredIntelOneApiVersions": [
-		"2024.1.0-2024.9999"
+		"2024.2.0-2024.9999"
 	],
-	"MinimumIntelOneApiVersion": "2024.1.0",
+	"MinimumIntelOneApiVersion": "2024.2.0",
 
 	"VisualStudioSuggestedComponents": [
 		"Microsoft.Net.Component.4.6.2.TargetingPack",
diff --git a/Engine/Source/Programs/UnrealBuildTool/ToolChain/ClangWarnings.cs b/Engine/Source/Programs/UnrealBuildTool/ToolChain/ClangWarnings.cs
--- a/Engine/Source/Programs/UnrealBuildTool/ToolChain/ClangWarnings.cs
+++ b/Engine/Source/Programs/UnrealBuildTool/ToolChain/ClangWarnings.cs
@@ -241,6 +241,7 @@ namespace UnrealBuildTool
 			Arguments.Add("-Wno-logical-not-parentheses");
 			Arguments.Add("-Wno-c++20-extensions");
 			Arguments.Add("-Wno-deprecated-declarations");
+			Arguments.Add("-Wno-shorten-64-to-32");
 		}
 
 		// TODO: Get valid checkers by calling clang -cc1 -analyzer-checker-help
