UnrealBuildTool: Intel ICX 2024.2 integration

diff --git a/Engine/Source/Programs/UnrealBuildTool/Platform/Windows/VCToolChain.cs b/Engine/Source/Programs/UnrealBuildTool/Platform/Windows/VCToolChain.cs
--- a/Engine/Source/Programs/UnrealBuildTool/Platform/Windows/VCToolChain.cs
+++ b/Engine/Source/Programs/UnrealBuildTool/Platform/Windows/VCToolChain.cs
@@ -960,7 +960,8 @@ namespace UnrealBuildTool
 					}
 
 					// https://clang.llvm.org/docs/UsersManual.html#cmdoption-feliminate-unused-debug-types
-					if (ClangVersion >= new VersionNumber(19))
+					// Intel ICX 2024.2 does not have this option, but reports Clang 19, so turn it off here
+					if (ClangVersion >= new VersionNumber(19) && !Target.WindowsPlatform.Compiler.IsIntel())
 					{
 						Arguments.Add("-fno-eliminate-unused-debug-types");
 					}
@@ -3182,6 +3183,13 @@ namespace UnrealBuildTool
 					if (Target.WindowsPlatform.bAllowClangLinker)
 					{
 						Log.TraceInformationOnce("Enabling Link-time optimization. Linking will take a while.");
+
+						// ThinLTO incremental cache
+						DirectoryReference? ThinLTOCacheDir = DirectoryReference.FromString(LinkEnvironment.ThinLTOCacheDirectory);
+						if (ThinLTOCacheDir != null)
+						{
+							Arguments.Add($"/lldltocache:\"{ThinLTOCacheDir}\"");
+						}
 					}
 					else
 					{
