UnrealBuildTool: Place resource intermediate files for executables in a subfolder

#rnx

diff --git a/Engine/Source/Programs/UnrealBuildTool/Configuration/UEBuildBinary.cs b/Engine/Source/Programs/UnrealBuildTool/Configuration/UEBuildBinary.cs
--- a/Engine/Source/Programs/UnrealBuildTool/Configuration/UEBuildBinary.cs
+++ b/Engine/Source/Programs/UnrealBuildTool/Configuration/UEBuildBinary.cs
@@ -892,7 +892,7 @@ namespace UnrealBuildTool
 			// Set the link output file.
 			BinaryLinkEnvironment.OutputFilePaths = OutputFilePaths.ToList();
 
-			// Rembmer the link type
+			// Remember the link type
 			BinaryLinkEnvironment.LinkType = Target.LinkType;
 
 			// Set whether the link is allowed to have exports.
@@ -939,6 +939,12 @@ namespace UnrealBuildTool
 						// Get the intermediate directory
 						DirectoryReference ResourceIntermediateDirectory = BinaryLinkEnvironment.IntermediateDirectory;
 
+						// Place resource intermediate files for executables in a subfolder
+						if (!BinaryLinkEnvironment.bIsBuildingDLL && !BinaryLinkEnvironment.bIsBuildingLibrary)
+						{
+							ResourceIntermediateDirectory = DirectoryReference.Combine(ResourceIntermediateDirectory, OutputFilePaths[0].GetFileName());
+						}
+
 						// Create a compile environment for resource files
 						CppCompileEnvironment ResourceCompileEnvironment = new CppCompileEnvironment(BinaryCompileEnvironment);
 
