Fastbuild executor fails when there are no multiple installations.

diff --git a/Engine/Source/Programs/UnrealBuildTool/Executors/Experimental/FASTBuild.cs b/Engine/Source/Programs/UnrealBuildTool/Executors/Experimental/FASTBuild.cs
--- a/Engine/Source/Programs/UnrealBuildTool/Executors/Experimental/FASTBuild.cs
+++ b/Engine/Source/Programs/UnrealBuildTool/Executors/Experimental/FASTBuild.cs
@@ -1009,10 +1009,17 @@ namespace UnrealBuildTool
 				AddText($"\t\t'$Root$/msobj{platformVersionNumber}.dll'\n");
 				AddText($"\t\t'$Root$/mspdb{platformVersionNumber}.dll'\n");
 
-				List<string> PotentialMSVCRedistPaths = new List<string>(Directory.EnumerateDirectories(String.Format("{0}/Redist/MSVC", VCEnv.GetVCInstallDirectory())));
+				string RedistInstallDirectory = String.Format("{0}/Redist/MSVC", VCEnv.GetVCInstallDirectory());
+				List<string> PotentialMSVCRedistPaths = new List<string>(Directory.EnumerateDirectories(RedistInstallDirectory));
 				string? PrefferedMSVCRedistPath = null;
 				string? FinalMSVCRedistPath = "";
 
+				if (PotentialMSVCRedistPaths.Count == 0)
+				{
+					// We cannot possibly derive a FinalMSVCRedistPath, as all paths are predicated on there being at least 1 install directory within PotentialMSVCRedistPaths
+					throw new Exception($"Couldn't find any installations under the VC Redist install directory: {RedistInstallDirectory}");
+				}
+
 				if (MsvcCRTRedistVersion.Length > 0)
 				{
 					PrefferedMSVCRedistPath = PotentialMSVCRedistPaths.Find(
@@ -1024,7 +1031,8 @@ namespace UnrealBuildTool
 
 				if (PrefferedMSVCRedistPath == null)
 				{
-					PrefferedMSVCRedistPath = PotentialMSVCRedistPaths[^2];
+					// If upon Visual Studio installation the user selects "latest", two folders are created. If however the user only specifies one, only the selected folder is created.
+					PrefferedMSVCRedistPath = PotentialMSVCRedistPaths.Count == 1 ? PotentialMSVCRedistPaths[0] : PotentialMSVCRedistPaths[^2];
 
 					if (MsvcCRTRedistVersion.Length > 0)
 					{
