Detect automation scripts in plugins for visual studio generation

diff --git a/Engine/Source/Programs/UnrealBuildTool/ProjectFiles/ProjectFileGenerator.cs b/Engine/Source/Programs/UnrealBuildTool/ProjectFiles/ProjectFileGenerator.cs
--- a/Engine/Source/Programs/UnrealBuildTool/ProjectFiles/ProjectFileGenerator.cs
+++ b/Engine/Source/Programs/UnrealBuildTool/ProjectFiles/ProjectFileGenerator.cs
@@ -815,6 +815,27 @@ namespace UnrealBuildTool
 			return ProjectFiles;
 		}
 
+		/// <summary>
+		/// Finds plugins associated with a set of provided game projects.
+		/// </summary>
+		/// <param name="Logger"></param>
+		/// <param name="GameProjects">The game projects to get plugins for.</param>
+		/// <returns>List of plugin files</returns>
+		public List<FileReference> FindGameProjectPlugins(ILogger Logger, List<FileReference> GameProjects)
+		{
+			List<FileReference> PluginFiles = new List<FileReference>();
+			foreach (FileReference UnrealProjectFile in GameProjects)
+			{
+				DirectoryReference GamePluginFolder = DirectoryReference.Combine(UnrealProjectFile.Directory, "Plugins");
+
+				if (DirectoryReference.Exists(GamePluginFolder))
+				{
+					PluginFiles.AddRange(PluginsBase.EnumeratePlugins(GamePluginFolder));
+				}
+			}
+			return PluginFiles;
+		}
+
 		/// <summary>
 		/// Gets the user's preferred IDE from their editor settings
 		/// </summary>
@@ -1251,8 +1272,10 @@ namespace UnrealBuildTool
 						ProgramsFolder.ChildProjects.Add(AutomationToolProject);
 					}
 
+					List<FileReference> AllPluginProjects = FindGameProjectPlugins(Logger, AllGameProjects);
+
 					// Add automation.csproj files to the primary project
-					AddRulesModules(Rules.RulesFileType.AutomationModule, "Automation", AutomationProjectFiles, AllGameProjects, RootFolder, ProgramsFolder, Logger);
+					AddRulesModules(Rules.RulesFileType.AutomationModule, "Automation", AutomationProjectFiles, AllGameProjects.Union(AllPluginProjects).ToList(), RootFolder, ProgramsFolder, Logger);
 
 					// Add ubtplugin.csproj files to the primary project
 					AddRulesModules(Rules.RulesFileType.UbtPlugin, "UnrealBuildTool.Plugins", UbtPluginProjectFiles, AllGameProjects, RootFolder, ProgramsFolder, Logger);
