Fix plist value causing issues with launch orientation (the github PR #12208).  Also (in additional to the PR), UISupportedInterfaceOrientations was being duplicated in the non-legacy plist writing (now removed) and modern xcode xcconfig file.

diff --git a/Engine/Source/Programs/UnrealBuildTool/Platform/IOS/UEDeployIOS.cs b/Engine/Source/Programs/UnrealBuildTool/Platform/IOS/UEDeployIOS.cs
--- a/Engine/Source/Programs/UnrealBuildTool/Platform/IOS/UEDeployIOS.cs
+++ b/Engine/Source/Programs/UnrealBuildTool/Platform/IOS/UEDeployIOS.cs
@@ -361,34 +361,6 @@ namespace UnrealBuildTool
 			{
 				Text.AppendLine(InterfaceOrientation);
 			}
-			Text.AppendLine("\t<key>UISupportedInterfaceOrientations</key>");
-			Text.AppendLine("\t<array>");
-			foreach (string Line in SupportedOrientations.Split("\n".ToCharArray()))
-			{
-				if (!String.IsNullOrWhiteSpace(Line))
-				{
-					Text.AppendLine(Line);
-				}
-			}
-			Text.AppendLine("\t</array>");
-
-			bool bEnableSplitView = false;
-			Ini.GetBool("/Script/IOSRuntimeSettings.IOSRuntimeSettings", "bEnableSplitView", out bEnableSplitView);
-			if (bEnableSplitView)
-			{
-				// As this is (currently) an iPad only feature, use the iPad descriminator to set it for iPad only
-				// as it also requires supporting all UIOrientations
-				Text.AppendLine("\t<key>UIRequiresFullScreen~ipad</key>");
-				Text.AppendLine("\t<false/>");
-
-				Text.AppendLine("\t<key>UISupportedInterfaceOrientations~ipad</key>");
-				Text.AppendLine("\t<array>");
-				Text.AppendLine($"\t\t<string>UIInterfaceOrientationPortrait</string>");
-				Text.AppendLine($"\t\t<string>UIInterfaceOrientationPortraitUpsideDown</string>");
-				Text.AppendLine($"\t\t<string>UIInterfaceOrientationLandscapeLeft</string>");
-				Text.AppendLine($"\t\t<string>UIInterfaceOrientationLandscapeRight</string>");
-				Text.AppendLine("\t</array>");
-			}
 
 			Text.AppendLine("\t<key>UIRequiredDeviceCapabilities</key>");
 			Text.AppendLine("\t<array>");
diff --git a/Engine/Source/Programs/UnrealBuildTool/ProjectFiles/Xcode/XcodeProject.cs b/Engine/Source/Programs/UnrealBuildTool/ProjectFiles/Xcode/XcodeProject.cs
--- a/Engine/Source/Programs/UnrealBuildTool/ProjectFiles/Xcode/XcodeProject.cs
+++ b/Engine/Source/Programs/UnrealBuildTool/ProjectFiles/Xcode/XcodeProject.cs
@@ -1890,7 +1890,7 @@ namespace UnrealBuildTool.XcodeProjectXcconfig
 
 					// only iphone deals with orientation
 					List<string> SupportedOrientations = XcodeUtils.GetSupportedOrientations(PlatformIni);
-					ExtraConfigLines.Add($"INFOPLIST_KEY_UISupportedInterfaceOrientations = \"{String.Join(" ", SupportedOrientations)}\"");
+					ExtraConfigLines.Add($"INFOPLIST_KEY_UISupportedInterfaceOrientations = " + String.Join(" ", SupportedOrientations));
 
 					// iPhone is always Fullscreen, however, an iPad can support SplitView mode (dynamic view resizeing), so check if that's enabled and then add iPad UISupportedInterfaceOrientations
 					bool bEnableSplitView = false;
