Upgrade to Gradle 8.7 with AGP 8.5.0 and use SDK 34 for compiling

#android

diff --git a/Engine/Extras/Android/SetupAndroid.bat b/Engine/Extras/Android/SetupAndroid.bat
--- a/Engine/Extras/Android/SetupAndroid.bat
+++ b/Engine/Extras/Android/SetupAndroid.bat
@@ -17,8 +17,8 @@ SET CMAKE_VERSION=%3
 SET NDK_VERSION=%4
 
 rem hardcoded versions for compatibility with non-Turnkey manual running
-if "%PLATFORMS_VERSION%" == "" SET PLATFORMS_VERSION=android-33
-if "%BUILDTOOLS_VERSION%" == "" SET BUILDTOOLS_VERSION=33.0.1
+if "%PLATFORMS_VERSION%" == "" SET PLATFORMS_VERSION=android-34
+if "%BUILDTOOLS_VERSION%" == "" SET BUILDTOOLS_VERSION=34.0.0
 if "%CMAKE_VERSION%" == "" SET CMAKE_VERSION=3.22.1
 if "%NDK_VERSION%" == "" SET NDK_VERSION=25.1.8937393
 
diff --git a/Engine/Extras/Android/SetupAndroid.command b/Engine/Extras/Android/SetupAndroid.command
--- a/Engine/Extras/Android/SetupAndroid.command
+++ b/Engine/Extras/Android/SetupAndroid.command
@@ -59,10 +59,10 @@ fi
 
 # hardcoded versions for compatibility with non-Turnkey manual running
 if [[ -z "${PLATFORMS_VERSION}" ]]; then
-    PLATFORMS_VERSION="android-33"
+    PLATFORMS_VERSION="android-34"
 fi
 if [[ -z "${BUILDTOOLS_VERSION}" ]]; then
-    BUILDTOOLS_VERSION="33.0.1"
+    BUILDTOOLS_VERSION="34.0.0"
 fi
 if [[ -z "${CMAKE_VERSION}" ]]; then
     CMAKE_VERSION="3.22.1"
diff --git a/Engine/Extras/Android/SetupAndroid.sh b/Engine/Extras/Android/SetupAndroid.sh
--- a/Engine/Extras/Android/SetupAndroid.sh
+++ b/Engine/Extras/Android/SetupAndroid.sh
@@ -27,10 +27,10 @@ fi
 
 # hardcoded versions for compatibility with non-Turnkey manual running
 if [[ -z "${PLATFORMS_VERSION}" ]]; then
-    PLATFORMS_VERSION="android-33"
+    PLATFORMS_VERSION="android-34"
 fi
 if [[ -z "${BUILDTOOLS_VERSION}" ]]; then
-    BUILDTOOLS_VERSION="33.0.1"
+    BUILDTOOLS_VERSION="34.0.0"
 fi
 if [[ -z "${CMAKE_VERSION}" ]]; then
     CMAKE_VERSION="3.22.1"
diff --git a/Engine/Plugins/Online/Android/OnlineSubsystemGooglePlay/Source/OnlineSubsystemGooglePlay_UPL.xml b/Engine/Plugins/Online/Android/OnlineSubsystemGooglePlay/Source/OnlineSubsystemGooglePlay_UPL.xml
--- a/Engine/Plugins/Online/Android/OnlineSubsystemGooglePlay/Source/OnlineSubsystemGooglePlay_UPL.xml
+++ b/Engine/Plugins/Online/Android/OnlineSubsystemGooglePlay/Source/OnlineSubsystemGooglePlay_UPL.xml
@@ -28,7 +28,7 @@
 			<true>
 				<insert>
 			dependencies {
-			implementation 'com.android.billingclient:billing:6.0.1'
+			implementation 'com.android.billingclient:billing:6.2.1'
 			}
 				</insert>
 			</true>
diff --git a/Engine/Source/Programs/UnrealBuildTool/Platform/Android/UEDeployAndroid.cs b/Engine/Source/Programs/UnrealBuildTool/Platform/Android/UEDeployAndroid.cs
--- a/Engine/Source/Programs/UnrealBuildTool/Platform/Android/UEDeployAndroid.cs
+++ b/Engine/Source/Programs/UnrealBuildTool/Platform/Android/UEDeployAndroid.cs
@@ -23,7 +23,7 @@ namespace UnrealBuildTool
 		private const string BUNDLETOOL_JAR = "bundletool-all-0.13.0.jar";
 
 		// classpath of default android build tools gradle plugin
-		private const string ANDROID_TOOLS_BUILD_GRADLE_VERSION = "com.android.tools.build:gradle:7.4.2";
+		private const string ANDROID_TOOLS_BUILD_GRADLE_VERSION = "com.android.tools.build:gradle:8.5.0";
 
 		// default NDK version if not set
 		private const string DEFAULT_NDK_VERSION = "25.1.8937393";
@@ -38,7 +38,7 @@ namespace UnrealBuildTool
 				new AliasedXMLNamespace { Alias = "tools", Url = "http://schemas.android.com/tools" } };
 
 		// Minimum Android SDK that must be used for Java compiling
-		readonly int MinimumSDKLevel = 30;
+		readonly int MinimumSDKLevel = 34;
 
 		// Minimum SDK version needed for App Bundles
 		readonly int MinimumSDKLevelForBundle = 21;
@@ -269,7 +269,7 @@ namespace UnrealBuildTool
 					SDKLevel = ToolChain.GetLargestApiLevel();
 				}
 
-				// make sure it is at least android-23
+				// make sure it is at least the minimum needed to compile
 				int SDKLevelInt = GetApiLevelInt(SDKLevel);
 				if (SDKLevelInt < MinimumSDKLevel)
 				{
@@ -387,7 +387,7 @@ namespace UnrealBuildTool
 
 			if (BestVersionString == null)
 			{
-				BestVersionString = "33.0.1";
+				BestVersionString = "34.0.0";
 				Logger.LogWarning("Failed to find %ANDROID_HOME%/build-tools subdirectory. Will attempt to use {BestVersionString}.", BestVersionString);
 			}
 
@@ -2719,9 +2719,6 @@ namespace UnrealBuildTool
 			bool bAllowIMU = true;
 			Ini.GetBool("/Script/AndroidRuntimeSettings.AndroidRuntimeSettings", "bAllowIMU", out bAllowIMU);
 
-			bool bExtractNativeLibs = true;
-			Ini.GetBool("/Script/AndroidRuntimeSettings.AndroidRuntimeSettings", "bExtractNativeLibs", out bExtractNativeLibs);
-
 			bool bPublicLogFiles = true;
 			Ini.GetBool("/Script/AndroidRuntimeSettings.AndroidRuntimeSettings", "bPublicLogFiles", out bPublicLogFiles);
 			if (!bUseExternalFilesDir)
@@ -2860,13 +2857,6 @@ namespace UnrealBuildTool
 			Text.AppendLine("\t<application android:label=\"@string/app_name\"");
 			Text.AppendLine("\t             android:icon=\"@drawable/icon\"");
 
-			AndroidToolChain.ClangSanitizer Sanitizer = ToolChain.BuildWithSanitizer();
-			// hwasan on NDK r26b+ requires wrap.sh that needs to be unpacked
-			if ((Sanitizer != AndroidToolChain.ClangSanitizer.None && (Sanitizer != AndroidToolChain.ClangSanitizer.HwAddress || ToolChain.HasEmbeddedHWASanSupport())) || bEnableScudoMemoryTracing)
-			{
-				bExtractNativeLibs = true;
-			}
-
 			bool bRequestedLegacyExternalStorage = false;
 			if (ExtraApplicationNodeTags != null)
 			{
@@ -2880,7 +2870,6 @@ namespace UnrealBuildTool
 				}
 			}
 			Text.AppendLine("\t             android:hardwareAccelerated=\"true\"");
-			Text.AppendLine(String.Format("\t             android:extractNativeLibs=\"{0}\"", bExtractNativeLibs ? "true" : "false"));
 			Text.AppendLine("\t				android:name=\"com.epicgames.unreal.GameApplication\"");
 			if (!bIsForDistribution && SDKLevelInt >= 29 && !bRequestedLegacyExternalStorage)
 			{
@@ -5216,7 +5205,7 @@ popd
 					int CompileSDKVersionInt;
 					if (!Int32.TryParse(CompileSDKVersion, out CompileSDKVersionInt))
 					{
-						CompileSDKVersionInt = 23;
+						CompileSDKVersionInt = MinimumSDKLevel;
 					}
 
 					bool bUpdatedCompileSDK = false;
@@ -5287,7 +5276,7 @@ popd
 				// Create local.properties
 				string LocalPropertiesFilename = Path.Combine(UnrealBuildGradlePath, "local.properties");
 				StringBuilder LocalProperties = new StringBuilder();
-				//				LocalProperties.AppendLine(string.Format("ndk.dir={0}", Environment.GetEnvironmentVariable("NDKROOT")!.Replace("\\", "/")));
+//				LocalProperties.AppendLine(string.Format("ndk.dir={0}", Environment.GetEnvironmentVariable("NDKROOT")!.Replace("\\", "/")));
 				LocalProperties.AppendLine(String.Format("sdk.dir={0}", Environment.GetEnvironmentVariable("ANDROID_HOME")!.Replace("\\", "/")));
 				File.WriteAllText(LocalPropertiesFilename, LocalProperties.ToString());
 
@@ -6423,6 +6412,7 @@ import java.util.Collection;
 				StringBuilder BuildGradleContent = new StringBuilder();
 				BuildGradleContent.AppendLine("apply plugin: 'com.android.library'");
 				BuildGradleContent.AppendLine("android {");
+				BuildGradleContent.AppendLine("\tndkVersion = NDK_VERSION");
 				BuildGradleContent.AppendLine("\tndkPath = System.getenv(\"NDKROOT\")");
 				BuildGradleContent.AppendLine("\tcompileSdkVersion = COMPILE_SDK_VERSION.toInteger()");
 				BuildGradleContent.AppendLine("\tbuildToolsVersion = BUILD_TOOLS_VERSION");
