UnrealBuildTool: Don't regenerate .props file contents over and over

#rnx

diff --git a/Engine/Source/Programs/UnrealBuildTool/ProjectFiles/ProjectFileGenerator.cs b/Engine/Source/Programs/UnrealBuildTool/ProjectFiles/ProjectFileGenerator.cs
--- a/Engine/Source/Programs/UnrealBuildTool/ProjectFiles/ProjectFileGenerator.cs
+++ b/Engine/Source/Programs/UnrealBuildTool/ProjectFiles/ProjectFileGenerator.cs
@@ -7,6 +7,8 @@ using System.IO;
 using System.Linq;
 using System.Text;
 using System.Threading;
+using System.Threading.Tasks;
+using System.Xml;
 using System.Xml.Linq;
 using EpicGames.Core;
 using Microsoft.Extensions.Logging;
@@ -564,8 +566,7 @@ namespace UnrealBuildTool
 
 					if (bCreatePropsFile)
 					{
-						FileReference PropsFile = new FileReference(ProjectFile.FullName + ".props");
-						CreateProjectPropsFile(PropsFile);
+						CreateProjectPropsFileAsync(new($"{ProjectFile}.props")).Wait();
 					}
 				}
 			}
@@ -609,18 +610,36 @@ namespace UnrealBuildTool
 		/// <summary>
 		/// Creates a .props file next to each project which specifies the path to the engine directory
 		/// </summary>
-		/// <param name="PropsFile">The properties file path</param>
-		static void CreateProjectPropsFile(FileReference PropsFile)
+		/// <param name="propsFile">The properties file path</param>
+		static async Task<bool> CreateProjectPropsFileAsync(FileReference propsFile) => await FileReference.WriteAllBytesIfDifferentAsync(propsFile, await s_projectPropsFileContents.Value);
+
+		static readonly Lazy<Task<byte[]>> s_projectPropsFileContents = new(async () =>
 		{
-			using StringWriter Writer = new StringWriter();
-			Writer.WriteLine("<?xml version=\"1.0\" encoding=\"utf-8\"?>");
-			Writer.WriteLine("<Project ToolsVersion=\"Current\" xmlns=\"http://schemas.microsoft.com/developer/msbuild/2003\">");
-			Writer.WriteLine("\t<PropertyGroup>");
-			Writer.WriteLine("\t\t<EngineDir Condition=\"'$(EngineDir)' == ''\">{0}</EngineDir>", Unreal.EngineDirectory);
-			Writer.WriteLine("\t</PropertyGroup>");
-			Writer.WriteLine("</Project>");
-			FileReference.WriteAllTextIfDifferent(PropsFile, Writer.ToString());
-		}
+			using MemoryStream stream = new MemoryStream();
+			using (XmlWriter writer = XmlWriter.Create(stream, new XmlWriterSettings() { Async = true, Encoding = Encoding.UTF8, Indent = true, IndentChars = "  " }))
+			{
+				await writer.WriteStartDocumentAsync();
+				{
+					await writer.WriteStartElementAsync(null, "Project", "http://schemas.microsoft.com/developer/msbuild/2003");
+					await writer.WriteAttributeStringAsync(null, "ToolsVersion", null, "Current");
+					{
+						await writer.WriteStartElementAsync(null, "PropertyGroup", null);
+						{
+							await writer.WriteStartElementAsync(null, "EngineDir", null);
+							await writer.WriteAttributeStringAsync(null, "Condition", null, "'$(EngineDir)' == ''");
+							{
+								await writer.WriteStringAsync(Unreal.EngineDirectory.FullName);
+							}
+							await writer.WriteEndElementAsync();
+						}
+						await writer.WriteEndElementAsync();
+					}
+					await writer.WriteEndElementAsync();
+				}
+				await writer.WriteEndDocumentAsync();
+			}
+			return stream.ToArray();
+		});
 
 		/// <summary>
 		/// Finds all csproj within Engine/Source/Programs, and add them if their .uprogram file exists.
