UnrealBuildTool: Rename SingleInstanceMutex to prepare moving to EpicGames.Build

#rnx

diff --git a/Engine/Source/Programs/UnrealBuildTool/Configuration/UEBuildTarget.cs b/Engine/Source/Programs/UnrealBuildTool/Configuration/UEBuildTarget.cs
--- a/Engine/Source/Programs/UnrealBuildTool/Configuration/UEBuildTarget.cs
+++ b/Engine/Source/Programs/UnrealBuildTool/Configuration/UEBuildTarget.cs
@@ -2462,11 +2462,11 @@ namespace UnrealBuildTool
 
 			FindSharedPCHs(Binaries, GlobalCompileEnvironment, Logger);
 
-			SingleInstanceMutex? mutex = null;
+			GlobalSingleInstanceMutex? mutex = null;
 			if (Rules.BuildEnvironment == TargetBuildEnvironment.Shared)
 			{
-				string MutexName = SingleInstanceMutex.GetUniqueMutexForPath("UnrealBuildTool_CreateCompileEnvironmentForProjectFiles", Unreal.EngineDirectory.FullName);
-				mutex = new SingleInstanceMutex(MutexName, true);
+				string MutexName = GlobalSingleInstanceMutex.GetUniqueMutexForPath("UnrealBuildTool_CreateCompileEnvironmentForProjectFiles", Unreal.EngineDirectory);
+				mutex = new GlobalSingleInstanceMutex(MutexName, true);
 			}
 
 			CreateSharedPCHInstances(Rules, TargetToolChain, Binaries, GlobalCompileEnvironment, new NullActionGraphBuilder(Logger), Logger);
diff --git a/Engine/Source/Programs/UnrealBuildTool/Modes/QueryMode.cs b/Engine/Source/Programs/UnrealBuildTool/Modes/QueryMode.cs
--- a/Engine/Source/Programs/UnrealBuildTool/Modes/QueryMode.cs
+++ b/Engine/Source/Programs/UnrealBuildTool/Modes/QueryMode.cs
@@ -265,8 +265,8 @@ namespace UnrealBuildTool
 								bool bUsePrecompiled = false;
 
 								// Prevent multiple conflicting processes building TargetRules at the same time
-								string MutexName = SingleInstanceMutex.GetUniqueMutexForPath("UnrealBuildTool_QueryMode_UEBuildTarget-Create", Unreal.RootDirectory.FullName);
-								using (new SingleInstanceMutex(MutexName, true))
+								string MutexName = GlobalSingleInstanceMutex.GetUniqueMutexForPath("UnrealBuildTool_QueryMode_UEBuildTarget-Create", Unreal.RootDirectory);
+								using (new GlobalSingleInstanceMutex(MutexName, true))
 								{
 									CurrentTarget = UEBuildTarget.Create(TargetDescriptors[0], false, false, bUsePrecompiled, TargetDescriptors[0].IntermediateEnvironment, NullLogger.Instance);
 								}
@@ -362,8 +362,8 @@ namespace UnrealBuildTool
 					bool bUsePrecompiled = false;
 
 					// Prevent multiple conflicting processes building TargetRules at the same time
-					string MutexName = SingleInstanceMutex.GetUniqueMutexForPath("UnrealBuildTool_QueryMode_UEBuildTarget-Create", Unreal.RootDirectory.FullName);
-					using (new SingleInstanceMutex(MutexName, true))
+					string MutexName = GlobalSingleInstanceMutex.GetUniqueMutexForPath("UnrealBuildTool_QueryMode_UEBuildTarget-Create", Unreal.RootDirectory);
+					using (new GlobalSingleInstanceMutex(MutexName, true))
 					{
 						CurrentTarget = UEBuildTarget.Create(TargetDescriptors[0], false, false, bUsePrecompiled, TargetDescriptors[0].IntermediateEnvironment, Logger);
 					}
diff --git a/Engine/Source/Programs/UnrealBuildTool/Modes/WriteMetadataMode.cs b/Engine/Source/Programs/UnrealBuildTool/Modes/WriteMetadataMode.cs
--- a/Engine/Source/Programs/UnrealBuildTool/Modes/WriteMetadataMode.cs
+++ b/Engine/Source/Programs/UnrealBuildTool/Modes/WriteMetadataMode.cs
@@ -96,8 +96,8 @@ namespace UnrealBuildTool
 		{
 			// Acquire a different mutex to the regular UBT instance, since this mode will be called as part of a build. We need the mutex to ensure that building two modular configurations 
 			// in parallel don't clash over writing shared *.modules files (eg. DebugGame and Development editors).
-			string MutexName = SingleInstanceMutex.GetUniqueMutexForPath("UnrealBuildTool_WriteMetadata", Unreal.RootDirectory.FullName);
-			using (new SingleInstanceMutex(MutexName, true))
+			string MutexName = GlobalSingleInstanceMutex.GetUniqueMutexForPath("UnrealBuildTool_WriteMetadata", Unreal.RootDirectory);
+			using (new GlobalSingleInstanceMutex(MutexName, true))
 			{
 				return ExecuteInternal(Arguments, Logger);
 			}
diff --git a/Engine/Source/Programs/UnrealBuildTool/Platform/IOS/UEDeployIOS.cs b/Engine/Source/Programs/UnrealBuildTool/Platform/IOS/UEDeployIOS.cs
--- a/Engine/Source/Programs/UnrealBuildTool/Platform/IOS/UEDeployIOS.cs
+++ b/Engine/Source/Programs/UnrealBuildTool/Platform/IOS/UEDeployIOS.cs
@@ -509,8 +509,8 @@ namespace UnrealBuildTool
 			if (BuildHostPlatform.Current.Platform == UnrealTargetPlatform.Mac && !bBuildAsFramework)
 			{
 				// need to make sure touching this plist is in the same lock as FinalizeAppWithXcode()
-				string MutexName = SingleInstanceMutex.GetUniqueMutexForPath("UnrealBuildTool_XcodeBuild", Unreal.RootDirectory.FullName);
-				using (new SingleInstanceMutex(MutexName, true))
+				string MutexName = GlobalSingleInstanceMutex.GetUniqueMutexForPath("UnrealBuildTool_XcodeBuild", Unreal.RootDirectory);
+				using (new GlobalSingleInstanceMutex(MutexName, true))
 				{
 					FileReference FinalPlistFile;
 					FinalPlistFile = new FileReference($"{ProjectDirectory}/Build/IOS/UBTGenerated/Info.Template.plist");
diff --git a/Engine/Source/Programs/UnrealBuildTool/System/GlobalSingleInstanceMutex.cs b/Engine/Source/Programs/UnrealBuildTool/System/GlobalSingleInstanceMutex.cs
new file mode 100644
--- /dev/null
+++ b/Engine/Source/Programs/UnrealBuildTool/System/GlobalSingleInstanceMutex.cs
@@ -0,0 +1,91 @@
+// Copyright Epic Games, Inc. All Rights Reserved.
+
+using System;
+using System.Threading;
+using EpicGames.Core;
+
+namespace UnrealBuildTool
+{
+	/// <summary>
+	/// System-wide mutex allowing only one instance of the program to run at a time
+	/// </summary>
+	class GlobalSingleInstanceMutex : IDisposable
+	{
+		/// <summary>
+		/// The global mutex instance
+		/// </summary>
+		Mutex? GlobalMutex;
+
+		/// <summary>
+		/// Constructor. Attempts to acquire the global mutex
+		/// </summary>
+		/// <param name="mutexName">Name of the mutex to acquire</param>
+		/// <param name="bWaitMutex">Allow waiting for the mutex to be acquired</param>
+		public GlobalSingleInstanceMutex(string mutexName, bool bWaitMutex)
+		{
+			// Try to create the mutex, with it initially locked
+			bool bCreatedMutex;
+			GlobalMutex = new Mutex(true, mutexName, out bCreatedMutex);
+
+			// If we didn't create the mutex, we can wait for it or fail immediately
+			if (!bCreatedMutex)
+			{
+				if (bWaitMutex)
+				{
+					try
+					{
+						GlobalMutex.WaitOne();
+					}
+					catch (AbandonedMutexException)
+					{
+					}
+				}
+				else
+				{
+					throw new BuildException("A conflicting instance of UnrealBuildTool is already running.");
+				}
+			}
+		}
+
+		/// <summary>
+		/// Gets the name of a mutex unique for the given path
+		/// </summary>
+		/// <param name="name">Base name of the mutex</param>
+		/// <param name="uniquePath">Path to identify a unique mutex</param>
+		public static string GetUniqueMutexForPath(string name, string uniquePath)
+		{
+			// generate an IoHash of the path, as GetHashCode is not guaranteed to generate a stable hash
+			return $"Global\\{name}_{IoHash.Compute(uniquePath.ToUpperInvariant())}";
+		}
+
+		/// <summary>
+		/// Gets the name of a mutex unique for the given path
+		/// </summary>
+		/// <param name="name">Base name of the mutex</param>
+		/// <param name="uniquePath">Path to identify a unique mutex</param>
+		public static string GetUniqueMutexForPath(string name, FileSystemReference? uniquePath) => GetUniqueMutexForPath(name, uniquePath?.FullName ?? String.Empty);
+
+		/// <summary>
+		/// Release the mutex and dispose of the object
+		/// </summary>
+		public void Dispose()
+		{
+			Dispose(true);
+			GC.SuppressFinalize(this);
+		}
+
+		/// <summary>
+		/// Overridable dispose method
+		/// </summary>
+		/// <param name="disposing">Whether the object should be disposed</param>
+		protected virtual void Dispose(bool disposing)
+		{
+			if (disposing)
+			{
+				GlobalMutex?.ReleaseMutex();
+				GlobalMutex?.Dispose();
+				GlobalMutex = null;
+			}
+		}
+	}
+}
diff --git a/Engine/Source/Programs/UnrealBuildTool/System/RulesCompiler.cs b/Engine/Source/Programs/UnrealBuildTool/System/RulesCompiler.cs
--- a/Engine/Source/Programs/UnrealBuildTool/System/RulesCompiler.cs
+++ b/Engine/Source/Programs/UnrealBuildTool/System/RulesCompiler.cs
@@ -82,8 +82,8 @@ namespace UnrealBuildTool
 		public static RulesAssembly CreateEngineRulesAssembly(bool bUsePrecompiled, bool bSkipCompile, bool bForceCompile, ILogger Logger)
 		{
 			// Prevent multiple conflicting processes building rules assembly at the same time
-			string MutexName = SingleInstanceMutex.GetUniqueMutexForPath("UnrealBuildTool_CreateTargetRulesAssembly", Unreal.EngineDirectory.FullName);
-			using (new SingleInstanceMutex(MutexName, true))
+			string MutexName = GlobalSingleInstanceMutex.GetUniqueMutexForPath("UnrealBuildTool_CreateTargetRulesAssembly", Unreal.EngineDirectory);
+			using (new GlobalSingleInstanceMutex(MutexName, true))
 			{
 				if (EngineRulesAssembly == null)
 				{
@@ -240,8 +240,8 @@ namespace UnrealBuildTool
 		public static RulesAssembly CreateProjectRulesAssembly(FileReference ProjectFileName, bool bUsePrecompiled, bool bSkipCompile, bool bForceCompile, ILogger Logger)
 		{
 			// Prevent multiple conflicting processes building rules assembly at the same time
-			string MutexName = SingleInstanceMutex.GetUniqueMutexForPath("UnrealBuildTool_CreateTargetRulesAssembly", ProjectFileName.FullName);
-			using (new SingleInstanceMutex(MutexName, true))
+			string MutexName = GlobalSingleInstanceMutex.GetUniqueMutexForPath("UnrealBuildTool_CreateTargetRulesAssembly", ProjectFileName);
+			using (new GlobalSingleInstanceMutex(MutexName, true))
 			{
 				// Check if there's an existing assembly for this project
 				return LoadedAssemblyMap.GetOrAdd(ProjectFileName, _ =>
@@ -350,8 +350,8 @@ namespace UnrealBuildTool
 		public static RulesAssembly CreatePluginRulesAssembly(FileReference PluginFileName, bool bSkipCompile, bool bForceCompile, RulesAssembly Parent, bool bBuildPluginAsLocal, bool bContainsEngineModules, ILogger Logger)
 		{
 			// Prevent multiple conflicting processes building rules assembly at the same time
-			string MutexName = SingleInstanceMutex.GetUniqueMutexForPath("UnrealBuildTool_CreateTargetRulesAssembly", PluginFileName.FullName);
-			using (new SingleInstanceMutex(MutexName, true))
+			string MutexName = GlobalSingleInstanceMutex.GetUniqueMutexForPath("UnrealBuildTool_CreateTargetRulesAssembly", PluginFileName);
+			using (new GlobalSingleInstanceMutex(MutexName, true))
 			{
 				// Check if there's an existing assembly for this project
 				return LoadedAssemblyMap.GetOrAdd(PluginFileName, _ =>
diff --git a/Engine/Source/Programs/UnrealBuildTool/System/SingleInstanceMutex.cs b/Engine/Source/Programs/UnrealBuildTool/System/SingleInstanceMutex.cs
deleted file mode 100644
diff --git a/Engine/Source/Programs/UnrealBuildTool/ToolChain/AppleToolChain.cs b/Engine/Source/Programs/UnrealBuildTool/ToolChain/AppleToolChain.cs
--- a/Engine/Source/Programs/UnrealBuildTool/ToolChain/AppleToolChain.cs
+++ b/Engine/Source/Programs/UnrealBuildTool/ToolChain/AppleToolChain.cs
@@ -946,8 +946,8 @@ namespace UnrealBuildTool
 		{
 			// Acquire a different mutex to the regular UBT instance, since this mode will be called as part of a build. We need the mutex to ensure that building two modular configurations 
 			// in parallel don't clash over writing shared *.modules files (eg. DebugGame and Development editors).
-			string MutexName = SingleInstanceMutex.GetUniqueMutexForPath("UnrealBuildTool_XcodeBuild", Unreal.RootDirectory.FullName);
-			using (new SingleInstanceMutex(MutexName, true))
+			string MutexName = GlobalSingleInstanceMutex.GetUniqueMutexForPath("UnrealBuildTool_XcodeBuild", Unreal.RootDirectory);
+			using (new GlobalSingleInstanceMutex(MutexName, true))
 			{
 				List<string> Arguments = new()
 				{
diff --git a/Engine/Source/Programs/UnrealBuildTool/UnrealBuildTool.cs b/Engine/Source/Programs/UnrealBuildTool/UnrealBuildTool.cs
--- a/Engine/Source/Programs/UnrealBuildTool/UnrealBuildTool.cs
+++ b/Engine/Source/Programs/UnrealBuildTool/UnrealBuildTool.cs
@@ -439,7 +439,7 @@ namespace UnrealBuildTool
 		{
 			FileReference? RunFile = null;
 			DirectoryReference? TempDirectory = null;
-			SingleInstanceMutex? Mutex = null;
+			GlobalSingleInstanceMutex? Mutex = null;
 			JsonTracer? Tracer = null;
 
 			ILogger Logger = Log.Logger;
@@ -553,8 +553,8 @@ namespace UnrealBuildTool
 				{
 					using (GlobalTracer.Instance.BuildSpan("XmlConfig.ReadConfigFiles()").StartActive())
 					{
-						string XmlConfigMutexName = SingleInstanceMutex.GetUniqueMutexForPath("UnrealBuildTool_Mutex_XmlConfig", Assembly.GetExecutingAssembly().Location);
-						using (SingleInstanceMutex XmlConfigMutex = new SingleInstanceMutex(XmlConfigMutexName, true))
+						string XmlConfigMutexName = GlobalSingleInstanceMutex.GetUniqueMutexForPath("UnrealBuildTool_Mutex_XmlConfig", FileReference.FromString(Assembly.GetExecutingAssembly().Location));
+						using (GlobalSingleInstanceMutex XmlConfigMutex = new GlobalSingleInstanceMutex(XmlConfigMutexName, true))
 						{
 							FileReference? XmlConfigCache = Arguments.GetFileReferenceOrDefault("-XmlConfigCache=", null);
 							XmlConfig.ReadConfigFiles(XmlConfigCache, null, Logger);
@@ -649,8 +649,8 @@ namespace UnrealBuildTool
 				{
 					using (GlobalTracer.Instance.BuildSpan("SingleInstanceMutex.Acquire()").StartActive())
 					{
-						string MutexName = SingleInstanceMutex.GetUniqueMutexForPath("UnrealBuildTool_Mutex", Assembly.GetExecutingAssembly().Location);
-						Mutex = new SingleInstanceMutex(MutexName, Options.bWaitMutex);
+						string MutexName = GlobalSingleInstanceMutex.GetUniqueMutexForPath("UnrealBuildTool_Mutex", FileReference.FromString(Assembly.GetExecutingAssembly().Location));
+						Mutex = new GlobalSingleInstanceMutex(MutexName, Options.bWaitMutex);
 					}
 				}
 
