Remove RHI breadcrumbs from default Android build config

diff --git a/Engine/Source/Programs/UnrealBuildTool/Platform/Android/UEBuildAndroid.cs b/Engine/Source/Programs/UnrealBuildTool/Platform/Android/UEBuildAndroid.cs
--- a/Engine/Source/Programs/UnrealBuildTool/Platform/Android/UEBuildAndroid.cs
+++ b/Engine/Source/Programs/UnrealBuildTool/Platform/Android/UEBuildAndroid.cs
@@ -508,7 +508,6 @@ namespace UnrealBuildTool
 
 			CompileEnvironment.Definitions.Add("WITH_EDITOR=0");
 			CompileEnvironment.Definitions.Add("USE_NULL_RHI=0");
- 			CompileEnvironment.Definitions.Add("WITH_RHI_BREADCRUMBS=1");
 
 			if (Target.AndroidPlatform.bEnableRayTracing && Target.AndroidPlatform.bSupportsVulkanSM5)
 			{
