- Changed Hybrid project checking to always check all enabled platforms to see if any need to be hybrid. This helps address an issue where a single-platform (say IOS) plugin was causing packaging to fail because the host platform (Mac) to clean up the hybrid target files, so then packaging for IOS would fail to find the compiled hybrid targets

diff --git a/Engine/Source/Programs/AutomationTool/AutomationUtils/ProjectUtils.cs b/Engine/Source/Programs/AutomationTool/AutomationUtils/ProjectUtils.cs
--- a/Engine/Source/Programs/AutomationTool/AutomationUtils/ProjectUtils.cs
+++ b/Engine/Source/Programs/AutomationTool/AutomationUtils/ProjectUtils.cs
@@ -262,7 +262,7 @@ namespace AutomationTool
 					}
 				}
 
-				if (NativeProjects.ConditionalMakeTempTargetForHybridProject(RawProjectPath, TargetPlatforms, Logger))
+				if (NativeProjects.ConditionalMakeTempTargetForHybridProject(RawProjectPath, Logger))
 				{
 					Properties.bWasGenerated = true;
 					string TempTargetDir = CommandUtils.CombinePaths(Path.GetDirectoryName(RawProjectPath.FullName), "Intermediate", "Source");
diff --git a/Engine/Source/Programs/UnrealBuildTool/Configuration/TargetDescriptor.cs b/Engine/Source/Programs/UnrealBuildTool/Configuration/TargetDescriptor.cs
--- a/Engine/Source/Programs/UnrealBuildTool/Configuration/TargetDescriptor.cs
+++ b/Engine/Source/Programs/UnrealBuildTool/Configuration/TargetDescriptor.cs
@@ -579,7 +579,7 @@ namespace UnrealBuildTool
 						// make a temp target for hybrid content-as-code projects
 						if (TargetProjectFile != null)
 						{
-							NativeProjects.ConditionalMakeTempTargetForHybridProject(TargetProjectFile, new List<UnrealTargetPlatform>(1) { Platform }, Logger);
+							NativeProjects.ConditionalMakeTempTargetForHybridProject(TargetProjectFile, Logger);
 						}
 
 						if (ParamArchitectures != null)
diff --git a/Engine/Source/Programs/UnrealBuildTool/ProjectFiles/ProjectFileGenerator.cs b/Engine/Source/Programs/UnrealBuildTool/ProjectFiles/ProjectFileGenerator.cs
--- a/Engine/Source/Programs/UnrealBuildTool/ProjectFiles/ProjectFileGenerator.cs
+++ b/Engine/Source/Programs/UnrealBuildTool/ProjectFiles/ProjectFileGenerator.cs
@@ -1020,7 +1020,7 @@ namespace UnrealBuildTool
 			{
 				foreach (FileReference GameUProjectFile in AllGameProjects)
 				{
-					NativeProjects.ConditionalMakeTempTargetForHybridProject(GameUProjectFile, PlatformProjectGenerators.GetRegisteredPlatforms(), Logger);
+					NativeProjects.ConditionalMakeTempTargetForHybridProject(GameUProjectFile, Logger);
 				}
 
 				// they are created in a temp location, which we need to scan
@@ -2939,7 +2939,7 @@ namespace UnrealBuildTool
 					ProjectFile.IsForeignProject = CheckProjectFile != null && !NativeProjects.IsNativeProject(CheckProjectFile, Logger);
 					ProjectFile.IsGeneratedProject = true;
 					ProjectFile.IsStubProject = Unreal.IsProjectInstalled();
-					ProjectFile.IsHybridContentOnlyProject = CheckProjectFile != null && NativeProjects.IsHybridContentOnlyProject(CheckProjectFile, Logger);
+					ProjectFile.IsHybridContentOnlyProject = CheckProjectFile != null && NativeProjects.IsHybridContentOnlyProject(CheckProjectFile, out _, Logger);
 					if (TargetRulesObject.bBuildInSolutionByDefault.HasValue)
 					{
 						ProjectFile.ShouldBuildByDefaultForSolutionTargets = TargetRulesObject.bBuildInSolutionByDefault.Value;
diff --git a/Engine/Source/Programs/UnrealBuildTool/System/NativeProjects.cs b/Engine/Source/Programs/UnrealBuildTool/System/NativeProjects.cs
--- a/Engine/Source/Programs/UnrealBuildTool/System/NativeProjects.cs
+++ b/Engine/Source/Programs/UnrealBuildTool/System/NativeProjects.cs
@@ -83,31 +83,13 @@ namespace UnrealBuildTool
 		/// Returns true if this project is a Hybrid content only project that requires it to be built as code
 		/// </summary>
 		/// <param name="UProjectFile"></param>
+		/// <param name="Reason"></param>
 		/// <param name="Logger"></param>
 		/// <returns></returns>
-		public static bool IsHybridContentOnlyProject(FileReference UProjectFile, ILogger Logger)
+		public static bool IsHybridContentOnlyProject(FileReference UProjectFile, [NotNullWhen(true)] out string? Reason, ILogger Logger)
 		{
 			return RequiresTempTarget(
 				UProjectFile,
-				new List<UnrealTargetPlatform>() { BuildHostPlatform.Current.Platform },
-				new List<UnrealTargetConfiguration>() { UnrealTargetConfiguration.Development, UnrealTargetConfiguration.Shipping },
-				out _,
-				Logger);
-		}
-
-		/// <summary>
-		/// Returns true if this project is a Hybrid content only project that requires it to be built as code
-		/// </summary>
-		/// <param name="UProjectFile"></param>
-		/// <param name="TargetPlatforms">The target platforms we are asking about.</param>
-		/// <param name="Reason">Contains a description of the reason the project is hybrid</param>
-		/// <param name="Logger"></param>
-		/// <returns></returns>
-		public static bool IsHybridContentOnlyProject(FileReference UProjectFile, List<UnrealTargetPlatform> TargetPlatforms, [NotNullWhen(true)] out string? Reason, ILogger Logger)
-		{
-			return RequiresTempTarget(
-				UProjectFile,
-				TargetPlatforms,
 				new List<UnrealTargetConfiguration>() { UnrealTargetConfiguration.Development, UnrealTargetConfiguration.Shipping },
 				out Reason,
 				Logger);
@@ -117,13 +99,12 @@ namespace UnrealBuildTool
 		/// Creates temporary target files, if needed, for a hybrid content only project
 		/// </summary>
 		/// <param name="UProjectFile"></param>
-		/// <param name="TargetPlatforms"></param>
 		/// <param name="Logger"></param>
 		/// <returns>True if the project is hybrid</returns>
-		public static bool ConditionalMakeTempTargetForHybridProject(FileReference UProjectFile, List<UnrealTargetPlatform> TargetPlatforms, ILogger Logger)
+		public static bool ConditionalMakeTempTargetForHybridProject(FileReference UProjectFile, ILogger Logger)
 		{
 			string? Reason;
-			bool bIsHybrid = IsHybridContentOnlyProject(UProjectFile, TargetPlatforms, out Reason, Logger);
+			bool bIsHybrid = IsHybridContentOnlyProject(UProjectFile, out Reason, Logger);
 
 			DirectoryReference TempDir = DirectoryReference.Combine(UProjectFile.Directory, "Intermediate", "Source");
 
@@ -146,7 +127,7 @@ namespace UnrealBuildTool
 				// clean up if needed
 				if (bWasHybrid)
 				{
-					Logger.LogWarning("Cleaning old temporary Target files for {Project} because it no longer being treated as a code-based project.", ProjectName);
+					Logger.LogWarning("Cleaning old temporary Target files for {Project} because it no longer being treated as a code-based project for any enabled platform.", ProjectName);
 					DirectoryReference.Delete(TempDir, bRecursive: true);
 				}
 				return false;
@@ -250,7 +231,7 @@ namespace UnrealBuildTool
 			return false;
 		}
 
-		private static bool RequiresTempTarget(FileReference UProjectFile, List<UnrealTargetPlatform> Platforms, List<UnrealTargetConfiguration> Configurations, [NotNullWhen(true)] out string? Reason, ILogger Logger)
+		private static bool RequiresTempTarget(FileReference UProjectFile, List<UnrealTargetConfiguration> Configurations, [NotNullWhen(true)] out string? Reason, ILogger Logger)
 		{
 			// no reason by default
 			Reason = null;
@@ -261,6 +242,11 @@ namespace UnrealBuildTool
 				return false;
 			}
 
+			List<UnrealTargetPlatform> Platforms = DataDrivenPlatformInfo.GetAllPlatformInfos()
+				.Where(x => x.Value.bIsEnabled)
+				.Select(x => UnrealTargetPlatform.Parse(x.Key.Replace("Windows", "Win64")))
+				.ToList();
+
 			bool bHasCode = ProjectHasCode(UProjectFile, bCheckForTempTargets: false);
 			foreach (UnrealTargetPlatform Platform in Platforms)
 			{
@@ -269,7 +255,8 @@ namespace UnrealBuildTool
 					string? InnerReason;
 					if (RequiresTempTarget(UProjectFile, bHasCode, Platform, Configuration, TargetType.Game, out InnerReason, Logger))
 					{
-						Reason = $"{UProjectFile.GetFileName()} is has no code, but is being treated as a code-based project because: {InnerReason}.";
+						string PlatformNames = string.Join(", ", Platforms.Select(x => x.ToString()));
+						Reason = $"{UProjectFile.GetFileName()} is has no code, but is being treated as a code-based project for platforms {PlatformNames} because: {InnerReason}.";
 						return true;
 					}
 				}
@@ -280,7 +267,7 @@ namespace UnrealBuildTool
 		/// <summary>
 		/// NOTE: This function must mirror the functionality of TargetPlatformBase::RequiresTempTarget
 		/// </summary>
-		public static bool RequiresTempTarget(FileReference RawProjectPath, bool bProjectHasCode, UnrealTargetPlatform Platform, UnrealTargetConfiguration Configuration, TargetType TargetType, out string? OutReason, ILogger Logger)
+		private static bool RequiresTempTarget(FileReference RawProjectPath, bool bProjectHasCode, UnrealTargetPlatform Platform, UnrealTargetConfiguration Configuration, TargetType TargetType, out string? OutReason, ILogger Logger)
 		{
 			// check to see if we already have a Target.cs file
 			if (bProjectHasCode)
