Use dotnet publish to get UBT into Engine\Binaries instead of setting the output path

#rnx

diff --git a/Engine/Build/BatchFiles/BuildUAT.bat b/Engine/Build/BatchFiles/BuildUAT.bat
--- a/Engine/Build/BatchFiles/BuildUAT.bat
+++ b/Engine/Build/BatchFiles/BuildUAT.bat
@@ -19,6 +19,10 @@ rem ## Change the CWD to /Engine/Source.
 pushd "%~dp0..\..\Source"
 if not exist ..\Build\BatchFiles\BuildUAT.bat goto Error_BatchFileInWrongLocation
 
+rem Ensure UnrealBuildTool is up to date as it is a prereq
+call %~dp0\BuildUBT.bat
+if errorlevel 1 goto Error_UBTCompileFailed
+
 rem Check to see if the files in the AutomationTool, EpicGames.Build, EpicGames.Core, or UnrealBuildTool directory have changed.
 rem find ".cs" files to only lines that match those names - excludes lines that will change for uninteresting reasons, like free space
 md ..\Intermediate\Build >nul 2>nul
@@ -115,11 +119,15 @@ goto Exit
 
 :Error_UATCompileFailed
 echo.
-echo RunUBT ERROR: UnrealBuildTool failed to compile.
+echo RunUAT ERROR: AutomationTool failed to compile.
 echo.
 set RUNUAT_EXITCODE=1
 goto Exit
 
+:Error_UBTCompileFailed
+set RUNUAT_EXITCODE=1
+goto Exit
+
 :Exit
 rem ## Restore original CWD in case we change it
 popd
diff --git a/Engine/Build/BatchFiles/BuildUAT.sh b/Engine/Build/BatchFiles/BuildUAT.sh
--- a/Engine/Build/BatchFiles/BuildUAT.sh
+++ b/Engine/Build/BatchFiles/BuildUAT.sh
@@ -19,6 +19,13 @@ if [ ! -f ../Build/BatchFiles/BuildUAT.sh ]; then
   exit 1
 fi
 
+# Ensure UnrealBuildTool is up to date as it is a prereq
+$SCRIPT_DIR/BuildUBT.sh
+if [ $? -ne 0 ]; then
+popd >/dev/null
+exit 1
+fi
+
 MSBuild_Verbosity="${1:-quiet}"
 
 # Check to see if the files in the AutomationTool, EpicGames.Build, EpicGames.Core, or UnrealBuildTool
diff --git a/Engine/Build/BatchFiles/BuildUBT.bat b/Engine/Build/BatchFiles/BuildUBT.bat
--- a/Engine/Build/BatchFiles/BuildUBT.bat
+++ b/Engine/Build/BatchFiles/BuildUBT.bat
@@ -90,6 +90,10 @@ echo Building UnrealBuildTool...
 dotnet build Programs\UnrealBuildTool\UnrealBuildTool.csproj -c Development -v %MSBUILD_LOGLEVEL%
 if errorlevel 1 goto Error_UBTCompileFailed
 
+echo Publishing UnrealBuildTool...
+dotnet publish Programs\UnrealBuildTool\UnrealBuildTool.csproj -c Development --output ..\Binaries\DotNET\UnrealBuildTool --no-build -v %MSBUILD_LOGLEVEL%
+if errorlevel 1 goto Error_UBTCompileFailed
+
 rem record input files - regardless of how we got here, these are now our point of reference
 copy /y ..\Intermediate\Build\UnrealBuildToolFiles.txt ..\Intermediate\Build\UnrealBuildToolPrevFiles.txt >nul
 
diff --git a/Engine/Build/BatchFiles/BuildUBT.sh b/Engine/Build/BatchFiles/BuildUBT.sh
--- a/Engine/Build/BatchFiles/BuildUBT.sh
+++ b/Engine/Build/BatchFiles/BuildUBT.sh
@@ -109,6 +109,13 @@ if [ $PERFORM_REBUILD -eq 1 ]; then
     exit 1
   fi
 
+  dotnet publish Programs/UnrealBuildTool/UnrealBuildTool.csproj -c Development --output ../Binaries/DotNET/UnrealBuildTool --no-build -v $MSBuild_Verbosity
+  if [ $? -ne 0 ]; then
+    echo "UnrealBuildTool publish failed"
+    popd >/dev/null
+    exit 1
+  fi
+
   touch ../Intermediate/Build/UnrealBuildToolLastBuildTime
 else
   echo "UnrealBuildTool.dll is up to date"
diff --git a/Engine/Source/Programs/Shared/EpicGames.MsBuild/EpicGames.MsBuild.csproj b/Engine/Source/Programs/Shared/EpicGames.MsBuild/EpicGames.MsBuild.csproj
--- a/Engine/Source/Programs/Shared/EpicGames.MsBuild/EpicGames.MsBuild.csproj
+++ b/Engine/Source/Programs/Shared/EpicGames.MsBuild/EpicGames.MsBuild.csproj
@@ -38,8 +38,8 @@
   </ItemGroup>
 
   <ItemGroup>
-    <PackageReference Include="Microsoft.Build" Version="17.1.0" ExcludeAssets="runtime" PrivateAssets="all" />
-    <PackageReference Include="Microsoft.Build.Locator" Version="1.4.1" PrivateAssets="all" />
+    <PackageReference Include="Microsoft.Build" Version="17.1.0" ExcludeAssets="runtime" />
+    <PackageReference Include="Microsoft.Build.Locator" Version="1.4.1" />
     <PackageReference Include="System.Drawing.Common" Version="4.7.3" />
   </ItemGroup>
 
diff --git a/Engine/Source/Programs/UnrealBuildTool/UnrealBuildTool.csproj b/Engine/Source/Programs/UnrealBuildTool/UnrealBuildTool.csproj
--- a/Engine/Source/Programs/UnrealBuildTool/UnrealBuildTool.csproj
+++ b/Engine/Source/Programs/UnrealBuildTool/UnrealBuildTool.csproj
@@ -5,7 +5,6 @@
     <TargetFramework>net8.0</TargetFramework>
     <Configuration Condition=" '$(Configuration)' == '' ">Development</Configuration>
     <OutputType>Exe</OutputType>
-    <OutputPath>..\..\..\Binaries\DotNET\UnrealBuildTool</OutputPath>
     <AppendTargetFrameworkToOutputPath>false</AppendTargetFrameworkToOutputPath>
     <GenerateAssemblyInfo>false</GenerateAssemblyInfo>
     <GenerateTargetFrameworkAttribute>false</GenerateTargetFrameworkAttribute>
@@ -25,28 +24,19 @@
   </PropertyGroup>
   <PropertyGroup Condition=" '$(Configuration)|$(Platform)' == 'Development|AnyCPU' ">
     <Optimize>True</Optimize>
-    <DocumentationFile>..\..\..\Binaries\DotNET\UnrealBuildTool\UnrealBuildTool.xml</DocumentationFile>
   </PropertyGroup>
   <PropertyGroup Condition=" '$(Configuration)|$(Platform)' == 'Release|AnyCPU' ">
     <Optimize>True</Optimize>
-    <PlatformTarget>AnyCPU</PlatformTarget>
-    <ErrorReport>prompt</ErrorReport>
-    <AllowUnsafeBlocks>false</AllowUnsafeBlocks>
-    <Prefer32Bit>false</Prefer32Bit>
-    <DocumentationFile>..\..\..\Binaries\DotNET\UnrealBuildTool\UnrealBuildTool.xml</DocumentationFile>
   </PropertyGroup>
   <PropertyGroup Condition=" '$(Configuration)|$(Platform)' == 'Debug|AnyCPU' ">
-    <DocumentationFile>..\..\..\Binaries\DotNET\UnrealBuildTool\UnrealBuildTool.xml</DocumentationFile>
     <DebugSymbols>true</DebugSymbols>
     <DebugType>full</DebugType>
     <DebugType Condition="'$([System.Runtime.InteropServices.RuntimeInformation]::IsOSPlatform($([System.Runtime.InteropServices.OSPlatform]::Windows)))' == 'true' And '$([System.Runtime.InteropServices.RuntimeInformation]::OSArchitecture)' == 'Arm64'">portable</DebugType>
   </PropertyGroup>
   <PropertyGroup Condition=" '$(Configuration)|$(Platform)' == 'Analyze|AnyCPU' ">
-    <DocumentationFile>..\..\..\Binaries\DotNET\UnrealBuildTool\UnrealBuildTool.xml</DocumentationFile>
     <RunAnalyzersDuringBuild>True</RunAnalyzersDuringBuild>
     <TreatWarningsAsErrors>false</TreatWarningsAsErrors>
     <GenerateDocumentationFile>True</GenerateDocumentationFile>
-    <DocumentationFile></DocumentationFile>
   </PropertyGroup>
   <ItemGroup>
     <Reference Include="Ionic.Zip.Reduced">
