UnrealBuildTool: If the clang compiler is in use, default to allowing the clang linker

#rnx

diff --git a/Engine/Source/Programs/UnrealBuildTool/Platform/Windows/UEBuildWindows.cs b/Engine/Source/Programs/UnrealBuildTool/Platform/Windows/UEBuildWindows.cs
--- a/Engine/Source/Programs/UnrealBuildTool/Platform/Windows/UEBuildWindows.cs
+++ b/Engine/Source/Programs/UnrealBuildTool/Platform/Windows/UEBuildWindows.cs
@@ -357,7 +357,8 @@ namespace UnrealBuildTool
 		[ConfigFile(ConfigHierarchyType.Engine, "/Script/WindowsTargetPlatform.WindowsTargetSettings", "bAllowClangLinker")]
 		[XmlConfigFile(Category = "WindowsPlatform")]
 		[CommandLine("-ClangLinker")]
-		public bool bAllowClangLinker = false;
+		[CommandLine("-NoClangLinker", Value = "false")]
+		public bool bAllowClangLinker = true;
 
 		/// <summary>
 		/// The specific Windows SDK version to use. This may be a specific version number (for example, "8.1", "10.0" or "10.0.10150.0"), or the string "Latest", to select the newest available version.
