Allow the extra LLVM verify check in the AutoRTFM compiler pass to be performed via a runtime toggable switch.

We are using this in our AutoRTFM compiler testing to vet our changes work as expected.

diff --git a/Engine/Source/Programs/UnrealBuildTool/Configuration/ReadOnlyTargetRules.cs b/Engine/Source/Programs/UnrealBuildTool/Configuration/ReadOnlyTargetRules.cs
--- a/Engine/Source/Programs/UnrealBuildTool/Configuration/ReadOnlyTargetRules.cs
+++ b/Engine/Source/Programs/UnrealBuildTool/Configuration/ReadOnlyTargetRules.cs
@@ -138,6 +138,7 @@ namespace UnrealBuildTool
 
 		public bool bUseAutoRTFMCompiler => Inner.bUseAutoRTFMCompiler;
 		public bool bUseAutoRTFMVerifier => Inner.bUseAutoRTFMVerifier;
+		public bool bAutoRTFMVerify => Inner.bAutoRTFMVerify;
 
 		[Obsolete("Deprecated in UE5.1 - No longer used as Chaos is always enabled.")]
 		public bool bCompileChaos => Inner.bCompileChaos;
diff --git a/Engine/Source/Programs/UnrealBuildTool/Configuration/TargetRules.cs b/Engine/Source/Programs/UnrealBuildTool/Configuration/TargetRules.cs
--- a/Engine/Source/Programs/UnrealBuildTool/Configuration/TargetRules.cs
+++ b/Engine/Source/Programs/UnrealBuildTool/Configuration/TargetRules.cs
@@ -912,6 +912,16 @@ namespace UnrealBuildTool
 		[XmlConfigFile(Category = "BuildConfiguration")]
 		public bool bUseAutoRTFMVerifier { get; set; } = false;
 
+		/// <summary>
+		/// Whether to run LLVM verification after the AutoRTFM compiler pass.
+		/// This is used by our compiler folks to ensure the pass works with
+		/// the various code-paths UBT can take us down.
+		/// </summary>
+		[RequiresUniqueBuildEnvironment]
+		[CommandLine("-AutoRTFMVerify", Value = "true")]
+		[XmlConfigFile(Category = "BuildConfiguration")]
+		public bool bAutoRTFMVerify { get; set; } = false;
+
 		/// <summary>
 		/// Whether to compile the Chaos physics plugin.
 		/// </summary>
diff --git a/Engine/Source/Programs/UnrealBuildTool/Configuration/UEBuildModuleCPP.cs b/Engine/Source/Programs/UnrealBuildTool/Configuration/UEBuildModuleCPP.cs
--- a/Engine/Source/Programs/UnrealBuildTool/Configuration/UEBuildModuleCPP.cs
+++ b/Engine/Source/Programs/UnrealBuildTool/Configuration/UEBuildModuleCPP.cs
@@ -1569,6 +1569,7 @@ namespace UnrealBuildTool
 			CompileEnvironment.bUseAutoRTFMCompiler = ModuleCompileEnvironment.bUseAutoRTFMCompiler;
 			CompileEnvironment.bDisableAutoRTFMInstrumentation = ModuleCompileEnvironment.bDisableAutoRTFMInstrumentation;
 			CompileEnvironment.bEnableAutoRTFMVerification = ModuleCompileEnvironment.bEnableAutoRTFMVerification;
+			CompileEnvironment.bAutoRTFMVerify = ModuleCompileEnvironment.bAutoRTFMVerify;
 			CompileEnvironment.bValidateFormatStrings = ModuleCompileEnvironment.bValidateFormatStrings;
 			CompileEnvironment.bValidateInternalApi = ModuleCompileEnvironment.bValidateInternalApi;
 		}
@@ -2172,6 +2173,7 @@ namespace UnrealBuildTool
 			{
 				Result.bDisableAutoRTFMInstrumentation = Rules.bDisableAutoRTFMInstrumentation;
 				Result.bEnableAutoRTFMVerification = Target.bUseAutoRTFMVerifier;
+				Result.bAutoRTFMVerify = Target.bAutoRTFMVerify;
 			}
 
 			if (Result.OptimizationLevel != Rules.OptimizationLevel)
diff --git a/Engine/Source/Programs/UnrealBuildTool/Platform/Windows/VCToolChain.cs b/Engine/Source/Programs/UnrealBuildTool/Platform/Windows/VCToolChain.cs
--- a/Engine/Source/Programs/UnrealBuildTool/Platform/Windows/VCToolChain.cs
+++ b/Engine/Source/Programs/UnrealBuildTool/Platform/Windows/VCToolChain.cs
@@ -1281,6 +1281,11 @@ namespace UnrealBuildTool
 						Arguments.Add("-fautortfm-verify");
 					}
 				}
+
+				if (CompileEnvironment.bAutoRTFMVerify)
+				{
+					Arguments.Add("-Xclang -mllvm -Xclang -autortfm-verify");
+				}
 			}
 
 			if (Target.WindowsPlatform.Compiler == WindowsCompiler.Intel)
@@ -1327,6 +1332,11 @@ namespace UnrealBuildTool
 						Arguments.Add("-fautortfm-verify");
 					}
 				}
+
+				if (CompileEnvironment.bAutoRTFMVerify)
+				{
+					Arguments.Add("-Xclang -mllvm -Xclang -autortfm-verify");
+				}
 			}
 		}
 
diff --git a/Engine/Source/Programs/UnrealBuildTool/System/CppCompileEnvironment.cs b/Engine/Source/Programs/UnrealBuildTool/System/CppCompileEnvironment.cs
--- a/Engine/Source/Programs/UnrealBuildTool/System/CppCompileEnvironment.cs
+++ b/Engine/Source/Programs/UnrealBuildTool/System/CppCompileEnvironment.cs
@@ -297,6 +297,11 @@ namespace UnrealBuildTool
 		/// </summary>
 		public bool bEnableAutoRTFMVerification = false;
 
+		/// <summary>
+		/// Whether we should run an extra LLVM verification pass after the AutoRTFM compiler pass.
+		/// </summary>
+		public bool bAutoRTFMVerify = false;
+
 		/// <summary>
 		/// If unity builds are enabled this can be used to override if this specific module will build using Unity.
 		/// This is set using the per module configurations in BuildConfiguration.
@@ -787,6 +792,7 @@ namespace UnrealBuildTool
 			bUseAutoRTFMCompiler = Other.bUseAutoRTFMCompiler;
 			bDisableAutoRTFMInstrumentation = Other.bDisableAutoRTFMInstrumentation;
 			bEnableAutoRTFMVerification = Other.bEnableAutoRTFMVerification;
+			bAutoRTFMVerify = Other.bAutoRTFMVerify;
 			bCodeCoverage = Other.bCodeCoverage;
 			OptimizationLevel = Other.OptimizationLevel;
 			FPSemantics = Other.FPSemantics;
diff --git a/Engine/Source/Programs/UnrealBuildTool/ToolChain/ClangToolChain.cs b/Engine/Source/Programs/UnrealBuildTool/ToolChain/ClangToolChain.cs
--- a/Engine/Source/Programs/UnrealBuildTool/ToolChain/ClangToolChain.cs
+++ b/Engine/Source/Programs/UnrealBuildTool/ToolChain/ClangToolChain.cs
@@ -405,6 +405,11 @@ namespace UnrealBuildTool
 					Arguments.Add("-fautortfm-verify");
 				}
 			}
+
+			if (CompileEnvironment.bAutoRTFMVerify)
+			{
+				Arguments.Add("-Xclang -mllvm -Xclang -autortfm-verify");
+			}
 		}
 
 		protected virtual void GetCStandardCompileArgument(CppCompileEnvironment CompileEnvironment, List<string> Arguments)
