Update to AGP 8.5.2

#android

diff --git a/Engine/Source/Programs/UnrealBuildTool/Platform/Android/UEDeployAndroid.cs b/Engine/Source/Programs/UnrealBuildTool/Platform/Android/UEDeployAndroid.cs
--- a/Engine/Source/Programs/UnrealBuildTool/Platform/Android/UEDeployAndroid.cs
+++ b/Engine/Source/Programs/UnrealBuildTool/Platform/Android/UEDeployAndroid.cs
@@ -23,7 +23,7 @@ namespace UnrealBuildTool
 		private const string BUNDLETOOL_JAR = "bundletool-all-0.13.0.jar";
 
 		// classpath of default android build tools gradle plugin
-		private const string ANDROID_TOOLS_BUILD_GRADLE_VERSION = "com.android.tools.build:gradle:8.5.0";
+		private const string ANDROID_TOOLS_BUILD_GRADLE_VERSION = "com.android.tools.build:gradle:8.5.2";
 
 		// default NDK version if not set
 		private const string DEFAULT_NDK_VERSION = "25.1.8937393";
