UnrealBuildTool: Show correct preferred msvc component string in message when VS is out of date or msvc is missing

#rnx

diff --git a/Engine/Source/Programs/UnrealBuildTool/Platform/Windows/UEBuildWindows.cs b/Engine/Source/Programs/UnrealBuildTool/Platform/Windows/UEBuildWindows.cs
--- a/Engine/Source/Programs/UnrealBuildTool/Platform/Windows/UEBuildWindows.cs
+++ b/Engine/Source/Programs/UnrealBuildTool/Platform/Windows/UEBuildWindows.cs
@@ -1419,25 +1419,34 @@ namespace UnrealBuildTool
 			if (!bSkipWarning)
 			{
 				UEBuildPlatformSDK? SDK = GetSDK(UnrealTargetPlatform.Win64);
-				string ToolSetWarning = Architecture == UnrealArch.X64 ?
-					"MSVC v143 - VS 2022 C++ x64/x86 build tools (Latest)" :
-					"MSVC v143 - VS 2022 C++ ARM64 build tools (Latest)";
-				string MinMsvcVersion = SDK?.GetVersionNumberFromConfig("MinimumVisualCppVersion")?.ToString() ?? "Latest";
-				string MinVSVersion = SDK?.GetVersionNumberFromConfig("MinimumVisualStudio2022Version")?.ToString() ?? "Latest";
+				VersionNumber? MinimumMsvcVersion = SDK?.GetVersionNumberFromConfig("MinimumVisualCppVersion");
+				VersionNumber? MinimumVsVersion = SDK?.GetVersionNumberFromConfig("MinimumVisualStudio2022Version");
+				VersionNumber? PreferredMsvcVersion = SDK?.GetVersionNumberRangeArrayFromConfig("PreferredVisualCppVersions")?.FirstOrDefault()?.Min;
+				string MinimumMsvcVersionStr = MinimumMsvcVersion?.ToString() ?? "Latest";
+				string MinimumVsVersionStr = MinimumVsVersion?.ToString() ?? "Latest";
+				string PreferredMsvcVersionStr = PreferredMsvcVersion?.ToString() ?? "Latest";
+				string PreferredComponentVersionStr = PreferredMsvcVersion != null
+					? $"v{PreferredMsvcVersion.Components[0]}.{PreferredMsvcVersion.Components[1]}-{PreferredMsvcVersion.Components[0] + 3}.{PreferredMsvcVersion.Components[1] - 30}"
+					: "Latest";
+				string PreferredComponentStr = Architecture == UnrealArch.X64 ?
+					$"MSVC v143 - VS 2022 C++ x64/x86 build tools ({PreferredComponentVersionStr})" :
+					$"MSVC v143 - VS 2022 C++ ARM64 build tools ({PreferredComponentVersionStr})";
 				// If we do have a Visual Studio installation, but we're missing just the C++ parts, warn about that.
 				if (TryGetVSInstallDirs(WindowsCompiler.VisualStudio2022, Logger) != null)
 				{
-					Logger.LogWarning("Visual Studio 2022 is installed, but is out of date or missing a valid C++ toolchain (minimum version {MinVersion}). Please update Visual Studio 2022 to {MinVSVersion} or later and verify that the \"{Component}\" component is selected in the Visual Studio 2022 installation options.",
-						MinMsvcVersion,
-						MinVSVersion,
-						ToolSetWarning);
+					Logger.LogWarning("Visual Studio 2022 is installed, but is out of date or missing a valid C++ toolchain (minimum version {MinVersion}, preferred version {PreferredVersion}). Please update Visual Studio 2022 to {MinimumVsVersionStr} or later and verify that the \"{Component}\" component is selected in the Visual Studio 2022 installation options.",
+						MinimumMsvcVersionStr,
+						PreferredMsvcVersionStr,
+						MinimumVsVersionStr,
+						PreferredComponentStr);
 				}
 				else
 				{
-					Logger.LogWarning("No valid Visual C++ toolchain was found (minimum version {MinVersion}). Please download and install Visual Studio 2022 {MinVSVersion} or later and verify that the \"{Component}\" component is selected in the Visual Studio 2022 installation options.",
-						MinMsvcVersion,
-						MinVSVersion,
-						ToolSetWarning);
+					Logger.LogWarning("No valid Visual C++ toolchain was found (minimum version {MinVersion}, preferred version {PreferredVersion}). Please download and install Visual Studio 2022 {MinimumVsVersionStr} or later and verify that the \"{Component}\" component is selected in the Visual Studio 2022 installation options.",
+						MinimumMsvcVersionStr,
+						PreferredMsvcVersionStr,
+						MinimumVsVersionStr,
+						PreferredComponentStr);
 				}
 			}
 
