[UBT]
* Added lock around code in UBAAgentCoordinatorHorde that removes worker to prevent race condition with GetCoreCount which can be called from. Note for reader, this code is really messy. Especially related to shutdown and Dispose. Not sure what can run in parallel at that point

diff --git a/Engine/Source/Programs/UnrealBuildTool/Executors/UnrealBuildAccelerator/UBAAgentCoordinatorHorde.cs b/Engine/Source/Programs/UnrealBuildTool/Executors/UnrealBuildAccelerator/UBAAgentCoordinatorHorde.cs
--- a/Engine/Source/Programs/UnrealBuildTool/Executors/UnrealBuildAccelerator/UBAAgentCoordinatorHorde.cs
+++ b/Engine/Source/Programs/UnrealBuildTool/Executors/UnrealBuildAccelerator/UBAAgentCoordinatorHorde.cs
@@ -115,7 +115,10 @@ namespace UnrealBuildTool
 				for (int idx = _workers.Count - 1; idx >= 0; idx--)
 				{
 					await _workers[idx].BackgroundTask!;
-					_workers.RemoveAt(idx);
+					lock (_workers)
+					{
+						_workers.RemoveAt(idx);
+					}
 				}
 			}
 
@@ -200,7 +203,9 @@ namespace UnrealBuildTool
 				{
 					await worker.BackgroundTask;
 					lock (_workers)
+					{
 						_workers.RemoveAt(idx--);
+					}
 				}
 			}
 		}
@@ -334,7 +339,9 @@ namespace UnrealBuildTool
 				};
 				worker.BackgroundTask = RunWorkerAsync(worker, lease, locator, exeName, workerLogger, hordeConfig, _cancellationTokenSource.Token);
 				lock (_workers)
+				{
 					_workers.Add(worker);
+				}
 				UpdateHordeStatus(null);
 				lease = null; // Will be disposed by RunWorkerAsync
 
@@ -771,8 +778,6 @@ namespace UnrealBuildTool
 					return;
 				}
 
-				//hordeSession.RemoveCompleteWorkersAsync();
-
 				if (queue.IsDone || _cancellationSource!.IsCancellationRequested)
 				{
 					hordeSession.UpdateHordeStatus(" - Requests stopped");
