UnrealBuildTool: Use Microsoft.VisualStudio.Setup.Configuration.Interop nuget package instead of dll

#rnx

diff --git a/Engine/Source/Programs/UnrealBuildTool/UnrealBuildTool.csproj b/Engine/Source/Programs/UnrealBuildTool/UnrealBuildTool.csproj
--- a/Engine/Source/Programs/UnrealBuildTool/UnrealBuildTool.csproj
+++ b/Engine/Source/Programs/UnrealBuildTool/UnrealBuildTool.csproj
@@ -42,10 +42,6 @@
     <Reference Include="Ionic.Zip.Reduced">
       <HintPath>..\..\..\Binaries\DotNET\Ionic.Zip.Reduced.dll</HintPath>
     </Reference>
-    <Reference Include="Microsoft.VisualStudio.Setup.Configuration.Interop">
-      <HintPath>..\..\..\Binaries\ThirdParty\VisualStudio\Microsoft.VisualStudio.Setup.Configuration.Interop.dll</HintPath>
-      <EmbedInteropTypes>True</EmbedInteropTypes>
-    </Reference>
   </ItemGroup>
   <Choose>
     <When Condition="Exists('..\..\..\Restricted\NotForLicensees\Source\Programs\UnrealBuildTool\Executors\FASTBuild.cs')">
@@ -115,6 +111,7 @@
   </ItemGroup>
   <ItemGroup>
     <PackageReference Include="Microsoft.CodeAnalysis.CSharp" Version="4.2.0" />
+    <PackageReference Include="Microsoft.VisualStudio.Setup.Configuration.Interop" Version="3.11.2177" />
     <PackageReference Include="Microsoft.Win32.Registry" Version="5.0.0" />
     <PackageReference Include="OpenTracing" Version="0.12.1" />
     <PackageReference Include="System.CodeDom" Version="8.0.0" />
