SDK: Common changes needed for SDK updates on some platforms [Part 1/4]

diff --git a/Engine/Source/Programs/UnrealBuildTool/ToolChain/ClangWarnings.cs b/Engine/Source/Programs/UnrealBuildTool/ToolChain/ClangWarnings.cs
--- a/Engine/Source/Programs/UnrealBuildTool/ToolChain/ClangWarnings.cs
+++ b/Engine/Source/Programs/UnrealBuildTool/ToolChain/ClangWarnings.cs
@@ -51,6 +51,12 @@ namespace UnrealBuildTool
 					Arguments.Add("-Wno-invalid-unevaluated-string");   // https://clang.llvm.org/docs/DiagnosticsReference.html#winvalid-unevaluated-string			// new warning for clang 17
 				}
 			}
+			if (ClangVersion >= new VersionNumber(18))
+			{
+				Arguments.Add("-Wno-nan-infinity-disabled");            // https://clang.llvm.org/docs/DiagnosticsReference.html#wnan-infinity-disabled					// We use the NAN macro in a few places to initialize floats while we also set -ffast-math, which disables NaN support.
+																																										// It could be easily fixed in the engine, but would create the risk of Win64 building fine and the code failing just on Clang-based platforms.
+																																										// We tend to use NAN just as a bit pattern, and not rely on it in calculations, so it should be reasonably safe to disable.
+			}
 
 			Arguments.Add("-Wno-gnu-string-literal-operator-template"); // https://clang.llvm.org/docs/DiagnosticsReference.html#wgnu-string-literal-operator-template	// We use this feature to allow static FNames.
 			Arguments.Add("-Wno-inconsistent-missing-override");        // https://clang.llvm.org/docs/DiagnosticsReference.html#winconsistent-missing-override			// ?? no reason given
