UnrealBuildTool: Don't ignore Restricted folders in VSCode

diff --git a/Engine/Source/Programs/UnrealBuildTool/ProjectFiles/VisualStudioCode/VSCodeProjectFileGenerator.cs b/Engine/Source/Programs/UnrealBuildTool/ProjectFiles/VisualStudioCode/VSCodeProjectFileGenerator.cs
--- a/Engine/Source/Programs/UnrealBuildTool/ProjectFiles/VisualStudioCode/VSCodeProjectFileGenerator.cs
+++ b/Engine/Source/Programs/UnrealBuildTool/ProjectFiles/VisualStudioCode/VSCodeProjectFileGenerator.cs
@@ -1770,7 +1770,7 @@ namespace UnrealBuildTool
 
 		private void GetExcludePathsCPP(DirectoryReference BaseDir, List<string> PathsToExclude)
 		{
-			string[] DirAllowList = { "Binaries", "Build", "Config", "Plugins", "Source", "Private", "Public", "Internal", "Classes", "Resources" };
+			string[] DirAllowList = { "Binaries", "Build", "Config", "Plugins", "Source", "Private", "Public", "Internal", "Classes", "Resources", "Restricted" };
 			foreach (DirectoryReference SubDir in DirectoryReference.EnumerateDirectories(BaseDir, "*", SearchOption.TopDirectoryOnly))
 			{
 				if (Array.Find(DirAllowList, Dir => Dir == SubDir.GetDirectoryName()) == null)
