UnrealBuildTool: Downgrade uba retry messages to info

#rnx

diff --git a/Engine/Source/Programs/UnrealBuildTool/Executors/UnrealBuildAccelerator/UBAExecutor.cs b/Engine/Source/Programs/UnrealBuildTool/Executors/UnrealBuildAccelerator/UBAExecutor.cs
--- a/Engine/Source/Programs/UnrealBuildTool/Executors/UnrealBuildAccelerator/UBAExecutor.cs
+++ b/Engine/Source/Programs/UnrealBuildTool/Executors/UnrealBuildAccelerator/UBAExecutor.cs
@@ -822,7 +822,7 @@ namespace UnrealBuildTool
 					Interlocked.Add(ref _localProcessedActions, 1);
 					if (!UBAConfig.bStoreObjFilesCompressed && (process.ExitCode != 0 && UBAConfig.bForcedRetry || (process.ExitCode >= 9000 && process.ExitCode < 10000)))
 					{
-						_threadedLogger.LogWarning("{Description} {StatusDescription}: Exited with error code {ExitCode}. This action will retry without UBA", action.CommandDescription, action.StatusDescription, process.ExitCode);
+						_threadedLogger.LogInformation("{Description} {StatusDescription}: Exited with error code {ExitCode}. This action will retry without UBA", action.CommandDescription, action.StatusDescription, process.ExitCode);
 						_forcedRetryActions.AddOrUpdate(action, false, (k, v) => false);
 						queue.RequeueAction(action);
 						return Task.CompletedTask;
@@ -1021,7 +1021,7 @@ namespace UnrealBuildTool
 				return;
 			}
 
-			_threadedLogger.LogWarning("{Description} {StatusDescription} [RemoteExecutor: {ExecutingHost}]: Exited with error code {ExitCode} with no output. This action will retry locally", action.CommandDescription, action.StatusDescription, executingHost, exitCode);
+			_threadedLogger.LogInformation("{Description} {StatusDescription} [RemoteExecutor: {ExecutingHost}]: Exited with error code {ExitCode} with no output. This action will retry locally", action.CommandDescription, action.StatusDescription, executingHost, exitCode);
 			_localRetryActions.AddOrUpdate(action, false, (k, v) => false);
 			queue.RequeueAction(action);
 
@@ -1039,7 +1039,7 @@ namespace UnrealBuildTool
 				return;
 			}
 
-			_threadedLogger.LogWarning("{Description} {StatusDescription} [RemoteExecutor: {ExecutingHost}]: Exited with error code {ExitCode} ({Error}). This action will retry locally", action.CommandDescription, action.StatusDescription, executingHost, exitCode, error);
+			_threadedLogger.LogInformation("{Description} {StatusDescription} [RemoteExecutor: {ExecutingHost}]: Exited with error code {ExitCode} ({Error}). This action will retry locally", action.CommandDescription, action.StatusDescription, executingHost, exitCode, error);
 			_localRetryActions.AddOrUpdate(action, false, (k, v) => false);
 			queue.RequeueAction(action);
 
