[UBA]
* Added code to uba horde agent coordinator to not add some environment variable usage if agent version is older than changelist the environment variables were added in

diff --git a/Engine/Source/Programs/UnrealBuildTool/Executors/UnrealBuildAccelerator/UBAAgentCoordinatorHorde.cs b/Engine/Source/Programs/UnrealBuildTool/Executors/UnrealBuildAccelerator/UBAAgentCoordinatorHorde.cs
--- a/Engine/Source/Programs/UnrealBuildTool/Executors/UnrealBuildAccelerator/UBAAgentCoordinatorHorde.cs
+++ b/Engine/Source/Programs/UnrealBuildTool/Executors/UnrealBuildAccelerator/UBAAgentCoordinatorHorde.cs
@@ -59,6 +59,7 @@ namespace UnrealBuildTool
 		readonly string _crypto;
 		IComputeClient? _client;
 		HordeHttpClient? _hordeHttpClient;
+		uint _agentChangelist = 0;
 
 		public class Worker
 		{
@@ -146,6 +147,15 @@ namespace UnrealBuildTool
 			GetServerInfoResponse serverInfo = await _hordeHttpClient.GetServerInfoAsync(cancellationToken);
 			_logger.LogInformation("Horde server: {ServerVersion}, agent: {AgentVersion}", serverInfo.ServerVersion, serverInfo.AgentVersion);
 
+			if (serverInfo.AgentVersion != null)
+			{
+				int versionLastDash = serverInfo.AgentVersion.LastIndexOf('-');
+				if (versionLastDash != -1 && uint.TryParse(serverInfo.AgentVersion.AsSpan(versionLastDash + 1), out uint changeList))
+				{
+					_agentChangelist = changeList;
+				}
+			}
+
 			DirectoryReference ubaDir = UBAExecutor.UbaBinariesDir;
 			List<string> agentFiles = new();
 			if (OperatingSystem.IsWindows())
@@ -551,10 +561,14 @@ namespace UnrealBuildTool
 
 						arguments.Add("-Dir=%UE_HORDE_SHARED_DIR%\\Uba");
 						arguments.Add("-Eventfile=%UE_HORDE_TERMINATION_SIGNAL_FILE%");
-						arguments.Add("-MaxCpu=%UE_HORDE_CPU_COUNT%");
-						arguments.Add("-MulCpu=%UE_HORDE_CPU_MULTIPLIER%");
 						arguments.Add("-MaxIdle=15");
 
+						if (_agentChangelist > 36297540)
+						{
+							arguments.Add("-MaxCpu=%UE_HORDE_CPU_COUNT%");
+							arguments.Add("-MulCpu=%UE_HORDE_CPU_MULTIPLIER%");
+						}
+
 						if (_owner.UBAConfig.bLogEnabled)
 						{
 							arguments.Add("-Log");
