UnrealBuildTool: Undo -program gpf change

#rnx

diff --git a/Engine/Source/Programs/UnrealBuildTool/ProjectFiles/ProjectFileGenerator.cs b/Engine/Source/Programs/UnrealBuildTool/ProjectFiles/ProjectFileGenerator.cs
--- a/Engine/Source/Programs/UnrealBuildTool/ProjectFiles/ProjectFileGenerator.cs
+++ b/Engine/Source/Programs/UnrealBuildTool/ProjectFiles/ProjectFileGenerator.cs
@@ -1551,12 +1551,9 @@ namespace UnrealBuildTool
 								bGatherThirdPartySource = true;
 								break;
 
-							case "-PROGRAMS":
-								bIncludeEnginePrograms = true;
-								break;
-
 							case "-NOPROGRAMS":
 								bIncludeEnginePrograms = false;
+								bIncludeDotNetPrograms = false;
 								break;
 
 							case "-GAME":
