Undo changelist 36885478 where we check if horde is newer than certain changelists. This was a workaround for one of the active farms which ran on an old horde. Now farm is updated so we can revert this (relying on changelists will not work for external teams so this was kind of a hack)

diff --git a/Engine/Source/Programs/UnrealBuildTool/Executors/UnrealBuildAccelerator/UBAAgentCoordinatorHorde.cs b/Engine/Source/Programs/UnrealBuildTool/Executors/UnrealBuildAccelerator/UBAAgentCoordinatorHorde.cs
--- a/Engine/Source/Programs/UnrealBuildTool/Executors/UnrealBuildAccelerator/UBAAgentCoordinatorHorde.cs
+++ b/Engine/Source/Programs/UnrealBuildTool/Executors/UnrealBuildAccelerator/UBAAgentCoordinatorHorde.cs
@@ -59,7 +59,6 @@ namespace UnrealBuildTool
 		readonly string _crypto;
 		IComputeClient? _client;
 		HordeHttpClient? _hordeHttpClient;
-		uint _agentChangelist = 0;
 
 		public class Worker
 		{
@@ -147,15 +146,6 @@ namespace UnrealBuildTool
 			GetServerInfoResponse serverInfo = await _hordeHttpClient.GetServerInfoAsync(cancellationToken);
 			_logger.LogInformation("Horde server: {ServerVersion}, agent: {AgentVersion}", serverInfo.ServerVersion, serverInfo.AgentVersion);
 
-			if (serverInfo.AgentVersion != null)
-			{
-				int versionLastDash = serverInfo.AgentVersion.LastIndexOf('-');
-				if (versionLastDash != -1 && uint.TryParse(serverInfo.AgentVersion.AsSpan(versionLastDash + 1), out uint changeList))
-				{
-					_agentChangelist = changeList;
-				}
-			}
-
 			DirectoryReference ubaDir = UBAExecutor.UbaBinariesDir;
 			List<string> agentFiles = new();
 			if (OperatingSystem.IsWindows())
@@ -561,14 +551,10 @@ namespace UnrealBuildTool
 
 						arguments.Add("-Dir=%UE_HORDE_SHARED_DIR%\\Uba");
 						arguments.Add("-Eventfile=%UE_HORDE_TERMINATION_SIGNAL_FILE%");
+						arguments.Add("-MaxCpu=%UE_HORDE_CPU_COUNT%");
+						arguments.Add("-MulCpu=%UE_HORDE_CPU_MULTIPLIER%");
 						arguments.Add("-MaxIdle=15");
 
-						if (_agentChangelist > 36297540)
-						{
-							arguments.Add("-MaxCpu=%UE_HORDE_CPU_COUNT%");
-							arguments.Add("-MulCpu=%UE_HORDE_CPU_MULTIPLIER%");
-						}
-
 						if (_owner.UBAConfig.bLogEnabled)
 						{
 							arguments.Add("-Log");
