FASTBUILD: Add the "Experimental" argument to the ignored list of tokens. The addition of the "sarif" file in 5.5 is confusing the old code to use it as the CompilerInputFiles when generating the BFF file.

diff --git a/Engine/Source/Programs/UnrealBuildTool/Executors/Experimental/FASTBuild.cs b/Engine/Source/Programs/UnrealBuildTool/Executors/Experimental/FASTBuild.cs
--- a/Engine/Source/Programs/UnrealBuildTool/Executors/Experimental/FASTBuild.cs
+++ b/Engine/Source/Programs/UnrealBuildTool/Executors/Experimental/FASTBuild.cs
@@ -762,7 +762,8 @@ namespace UnrealBuildTool
 					(Token == "-weak_framework") ||
 					(Token == "-framework") ||
 					(Token == "/sourceDependencies") ||
-					(Token == "/sourceDependencies:directives"))
+					(Token == "/sourceDependencies:directives") ||
+					(Token.Contains("/experimental")))
 				{
 					++i;
 				}
