[UBA]
* Fixed so name of agents are Worker[0...] in logs and trace file if agent has EC2 property. This makes it much easier to search the log file and figure out what errors that came from where, and also find the lease for the agent.

diff --git a/Engine/Source/Programs/UnrealBuildTool/Executors/UnrealBuildAccelerator/UBAAgentCoordinatorHorde.cs b/Engine/Source/Programs/UnrealBuildTool/Executors/UnrealBuildAccelerator/UBAAgentCoordinatorHorde.cs
--- a/Engine/Source/Programs/UnrealBuildTool/Executors/UnrealBuildAccelerator/UBAAgentCoordinatorHorde.cs
+++ b/Engine/Source/Programs/UnrealBuildTool/Executors/UnrealBuildAccelerator/UBAAgentCoordinatorHorde.cs
@@ -247,7 +247,8 @@ namespace UnrealBuildTool
 				"RAM",
 				"DiskFreeSpace",
 				"PhysicalCores",
-				"LogicalCores"
+				"LogicalCores",
+				"EC2",
 			};
 
 		public async Task<bool> AddWorkerAsync(Requirements requirements, UnrealBuildAcceleratorHordeConfig hordeConfig, CancellationToken cancellationToken, int activeCores)
@@ -257,7 +258,9 @@ namespace UnrealBuildTool
 				throw new InvalidOperationException($"Session not initialized. Call {nameof(InitAsync)} first");
 			}
 
-			PrefixLogger workerLogger = new($"[Worker{_workerId}]", _logger);
+			int workerId = _workerId++;
+			
+			PrefixLogger workerLogger = new($"[Worker{workerId}]", _logger);
 
 			const string UbaPortName = "UbaPort";
 			const string UbaProxyPortName = "UbaProxyPort";
@@ -266,7 +269,7 @@ namespace UnrealBuildTool
 
 			// Request ID that is unique per attempt to acquire the same compute lease/worker
 			// Primarily for tracking worker demand on Horde server as UBAExecutor will repeatedly try adding a new worker
-			string requestId = $"{_id}-worker-{_workerId}";
+			string requestId = $"{_id}-worker-{workerId}";
 			IComputeLease? lease = null;
 			try
 			{
@@ -289,10 +292,9 @@ namespace UnrealBuildTool
 					return false;
 				}
 
-				_workerId++;
-
 				workerLogger.LogDebug("Agent properties:");
 
+				string agentName = String.Empty;
 				int numLogicalCores = 24; // Assume 24 if something goes wrong here and property is not found
 				string computeIp = String.Empty;
 				foreach (string property in lease.Properties)
@@ -311,6 +313,10 @@ namespace UnrealBuildTool
 						{
 							computeIp = property[(equalsIdx + 1)..];
 						}
+						else if (propertyName == "EC2")
+						{
+							agentName = $"Worker{workerId}";
+						}
 					}
 				}
 
@@ -337,7 +343,7 @@ namespace UnrealBuildTool
 					Port = ubaPort,
 					ProxyPort = ubaProxyPort,
 				};
-				worker.BackgroundTask = RunWorkerAsync(worker, lease, locator, exeName, workerLogger, hordeConfig, _cancellationTokenSource.Token);
+				worker.BackgroundTask = RunWorkerAsync(worker, lease, locator, exeName, workerLogger, hordeConfig, _cancellationTokenSource.Token, agentName);
 				lock (_workers)
 				{
 					_workers.Add(worker);
@@ -476,7 +482,7 @@ namespace UnrealBuildTool
 			return result.AccessToken;
 		}
 
-		async Task RunWorkerAsync(Worker self, IComputeLease lease, BlobLocator tool, string executable, ILogger logger, UnrealBuildAcceleratorHordeConfig hordeConfig, CancellationToken cancellationToken)
+		async Task RunWorkerAsync(Worker self, IComputeLease lease, BlobLocator tool, string executable, ILogger logger, UnrealBuildAcceleratorHordeConfig hordeConfig, CancellationToken cancellationToken, string agentName)
 		{
 			logger.LogDebug("Running worker task..");
 			try
@@ -513,6 +519,11 @@ namespace UnrealBuildTool
 						bool useListen = !String.IsNullOrEmpty(hordeConfig.HordeHost);
 						List<string> arguments = new();
 
+						if (!String.IsNullOrEmpty(agentName))
+						{
+							arguments.Add($"-name={agentName}");
+						}
+
 						if (useListen)
 						{
 							arguments.Add($"-Host={hordeHost}:{_owner.UBAConfig.Port}");
