[UBA]
* Added so a fail fast exception is retried

diff --git a/Engine/Source/Programs/UnrealBuildTool/Executors/UnrealBuildAccelerator/UBAExecutor.cs b/Engine/Source/Programs/UnrealBuildTool/Executors/UnrealBuildAccelerator/UBAExecutor.cs
--- a/Engine/Source/Programs/UnrealBuildTool/Executors/UnrealBuildAccelerator/UBAExecutor.cs
+++ b/Engine/Source/Programs/UnrealBuildTool/Executors/UnrealBuildAccelerator/UBAExecutor.cs
@@ -925,6 +925,11 @@ namespace UnrealBuildTool
 						RemoteActionFailedCrash(queue, action, e.ExitCode, e.ExecutingHost ?? "Unknown", "Stack buffer overflow");
 						return;
 					}
+					else if ((uint)e.ExitCode == 0xC0000602)
+					{
+						RemoteActionFailedCrash(queue, action, e.ExitCode, e.ExecutingHost ?? "Unknown", "Fail Fast Exception");
+						return;
+					}
 					else if (e.ExitCode != 0 && e.LogLines.Any(x => x.Contains(" C1001: ", StringComparison.Ordinal)))
 					{
 						RemoteActionFailedCrash(queue, action, e.ExitCode, e.ExecutingHost ?? "Unknown", "C1001");
