Merging

//UE5/Main/Engine/Source/Programs/UnrealBuildTool/ProjectFiles/ProjectFileGenerator.cs
//UE5/Main/Engine/Source/Programs/UnrealGameSync/...

to //UE5/Release-5.5/Engine/Source/Programs/...

diff --git a/Engine/Source/Programs/UnrealBuildTool/ProjectFiles/ProjectFileGenerator.cs b/Engine/Source/Programs/UnrealBuildTool/ProjectFiles/ProjectFileGenerator.cs
--- a/Engine/Source/Programs/UnrealBuildTool/ProjectFiles/ProjectFileGenerator.cs
+++ b/Engine/Source/Programs/UnrealBuildTool/ProjectFiles/ProjectFileGenerator.cs
@@ -7,8 +7,6 @@ using System.IO;
 using System.Linq;
 using System.Text;
 using System.Threading;
-using System.Threading.Tasks;
-using System.Xml;
 using System.Xml.Linq;
 using EpicGames.Core;
 using Microsoft.Extensions.Logging;
@@ -566,7 +564,8 @@ namespace UnrealBuildTool
 
 					if (bCreatePropsFile)
 					{
-						CreateProjectPropsFileAsync(new($"{ProjectFile}.props")).Wait();
+						FileReference PropsFile = new FileReference(ProjectFile.FullName + ".props");
+						CreateProjectPropsFile(PropsFile);
 					}
 				}
 			}
@@ -610,36 +609,18 @@ namespace UnrealBuildTool
 		/// <summary>
 		/// Creates a .props file next to each project which specifies the path to the engine directory
 		/// </summary>
-		/// <param name="propsFile">The properties file path</param>
-		static async Task<bool> CreateProjectPropsFileAsync(FileReference propsFile) => await FileReference.WriteAllBytesIfDifferentAsync(propsFile, await s_projectPropsFileContents.Value);
-
-		static readonly Lazy<Task<byte[]>> s_projectPropsFileContents = new(async () =>
+		/// <param name="PropsFile">The properties file path</param>
+		static void CreateProjectPropsFile(FileReference PropsFile)
 		{
-			using MemoryStream stream = new MemoryStream();
-			using (XmlWriter writer = XmlWriter.Create(stream, new XmlWriterSettings() { Async = true, Encoding = Encoding.UTF8, Indent = true, IndentChars = "  " }))
-			{
-				await writer.WriteStartDocumentAsync();
-				{
-					await writer.WriteStartElementAsync(null, "Project", "http://schemas.microsoft.com/developer/msbuild/2003");
-					await writer.WriteAttributeStringAsync(null, "ToolsVersion", null, "Current");
-					{
-						await writer.WriteStartElementAsync(null, "PropertyGroup", null);
-						{
-							await writer.WriteStartElementAsync(null, "EngineDir", null);
-							await writer.WriteAttributeStringAsync(null, "Condition", null, "'$(EngineDir)' == ''");
-							{
-								await writer.WriteStringAsync(Unreal.EngineDirectory.FullName);
-							}
-							await writer.WriteEndElementAsync();
-						}
-						await writer.WriteEndElementAsync();
-					}
-					await writer.WriteEndElementAsync();
-				}
-				await writer.WriteEndDocumentAsync();
-			}
-			return stream.ToArray();
-		});
+			using StringWriter Writer = new StringWriter();
+			Writer.WriteLine("<?xml version=\"1.0\" encoding=\"utf-8\"?>");
+			Writer.WriteLine("<Project ToolsVersion=\"Current\" xmlns=\"http://schemas.microsoft.com/developer/msbuild/2003\">");
+			Writer.WriteLine("\t<PropertyGroup>");
+			Writer.WriteLine("\t\t<EngineDir Condition=\"'$(EngineDir)' == ''\">{0}</EngineDir>", Unreal.EngineDirectory);
+			Writer.WriteLine("\t</PropertyGroup>");
+			Writer.WriteLine("</Project>");
+			FileReference.WriteAllTextIfDifferent(PropsFile, Writer.ToString());
+		}
 
 		/// <summary>
 		/// Finds all csproj within Engine/Source/Programs, and add them if their .uprogram file exists.
@@ -834,27 +815,6 @@ namespace UnrealBuildTool
 			return ProjectFiles;
 		}
 
-		/// <summary>
-		/// Finds plugins associated with a set of provided game projects.
-		/// </summary>
-		/// <param name="Logger"></param>
-		/// <param name="GameProjects">The game projects to get plugins for.</param>
-		/// <returns>List of plugin files</returns>
-		public List<FileReference> FindGameProjectPlugins(ILogger Logger, List<FileReference> GameProjects)
-		{
-			List<FileReference> PluginFiles = new List<FileReference>();
-			foreach (FileReference UnrealProjectFile in GameProjects)
-			{
-				DirectoryReference GamePluginFolder = DirectoryReference.Combine(UnrealProjectFile.Directory, "Plugins");
-
-				if (DirectoryReference.Exists(GamePluginFolder))
-				{
-					PluginFiles.AddRange(PluginsBase.EnumeratePlugins(GamePluginFolder));
-				}
-			}
-			return PluginFiles;
-		}
-
 		/// <summary>
 		/// Gets the user's preferred IDE from their editor settings
 		/// </summary>
@@ -1297,10 +1257,8 @@ namespace UnrealBuildTool
 						ProgramsFolder.ChildProjects.Add(AutomationToolProject);
 					}
 
-					List<FileReference> AllPluginProjects = FindGameProjectPlugins(Logger, AllGameProjects);
-
 					// Add automation.csproj files to the primary project
-					AddRulesModules(Rules.RulesFileType.AutomationModule, "Automation", AutomationProjectFiles, AllGameProjects.Union(AllPluginProjects).ToList(), RootFolder, ProgramsFolder, Logger);
+					AddRulesModules(Rules.RulesFileType.AutomationModule, "Automation", AutomationProjectFiles, AllGameProjects, RootFolder, ProgramsFolder, Logger);
 
 					// Add ubtplugin.csproj files to the primary project
 					AddRulesModules(Rules.RulesFileType.UbtPlugin, "UnrealBuildTool.Plugins", UbtPluginProjectFiles, AllGameProjects, RootFolder, ProgramsFolder, Logger);
@@ -2081,6 +2039,9 @@ namespace UnrealBuildTool
 				VCSharpProjectFile UnrealBuildToolProject = new VCSharpProjectFile(ProjectFileName, Logger);
 				UnrealBuildToolProject.ShouldBuildForAllSolutionTargets = true;
 
+				// Store it off as we need it when generating target projects.
+				UBTProject = UnrealBuildToolProject;
+
 				// Add the project
 				AddExistingProjectFile(UnrealBuildToolProject, bNeedsAllPlatformAndConfigurations: true, bForceDevelopmentConfiguration: true);
 
