Update more nuget packages

#rnx

diff --git a/Engine/Source/Programs/AutomationTool/AutomationUtils/AutomationUtils.Automation.csproj b/Engine/Source/Programs/AutomationTool/AutomationUtils/AutomationUtils.Automation.csproj
--- a/Engine/Source/Programs/AutomationTool/AutomationUtils/AutomationUtils.Automation.csproj
+++ b/Engine/Source/Programs/AutomationTool/AutomationUtils/AutomationUtils.Automation.csproj
@@ -68,7 +68,7 @@
 
     <!--Required by Ionic -->
     <PackageReference Include="System.Security.Permissions" Version="4.7.0" />
-    <PackageReference Include="System.Text.Encoding.CodePages" Version="6.0.0" />
+    <PackageReference Include="System.Text.Encoding.CodePages" Version="8.0.0" />
 
     <PackageReference Include="Google.Apis.Drive.v3" Version="1.33.1.1225" PrivateAssets="all" />
     <PackageReference Include="Google.Apis.Sheets.v4" Version="1.33.1.1229" PrivateAssets="all" />
diff --git a/Engine/Source/Programs/Horde/Drivers/JobDriver/JobDriver.csproj b/Engine/Source/Programs/Horde/Drivers/JobDriver/JobDriver.csproj
--- a/Engine/Source/Programs/Horde/Drivers/JobDriver/JobDriver.csproj
+++ b/Engine/Source/Programs/Horde/Drivers/JobDriver/JobDriver.csproj
@@ -69,8 +69,8 @@
     <PackageReference Include="Microsoft.Extensions.Hosting" Version="8.0.0" />
     <PackageReference Include="Microsoft.Extensions.Hosting.Abstractions" Version="8.0.0" />
     <PackageReference Include="Microsoft.Extensions.Hosting.WindowsServices" Version="8.0.0" />
-    <PackageReference Include="Microsoft.Extensions.Http" Version="8.0.0" />
-    <PackageReference Include="Microsoft.Extensions.Http.Polly" Version="8.0.0" />
+    <PackageReference Include="Microsoft.Extensions.Http" Version="8.0.1" />
+    <PackageReference Include="Microsoft.Extensions.Http.Polly" Version="8.0.10" />
     <PackageReference Include="Microsoft.Extensions.Logging.Abstractions" Version="8.0.2" />
     <PackageReference Include="Microsoft.Extensions.Options.DataAnnotations" Version="8.0.0" />
     <PackageReference Include="Microsoft.Management.Infrastructure" Version="3.0.0" />
@@ -90,8 +90,8 @@
     <PackageReference Include="Serilog.Settings.Configuration" Version="8.0.0" />
     <PackageReference Include="Serilog.Sinks.Console" Version="5.0.1" />
     <PackageReference Include="Serilog.Sinks.File" Version="5.0.0" />
-    <PackageReference Include="System.Diagnostics.PerformanceCounter" Version="8.0.0" />
-    <PackageReference Include="System.ServiceProcess.ServiceController" Version="8.0.0" />
+    <PackageReference Include="System.Diagnostics.PerformanceCounter" Version="8.0.1" />
+    <PackageReference Include="System.ServiceProcess.ServiceController" Version="8.0.1" />
     <PackageReference Include="System.Text.Json" Version="8.0.5" />
   </ItemGroup>
   
diff --git a/Engine/Source/Programs/Horde/HordeAgent/HordeAgent.csproj b/Engine/Source/Programs/Horde/HordeAgent/HordeAgent.csproj
--- a/Engine/Source/Programs/Horde/HordeAgent/HordeAgent.csproj
+++ b/Engine/Source/Programs/Horde/HordeAgent/HordeAgent.csproj
@@ -51,8 +51,8 @@
     <PackageReference Include="Microsoft.Extensions.Hosting" Version="8.0.0" />
     <PackageReference Include="Microsoft.Extensions.Hosting.Abstractions" Version="8.0.0" />
     <PackageReference Include="Microsoft.Extensions.Hosting.WindowsServices" Version="8.0.0" />
-    <PackageReference Include="Microsoft.Extensions.Http" Version="8.0.0" />
-    <PackageReference Include="Microsoft.Extensions.Http.Polly" Version="8.0.0" />
+    <PackageReference Include="Microsoft.Extensions.Http" Version="8.0.1" />
+    <PackageReference Include="Microsoft.Extensions.Http.Polly" Version="8.0.10" />
     <PackageReference Include="Microsoft.Extensions.Logging.Abstractions" Version="8.0.2" />
     <PackageReference Include="Microsoft.Extensions.Options.DataAnnotations" Version="8.0.0" />
     <PackageReference Include="Microsoft.Management.Infrastructure" Version="3.0.0" />
@@ -72,8 +72,8 @@
     <PackageReference Include="Serilog.Settings.Configuration" Version="8.0.0" />
     <PackageReference Include="Serilog.Sinks.Console" Version="5.0.1" />
     <PackageReference Include="Serilog.Sinks.File" Version="5.0.0" />
-    <PackageReference Include="System.Diagnostics.PerformanceCounter" Version="8.0.0" />
-    <PackageReference Include="System.ServiceProcess.ServiceController" Version="8.0.0" />
+    <PackageReference Include="System.Diagnostics.PerformanceCounter" Version="8.0.1" />
+    <PackageReference Include="System.ServiceProcess.ServiceController" Version="8.0.1" />
     <PackageReference Include="System.Text.Json" Version="8.0.5" />
   </ItemGroup>
 
diff --git a/Engine/Source/Programs/Horde/HordeServer/HordeServer.csproj b/Engine/Source/Programs/Horde/HordeServer/HordeServer.csproj
--- a/Engine/Source/Programs/Horde/HordeServer/HordeServer.csproj
+++ b/Engine/Source/Programs/Horde/HordeServer/HordeServer.csproj
@@ -51,8 +51,8 @@
     <PackageReference Include="Microsoft.AspNetCore.SpaServices.Extensions" Version="8.0.0" />
     <PackageReference Include="Microsoft.Extensions.Caching.Memory" Version="8.0.0" />
     <PackageReference Include="Microsoft.Extensions.Hosting.WindowsServices" Version="8.0.0" />
-    <PackageReference Include="Microsoft.Extensions.Http" Version="8.0.0" />
-    <PackageReference Include="Microsoft.Extensions.Http.Polly" Version="8.0.0" />
+    <PackageReference Include="Microsoft.Extensions.Http" Version="8.0.1" />
+    <PackageReference Include="Microsoft.Extensions.Http.Polly" Version="8.0.10" />
     <PackageReference Include="Microsoft.Extensions.Logging" Version="8.0.1" />
     <PackageReference Include="Microsoft.VisualStudio.Threading.Analyzers" Version="17.8.14">
       <PrivateAssets>all</PrivateAssets>
diff --git a/Engine/Source/Programs/Shared/EpicGames.Build/EpicGames.Build.csproj b/Engine/Source/Programs/Shared/EpicGames.Build/EpicGames.Build.csproj
--- a/Engine/Source/Programs/Shared/EpicGames.Build/EpicGames.Build.csproj
+++ b/Engine/Source/Programs/Shared/EpicGames.Build/EpicGames.Build.csproj
@@ -45,7 +45,7 @@
   </ItemGroup>
   <ItemGroup>
     <PackageReference Include="Microsoft.CSharp" Version="4.7.0" />
-    <PackageReference Include="Microsoft.Extensions.FileSystemGlobbing" Version="6.0.0" />
+    <PackageReference Include="Microsoft.Extensions.FileSystemGlobbing" Version="8.0.0" />
     <PackageReference Include="System.Security.Permissions" Version="4.7.0" />
   </ItemGroup>
 </Project>
diff --git a/Engine/Source/Programs/Shared/EpicGames.Core/EpicGames.Core.csproj b/Engine/Source/Programs/Shared/EpicGames.Core/EpicGames.Core.csproj
--- a/Engine/Source/Programs/Shared/EpicGames.Core/EpicGames.Core.csproj
+++ b/Engine/Source/Programs/Shared/EpicGames.Core/EpicGames.Core.csproj
@@ -54,9 +54,9 @@
   <ItemGroup>
     <PackageReference Include="JetBrains.Annotations" Version="2022.1.0" />
     <PackageReference Include="Microsoft.Extensions.Configuration" Version="8.0.0" />
-    <PackageReference Include="Microsoft.Extensions.Http.Polly" Version="6.0.26" />
+    <PackageReference Include="Microsoft.Extensions.Http.Polly" Version="8.0.10" />
     <PackageReference Include="Microsoft.Extensions.Logging.Console" Version="8.0.1" />
-    <PackageReference Include="Microsoft.Extensions.ObjectPool" Version="6.0.8" />
+    <PackageReference Include="Microsoft.Extensions.ObjectPool" Version="8.0.10" />
     <PackageReference Include="OpenTracing" Version="0.12.1" />
     <PackageReference Include="System.Memory" Version="4.5.5" />
     <PackageReference Include="System.Text.Json" Version="8.0.5" />
diff --git a/Engine/Source/Programs/UnrealBuildTool/Executors/SNDBS.cs b/Engine/Source/Programs/UnrealBuildTool/Executors/SNDBS.cs
--- a/Engine/Source/Programs/UnrealBuildTool/Executors/SNDBS.cs
+++ b/Engine/Source/Programs/UnrealBuildTool/Executors/SNDBS.cs
@@ -198,7 +198,7 @@ namespace UnrealBuildTool
 			}
 
 			// Check the service is running
-			if (!ServiceController.GetServices().Any(s => s.ServiceName.StartsWith("SNDBS") && s.Status == ServiceControllerStatus.Running))
+			if (OperatingSystem.IsWindows() && !ServiceController.GetServices().Any(s => OperatingSystem.IsWindows() && s.ServiceName.StartsWith("SNDBS") && s.Status == ServiceControllerStatus.Running))
 			{
 				return false;
 			}
diff --git a/Engine/Source/Programs/UnrealBuildTool/Executors/XGE.cs b/Engine/Source/Programs/UnrealBuildTool/Executors/XGE.cs
--- a/Engine/Source/Programs/UnrealBuildTool/Executors/XGE.cs
+++ b/Engine/Source/Programs/UnrealBuildTool/Executors/XGE.cs
@@ -318,7 +318,7 @@ namespace UnrealBuildTool
 			}
 
 			// on windows check the service is actually running
-			if (BuildHostPlatform.Current.Platform == UnrealTargetPlatform.Win64)
+			if (OperatingSystem.IsWindows())
 			{
 				try
 				{
diff --git a/Engine/Source/Programs/UnrealBuildTool/UnrealBuildTool.csproj b/Engine/Source/Programs/UnrealBuildTool/UnrealBuildTool.csproj
--- a/Engine/Source/Programs/UnrealBuildTool/UnrealBuildTool.csproj
+++ b/Engine/Source/Programs/UnrealBuildTool/UnrealBuildTool.csproj
@@ -110,21 +110,21 @@
     <ProjectReference Include="..\Shared\EpicGames.UHT\EpicGames.UHT.csproj" />
   </ItemGroup>
   <ItemGroup>
-    <PackageReference Include="Microsoft.CodeAnalysis.CSharp" Version="4.2.0" />
+    <PackageReference Include="Microsoft.CodeAnalysis.CSharp" Version="4.11.0" />
     <PackageReference Include="Microsoft.VisualStudio.Setup.Configuration.Interop" Version="3.11.2177" />
     <PackageReference Include="Microsoft.Win32.Registry" Version="5.0.0" />
     <PackageReference Include="OpenTracing" Version="0.12.1" />
     <PackageReference Include="System.CodeDom" Version="8.0.0" />
-    <PackageReference Include="System.Management" Version="4.7.0" />
+    <PackageReference Include="System.Management" Version="8.0.0" />
     <PackageReference Include="System.Security.Cryptography.Csp" Version="4.3.0" />
-    <PackageReference Include="System.ServiceProcess.ServiceController" Version="4.7.0" />
+    <PackageReference Include="System.ServiceProcess.ServiceController" Version="8.0.1" />
 
     <!-- Required by Ionic -->
     <PackageReference Include="System.Security.Permissions" Version="4.7.0" />
-    <PackageReference Include="System.Text.Encoding.CodePages" Version="6.0.0" />
+    <PackageReference Include="System.Text.Encoding.CodePages" Version="8.0.0" />
 
     <!-- Required for Hololens -->
-    <PackageReference Include="System.Reflection.MetadataLoadContext" Version="4.7.2" /> 
+    <PackageReference Include="System.Reflection.MetadataLoadContext" Version="8.0.1" /> 
     
     <!-- Required for EpicGames.MsBuild -->
     <PackageReference Include="Microsoft.Build" Version="17.1.0" ExcludeAssets="runtime" PrivateAssets="all" />
