UnrealBuildTool: Add all assemblies loaded by UBT to the reference list for any Rules being compiled

#rnx

diff --git a/Engine/Source/Programs/UnrealBuildTool/System/DynamicCompilation.cs b/Engine/Source/Programs/UnrealBuildTool/System/DynamicCompilation.cs
--- a/Engine/Source/Programs/UnrealBuildTool/System/DynamicCompilation.cs
+++ b/Engine/Source/Programs/UnrealBuildTool/System/DynamicCompilation.cs
@@ -2,6 +2,7 @@
 
 using System;
 using System.Collections.Generic;
+using System.Diagnostics;
 using System.IO;
 using System.Linq;
 using System.Reflection;
@@ -157,7 +158,7 @@ namespace UnrealBuildTool
 			}
 		}
 
-		private static Assembly? CompileAssembly(FileReference OutputAssemblyPath, IEnumerable<FileReference> SourceFileNames, ILogger Logger, IEnumerable<string>? ReferencedAssembies, IEnumerable<string>? PreprocessorDefines = null, bool TreatWarningsAsErrors = false)
+		private static Assembly? CompileAssembly(FileReference OutputAssemblyPath, IEnumerable<FileReference> SourceFileNames, ILogger Logger, IEnumerable<string>? ReferencedAssemblies, IEnumerable<string>? PreprocessorDefines = null, bool TreatWarningsAsErrors = false)
 		{
 			CSharpParseOptions ParseOptions = new CSharpParseOptions(
 				languageVersion: LanguageVersion.Latest,
@@ -197,53 +198,61 @@ namespace UnrealBuildTool
 				}
 			}
 
-			List<MetadataReference> MetadataReferences = new List<MetadataReference>();
-			if (ReferencedAssembies != null)
+			// Should be loaded from the current app domain assemblies
+			List<Type> typesForAssmebly = [
+				typeof(object),
+				typeof(FileReference),
+				typeof(IoHash),
+				typeof(UEBuildPlatformSDK),
+				typeof(UnrealBuildTool),
+			];
+
+			// Most of these will be loaded from the current app domain assemblies
+			List<string> namesForAssembly = [
+				"Microsoft.CodeAnalysis.CSharp",
+				"Microsoft.Extensions.Logging.Abstractions",
+				"netstandard",
+				"System.CodeDom",
+				"System.Collections",
+				"System.Console",
+				"System.IO",
+				"System.IO.FileSystem",
+				"System.Linq",
+				"System.Private.Xml",
+				"System.Private.Xml.Linq",
+				"System.Runtime",
+				"System.Runtime.Extensions",
+				"System.Text.RegularExpressions",
+
+				// process start dependencies
+				"System.ComponentModel.Primitives",
+				"System.Diagnostics.Process",
+
+				// registry access
+				"Microsoft.Win32.Registry",
+
+				// RNGCryptoServiceProvider, used to generate random hex bytes
+				"System.Security.Cryptography",
+				"System.Security.Cryptography.Algorithms",
+				"System.Security.Cryptography.Csp",
+			];
+
+			HashSet<string> assemblyLocations = new();
+			assemblyLocations.UnionWith(AppDomain.CurrentDomain.GetAssemblies().Where(x => !String.IsNullOrEmpty(x.Location)).Select(x => x.Location));
+			assemblyLocations.UnionWith(ReferencedAssemblies ?? []);
+			assemblyLocations.UnionWith(typesForAssmebly.Where(x => !String.IsNullOrEmpty(x.Assembly.Location)).Select(x => x.Assembly.Location));
+			foreach (string assembly in namesForAssembly)
 			{
-				foreach (string Reference in ReferencedAssembies)
+				try
+				{
+					assemblyLocations.Add(Assembly.Load(assembly).Location);
+				}
+				catch (Exception)
 				{
-					MetadataReferences.Add(MetadataReference.CreateFromFile(Reference));
 				}
 			}
 
-			MetadataReferences.Add(MetadataReference.CreateFromFile(typeof(object).Assembly.Location));
-			MetadataReferences.Add(MetadataReference.CreateFromFile(Assembly.Load("System.Runtime").Location));
-			MetadataReferences.Add(MetadataReference.CreateFromFile(Assembly.Load("System.CodeDom").Location));
-			MetadataReferences.Add(MetadataReference.CreateFromFile(Assembly.Load("System.Collections").Location));
-			MetadataReferences.Add(MetadataReference.CreateFromFile(Assembly.Load("System.IO").Location));
-			MetadataReferences.Add(MetadataReference.CreateFromFile(Assembly.Load("System.IO.FileSystem").Location));
-			MetadataReferences.Add(MetadataReference.CreateFromFile(Assembly.Load("System.Linq").Location));
-			MetadataReferences.Add(MetadataReference.CreateFromFile(Assembly.Load("System.Private.Xml").Location));
-			MetadataReferences.Add(MetadataReference.CreateFromFile(Assembly.Load("System.Private.Xml.Linq").Location));
-			MetadataReferences.Add(MetadataReference.CreateFromFile(Assembly.Load("System.Text.RegularExpressions").Location));
-			MetadataReferences.Add(MetadataReference.CreateFromFile(Assembly.Load("Microsoft.CodeAnalysis.CSharp").Location));
-			MetadataReferences.Add(MetadataReference.CreateFromFile(Assembly.Load("System.Console").Location));
-			MetadataReferences.Add(MetadataReference.CreateFromFile(Assembly.Load("System.Runtime.Extensions").Location));
-			MetadataReferences.Add(MetadataReference.CreateFromFile(Assembly.Load("Microsoft.Extensions.Logging.Abstractions").Location));
-			MetadataReferences.Add(MetadataReference.CreateFromFile(Assembly.Load("netstandard").Location));
-
-			// process start dependencies
-			MetadataReferences.Add(MetadataReference.CreateFromFile(Assembly.Load("System.ComponentModel.Primitives").Location));
-			MetadataReferences.Add(MetadataReference.CreateFromFile(Assembly.Load("System.Diagnostics.Process").Location));
-
-			// registry access
-			MetadataReferences.Add(MetadataReference.CreateFromFile(Assembly.Load("Microsoft.Win32.Registry").Location));
-
-			// RNGCryptoServiceProvider, used to generate random hex bytes
-			try
-			{
-				// Does not exist for .NET6
-				MetadataReferences.Add(MetadataReference.CreateFromFile(Assembly.Load("System.Security.Cryptography").Location));
-			}
-			catch (FileNotFoundException)
-			{
-			}
-			MetadataReferences.Add(MetadataReference.CreateFromFile(Assembly.Load("System.Security.Cryptography.Algorithms").Location));
-			MetadataReferences.Add(MetadataReference.CreateFromFile(Assembly.Load("System.Security.Cryptography.Csp").Location));
-
-			MetadataReferences.Add(MetadataReference.CreateFromFile(typeof(UnrealBuildTool).Assembly.Location));
-			MetadataReferences.Add(MetadataReference.CreateFromFile(typeof(FileReference).Assembly.Location));
-			MetadataReferences.Add(MetadataReference.CreateFromFile(typeof(UEBuildPlatformSDK).Assembly.Location));
+			IEnumerable<MetadataReference> MetadataReferences = assemblyLocations.Select(x => MetadataReference.CreateFromFile(x));
 
 			CSharpCompilationOptions CompilationOptions = new CSharpCompilationOptions(
 				outputKind: OutputKind.DynamicallyLinkedLibrary,
@@ -263,7 +272,7 @@ namespace UnrealBuildTool
 				syntaxTrees: SyntaxTrees,
 				references: MetadataReferences,
 				options: CompilationOptions
-				);
+			);
 
 			using (FileStream AssemblyStream = FileReference.Open(OutputAssemblyPath, FileMode.Create))
 			{
@@ -297,13 +306,13 @@ namespace UnrealBuildTool
 		/// <param name="OutputAssemblyPath">Full path to the assembly to be created</param>
 		/// <param name="SourceFileNames">List of source file name</param>
 		/// <param name="Logger">Logger for output</param>
-		/// <param name="ReferencedAssembies"></param>
+		/// <param name="ReferencedAssemblies"></param>
 		/// <param name="PreprocessorDefines"></param>
 		/// <param name="DoNotCompile"></param>
 		/// <param name="ForceCompile"></param>
 		/// <param name="TreatWarningsAsErrors"></param>
 		/// <returns>The assembly that was loaded</returns>
-		public static Assembly? CompileAndLoadAssembly(FileReference OutputAssemblyPath, IEnumerable<FileReference> SourceFileNames, ILogger Logger, IEnumerable<string>? ReferencedAssembies = null, IEnumerable<string>? PreprocessorDefines = null, bool DoNotCompile = false, bool ForceCompile = false, bool TreatWarningsAsErrors = false)
+		public static Assembly? CompileAndLoadAssembly(FileReference OutputAssemblyPath, IEnumerable<FileReference> SourceFileNames, ILogger Logger, IEnumerable<string>? ReferencedAssemblies, IEnumerable<string>? PreprocessorDefines = null, bool DoNotCompile = false, bool ForceCompile = false, bool TreatWarningsAsErrors = false)
 		{
 			// Check to see if the resulting assembly is compiled and up to date
 			FileReference AssemblyManifestFilePath = FileReference.Combine(OutputAssemblyPath.Directory, Path.GetFileNameWithoutExtension(OutputAssemblyPath.FullName) + "Manifest.json");
@@ -348,7 +357,7 @@ namespace UnrealBuildTool
 			{
 				using (GlobalTracer.Instance.BuildSpan(String.Format("Compiling rules assembly ({0})", OutputAssemblyPath.GetFileName())).StartActive())
 				{
-					CompiledAssembly = CompileAssembly(OutputAssemblyPath, SourceFileNames, Logger, ReferencedAssembies, PreprocessorDefines, TreatWarningsAsErrors);
+					CompiledAssembly = CompileAssembly(OutputAssemblyPath, SourceFileNames, Logger, ReferencedAssemblies, PreprocessorDefines, TreatWarningsAsErrors);
 				}
 
 				using (JsonWriter Writer = new JsonWriter(AssemblyManifestFilePath))
diff --git a/Engine/Source/Programs/UnrealBuildTool/System/RulesAssembly.cs b/Engine/Source/Programs/UnrealBuildTool/System/RulesAssembly.cs
--- a/Engine/Source/Programs/UnrealBuildTool/System/RulesAssembly.cs
+++ b/Engine/Source/Programs/UnrealBuildTool/System/RulesAssembly.cs
@@ -177,18 +177,18 @@ namespace UnrealBuildTool
 			if (AssemblySourceFiles.Any())
 			{
 				// Add the parent assemblies as references so they can be used
-				List<string>? ReferencedAssembies = null;
+				List<string>? ReferencedAssemblies = null;
 				RulesAssembly? CurrentParent = Parent;
 				while (CurrentParent != null && CurrentParent.CompiledAssembly != null)
 				{
-					ReferencedAssembies ??= new List<string>();
+					ReferencedAssemblies ??= new List<string>();
 
-					ReferencedAssembies.Add(CurrentParent.CompiledAssembly.Location);
+					ReferencedAssemblies.Add(CurrentParent.CompiledAssembly.Location);
 					CurrentParent = CurrentParent.Parent;
 				}
 
 				PreprocessorDefines = GetPreprocessorDefinitions();
-				CompiledAssembly = DynamicCompilation.CompileAndLoadAssembly(AssemblyFileName, AssemblySourceFiles, Logger, ReferencedAssembies: ReferencedAssembies, PreprocessorDefines: PreprocessorDefines, DoNotCompile: bSkipCompile, ForceCompile: bForceCompile);
+				CompiledAssembly = DynamicCompilation.CompileAndLoadAssembly(AssemblyFileName, AssemblySourceFiles, Logger, ReferencedAssemblies: ReferencedAssemblies, PreprocessorDefines: PreprocessorDefines, DoNotCompile: bSkipCompile, ForceCompile: bForceCompile);
 			}
 
 			// Setup the module map
