UnrealBuildTool: Remove unreferened nuget packages that are pulled in via EpicGames.MsBuild

#rnx

diff --git a/Engine/Source/Programs/UnrealBuildTool/UnrealBuildTool.csproj b/Engine/Source/Programs/UnrealBuildTool/UnrealBuildTool.csproj
--- a/Engine/Source/Programs/UnrealBuildTool/UnrealBuildTool.csproj
+++ b/Engine/Source/Programs/UnrealBuildTool/UnrealBuildTool.csproj
@@ -63,9 +63,6 @@
       <Link>Properties\MetaData.cs</Link>
     </Compile>
   </ItemGroup>
-  <ItemGroup>
-     <PackageReference Include="Microsoft.Extensions.Logging.Console" Version="8.0.1" />
-  </ItemGroup>
 
   <!-- The following will glob all platform extensions' UnrealBuildTool source
   files and include them in the build. They will also appear in VisualStudio.
@@ -111,6 +108,7 @@
   </ItemGroup>
   <ItemGroup>
     <PackageReference Include="Microsoft.CodeAnalysis.CSharp" Version="4.11.0" />
+    <PackageReference Include="Microsoft.Extensions.Logging.Console" Version="8.0.1" />
     <PackageReference Include="Microsoft.VisualStudio.Setup.Configuration.Interop" Version="3.11.2177" />
     <PackageReference Include="Microsoft.Win32.Registry" Version="5.0.0" />
     <PackageReference Include="OpenTracing" Version="0.12.1" />
@@ -123,15 +121,6 @@
     <PackageReference Include="System.Security.Permissions" Version="4.7.0" />
     <PackageReference Include="System.Text.Encoding.CodePages" Version="8.0.0" />
 
-    <!-- Required for Hololens -->
-    <PackageReference Include="System.Reflection.MetadataLoadContext" Version="8.0.1" /> 
-    
-    <!-- Required for EpicGames.MsBuild -->
-    <PackageReference Include="Microsoft.Build" Version="17.11.4" ExcludeAssets="runtime" PrivateAssets="all">
-      <IncludeAssets>compile; build; native; contentfiles; analyzers; buildtransitive</IncludeAssets>
-    </PackageReference>
-    <PackageReference Include="Microsoft.Build.Locator" Version="1.7.8" PrivateAssets="all" />
-
     <!-- Added as a temporary workaround to running from visual studio where the system version of dotnet is run. 
          Adding a loader context resolves many of the issues, but it still fails when MSBuild tries to write the json -->
     <PackageReference Include="Newtonsoft.Json" Version="13.0.3" />
