UnrealBuildTool: Move OriginalCompilationRootDirectory to Unreal namespace

#rnx

diff --git a/Engine/Source/Programs/Shared/EpicGames.Build/Unreal.cs b/Engine/Source/Programs/Shared/EpicGames.Build/Unreal.cs
--- a/Engine/Source/Programs/Shared/EpicGames.Build/Unreal.cs
+++ b/Engine/Source/Programs/Shared/EpicGames.Build/Unreal.cs
@@ -236,6 +236,13 @@ namespace UnrealBuildBase
 			return FileReference.Exists(installedProjectLocationFile) ? FileReference.Combine(RootDirectory, FileReference.ReadAllText(installedProjectLocationFile).Trim()) : null;
 		});
 
+		/// <summary>
+		/// The original root directory that was used to compile the installed engine
+		/// Used to remap source code paths when debugging.
+		/// </summary>
+		public static DirectoryReference OriginalCompilationRootDirectory => _OriginalCompilationRootDirectory.Value;
+		private static readonly Lazy<DirectoryReference> _OriginalCompilationRootDirectory = new(FindOriginalCompilationRootDirectory);
+
 		/// <summary>
 		/// Returns where another platform's Dotnet is located
 		/// </summary>
@@ -273,6 +280,40 @@ namespace UnrealBuildBase
 			return false;
 		}
 
+		private static DirectoryReference FindOriginalCompilationRootDirectory()
+		{
+			if (IsEngineInstalled())
+			{
+				// Load Engine\Intermediate\Build\BuildRules\*RulesManifest.json
+				DirectoryReference BuildRules = DirectoryReference.Combine(EngineDirectory, "Intermediate", "Build", "BuildRules");
+				FileReference? RulesManifest = DirectoryReference.EnumerateFiles(BuildRules, "*RulesManifest.json").FirstOrDefault();
+				if (RulesManifest != null)
+				{
+					JsonObject Manifest = JsonObject.Read(RulesManifest);
+					if (Manifest.TryGetStringArrayField("SourceFiles", out string[]? SourceFiles))
+					{
+						FileReference? SourceFile = FileReference.FromString(SourceFiles.FirstOrDefault());
+						if (SourceFile != null && !SourceFile.IsUnderDirectory(EngineDirectory))
+						{
+							// Walk up parent directory until Engine is found
+							DirectoryReference? Directory = SourceFile.Directory;
+							while (Directory != null && !Directory.IsRootDirectory())
+							{
+								if (Directory.GetDirectoryName() == "Engine" && Directory.ParentDirectory != null)
+								{
+									return Directory.ParentDirectory;
+								}
+
+								Directory = Directory.ParentDirectory;
+							}
+						}
+					}
+				}
+			}
+
+			return RootDirectory;
+		}
+
 		public static class LocationOverride
 		{
 			/// <summary>
diff --git a/Engine/Source/Programs/UnrealBuildTool/ProjectFiles/VisualStudioCode/VSCodeProjectFileGenerator.cs b/Engine/Source/Programs/UnrealBuildTool/ProjectFiles/VisualStudioCode/VSCodeProjectFileGenerator.cs
--- a/Engine/Source/Programs/UnrealBuildTool/ProjectFiles/VisualStudioCode/VSCodeProjectFileGenerator.cs
+++ b/Engine/Source/Programs/UnrealBuildTool/ProjectFiles/VisualStudioCode/VSCodeProjectFileGenerator.cs
@@ -1405,12 +1405,12 @@ namespace UnrealBuildTool
 						OutFile.AddField("type", "lldb");
 					}
 
-					if (UnrealBuildTool.OriginalCompilationRootDirectory != ProjectRoot)
+					if (Unreal.OriginalCompilationRootDirectory != ProjectRoot)
 					{
 						OutFile.BeginObject("sourceFileMap");
 						{
 							OutFile.AddField(
-								MakeUnquotedPathString(UnrealBuildTool.OriginalCompilationRootDirectory, EPathType.Absolute),
+								MakeUnquotedPathString(Unreal.OriginalCompilationRootDirectory, EPathType.Absolute),
 								MakeUnquotedPathString(ProjectRoot, EPathType.Absolute));
 						}
 						OutFile.EndObject();
diff --git a/Engine/Source/Programs/UnrealBuildTool/UnrealBuildTool.cs b/Engine/Source/Programs/UnrealBuildTool/UnrealBuildTool.cs
--- a/Engine/Source/Programs/UnrealBuildTool/UnrealBuildTool.cs
+++ b/Engine/Source/Programs/UnrealBuildTool/UnrealBuildTool.cs
@@ -54,12 +54,6 @@ namespace UnrealBuildTool
 		[Obsolete("Replace with Unreal.EngineSourceDirectory")]
 		public static readonly DirectoryReference EngineSourceDirectory = Unreal.EngineSourceDirectory;
 
-		/// <summary>
-		/// Cached copy of the source root directory that was used to compile the installed engine
-		/// Used to remap source code paths when debugging.
-		/// </summary>
-		static DirectoryReference? CachedOriginalCompilationRootDirectory;
-
 		/// <summary>
 		/// Writable engine directory. Uses the user's settings folder for installed builds.
 		/// </summary>
@@ -76,55 +70,8 @@ namespace UnrealBuildTool
 		/// The original root directory that was used to compile the installed engine
 		/// Used to remap source code paths when debugging.
 		/// </summary>
-		public static DirectoryReference OriginalCompilationRootDirectory
-		{
-			get
-			{
-				if (CachedOriginalCompilationRootDirectory == null)
-				{
-					if (Unreal.IsEngineInstalled())
-					{
-						// Load Engine\Intermediate\Build\BuildRules\*RulesManifest.json
-						DirectoryReference BuildRules = DirectoryReference.Combine(Unreal.EngineDirectory, "Intermediate", "Build", "BuildRules");
-						FileReference? RulesManifest = DirectoryReference.EnumerateFiles(BuildRules, "*RulesManifest.json").FirstOrDefault();
-						if (RulesManifest != null)
-						{
-							JsonObject Manifest = JsonObject.Read(RulesManifest);
-							if (Manifest.TryGetStringArrayField("SourceFiles", out string[]? SourceFiles))
-							{
-								FileReference? SourceFile = FileReference.FromString(SourceFiles.FirstOrDefault());
-								if (SourceFile != null && !SourceFile.IsUnderDirectory(Unreal.EngineDirectory))
-								{
-									// Walk up parent directory until Engine is found
-									DirectoryReference? Directory = SourceFile.Directory;
-									while (Directory != null && !Directory.IsRootDirectory())
-									{
-										if (Directory.GetDirectoryName() == "Engine" && Directory.ParentDirectory != null)
-										{
-											CachedOriginalCompilationRootDirectory = Directory.ParentDirectory;
-											break;
-										}
-
-										Directory = Directory.ParentDirectory;
-									}
-								}
-							}
-						}
-					}
-
-					if (CachedOriginalCompilationRootDirectory == null)
-					{
-						CachedOriginalCompilationRootDirectory = Unreal.RootDirectory;
-					}
-				}
-				return CachedOriginalCompilationRootDirectory;
-			}
-		}
-
-		/// <summary>
-		/// The Remote Ini directory.  This should always be valid when compiling using a remote server.
-		/// </summary>
-		static string? RemoteIniPath = null;
+		[Obsolete("Deprecated in UE5.6 - use Unreal.OriginalCompilationRootDirectory")]
+		public static DirectoryReference OriginalCompilationRootDirectory => Unreal.OriginalCompilationRootDirectory;
 
 		/// <summary>
 		/// Returns true if UnrealBuildTool is running using an installed project (ie. a mod kit)
@@ -152,22 +99,25 @@ namespace UnrealBuildTool
 		/// Gets the absolute path to the UBT assembly.
 		/// </summary>
 		/// <returns>A string containing the path to the UBT assembly.</returns>
-		[Obsolete("Deprecated in UE5.1 - use UnrealBuildTool.DotnetPath Unreal.UnrealBuildToolDllPath")]
+		[Obsolete("Deprecated in UE5.1 - use Unreal.DotnetPath and Unreal.UnrealBuildToolDllPath")]
 		public static FileReference GetUBTPath() => Unreal.UnrealBuildToolPath;
 
 		/// <summary>
 		/// The Unreal remote tool ini directory.  This should be valid if compiling using a remote server
 		/// </summary>
 		/// <returns>The directory path</returns>
-		public static string? GetRemoteIniPath()
-		{
-			return RemoteIniPath;
-		}
+		public static string? GetRemoteIniPath() => _RemoteIniPath;
 
-		public static void SetRemoteIniPath(string Path)
-		{
-			RemoteIniPath = Path;
-		}
+		/// <summary>
+		/// Set the Unreal remote tool ini directory.
+		/// </summary>
+		/// <param name="Path"> remote path</param>
+		public static void SetRemoteIniPath(string Path) => _RemoteIniPath = Path;
+
+		/// <summary>
+		/// The Remote Ini directory.  This should always be valid when compiling using a remote server.
+		/// </summary>
+		private static string? _RemoteIniPath = null;
 
 		/// <summary>
 		/// Global options for UBT (any modes)
