UnrealBuildTool: Allow .sarif files to be disabled for fastbuild users

#rnx

diff --git a/Engine/Source/Programs/UnrealBuildTool/Platform/Windows/UEBuildWindows.cs b/Engine/Source/Programs/UnrealBuildTool/Platform/Windows/UEBuildWindows.cs
--- a/Engine/Source/Programs/UnrealBuildTool/Platform/Windows/UEBuildWindows.cs
+++ b/Engine/Source/Programs/UnrealBuildTool/Platform/Windows/UEBuildWindows.cs
@@ -463,7 +463,8 @@ namespace UnrealBuildTool
 		/// <summary>
 		/// Whether .sarif files containing errors and warnings are written alongside each .obj, if supported
 		/// </summary>
-		[Obsolete("No longer necessary, .sarif files will always be written")]
+		[ConfigFile(ConfigHierarchyType.Engine, "/Script/WindowsTargetPlatform.WindowsTargetSettings", "bWriteSarif")]
+		[XmlConfigFile(Category = "WindowsPlatform")]
 		public bool bWriteSarif = true;
 
 		/// <summary>
@@ -909,6 +910,8 @@ namespace UnrealBuildTool
 
 		public bool bEnableLibFuzzer => Inner.bEnableLibFuzzer;
 
+		public bool bWriteSarif => Inner.bWriteSarif;
+
 		public string? ObjSrcMapFile => Inner.ObjSrcMapFile;
 
 		public string? LinkReproDir => Inner.LinkReproDir;
diff --git a/Engine/Source/Programs/UnrealBuildTool/Platform/Windows/VCToolChain.cs b/Engine/Source/Programs/UnrealBuildTool/Platform/Windows/VCToolChain.cs
--- a/Engine/Source/Programs/UnrealBuildTool/Platform/Windows/VCToolChain.cs
+++ b/Engine/Source/Programs/UnrealBuildTool/Platform/Windows/VCToolChain.cs
@@ -2081,7 +2081,7 @@ namespace UnrealBuildTool
 				}
 
 				// Write cl errors and warnings to a file
-				if (Target.WindowsPlatform.Compiler.IsMSVC())
+				if (Target.WindowsPlatform.Compiler.IsMSVC() && Target.WindowsPlatform.bWriteSarif)
 				{
 					if (Target.StaticAnalyzer == StaticAnalyzer.Default && !CompileEnvironment.bDisableStaticAnalysis)
 					{
