[Fastbuild] Fastbuild is issuing a local duplicate build

diff --git a/Engine/Source/Programs/UnrealBuildTool/Executors/Experimental/FASTBuild.cs b/Engine/Source/Programs/UnrealBuildTool/Executors/Experimental/FASTBuild.cs
--- a/Engine/Source/Programs/UnrealBuildTool/Executors/Experimental/FASTBuild.cs
+++ b/Engine/Source/Programs/UnrealBuildTool/Executors/Experimental/FASTBuild.cs
@@ -502,6 +502,7 @@ namespace UnrealBuildTool
 
 				bool result = ExecuteBffFile(FASTBuildFilePath, Logger);
 				telemetryEvent = new TelemetryExecutorEvent(Name, startTimeUTC, result, Actions.Count(), -1, -1, 0, 0, DateTime.UtcNow);
+				return result;
 			}
 
 			return await LocalExecutor.ExecuteActionsAsync(Actions, Logger, actionArtifactCache);
