UnrealBuildTool: Add missing ispc header outputs to action

#rnx

diff --git a/Engine/Source/Programs/UnrealBuildTool/ToolChain/ISPCToolChain.cs b/Engine/Source/Programs/UnrealBuildTool/ToolChain/ISPCToolChain.cs
--- a/Engine/Source/Programs/UnrealBuildTool/ToolChain/ISPCToolChain.cs
+++ b/Engine/Source/Programs/UnrealBuildTool/ToolChain/ISPCToolChain.cs
@@ -476,6 +476,33 @@ namespace UnrealBuildTool
 						)
 					);
 
+				// Add the ISPC h files to the produced item list.
+				CompileAction.ProducedItems.Add(ISPCIncludeHeaderFile);
+				
+				// Add the target specific dummy headers to the produced item list.
+				// These are unused but should still be tracked.
+				if (CompileTargets.Count > 1)
+				{
+					foreach (string Target in CompileTargets)
+					{
+						string HeaderTarget = Target;
+
+						if (Target.Contains('-'))
+						{
+							// Remove lane width and gang size from header file name
+							HeaderTarget = Target.Split('-')[0];
+						}
+
+						FileItem ISPCTargetIncludeHeaderFile = FileItem.GetItemByFileReference(
+							FileReference.Combine(
+								OutputDir,
+								Path.GetFileName(ISPCFile.AbsolutePath) + ".generated.dummy_" + HeaderTarget + ".h"
+								)
+							);
+						CompileAction.ProducedItems.Add(ISPCTargetIncludeHeaderFile);
+					}
+				}
+
 				// Add the ISPC file to be compiled.
 				Arguments.Add($"-h \"{NormalizeCommandLinePath(ISPCIncludeHeaderFile)}\"");
 
@@ -487,8 +514,6 @@ namespace UnrealBuildTool
 
 				Arguments.AddRange(GlobalArguments);
 
-				CompileAction.ProducedItems.Add(ISPCIncludeHeaderFile);
-
 				FileReference ResponseFileName = GetResponseFileName(CompileEnvironment, ISPCIncludeHeaderFile);
 				FileItem ResponseFileItem = Graph.CreateIntermediateTextFile(ResponseFileName, Arguments.Select(x => Utils.ExpandVariables(x)));
 				CompileAction.CommandArguments = $"@\"{NormalizeCommandLinePath(ResponseFileName)}\"";
