UnrealBuildTool: HostArm64 fixes

#rnx

diff --git a/Engine/Source/Programs/Shared/EpicGames.Core/ManagedProcess.cs b/Engine/Source/Programs/Shared/EpicGames.Core/ManagedProcess.cs
--- a/Engine/Source/Programs/Shared/EpicGames.Core/ManagedProcess.cs
+++ b/Engine/Source/Programs/Shared/EpicGames.Core/ManagedProcess.cs
@@ -200,17 +200,12 @@ namespace EpicGames.Core
 				if (SupportsJobObjects)
 				{
 					JOBOBJECT_BASIC_ACCOUNTING_INFORMATION accountingInformation = new JOBOBJECT_BASIC_ACCOUNTING_INFORMATION();
-					if (QueryInformationJobObject(JobHandle!, JobObjectBasicAccountingInformation, ref accountingInformation, Marshal.SizeOf(typeof(JOBOBJECT_BASIC_ACCOUNTING_INFORMATION))) == false)
+					if (QueryInformationJobObject(JobHandle!, JobObjectBasicAccountingInformation, ref accountingInformation, Marshal.SizeOf(typeof(JOBOBJECT_BASIC_ACCOUNTING_INFORMATION))))
 					{
-						throw new Win32Exception();
+						return new TimeSpan((long)accountingInformation.TotalUserTime + (long)accountingInformation.TotalKernelTime);
 					}
-
-					return new TimeSpan((long)accountingInformation.TotalUserTime + (long)accountingInformation.TotalKernelTime);
-				}
-				else
-				{
-					return new TimeSpan();
 				}
+				return new TimeSpan();
 			}
 		}
 
diff --git a/Engine/Source/Programs/UnrealBuildTool/Platform/Windows/VCEnvironment.cs b/Engine/Source/Programs/UnrealBuildTool/Platform/Windows/VCEnvironment.cs
--- a/Engine/Source/Programs/UnrealBuildTool/Platform/Windows/VCEnvironment.cs
+++ b/Engine/Source/Programs/UnrealBuildTool/Platform/Windows/VCEnvironment.cs
@@ -4,6 +4,7 @@ using System;
 using System.Collections.Generic;
 using System.Diagnostics;
 using System.Linq;
+using System.Runtime.InteropServices;
 using System.Runtime.Versioning;
 using EpicGames.Core;
 using Microsoft.Extensions.Logging;
@@ -165,12 +166,7 @@ namespace UnrealBuildTool
 			// Add the Windows SDK directory to the path too, for mt.exe.
 			if (WindowsSdkVersion >= new VersionNumber(10))
 			{
-				string BuildHostArch = Architecture.ToString();
-				if (Architecture == UnrealArch.Arm64 || Architecture == UnrealArch.Arm64ec)
-				{
-					Debug.Assert(BuildHostPlatform.Current.Platform == UnrealTargetPlatform.Win64);
-					BuildHostArch = "x64";
-				}
+				string BuildHostArch = RuntimeInformation.ProcessArchitecture == System.Runtime.InteropServices.Architecture.Arm64 ? "arm64" : "x64";
 				DirectoryReference.AddDirectoryToPath(DirectoryReference.Combine(WindowsSdkDir, "bin", WindowsSdkVersion.ToString(), BuildHostArch));
 
 			}
@@ -250,13 +246,14 @@ namespace UnrealBuildTool
 		/// </summary>
 		protected virtual FileReference GetResourceCompilerToolPath(DirectoryReference WindowsSdkDir, VersionNumber WindowsSdkVersion, ILogger Logger)
 		{
-			FileReference ResourceCompilerPath = FileReference.Combine(WindowsSdkDir, "bin", WindowsSdkVersion.ToString(), "x64", "rc.exe");
+			string BuildHostArch = RuntimeInformation.ProcessArchitecture == System.Runtime.InteropServices.Architecture.Arm64 ? "arm64" : "x64";
+			FileReference ResourceCompilerPath = FileReference.Combine(WindowsSdkDir, "bin", WindowsSdkVersion.ToString(), BuildHostArch, "rc.exe");
 			if (FileReference.Exists(ResourceCompilerPath))
 			{
 				return ResourceCompilerPath;
 			}
 
-			ResourceCompilerPath = FileReference.Combine(WindowsSdkDir, "bin", "x64", "rc.exe");
+			ResourceCompilerPath = FileReference.Combine(WindowsSdkDir, "bin", BuildHostArch, "rc.exe");
 			if (FileReference.Exists(ResourceCompilerPath))
 			{
 				return ResourceCompilerPath;
