UnrealBuildTool: Support arm64e cross-compile with clang (currently has build errors that need to be addressed)

#rnx

diff --git a/Engine/Source/Programs/UnrealBuildTool/Platform/Windows/MicrosoftPlatformSDK.cs b/Engine/Source/Programs/UnrealBuildTool/Platform/Windows/MicrosoftPlatformSDK.cs
--- a/Engine/Source/Programs/UnrealBuildTool/Platform/Windows/MicrosoftPlatformSDK.cs
+++ b/Engine/Source/Programs/UnrealBuildTool/Platform/Windows/MicrosoftPlatformSDK.cs
@@ -1072,6 +1072,8 @@ namespace UnrealBuildTool
 
 				Logger.LogDebug("Found Clang toolchain: {ToolChainDir} (Version={Version}, Is64Bit={Is64Bit}, Rank={Rank}, Error={Error})", ToolChainDir, Version, Is64Bit, Rank, Error != null);
 				ToolChains.Add(new ToolChainInstallation(Family, Rank, Version, Is64Bit, WindowsCompilerChannel.Any, UnrealArch.X64, Error, ToolChainDir, null, IsAutoSdk));
+				ToolChains.Add(new ToolChainInstallation(Family, Rank, Version, Is64Bit, WindowsCompilerChannel.Any, UnrealArch.Arm64, Error, ToolChainDir, null, IsAutoSdk));
+				ToolChains.Add(new ToolChainInstallation(Family, Rank, Version, Is64Bit, WindowsCompilerChannel.Any, UnrealArch.Arm64ec, Error, ToolChainDir, null, IsAutoSdk));
 			}
 		}
 
@@ -1108,6 +1110,8 @@ namespace UnrealBuildTool
 
 				Logger.LogDebug("Found Intel OneAPI toolchain: {ToolChainDir} (Version={Version}, Is64Bit={Is64Bit}, Rank={Rank}, Error={Error})", ToolChainDir, Version, Is64Bit, Rank, Error != null);
 				ToolChains.Add(new ToolChainInstallation(Family, Rank, Version, Is64Bit, WindowsCompilerChannel.Any, UnrealArch.X64, Error, ToolChainDir, null, IsAutoSdk));
+				ToolChains.Add(new ToolChainInstallation(Family, Rank, Version, Is64Bit, WindowsCompilerChannel.Any, UnrealArch.Arm64, Error, ToolChainDir, null, IsAutoSdk));
+				ToolChains.Add(new ToolChainInstallation(Family, Rank, Version, Is64Bit, WindowsCompilerChannel.Any, UnrealArch.Arm64ec, Error, ToolChainDir, null, IsAutoSdk));
 			}
 		}
 
diff --git a/Engine/Source/Programs/UnrealBuildTool/Platform/Windows/VCToolChain.cs b/Engine/Source/Programs/UnrealBuildTool/Platform/Windows/VCToolChain.cs
--- a/Engine/Source/Programs/UnrealBuildTool/Platform/Windows/VCToolChain.cs
+++ b/Engine/Source/Programs/UnrealBuildTool/Platform/Windows/VCToolChain.cs
@@ -477,13 +477,24 @@ namespace UnrealBuildTool
 			if (Target.WindowsPlatform.Compiler.IsClang() && Target.Platform.IsInGroup(UnrealPlatformGroup.Windows))
 			{
 				// Tell the Clang compiler to generate 64-bit code
-				Arguments.Add("--target=x86_64-pc-windows-msvc");
+				if (CompileEnvironment.Architecture.bIsX64)
+				{
+					Arguments.Add("--target=x86_64-pc-windows-msvc");
 
-				// UE5 minspec is 4.2
-				Arguments.Add("-msse4.2");
+					// UE5 minspec is 4.2
+					Arguments.Add("-msse4.2");
 
-				// Use tpause on supported processors.
-				Arguments.Add("-mwaitpkg");
+					// Use tpause on supported processors.
+					Arguments.Add("-mwaitpkg");
+				}
+				else if (CompileEnvironment.Architecture == UnrealArch.Arm64)
+				{
+					Arguments.Add("--target=arm64-pc-windows-msvc");
+				}
+				else if (CompileEnvironment.Architecture == UnrealArch.Arm64ec)
+				{
+					Arguments.Add("--target=arm64ec-pc-windows-msvc");
+				}
 			}
 
 			// Compile into an .obj file, and skip linking.
