Support LiveCoding with verse-clang.exe

#rnx

diff --git a/Engine/Source/Developer/Windows/LiveCodingServer/Private/External/LC_Compiler.cpp b/Engine/Source/Developer/Windows/LiveCodingServer/Private/External/LC_Compiler.cpp
--- a/Engine/Source/Developer/Windows/LiveCodingServer/Private/External/LC_Compiler.cpp
+++ b/Engine/Source/Developer/Windows/LiveCodingServer/Private/External/LC_Compiler.cpp
@@ -162,6 +162,10 @@ Process::Environment compiler::CreateEnvironmentCacheEntry(const wchar_t* absolu
 		{
 			return Process::Environment { 0u, nullptr };
 		}
+		else if (string::Matches(toolFilename.GetString(), L"verse-clang-cl.exe"))
+		{
+			return Process::Environment { 0u, nullptr };
+		}
 
 		// LLD
 		else if (string::Matches(toolFilename.GetString(), L"lld.exe"))
diff --git a/Engine/Source/Programs/UnrealBuildTool/Platform/Windows/VCToolChain.cs b/Engine/Source/Programs/UnrealBuildTool/Platform/Windows/VCToolChain.cs
--- a/Engine/Source/Programs/UnrealBuildTool/Platform/Windows/VCToolChain.cs
+++ b/Engine/Source/Programs/UnrealBuildTool/Platform/Windows/VCToolChain.cs
@@ -353,6 +353,14 @@ namespace UnrealBuildTool
 
 				// We have 'this' vs nullptr comparisons that get optimized away for newer versions of Clang, which is undesirable until we refactor these checks.
 				Arguments.Add("-fno-delete-null-pointer-checks");
+				
+				if (Target.bWithLiveCoding)
+				{
+					Arguments.Add("-Z7");
+					Arguments.Add("-fms-hotpatch");
+					Arguments.Add("-Gy");
+					Arguments.Add("-Xclang -mno-constructor-aliases");
+				}
 
 				if (Target.StaticAnalyzer == StaticAnalyzer.Default && CompileEnvironment.PrecompiledHeaderAction != PrecompiledHeaderAction.Create && !CompileEnvironment.bDisableStaticAnalysis)
 				{
diff --git a/Engine/Source/Programs/UnrealBuildTool/System/HotReload.cs b/Engine/Source/Programs/UnrealBuildTool/System/HotReload.cs
--- a/Engine/Source/Programs/UnrealBuildTool/System/HotReload.cs
+++ b/Engine/Source/Programs/UnrealBuildTool/System/HotReload.cs
@@ -874,6 +874,7 @@ namespace UnrealBuildTool
 				(Action.CommandPath.GetFileName().Equals("cl-filter.exe", StringComparison.OrdinalIgnoreCase)
 					|| Action.CommandPath.GetFileName().Equals("cl.exe", StringComparison.OrdinalIgnoreCase)
 					|| Action.CommandPath.GetFileName().Equals("clang-cl.exe", StringComparison.OrdinalIgnoreCase)
+					|| Action.CommandPath.GetFileName().Equals("verse-clang-cl.exe", StringComparison.OrdinalIgnoreCase)
 				);
 		}
 
@@ -898,7 +899,9 @@ namespace UnrealBuildTool
 				{
 					if (!Action.CommandPath.GetFileName().Equals("cl-filter.exe", StringComparison.OrdinalIgnoreCase)
 						&& !Action.CommandPath.GetFileName().Equals("cl.exe", StringComparison.OrdinalIgnoreCase)
-						&& !Action.CommandPath.GetFileName().Equals("clang-cl.exe", StringComparison.OrdinalIgnoreCase))
+						&& !Action.CommandPath.GetFileName().Equals("clang-cl.exe", StringComparison.OrdinalIgnoreCase)
+						&& !Action.CommandPath.GetFileName().Equals("verse-clang-cl.exe", StringComparison.OrdinalIgnoreCase)
+						)
 					{
 						throw new BuildException("Unable to patch action graph - unexpected executable in compile action ({0})", Action.CommandPath);
 					}
@@ -989,8 +992,9 @@ namespace UnrealBuildTool
 					}
 
 					// Update dependency file path for cl or clang-cl which is in the response file
-					if (Action.CommandPath.GetFileName().Equals("cl.exe", StringComparison.OrdinalIgnoreCase) ||
-						Action.CommandPath.GetFileName().Equals("clang-cl.exe", StringComparison.OrdinalIgnoreCase))
+					if (Action.CommandPath.GetFileName().Equals("cl.exe", StringComparison.OrdinalIgnoreCase) 
+					    || Action.CommandPath.GetFileName().Equals("clang-cl.exe", StringComparison.OrdinalIgnoreCase)
+						|| Action.CommandPath.GetFileName().Equals("verse-clang-cl.exe", StringComparison.OrdinalIgnoreCase))
 					{
 						string SourceDependencyPrefix = Action.CommandPath.GetFileName().Equals("cl.exe", StringComparison.OrdinalIgnoreCase) ? "/sourceDependencies" : "/clang:-MD /clang:-MF";
 						string NewExtension = Action.CommandPath.GetFileName().Equals("cl.exe", StringComparison.OrdinalIgnoreCase) ? clSourceDepFileExtension : clangSourceDepFileExtension;
