UnrealBuildTool: Register all host architecture toolchains and register them with UBA

#rnx

diff --git a/Engine/Source/Programs/Shared/EpicGames.UBA/Utils.cs b/Engine/Source/Programs/Shared/EpicGames.UBA/Utils.cs
--- a/Engine/Source/Programs/Shared/EpicGames.UBA/Utils.cs
+++ b/Engine/Source/Programs/Shared/EpicGames.UBA/Utils.cs
@@ -26,6 +26,13 @@ namespace EpicGames.UBA
 		public static IEnumerable<string> DisallowedPaths() => s_disallowedPaths;
 		static readonly SortedSet<string> s_disallowedPaths = [];
 
+		/// <summary>
+		/// Mapping of binary paths for cross architecture host binaries, to allow for using helpers of a different architecture.
+		/// </summary>
+		/// <returns>dictionary of binary mappings, where the key is the binary for the current host architecture</returns>
+		public static IDictionary<string, string> CrossArchitecturePaths => s_crossArchitecturePaths;
+		static readonly Dictionary<string, string> s_crossArchitecturePaths = [];
+
 		/// <summary>
 		/// Registers a path that is not allowed to be transferred over the network for UBA remote agents.
 		/// </summary>
@@ -36,6 +43,17 @@ namespace EpicGames.UBA
 			DisallowedPathRegistered?.Invoke(s_disallowedPaths, new(paths));
 		}
 
+		/// <summary>
+		/// Registers a path mapping for cross architecture binaries
+		/// </summary>
+		/// <param name="path">host architecture path</param>
+		/// <param name="crossPath">cross architecture path</param>
+		public static void RegisterCrossArchitecturePath(string path, string crossPath)
+		{
+			s_crossArchitecturePaths.TryAdd(path, crossPath);
+			CrossArchitecturePathRegistered?.Invoke(s_crossArchitecturePaths, new(path, crossPath));
+		}
+
 		/// <summary>
 		/// Delegate for registering a remote disallowed path
 		/// </summary>
@@ -43,11 +61,23 @@ namespace EpicGames.UBA
 		/// <param name="e">event args containing which paths were added</param>
 		public delegate void DisallowedPathRegisteredEventHandler(SortedSet<string> sender, DisallowedPathRegisteredEventArgs e);
 
+		/// <summary>
+		/// Delegate for registering a cross architecture path
+		/// </summary>
+		/// <param name="sender">collection that is being changed</param>
+		/// <param name="e">event args containing which path was added</param>
+		public delegate void CrossArchitecturePathRegisteredEventHandler(Dictionary<string, string> sender, CrossArchitecturePathRegisteredEventArgs e);
+
 		/// <summary>
 		/// Remote disallowed path registered event handler
 		/// </summary>
 		public static event DisallowedPathRegisteredEventHandler? DisallowedPathRegistered;
 
+		/// <summary>
+		/// Cross architecture path registered event handler
+		/// </summary>
+		public static event CrossArchitecturePathRegisteredEventHandler? CrossArchitecturePathRegistered;
+
 		/// <summary>
 		/// Get the path to the p/invoke library that would be loaded 
 		/// </summary>
@@ -84,4 +114,20 @@ namespace EpicGames.UBA
 		/// </summary>
 		public IEnumerable<string> Paths { get; } = paths;
 	}
+
+	/// <summary>
+	/// Event args for registering a cross architecture path
+	/// </summary>
+	public class CrossArchitecturePathRegisteredEventArgs(string path, string crossPath) : EventArgs
+	{
+		/// <summary>
+		/// The host architecture path being registered
+		/// </summary>
+		public string Paths { get; } = path;
+
+		/// <summary>
+		/// The cross architecture path being registered
+		/// </summary>
+		public string CrossPath { get; } = crossPath;
+	}
 }
diff --git a/Engine/Source/Programs/UnrealBuildTool/Platform/Windows/MicrosoftPlatformSDK.cs b/Engine/Source/Programs/UnrealBuildTool/Platform/Windows/MicrosoftPlatformSDK.cs
--- a/Engine/Source/Programs/UnrealBuildTool/Platform/Windows/MicrosoftPlatformSDK.cs
+++ b/Engine/Source/Programs/UnrealBuildTool/Platform/Windows/MicrosoftPlatformSDK.cs
@@ -630,9 +630,8 @@ namespace UnrealBuildTool
 		/// <returns></returns>
 		static ToolChainInstallation? SelectToolChain(IEnumerable<ToolChainInstallation> ToolChains, Func<IOrderedEnumerable<ToolChainInstallation>, IOrderedEnumerable<ToolChainInstallation>> Preference, UnrealArch Architecture)
 		{
-			ToolChainInstallation? ToolChain = Preference(ToolChains.Where(x => x.Architecture == Architecture)
+			ToolChainInstallation? ToolChain = Preference(ToolChains.Where(x => x.Architecture == Architecture && x.HostArchitecture == MicrosoftPlatformSDK.MSVCHostUnrealArch)
 				.OrderByDescending(x => x.Error == null))
-				.ThenByDescending(x => x.Is64Bit)
 				.ThenBy(x => x.ReleaseChannel)
 				.ThenBy(x => x.FamilyRank)
 				.ThenByDescending(x => x.IsAutoSdk)
@@ -663,11 +662,11 @@ namespace UnrealBuildTool
 		{
 			IOrderedEnumerable<ToolChainInstallation> SortedToolChains = Preference(ToolChains.Where(x => x.Architecture == Architecture)
 				.OrderByDescending(x => x.Error == null))
-				.ThenByDescending(x => x.Is64Bit)
 				.ThenBy(x => x.ReleaseChannel)
 				.ThenBy(x => x.FamilyRank)
 				.ThenByDescending(x => x.IsAutoSdk)
-				.ThenByDescending(x => x.Version);
+				.ThenByDescending(x => x.Version)
+				.ThenBy(x => x.HostArchitecture.WindowsName);
 
 			if (SortedToolChains.Any())
 			{
@@ -679,7 +678,7 @@ namespace UnrealBuildTool
 					Logger.LogInformation("Available {Architecture} toolchains ({Count}):", Architecture, AvailableToolChains.Count());
 					foreach (ToolChainInstallation ToolChain in AvailableToolChains)
 					{
-						Logger.LogInformation(" * {ToolChainDir}\n    (Family={Family}, FamilyRank={FamilyRank}, Version={Version}, Is64Bit={Is64Bit}, ReleaseChannel={ReleaseChannel}, Architecture={Arch})", ToolChain.BaseDir, ToolChain.Family, ToolChain.FamilyRank, ToolChain.Version, ToolChain.Is64Bit, ToolChain.ReleaseChannel, ToolChain.Architecture);
+						Logger.LogInformation(" * {ToolChainDir}\n    (Family={Family}, FamilyRank={FamilyRank}, Version={Version}, HostArchitecture={HostArchitecture}, ReleaseChannel={ReleaseChannel}, Architecture={Arch})", ToolChain.BaseDir, ToolChain.Family, ToolChain.FamilyRank, ToolChain.Version, ToolChain.HostArchitecture, ToolChain.ReleaseChannel, ToolChain.Architecture);
 					}
 				}
 
@@ -688,7 +687,7 @@ namespace UnrealBuildTool
 					Logger.LogWarning("Unavailable {Architecture} toolchains ({Count}):", Architecture, UnavailableToolChains.Count());
 					foreach (ToolChainInstallation ToolChain in UnavailableToolChains.Where(x => x.Error != null))
 					{
-						Logger.LogWarning(" * {ToolChainDir}\n    (Family={Family}, FamilyRank={FamilyRank}, Version={Version}, Is64Bit={Is64Bit}, ReleaseChannel={ReleaseChannel}, Architecture={Arch}, Error=\"{Error}\")", ToolChain.BaseDir, ToolChain.Family, ToolChain.FamilyRank, ToolChain.Version, ToolChain.Is64Bit, ToolChain.ReleaseChannel, ToolChain.Architecture, ToolChain.Error);
+						Logger.LogWarning(" * {ToolChainDir}\n    (Family={Family}, FamilyRank={FamilyRank}, Version={Version}, HostArchitecture={HostArchitecture}, ReleaseChannel={ReleaseChannel}, Architecture={Arch}, Error=\"{Error}\")", ToolChain.BaseDir, ToolChain.Family, ToolChain.FamilyRank, ToolChain.Version, ToolChain.HostArchitecture, ToolChain.ReleaseChannel, ToolChain.Architecture, ToolChain.Error);
 					}
 				}
 			}
@@ -738,6 +737,7 @@ namespace UnrealBuildTool
 						{
 							AddClangToolChain(Compiler, DirectoryReference.Combine(Installation.BaseDir, "VC", "Tools", "Llvm"), ToolChains, IsAutoSdk: false, Logger);
 							AddClangToolChain(Compiler, DirectoryReference.Combine(Installation.BaseDir, "VC", "Tools", "Llvm", "x64"), ToolChains, IsAutoSdk: false, Logger);
+							AddClangToolChain(Compiler, DirectoryReference.Combine(Installation.BaseDir, "VC", "Tools", "Llvm", "ARM64"), ToolChains, IsAutoSdk: false, Logger);
 						}
 
 						// Check for AutoSDK paths
@@ -831,6 +831,31 @@ namespace UnrealBuildTool
 					{
 						throw new BuildException("Unsupported compiler version ({0})", Compiler);
 					}
+
+					foreach (ToolChainInstallation? ToolChain in ToolChains.Where(x => x.Error == null && x.HostArchitecture == MSVCHostUnrealArch))
+					{
+						ToolChainInstallation? CrossToolChain = ToolChains.FirstOrDefault(x =>
+							x.Error == null &&
+							x.Architecture == ToolChain.Architecture &&
+							x.ReleaseChannel == ToolChain.ReleaseChannel &&
+							x.Version == ToolChain.Version &&
+							x.HostArchitecture != MSVCHostUnrealArch);
+
+						if (CrossToolChain != null)
+						{
+							DirectoryReference BinDir = DirectoryReference.Combine(ToolChain.BaseDir, "bin", Compiler.IsMSVC() ? UnrealArchToMSVCHostDirectoryName(ToolChain.HostArchitecture) : String.Empty);
+							DirectoryReference CrossBinDir = DirectoryReference.Combine(CrossToolChain.BaseDir, "bin", Compiler.IsMSVC() ? UnrealArchToMSVCHostDirectoryName(CrossToolChain.HostArchitecture) : String.Empty);
+							foreach (FileReference Binary in DirectoryReference.EnumerateFiles(BinDir, "*.exe", SearchOption.AllDirectories))
+							{
+								FileReference CrossBinary = FileReference.Combine(CrossBinDir, Path.GetRelativePath(BinDir.FullName, Binary.FullName));
+								if (FileReference.Exists(CrossBinary))
+								{
+									EpicGames.UBA.Utils.RegisterCrossArchitecturePath(Binary.FullName, CrossBinary.FullName);
+								}
+							}
+						}
+
+					}
 				}
 				return ToolChains;
 			});
@@ -947,11 +972,16 @@ namespace UnrealBuildTool
 				foreach (DirectoryReference ToolChainDir in DirectoryReference.EnumerateDirectories(BaseDir))
 				{
 					VersionNumber? Version;
-					if (IsValidToolChainDirMSVC(ToolChainDir, out Version))
+					if (IsValidToolChainDirMSVC(ToolChainDir, UnrealArch.X64, out Version))
 					{
 						DirectoryReference? RedistDir = FindVisualStudioRedistForToolChain(ToolChainDir, OptionalRedistDir, Version);
+						AddVisualCppToolChain(Version, UnrealArch.X64, ReleaseChannel, ToolChainDir, RedistDir, ToolChains, IsAutoSdk, Logger);
+					}
 
-						AddVisualCppToolChain(Version, ReleaseChannel, ToolChainDir, RedistDir, ToolChains, IsAutoSdk, Logger);
+					if (IsValidToolChainDirMSVC(ToolChainDir, UnrealArch.Arm64, out Version))
+					{
+						DirectoryReference? RedistDir = FindVisualStudioRedistForToolChain(ToolChainDir, OptionalRedistDir, Version);
+						AddVisualCppToolChain(Version, UnrealArch.Arm64, ReleaseChannel, ToolChainDir, RedistDir, ToolChains, IsAutoSdk, Logger);
 					}
 				}
 			}
@@ -997,16 +1027,15 @@ namespace UnrealBuildTool
 		/// Adds a Visual C++ toolchain to a list of installations
 		/// </summary>
 		/// <param name="Version"></param>
+		/// <param name="HostArchitecture"></param>
 		/// <param name="ReleaseChannel"></param>
 		/// <param name="ToolChainDir"></param>
 		/// <param name="RedistDir"></param>
 		/// <param name="ToolChains"></param>
 		/// <param name="IsAutoSdk"></param>
 		/// <param name="Logger"></param>
-		static void AddVisualCppToolChain(VersionNumber Version, WindowsCompilerChannel ReleaseChannel, DirectoryReference ToolChainDir, DirectoryReference? RedistDir, List<ToolChainInstallation> ToolChains, bool IsAutoSdk, ILogger Logger)
+		static void AddVisualCppToolChain(VersionNumber Version, UnrealArch HostArchitecture, WindowsCompilerChannel ReleaseChannel, DirectoryReference ToolChainDir, DirectoryReference? RedistDir, List<ToolChainInstallation> ToolChains, bool IsAutoSdk, ILogger Logger)
 		{
-			bool Is64Bit = Has64BitToolChain(ToolChainDir);
-
 			VersionNumber? Family;
 			if (!VersionNumber.TryParse(ToolChainDir.GetDirectoryName(), out Family))
 			{
@@ -1027,18 +1056,21 @@ namespace UnrealBuildTool
 				Error = $"UnrealBuildTool has banned the MSVC {Banned} toolchains due to compiler issues. Please install a different toolchain such as {PreferredVisualCppVersions.Select(x => x.Min).Max()} by opening the generated solution and installing recommended components or from the Visual Studio installer.";
 			}
 
-			Logger.LogDebug("Found Visual Studio toolchain: {ToolChainDir} (Family={Family}, FamilyRank={FamilyRank}, Version={Version}, Is64Bit={Is64Bit}, ReleaseChannel={ReleaseChannel}, Architecture={Arch}, Error={Error}, Redist={RedistDir})", ToolChainDir, Family, FamilyRank, Version, Is64Bit, ReleaseChannel, UnrealArch.X64.ToString(), Error != null, RedistDir);
-			ToolChains.Add(new ToolChainInstallation(Family, FamilyRank, Version, Is64Bit, ReleaseChannel, UnrealArch.X64, Error, ToolChainDir, RedistDir, IsAutoSdk));
+			if (HasX64ToolChain(ToolChainDir, HostArchitecture))
+			{
+				Logger.LogDebug("Found Visual Studio toolchain: {ToolChainDir} (Family={Family}, FamilyRank={FamilyRank}, Version={Version}, HostArchitecture={HostArchitecture}, ReleaseChannel={ReleaseChannel}, Architecture={Arch}, Error={Error}, Redist={RedistDir})", ToolChainDir, Family, FamilyRank, Version, HostArchitecture, ReleaseChannel, UnrealArch.X64.ToString(), Error != null, RedistDir);
+				ToolChains.Add(new ToolChainInstallation(Family, FamilyRank, Version, HostArchitecture, ReleaseChannel, UnrealArch.X64, Error, ToolChainDir, RedistDir, IsAutoSdk));
+			}
 
-			if (HasArm64ToolChain(ToolChainDir))
+			if (HasArm64ToolChain(ToolChainDir, HostArchitecture))
 			{
-				Logger.LogDebug("Found Visual Studio toolchain: {ToolChainDir} (Family={Family}, FamilyRank={FamilyRank}, Version={Version}, Is64Bit={Is64Bit}, ReleaseChannel={ReleaseChannel}, Architecture={Arch}, Error={Error}, Redist={RedistDir})", ToolChainDir, Family, FamilyRank, Version, Is64Bit, ReleaseChannel, UnrealArch.Arm64.ToString(), Error != null, RedistDir);
-				ToolChains.Add(new ToolChainInstallation(Family, FamilyRank, Version, Is64Bit, ReleaseChannel, UnrealArch.Arm64, Error, ToolChainDir, RedistDir, IsAutoSdk));
+				Logger.LogDebug("Found Visual Studio toolchain: {ToolChainDir} (Family={Family}, FamilyRank={FamilyRank}, Version={Version}, HostArchitecture={HostArchitecture}, ReleaseChannel={ReleaseChannel}, Architecture={Arch}, Error={Error}, Redist={RedistDir})", ToolChainDir, Family, FamilyRank, Version, HostArchitecture, ReleaseChannel, UnrealArch.Arm64.ToString(), Error != null, RedistDir);
+				ToolChains.Add(new ToolChainInstallation(Family, FamilyRank, Version, HostArchitecture, ReleaseChannel, UnrealArch.Arm64, Error, ToolChainDir, RedistDir, IsAutoSdk));
 
 				if (HasArm64ECToolChain(ToolChainDir))
 				{
-					Logger.LogDebug("Found Visual Studio toolchain: {ToolChainDir} (Family={Family}, FamilyRank={FamilyRank}, Version={Version}, Is64Bit={Is64Bit}, ReleaseChannel={ReleaseChannel}, Architecture={Arch}, Error={Error}, Redist={RedistDir})", ToolChainDir, Family, FamilyRank, Version, Is64Bit, ReleaseChannel, UnrealArch.Arm64ec.ToString(), Error != null, RedistDir);
-					ToolChains.Add(new ToolChainInstallation(Family, FamilyRank, Version, Is64Bit, ReleaseChannel, UnrealArch.Arm64ec, Error, ToolChainDir, RedistDir, IsAutoSdk));
+					Logger.LogDebug("Found Visual Studio toolchain: {ToolChainDir} (Family={Family}, FamilyRank={FamilyRank}, Version={Version}, HostArchitecture={HostArchitecture}, ReleaseChannel={ReleaseChannel}, Architecture={Arch}, Error={Error}, Redist={RedistDir})", ToolChainDir, Family, FamilyRank, Version, HostArchitecture, ReleaseChannel, UnrealArch.Arm64ec.ToString(), Error != null, RedistDir);
+					ToolChains.Add(new ToolChainInstallation(Family, FamilyRank, Version, HostArchitecture, ReleaseChannel, UnrealArch.Arm64ec, Error, ToolChainDir, RedistDir, IsAutoSdk));
 				}
 			}
 		}
@@ -1062,7 +1094,6 @@ namespace UnrealBuildTool
 				VersionNumber Family = new VersionNumber(VersionInfo.FileMajorPart);
 
 				int Rank = PreferredClangVersions.TakeWhile(x => !x.Contains(Version)).Count();
-				bool Is64Bit = Is64BitExecutable(CompilerFile);
 
 				string? Error = null;
 				if (Version < MinimumClangVersion)
@@ -1070,10 +1101,26 @@ namespace UnrealBuildTool
 					Error = $"UnrealBuildTool requires at minimum the Clang {MinimumClangVersion} toolchain. Please install a later toolchain such as {PreferredClangVersions.Select(x => x.Min).Max()} from LLVM.";
 				}
 
-				Logger.LogDebug("Found Clang toolchain: {ToolChainDir} (Version={Version}, Is64Bit={Is64Bit}, Rank={Rank}, Error={Error})", ToolChainDir, Version, Is64Bit, Rank, Error != null);
-				ToolChains.Add(new ToolChainInstallation(Family, Rank, Version, Is64Bit, WindowsCompilerChannel.Any, UnrealArch.X64, Error, ToolChainDir, null, IsAutoSdk));
-				ToolChains.Add(new ToolChainInstallation(Family, Rank, Version, Is64Bit, WindowsCompilerChannel.Any, UnrealArch.Arm64, Error, ToolChainDir, null, IsAutoSdk));
-				ToolChains.Add(new ToolChainInstallation(Family, Rank, Version, Is64Bit, WindowsCompilerChannel.Any, UnrealArch.Arm64ec, Error, ToolChainDir, null, IsAutoSdk));
+				UnrealArch HostArchitecture;
+				if (IsX64Executable(CompilerFile))
+				{
+					HostArchitecture = UnrealArch.X64;
+				}
+				else if (IsArm64Executable(CompilerFile))
+				{
+					HostArchitecture = UnrealArch.Arm64;
+				}
+				else
+				{
+					// Unknown binary architecture
+					HostArchitecture = UnrealArch.Deprecated;
+					Error = $"Unable to determine architecture of clang binary.";
+				}
+
+				Logger.LogDebug("Found Clang toolchain: {ToolChainDir} (Version={Version}, HostArchitecture={HostArchitecture}, Rank={Rank}, Error={Error})", ToolChainDir, Version, HostArchitecture, Rank, Error != null);
+				ToolChains.Add(new ToolChainInstallation(Family, Rank, Version, HostArchitecture, WindowsCompilerChannel.Any, UnrealArch.X64, Error, ToolChainDir, null, IsAutoSdk));
+				ToolChains.Add(new ToolChainInstallation(Family, Rank, Version, HostArchitecture, WindowsCompilerChannel.Any, UnrealArch.Arm64, Error, ToolChainDir, null, IsAutoSdk));
+				ToolChains.Add(new ToolChainInstallation(Family, Rank, Version, HostArchitecture, WindowsCompilerChannel.Any, UnrealArch.Arm64ec, Error, ToolChainDir, null, IsAutoSdk));
 			}
 		}
 
@@ -1100,7 +1147,7 @@ namespace UnrealBuildTool
 				}
 
 				int Rank = PreferredIntelOneApiVersions.TakeWhile(x => !x.Contains(Version)).Count();
-				bool Is64Bit = Is64BitExecutable(CompilerFile);
+				bool Is64Bit = IsX64Executable(CompilerFile);
 
 				string? Error = null;
 				if (Version < MinimumIntelOneApiVersion)
@@ -1108,10 +1155,26 @@ namespace UnrealBuildTool
 					Error = $"UnrealBuildTool requires at minimum the Intel OneAPI {MinimumIntelOneApiVersion} toolchain. Please install a later toolchain such as {PreferredIntelOneApiVersions.Select(x => x.Min).Max()} from Intel.";
 				}
 
-				Logger.LogDebug("Found Intel OneAPI toolchain: {ToolChainDir} (Version={Version}, Is64Bit={Is64Bit}, Rank={Rank}, Error={Error})", ToolChainDir, Version, Is64Bit, Rank, Error != null);
-				ToolChains.Add(new ToolChainInstallation(Family, Rank, Version, Is64Bit, WindowsCompilerChannel.Any, UnrealArch.X64, Error, ToolChainDir, null, IsAutoSdk));
-				ToolChains.Add(new ToolChainInstallation(Family, Rank, Version, Is64Bit, WindowsCompilerChannel.Any, UnrealArch.Arm64, Error, ToolChainDir, null, IsAutoSdk));
-				ToolChains.Add(new ToolChainInstallation(Family, Rank, Version, Is64Bit, WindowsCompilerChannel.Any, UnrealArch.Arm64ec, Error, ToolChainDir, null, IsAutoSdk));
+				UnrealArch HostArchitecture;
+				if (IsX64Executable(CompilerFile))
+				{
+					HostArchitecture = UnrealArch.X64;
+				}
+				else if (IsArm64Executable(CompilerFile))
+				{
+					HostArchitecture = UnrealArch.Arm64;
+				}
+				else
+				{
+					// Unknown binary architecture
+					HostArchitecture = UnrealArch.Deprecated;
+					Error = $"Unable to determine architecture of clang binary.";
+				}
+
+				Logger.LogDebug("Found Intel OneAPI toolchain: {ToolChainDir} (Version={Version}, HostArchitecture={HostArchitecture}, Rank={Rank}, Error={Error})", ToolChainDir, Version, Is64Bit, Rank, Error != null);
+				ToolChains.Add(new ToolChainInstallation(Family, Rank, Version, HostArchitecture, WindowsCompilerChannel.Any, UnrealArch.X64, Error, ToolChainDir, null, IsAutoSdk));
+				ToolChains.Add(new ToolChainInstallation(Family, Rank, Version, HostArchitecture, WindowsCompilerChannel.Any, UnrealArch.Arm64, Error, ToolChainDir, null, IsAutoSdk));
+				ToolChains.Add(new ToolChainInstallation(Family, Rank, Version, HostArchitecture, WindowsCompilerChannel.Any, UnrealArch.Arm64ec, Error, ToolChainDir, null, IsAutoSdk));
 			}
 		}
 
@@ -1134,11 +1197,28 @@ namespace UnrealBuildTool
 		}
 
 		/// <summary>
-		/// Test whether an executable is 64-bit
+		/// The compiler host directory name for the running process architecture.
+		/// </summary>
+		public static UnrealArch MSVCHostUnrealArch => RuntimeInformation.ProcessArchitecture == Architecture.Arm64 ? UnrealArch.Arm64 : UnrealArch.X64;
+
+		/// <summary>
+		/// The compiler host directory name for an UnrealArch.
+		/// </summary>
+		/// <param name="HostArchitecture">The binary architecture for the toolchain</param>
+		public static string UnrealArchToMSVCHostDirectoryName(UnrealArch HostArchitecture) => HostArchitecture.bIsX64 ? "Hostx64" : "Hostarm64";
+
+		/// <summary>
+		/// The compiler host directory name for the running process architecture.
+		/// </summary>
+		public static string MSVCHostDirectoryName => UnrealArchToMSVCHostDirectoryName(MSVCHostUnrealArch);
+
+		/// <summary>
+		/// Test whether an executable is a specific machine type
 		/// </summary>
 		/// <param name="File">Executable to test</param>
+		/// <param name="MachineType">The machine type to check</param>
 		/// <returns></returns>
-		static bool Is64BitExecutable(FileReference File)
+		static bool IsExecutableMachineType(FileReference File, ushort MachineType)
 		{
 			using (FileStream Stream = new FileStream(File.FullName, FileMode.Open, FileAccess.Read, FileShare.Read | FileShare.Delete))
 			{
@@ -1168,27 +1248,47 @@ namespace UnrealBuildTool
 					return false;
 				}
 
-				ushort MachineType = BinaryPrimitives.ReadUInt16LittleEndian(PeHeader.AsSpan(4));
-
-				const ushort IMAGE_FILE_MACHINE_AMD64 = 0x8664;
-				return MachineType == IMAGE_FILE_MACHINE_AMD64;
+				return MachineType == BinaryPrimitives.ReadUInt16LittleEndian(PeHeader.AsSpan(4));
 			}
 		}
 
 		/// <summary>
-		/// The compiler host directory name for the running process architecture.
+		/// Test whether an executable is x64
 		/// </summary>
-		public static string MSVCHostDirectoryName => RuntimeInformation.ProcessArchitecture == Architecture.Arm64 ? "Hostarm64" : "Hostx64";
+		/// <param name="File">Executable to test</param>
+		/// <returns></returns>
+		static bool IsX64Executable(FileReference File)
+		{
+			const ushort IMAGE_FILE_MACHINE_AMD64 = 0x8664;
+			return IsExecutableMachineType(File, IMAGE_FILE_MACHINE_AMD64);
+		}
+
+		/// <summary>
+		/// Test whether an executable is arm64
+		/// </summary>
+		/// <param name="File">Executable to test</param>
+		/// <returns></returns>
+		static bool IsArm64Executable(FileReference File)
+		{
+			const ushort IMAGE_FILE_MACHINE_ARM64 = 0xaa64;
+			return IsExecutableMachineType(File, IMAGE_FILE_MACHINE_ARM64);
+		}
 
 		/// <summary>
 		/// Determines if the given path is a valid Visual C++ version number
 		/// </summary>
 		/// <param name="ToolChainDir">The toolchain directory</param>
+		/// <param name="HostArchitecture">The binary architecture for the toolchain</param>
 		/// <param name="Version">The version number for the toolchain</param>
 		/// <returns>True if the path is a valid version</returns>
-		static bool IsValidToolChainDirMSVC(DirectoryReference ToolChainDir, [NotNullWhen(true)] out VersionNumber? Version)
+		static bool IsValidToolChainDirMSVC(DirectoryReference ToolChainDir, UnrealArch HostArchitecture, [NotNullWhen(true)] out VersionNumber? Version)
 		{
-			FileReference CompilerExe = FileReference.Combine(ToolChainDir, "bin", MSVCHostDirectoryName, "x64", "cl.exe");
+			FileReference CompilerExe = FileReference.Combine(ToolChainDir, "bin", UnrealArchToMSVCHostDirectoryName(HostArchitecture), "x64", "cl.exe");
+			if (!FileReference.Exists(CompilerExe))
+			{
+				CompilerExe = FileReference.Combine(ToolChainDir, "bin", UnrealArchToMSVCHostDirectoryName(HostArchitecture), "arm64", "cl.exe");
+			}
+
 			if (!FileReference.Exists(CompilerExe))
 			{
 				Version = null;
@@ -1206,27 +1306,29 @@ namespace UnrealBuildTool
 		}
 
 		/// <summary>
-		/// Checks if the given directory contains a 64-bit toolchain. Used to prefer regular Visual Studio versions over express editions.
+		/// Checks if the given directory contains a x64 toolchain. 
 		/// </summary>
 		/// <param name="ToolChainDir">Directory to check</param>
+		/// <param name="HostArchitecture">The host binary architecture to check</param>
 		/// <returns>True if the given directory contains a 64-bit toolchain</returns>
-		static bool Has64BitToolChain(DirectoryReference ToolChainDir)
+		static bool HasX64ToolChain(DirectoryReference ToolChainDir, UnrealArch HostArchitecture)
 		{
-			return FileReference.Exists(FileReference.Combine(ToolChainDir, "bin", MSVCHostDirectoryName, "x64", "cl.exe"));
+			return FileReference.Exists(FileReference.Combine(ToolChainDir, "bin", UnrealArchToMSVCHostDirectoryName(HostArchitecture), "x64", "cl.exe"));
 		}
 
 		/// <summary>
-		/// Checks if the given directory contains an arm64 toolchain.  Used to require arm64, which is an optional install item, when that is our target architecture.
+		/// Checks if the given directory contains an arm64 toolchain.
 		/// </summary>
 		/// <param name="ToolChainDir">Directory to check</param>
+		/// <param name="HostArchitecture">The host binary architecture to check</param>
 		/// <returns>True if the given directory contains the arm64 toolchain</returns>
-		static bool HasArm64ToolChain(DirectoryReference ToolChainDir)
+		static bool HasArm64ToolChain(DirectoryReference ToolChainDir, UnrealArch HostArchitecture)
 		{
-			return FileReference.Exists(FileReference.Combine(ToolChainDir, "bin", MSVCHostDirectoryName, "arm64", "cl.exe"));
+			return FileReference.Exists(FileReference.Combine(ToolChainDir, "bin", UnrealArchToMSVCHostDirectoryName(HostArchitecture), "arm64", "cl.exe"));
 		}
 
 		/// <summary>
-		/// Checks if the given directory contains a arm64ec toolchain.  Used to require arm64, which is an optional install item, when that is our target architecture.
+		/// Checks if the given directory contains an arm64ec toolchain.
 		/// </summary>
 		/// <param name="ToolChainDir">Directory to check</param>
 		/// <returns>True if the given directory contains the arm64ec toolchain</returns>
@@ -1396,9 +1498,9 @@ namespace UnrealBuildTool
 		public VersionNumber Version { get; }
 
 		/// <summary>
-		/// Whether this is a 64-bit toolchain
+		/// The host architecture of this ToolChainInstallation (multiple ToolChainInstallation instances may be created, one per host architecture).
 		/// </summary>
-		public bool Is64Bit { get; }
+		public UnrealArch HostArchitecture { get; }
 
 		/// <summary>
 		/// The release channel of this toolchain
@@ -1436,19 +1538,19 @@ namespace UnrealBuildTool
 		/// <param name="Family"></param>
 		/// <param name="FamilyRank"></param>
 		/// <param name="Version"></param>
-		/// <param name="Is64Bit"></param>
+		/// <param name="HostArchitecture">The architecture of the binaries in this toolchain</param>
 		/// <param name="ReleaseChannel">The release channel of this toolchain</param>
-		/// <param name="Architecture"></param>
+		/// <param name="Architecture">The target architecture this toolchain can compile</param>
 		/// <param name="Error"></param>
 		/// <param name="BaseDir">Base directory for the toolchain</param>
 		/// <param name="RedistDir">Optional directory for redistributable components (DLLs etc)</param>
 		/// <param name="IsAutoSdk">Whether this toolchain comes from AutoSDK</param>
-		public ToolChainInstallation(VersionNumber Family, int FamilyRank, VersionNumber Version, bool Is64Bit, WindowsCompilerChannel ReleaseChannel, UnrealArch Architecture, string? Error, DirectoryReference BaseDir, DirectoryReference? RedistDir, bool IsAutoSdk)
+		public ToolChainInstallation(VersionNumber Family, int FamilyRank, VersionNumber Version, UnrealArch HostArchitecture, WindowsCompilerChannel ReleaseChannel, UnrealArch Architecture, string? Error, DirectoryReference BaseDir, DirectoryReference? RedistDir, bool IsAutoSdk)
 		{
 			this.Family = Family;
 			this.FamilyRank = FamilyRank;
 			this.Version = Version;
-			this.Is64Bit = Is64Bit;
+			this.HostArchitecture = HostArchitecture;
 			this.ReleaseChannel = ReleaseChannel;
 			this.Architecture = Architecture;
 			this.Error = Error;
diff --git a/Engine/Source/Programs/UnrealBuildTool/Platform/Windows/VCEnvironment.cs b/Engine/Source/Programs/UnrealBuildTool/Platform/Windows/VCEnvironment.cs
--- a/Engine/Source/Programs/UnrealBuildTool/Platform/Windows/VCEnvironment.cs
+++ b/Engine/Source/Programs/UnrealBuildTool/Platform/Windows/VCEnvironment.cs
@@ -3,6 +3,7 @@
 using System;
 using System.Collections.Generic;
 using System.Diagnostics;
+using System.IO;
 using System.Linq;
 using System.Runtime.InteropServices;
 using System.Runtime.Versioning;
@@ -141,6 +142,9 @@ namespace UnrealBuildTool
 			// Get the resource compiler path from the Windows SDK
 			ResourceCompilerPath = GetResourceCompilerToolPath(WindowsSdkDir, WindowsSdkVersion, Logger);
 
+			// Register any executables for cross-architecture support with UBA
+			RegisterCrossArchitectureToolPaths(WindowsSdkDir, WindowsSdkVersion);
+
 			// Get all the system include paths
 			SetupEnvironment(Logger);
 		}
@@ -253,13 +257,24 @@ namespace UnrealBuildTool
 				return ResourceCompilerPath;
 			}
 
-			ResourceCompilerPath = FileReference.Combine(WindowsSdkDir, "bin", BuildHostArch, "rc.exe");
-			if (FileReference.Exists(ResourceCompilerPath))
+			throw new BuildException("Unable to find path to the Windows resource compiler under {0} (version {1})", WindowsSdkDir, WindowsSdkVersion);
+		}
+
+		/// <summary>
+		/// Registers all SDK binaries with UBA for cross-architecture support
+		/// </summary>
+		void RegisterCrossArchitectureToolPaths(DirectoryReference WindowsSdkDir, VersionNumber WindowsSdkVersion)
+		{
+			DirectoryReference HostBinDir = DirectoryReference.Combine(WindowsSdkDir, "bin", WindowsSdkVersion.ToString(), RuntimeInformation.ProcessArchitecture == System.Runtime.InteropServices.Architecture.Arm64 ? "arm64" : "x64");
+			DirectoryReference CrossBinDir = DirectoryReference.Combine(WindowsSdkDir, "bin", WindowsSdkVersion.ToString(), RuntimeInformation.ProcessArchitecture != System.Runtime.InteropServices.Architecture.Arm64 ? "arm64" : "x64");
+			foreach (FileReference Binary in DirectoryReference.EnumerateFiles(HostBinDir, "*.exe", SearchOption.AllDirectories))
 			{
-				return ResourceCompilerPath;
+				FileReference CrossBinary = FileReference.Combine(CrossBinDir, Path.GetRelativePath(HostBinDir.FullName, Binary.FullName));
+				if (FileReference.Exists(CrossBinary))
+				{
+					EpicGames.UBA.Utils.RegisterCrossArchitecturePath(Binary.FullName, CrossBinary.FullName);
+				}
 			}
-
-			throw new BuildException("Unable to find path to the Windows resource compiler under {0} (version {1})", WindowsSdkDir, WindowsSdkVersion);
 		}
 
 		/// <summary>
