UnrealBuildTool: Register cross directories instead of exe paths

#rnx

diff --git a/Engine/Source/Programs/UnrealBuildTool/Platform/Windows/MicrosoftPlatformSDK.cs b/Engine/Source/Programs/UnrealBuildTool/Platform/Windows/MicrosoftPlatformSDK.cs
--- a/Engine/Source/Programs/UnrealBuildTool/Platform/Windows/MicrosoftPlatformSDK.cs
+++ b/Engine/Source/Programs/UnrealBuildTool/Platform/Windows/MicrosoftPlatformSDK.cs
@@ -834,24 +834,35 @@ namespace UnrealBuildTool
 
 					foreach (ToolChainInstallation? ToolChain in ToolChains.Where(x => x.Error == null && x.HostArchitecture == MSVCHostUnrealArch))
 					{
-						ToolChainInstallation? CrossToolChain = ToolChains.FirstOrDefault(x =>
-							x.Error == null &&
-							x.Architecture == ToolChain.Architecture &&
-							x.ReleaseChannel == ToolChain.ReleaseChannel &&
-							x.Version == ToolChain.Version &&
-							x.HostArchitecture != MSVCHostUnrealArch);
+						ToolChainInstallation? CrossToolChain;
+						if (Compiler.IsMSVC())
+						{
+							// MSVC has the same base dir. Version may not match between Hostx64 and Hostarm64.
+							CrossToolChain = ToolChains.FirstOrDefault(x =>
+								x.Error == null &&
+								x.Architecture == ToolChain.Architecture &&
+								x.ReleaseChannel == ToolChain.ReleaseChannel &&
+								x.BaseDir == ToolChain.BaseDir &&
+								x.HostArchitecture != MSVCHostUnrealArch);
+						}
+						else
+						{
+							// Clang the version should match.
+							CrossToolChain = ToolChains.FirstOrDefault(x =>
+								x.Error == null &&
+								x.Architecture == ToolChain.Architecture &&
+								x.ReleaseChannel == ToolChain.ReleaseChannel &&
+								x.Version == ToolChain.Version &&
+								x.HostArchitecture != MSVCHostUnrealArch);
+						}
 
 						if (CrossToolChain != null)
 						{
-							DirectoryReference BinDir = DirectoryReference.Combine(ToolChain.BaseDir, "bin", Compiler.IsMSVC() ? UnrealArchToMSVCHostDirectoryName(ToolChain.HostArchitecture) : String.Empty);
+							DirectoryReference HostBinDir = DirectoryReference.Combine(ToolChain.BaseDir, "bin", Compiler.IsMSVC() ? UnrealArchToMSVCHostDirectoryName(ToolChain.HostArchitecture) : String.Empty);
 							DirectoryReference CrossBinDir = DirectoryReference.Combine(CrossToolChain.BaseDir, "bin", Compiler.IsMSVC() ? UnrealArchToMSVCHostDirectoryName(CrossToolChain.HostArchitecture) : String.Empty);
-							foreach (FileReference Binary in DirectoryReference.EnumerateFiles(BinDir, "*.exe", SearchOption.AllDirectories))
+							if (FileReference.Exists(HostBinDir) && FileReference.Exists(CrossBinDir))
 							{
-								FileReference CrossBinary = FileReference.Combine(CrossBinDir, Path.GetRelativePath(BinDir.FullName, Binary.FullName));
-								if (FileReference.Exists(CrossBinary))
-								{
-									EpicGames.UBA.Utils.RegisterCrossArchitecturePath(Binary.FullName, CrossBinary.FullName);
-								}
+								EpicGames.UBA.Utils.RegisterCrossArchitecturePath(HostBinDir.FullName, CrossBinDir.FullName);
 							}
 						}
 
diff --git a/Engine/Source/Programs/UnrealBuildTool/Platform/Windows/VCEnvironment.cs b/Engine/Source/Programs/UnrealBuildTool/Platform/Windows/VCEnvironment.cs
--- a/Engine/Source/Programs/UnrealBuildTool/Platform/Windows/VCEnvironment.cs
+++ b/Engine/Source/Programs/UnrealBuildTool/Platform/Windows/VCEnvironment.cs
@@ -267,13 +267,9 @@ namespace UnrealBuildTool
 		{
 			DirectoryReference HostBinDir = DirectoryReference.Combine(WindowsSdkDir, "bin", WindowsSdkVersion.ToString(), RuntimeInformation.ProcessArchitecture == System.Runtime.InteropServices.Architecture.Arm64 ? "arm64" : "x64");
 			DirectoryReference CrossBinDir = DirectoryReference.Combine(WindowsSdkDir, "bin", WindowsSdkVersion.ToString(), RuntimeInformation.ProcessArchitecture != System.Runtime.InteropServices.Architecture.Arm64 ? "arm64" : "x64");
-			foreach (FileReference Binary in DirectoryReference.EnumerateFiles(HostBinDir, "*.exe", SearchOption.AllDirectories))
+			if (FileReference.Exists(HostBinDir) && FileReference.Exists(CrossBinDir))
 			{
-				FileReference CrossBinary = FileReference.Combine(CrossBinDir, Path.GetRelativePath(HostBinDir.FullName, Binary.FullName));
-				if (FileReference.Exists(CrossBinary))
-				{
-					EpicGames.UBA.Utils.RegisterCrossArchitecturePath(Binary.FullName, CrossBinary.FullName);
-				}
+				EpicGames.UBA.Utils.RegisterCrossArchitecturePath(HostBinDir.FullName, CrossBinDir.FullName);
 			}
 		}
 
