UnrealBuildTool: Action property to force deleting outputs on error, and to force warnings as error

#rnx

diff --git a/Engine/Source/Programs/UnrealBuildTool/Artifacts/Artifacts.cs b/Engine/Source/Programs/UnrealBuildTool/Artifacts/Artifacts.cs
--- a/Engine/Source/Programs/UnrealBuildTool/Artifacts/Artifacts.cs
+++ b/Engine/Source/Programs/UnrealBuildTool/Artifacts/Artifacts.cs
@@ -235,7 +235,8 @@ namespace UnrealBuildTool.Artifacts
 	/// </summary>
 	/// <param name="Success">Is set to true if succeeded in fetching artifacts</param>
 	/// <param name="LogLines">Contains log lines if any</param>
-	public readonly record struct ActionArtifactResult(bool Success, List<string> LogLines);
+	/// <param name="ExitCode">Process exit code (expected to be 0 in almost all cases as errors aren't cached)</param>
+	public readonly record struct ActionArtifactResult(bool Success, List<string> LogLines, int ExitCode = 0);
 
 	/// <summary>
 	/// Interface for action specific support of artifacts.
diff --git a/Engine/Source/Programs/UnrealBuildTool/Executors/ImmediateActionQueue.cs b/Engine/Source/Programs/UnrealBuildTool/Executors/ImmediateActionQueue.cs
--- a/Engine/Source/Programs/UnrealBuildTool/Executors/ImmediateActionQueue.cs
+++ b/Engine/Source/Programs/UnrealBuildTool/Executors/ImmediateActionQueue.cs
@@ -433,9 +433,9 @@ namespace UnrealBuildTool
 						ActionArtifactResult result = await _actionArtifactCache!.CompleteActionFromCacheAsync(action, CancellationToken);
 						if (result.Success)
 						{
-							ExecuteResults results = new(result.LogLines, 0, TimeSpan.Zero, TimeSpan.Zero, "[Cache]");
+							ExecuteResults results = new(result.LogLines, result.ExitCode, TimeSpan.Zero, TimeSpan.Zero, "[Cache]");
 							Interlocked.Increment(ref _cacheHitActions);
-							OnActionCompleted(action, true, results);
+							OnActionCompleted(action, results.ExitCode == 0, results);
 							OnArtifactsRead?.Invoke(action);
 						}
 						else
@@ -905,6 +905,25 @@ namespace UnrealBuildTool
 		/// <param name="results">Optional execution results</param>
 		public void OnActionCompleted(LinkedAction action, bool success, ExecuteResults? results)
 		{
+			// Delete produced items on error if requested
+			if (action.bDeleteProducedItemsOnError && (results?.ExitCode ?? 0) != 0)
+			{
+				foreach (FileItem output in action.ProducedItems.Distinct().Where(x => FileReference.Exists(x.Location)))
+				{
+					FileReference.Delete(output.Location);
+				}
+			}
+
+			// On success invalidate the file info for all the produced items.  This needs to be
+			// done so the artifact system sees any compiled files.
+			if (success)
+			{
+				foreach (FileItem output in action.ProducedItems)
+				{
+					output.ResetCachedInfo();
+				}
+			}
+
 			SetActionState(action, success ? ActionStatus.Finished : ActionStatus.Error, results);
 		}
 
@@ -920,16 +939,6 @@ namespace UnrealBuildTool
 			int actionIndex = action.SortIndex;
 			int completedActions = 0;
 
-			// If we are finished, then invalidate the file info for all the produced items.  This needs to be
-			// done so the artifact system sees any compiled files.
-			if (status == ActionStatus.Finished)
-			{
-				foreach (FileItem output in action.ProducedItems)
-				{
-					output.ResetCachedInfo();
-				}
-			}
-
 			// Update the actions data
 			lock (Actions)
 			{
@@ -946,7 +955,7 @@ namespace UnrealBuildTool
 				{
 					wasArtifactCheck = true;
 					Actions[actionIndex].Phase = ActionPhase.Compile;
-					if (status != ActionStatus.Finished)
+					if (status != ActionStatus.Finished && (results?.ExitCode ?? 0) == 0)
 					{
 						status = ActionStatus.Queued;
 					}
diff --git a/Engine/Source/Programs/UnrealBuildTool/Executors/ParallelExecutor.cs b/Engine/Source/Programs/UnrealBuildTool/Executors/ParallelExecutor.cs
--- a/Engine/Source/Programs/UnrealBuildTool/Executors/ParallelExecutor.cs
+++ b/Engine/Source/Programs/UnrealBuildTool/Executors/ParallelExecutor.cs
@@ -248,6 +248,13 @@ namespace UnrealBuildTool
 			int ExitCode = Process.ExitCode;
 			TimeSpan ProcessorTime = Process.TotalProcessorTime;
 			TimeSpan ExecutionTime = Process.ExitTime - Process.StartTime;
+
+			if (ExitCode == 0 && Action.bForceWarningsAsError && LogLines.Any(x => x.Contains("): warning: ", StringComparison.OrdinalIgnoreCase)))
+			{
+				ExitCode = 1;
+				LogLines = LogLines.Select(x => x.Replace("): warning: ", "): error: ", StringComparison.OrdinalIgnoreCase)).ToList();
+			}
+
 			return new ExecuteResults(LogLines, ExitCode, ExecutionTime, ProcessorTime, AdditionalDescription);
 		}
 	}
diff --git a/Engine/Source/Programs/UnrealBuildTool/Executors/UnrealBuildAccelerator/UBAExecutor.cs b/Engine/Source/Programs/UnrealBuildTool/Executors/UnrealBuildAccelerator/UBAExecutor.cs
--- a/Engine/Source/Programs/UnrealBuildTool/Executors/UnrealBuildAccelerator/UBAExecutor.cs
+++ b/Engine/Source/Programs/UnrealBuildTool/Executors/UnrealBuildAccelerator/UBAExecutor.cs
@@ -262,8 +262,15 @@ namespace UnrealBuildTool
 					using (IRootPaths rootPaths = _executor.GetActionRootPaths(action))
 					{
 						FetchFromCacheResult result = _executor._cacheClient!.FetchFromCache(rootPaths, bucket, startInfo);
+						int exitCode = 0;
 						RemoveLogLineSpam(result.LogLines);
-						return new ActionArtifactResult(result.Success, result.LogLines);
+						if (result.Success && action.bForceWarningsAsError && result.LogLines.Any(x => x.Contains("): warning: ")))
+						{
+							result = new(true, result.LogLines.Select(x => x.Replace("): warning: ", "): error: ", StringComparison.OrdinalIgnoreCase)).ToList());
+							exitCode = 1;
+							++_executor._errorCount;
+						}
+						return new ActionArtifactResult(result.Success, result.LogLines, exitCode);
 					}
 				}, cancellationToken, TaskCreationOptions.LongRunning | TaskCreationOptions.PreferFairness, TaskScheduler.Default);
 			}
@@ -1002,6 +1009,12 @@ namespace UnrealBuildTool
 				_session!.SetCustomCasKeyFromTrackedInputs(pchItem.FullName, action.WorkingDirectory.FullName, process);
 			}
 
+			if (results.ExitCode == 0 && action.bForceWarningsAsError && results.LogLines.Any(x => x.Contains("): warning: ")))
+			{
+				results = new(results.LogLines.Select(x => x.Replace("): warning: ", "): error: ", StringComparison.OrdinalIgnoreCase)).ToList(), 1, results.ExecutionTime, results.ProcessorTime, results.AdditionalDescription);
+				success = false;
+			}
+
 			queue.OnActionCompleted(action, success, results);
 
 			if (!success)
diff --git a/Engine/Source/Programs/UnrealBuildTool/Platform/Windows/VCCompileAction.cs b/Engine/Source/Programs/UnrealBuildTool/Platform/Windows/VCCompileAction.cs
--- a/Engine/Source/Programs/UnrealBuildTool/Platform/Windows/VCCompileAction.cs
+++ b/Engine/Source/Programs/UnrealBuildTool/Platform/Windows/VCCompileAction.cs
@@ -169,6 +169,8 @@ namespace UnrealBuildTool
 		/// </summary>
 		public bool bIsAnalyzing { get; set; }
 
+		public bool bWarningsAsErrors { get; set; }
+
 		#region Public IAction implementation
 
 		/// <summary>
@@ -227,6 +229,8 @@ namespace UnrealBuildTool
 			}
 		}
 		bool IExternalAction.bIsGCCCompiler => false;
+		bool IExternalAction.bDeleteProducedItemsOnError => CompilerType.IsClang() && bIsAnalyzing && bWarningsAsErrors;
+		bool IExternalAction.bForceWarningsAsError => CompilerType.IsClang() && bIsAnalyzing && bWarningsAsErrors;
 		bool IExternalAction.bProducesImportLibrary => false;
 		string IExternalAction.StatusDescription => (SourceFile == null) ? "Compiling" : SourceFile.Location.GetFileName();
 		bool IExternalAction.bShouldOutputStatusDescription => CompilerType.IsClang();
@@ -387,6 +391,7 @@ namespace UnrealBuildTool
 			bCanExecuteRemotelyWithXGE = InAction.bCanExecuteRemotelyWithXGE;
 			Architecture = InAction.Architecture;
 			bIsAnalyzing = InAction.bIsAnalyzing;
+			bWarningsAsErrors = InAction.bWarningsAsErrors;
 			Weight = InAction.Weight;
 
 			AdditionalPrerequisiteItems = new List<FileItem>(InAction.AdditionalPrerequisiteItems);
@@ -431,6 +436,7 @@ namespace UnrealBuildTool
 			bCanExecuteInUBA = Reader.ReadBool();
 			Architecture = UnrealArch.Parse(Reader.ReadString()!);
 			bIsAnalyzing = Reader.ReadBool();
+			bWarningsAsErrors = Reader.ReadBool();
 			Weight = Reader.ReadDouble();
 
 			AdditionalPrerequisiteItems = Reader.ReadList(() => Reader.ReadFileItem())!;
@@ -472,6 +478,7 @@ namespace UnrealBuildTool
 			Writer.WriteBool(bCanExecuteInUBA);
 			Writer.WriteString(Architecture.ToString());
 			Writer.WriteBool(bIsAnalyzing);
+			Writer.WriteBool(bWarningsAsErrors);
 			Writer.WriteDouble(Weight);
 
 			Writer.WriteList(AdditionalPrerequisiteItems, Item => Writer.WriteFileItem(Item));
diff --git a/Engine/Source/Programs/UnrealBuildTool/Platform/Windows/VCToolChain.cs b/Engine/Source/Programs/UnrealBuildTool/Platform/Windows/VCToolChain.cs
--- a/Engine/Source/Programs/UnrealBuildTool/Platform/Windows/VCToolChain.cs
+++ b/Engine/Source/Programs/UnrealBuildTool/Platform/Windows/VCToolChain.cs
@@ -1739,6 +1739,7 @@ namespace UnrealBuildTool
 			AppendCLArguments_Global(CompileEnvironment, BaseCompileAction.Arguments);
 
 			BaseCompileAction.bIsAnalyzing = Target.StaticAnalyzer != StaticAnalyzer.None && !CompileEnvironment.bDisableStaticAnalysis && !(Target.WindowsPlatform.Compiler.IsClang() && CompileEnvironment.PrecompiledHeaderAction == PrecompiledHeaderAction.Create);
+			BaseCompileAction.bWarningsAsErrors = CompileEnvironment.bWarningsAsErrors;
 
 			// Add include paths to the argument list.
 			BaseCompileAction.IncludePaths.AddRange(CompileEnvironment.UserIncludePaths);
@@ -2128,6 +2129,7 @@ namespace UnrealBuildTool
 				}
 
 				CompileAction.bIsAnalyzing = Target.StaticAnalyzer != StaticAnalyzer.None && !CompileEnvironment.bDisableStaticAnalysis && !(Target.WindowsPlatform.Compiler.IsClang() && CompileEnvironment.PrecompiledHeaderAction == PrecompiledHeaderAction.Create);
+				CompileAction.bWarningsAsErrors = CompileEnvironment.bWarningsAsErrors;
 
 				// Update the output
 				Graph.AddAction(CompileAction);
diff --git a/Engine/Source/Programs/UnrealBuildTool/System/Action.cs b/Engine/Source/Programs/UnrealBuildTool/System/Action.cs
--- a/Engine/Source/Programs/UnrealBuildTool/System/Action.cs
+++ b/Engine/Source/Programs/UnrealBuildTool/System/Action.cs
@@ -160,6 +160,19 @@ namespace UnrealBuildTool
 		/// </summary>
 		bool bIsGCCCompiler { get; }
 
+		/// <summary>
+		/// True if the action executor should delete output files on an error.
+		/// May not work for XGE, SN-DBS, and FASTBuild executors.
+		/// </summary>
+		bool bDeleteProducedItemsOnError { get; }
+
+		/// <summary>
+		/// True if the action executor should force warning output to fail the action when the exit code would otherwise be zero.
+		/// Used if the command to be run does not support a warnings as error flag.
+		/// May not work for XGE, SN-DBS, and FASTBuild executors.
+		/// </summary>
+		bool bForceWarningsAsError { get; }
+
 		/// <summary>
 		/// Whether we should log this action, whether executed locally or remotely.  This is useful for actions that take time
 		/// but invoke tools without any console output.
@@ -294,6 +307,19 @@ namespace UnrealBuildTool
 		/// </summary>
 		public bool bIsGCCCompiler { get; set; } = false;
 
+		/// <summary>
+		/// True if the action executor should delete output files on an error.
+		/// May not work for XGE, SN-DBS, and FASTBuild executors.
+		/// </summary>
+		public bool bDeleteProducedItemsOnError { get; set; } = false;
+
+		/// <summary>
+		/// True if the action executor should force warning output to fail the action when the exit code would otherwise be zero.
+		/// Used if the command to be run does not support a warnings as error flag.
+		/// May not work for XGE, SN-DBS, and FASTBuild executors.
+		/// </summary>
+		public bool bForceWarningsAsError { get; set; } = false;
+
 		/// <summary>
 		/// Whether we should log this action, whether executed locally or remotely.  This is useful for actions that take time
 		/// but invoke tools without any console output.
@@ -353,6 +379,8 @@ namespace UnrealBuildTool
 			bCanExecuteRemotelyWithXGE = InOther.bCanExecuteRemotelyWithXGE;
 			bCanExecuteInUBA = InOther.bCanExecuteInUBA;
 			bIsGCCCompiler = InOther.bIsGCCCompiler;
+			bDeleteProducedItemsOnError = InOther.bDeleteProducedItemsOnError;
+			bForceWarningsAsError = InOther.bForceWarningsAsError;
 			bShouldOutputStatusDescription = InOther.bShouldOutputStatusDescription;
 			bShouldOutputLog = InOther.bShouldOutputLog;
 			bProducesImportLibrary = InOther.bProducesImportLibrary;
@@ -376,6 +404,8 @@ namespace UnrealBuildTool
 			bCanExecuteRemotelyWithXGE = Reader.ReadBool();
 			bCanExecuteInUBA = Reader.ReadBool();
 			bIsGCCCompiler = Reader.ReadBool();
+			bDeleteProducedItemsOnError = Reader.ReadBool();
+			bForceWarningsAsError = Reader.ReadBool();
 			bShouldOutputStatusDescription = Reader.ReadBool();
 			bShouldOutputLog = Reader.ReadBool();
 			bProducesImportLibrary = Reader.ReadBool();
@@ -408,6 +438,8 @@ namespace UnrealBuildTool
 			Writer.WriteBool(bCanExecuteRemotelyWithXGE);
 			Writer.WriteBool(bCanExecuteInUBA);
 			Writer.WriteBool(bIsGCCCompiler);
+			Writer.WriteBool(bDeleteProducedItemsOnError);
+			Writer.WriteBool(bForceWarningsAsError);
 			Writer.WriteBool(bShouldOutputStatusDescription);
 			Writer.WriteBool(bShouldOutputLog);
 			Writer.WriteBool(bProducesImportLibrary);
@@ -489,6 +521,16 @@ namespace UnrealBuildTool
 				Action.bIsGCCCompiler = bIsGCCCompiler;
 			}
 
+			if (Object.TryGetBoolField("bDeleteProducedItemsOnError", out bool bDeleteProducedItemsOnError))
+			{
+				Action.bDeleteProducedItemsOnError = bDeleteProducedItemsOnError;
+			}
+
+			if (Object.TryGetBoolField("bForceWarningsAsError", out bool bForceWarningsAsError))
+			{
+				Action.bForceWarningsAsError = bForceWarningsAsError;
+			}
+
 			if (Object.TryGetBoolField("bShouldOutputStatusDescription", out bool bShouldOutputStatusDescription))
 			{
 				Action.bShouldOutputStatusDescription = bShouldOutputStatusDescription;
@@ -596,6 +638,8 @@ namespace UnrealBuildTool
 			Writer.WriteValue("bCanExecuteRemotelyWithXGE", Action.bCanExecuteRemotelyWithXGE);
 			Writer.WriteValue("bCanExecuteInUBA", Action.bCanExecuteInUBA);
 			Writer.WriteValue("bIsGCCCompiler", Action.bIsGCCCompiler);
+			Writer.WriteValue("bDeleteProducedItemsOnError", Action.bDeleteProducedItemsOnError);
+			Writer.WriteValue("bForceWarningsAsError", Action.bForceWarningsAsError);			
 			Writer.WriteValue("bShouldOutputStatusDescription", Action.bShouldOutputStatusDescription);
 			Writer.WriteValue("bShouldOutputLog", Action.bShouldOutputLog);
 			Writer.WriteValue("bProducesImportLibrary", Action.bProducesImportLibrary);
@@ -727,6 +771,8 @@ namespace UnrealBuildTool
 		public bool bCanExecuteRemotelyWithXGE => Inner.bCanExecuteRemotelyWithXGE;
 		public bool bCanExecuteInUBA => Inner.bCanExecuteInUBA;
 		public bool bIsGCCCompiler => Inner.bIsGCCCompiler;
+		public bool bDeleteProducedItemsOnError => Inner.bDeleteProducedItemsOnError;
+		public bool bForceWarningsAsError => Inner.bForceWarningsAsError;
 		public bool bShouldOutputStatusDescription => Inner.bShouldOutputStatusDescription;
 		public bool bShouldOutputLog => Inner.bShouldOutputLog;
 		public bool bProducesImportLibrary => Inner.bProducesImportLibrary;
diff --git a/Engine/Source/Programs/UnrealBuildTool/ToolChain/ClangToolChain.cs b/Engine/Source/Programs/UnrealBuildTool/ToolChain/ClangToolChain.cs
--- a/Engine/Source/Programs/UnrealBuildTool/ToolChain/ClangToolChain.cs
+++ b/Engine/Source/Programs/UnrealBuildTool/ToolChain/ClangToolChain.cs
@@ -1302,6 +1302,12 @@ namespace UnrealBuildTool
 				CompileAction.ArtifactMode |= ArtifactMode.AbsolutePath; // Unfortunately we require matching absolute paths for pch to be cached
 			}
 
+			if (IsAnalyzing(CompileEnvironment) && CompileEnvironment.bWarningsAsErrors)
+			{
+				CompileAction.bForceWarningsAsError = true;
+				CompileAction.bDeleteProducedItemsOnError = true;
+			}
+
 			// Two-pass compile where the preprocessor is run first to output the dependency list
 			if (PreprocessDepends)
 			{
