[UBT]
* Added bCanExecuteInUBACrossArchitecture to Action and hooked it up with UBA scheduler... now actions can decide wether they support cross architecture or not
* Disabled cross architecture support for compile actions using msvc pch since msvc pch is a memory dump and does not work cross architecture

diff --git a/Engine/Source/Programs/UnrealBuildTool/Executors/UnrealBuildAccelerator/UBAExecutor.cs b/Engine/Source/Programs/UnrealBuildTool/Executors/UnrealBuildAccelerator/UBAExecutor.cs
--- a/Engine/Source/Programs/UnrealBuildTool/Executors/UnrealBuildAccelerator/UBAExecutor.cs
+++ b/Engine/Source/Programs/UnrealBuildTool/Executors/UnrealBuildAccelerator/UBAExecutor.cs
@@ -873,10 +873,7 @@ namespace UnrealBuildTool
 				return null;
 			}
 
-			// TODO: Implement this
-			bool canExecuteCrossArchitecture = true; // action.bCanExecuteCrossArchitecture
-
-			if (isCrossArchitecture && !canExecuteCrossArchitecture)
+			if (isCrossArchitecture && !action.bCanExecuteInUBACrossArchitecture)
 			{
 				return null;
 			}
@@ -926,9 +923,6 @@ namespace UnrealBuildTool
 					}
 				}
 
-				// TODO: Implement this
-				bool canExecuteCrossArchitecture = true; // action.bCanExecuteCrossArchitecture
-
 				ProcessStartInfo startInfo = GetActionStartInfo(action, out FileItem? pchItem);
 				_session!.RunProcessRemote(startInfo, (s, e) =>
 				{
@@ -984,7 +978,7 @@ namespace UnrealBuildTool
 					List<string> logLines = e.LogLines;
 					logLines.RemoveAll((line) => line.StartsWith("   Creating library ", StringComparison.OrdinalIgnoreCase) && line.EndsWith(".exp", StringComparison.OrdinalIgnoreCase));
 					ActionFinished(queue, new ExecuteResults(logLines, e.ExitCode, executionTime, processorTime, additionalDescription), action, pchItem, process);
-				}, action.Weight, knownInputs, knownInputsCount, canExecuteCrossArchitecture);
+				}, action.Weight, knownInputs, knownInputsCount, action.bCanExecuteInUBACrossArchitecture);
 				return Task.CompletedTask;
 			};
 		}
diff --git a/Engine/Source/Programs/UnrealBuildTool/Platform/Windows/VCCompileAction.cs b/Engine/Source/Programs/UnrealBuildTool/Platform/Windows/VCCompileAction.cs
--- a/Engine/Source/Programs/UnrealBuildTool/Platform/Windows/VCCompileAction.cs
+++ b/Engine/Source/Programs/UnrealBuildTool/Platform/Windows/VCCompileAction.cs
@@ -195,6 +195,9 @@ namespace UnrealBuildTool
 		/// <inheritdoc/>
 		public bool bCanExecuteInUBA { get; set; } = true;
 
+		/// <inheritdoc/>
+		public bool bCanExecuteInUBACrossArchitecture { get; set; } = true;
+
 		/// <inheritdoc/>
 		public bool bUseActionHistory => true;
 
@@ -434,6 +437,7 @@ namespace UnrealBuildTool
 			bCanExecuteRemotelyWithSNDBS = Reader.ReadBool();
 			bCanExecuteRemotelyWithXGE = Reader.ReadBool();
 			bCanExecuteInUBA = Reader.ReadBool();
+			bCanExecuteInUBACrossArchitecture = Reader.ReadBool();
 			Architecture = UnrealArch.Parse(Reader.ReadString()!);
 			bIsAnalyzing = Reader.ReadBool();
 			bWarningsAsErrors = Reader.ReadBool();
@@ -476,6 +480,7 @@ namespace UnrealBuildTool
 			Writer.WriteBool(bCanExecuteRemotelyWithSNDBS);
 			Writer.WriteBool(bCanExecuteRemotelyWithXGE);
 			Writer.WriteBool(bCanExecuteInUBA);
+			Writer.WriteBool(bCanExecuteInUBACrossArchitecture);
 			Writer.WriteString(Architecture.ToString());
 			Writer.WriteBool(bIsAnalyzing);
 			Writer.WriteBool(bWarningsAsErrors);
diff --git a/Engine/Source/Programs/UnrealBuildTool/Platform/Windows/VCToolChain.cs b/Engine/Source/Programs/UnrealBuildTool/Platform/Windows/VCToolChain.cs
--- a/Engine/Source/Programs/UnrealBuildTool/Platform/Windows/VCToolChain.cs
+++ b/Engine/Source/Programs/UnrealBuildTool/Platform/Windows/VCToolChain.cs
@@ -2111,6 +2111,12 @@ namespace UnrealBuildTool
 					CompileAction.ExperimentalLogFile = FileItem.GetItemByFileReference(FileReference.Combine(OutputDir, $"{FileName}.sarif"));
 				}
 
+				// Pch are not compatible between different architectures (msvc pch is just a memory dump)
+				if (Target.WindowsPlatform.Compiler.IsMSVC() && CompileEnvironment.PrecompiledHeaderAction != PrecompiledHeaderAction.None)
+				{
+					CompileAction.bCanExecuteInUBACrossArchitecture = false;
+				}
+
 				// Allow derived toolchains to make further changes
 				ModifyFinalCompileAction(CompileAction, CompileEnvironment, SourceFile, OutputDir, ModuleName);
 
diff --git a/Engine/Source/Programs/UnrealBuildTool/System/Action.cs b/Engine/Source/Programs/UnrealBuildTool/System/Action.cs
--- a/Engine/Source/Programs/UnrealBuildTool/System/Action.cs
+++ b/Engine/Source/Programs/UnrealBuildTool/System/Action.cs
@@ -151,10 +151,15 @@ namespace UnrealBuildTool
 		bool bCanExecuteRemotelyWithXGE { get; }
 
 		/// <summary>
-		/// True if this action can be executed by box
+		/// True if this action can be executed by Unreal Build Accelerator
 		/// </summary>
 		bool bCanExecuteInUBA { get; }
 
+		/// <summary>
+		/// True if this action can be executed on a cross-architecture Unreal Build Accelerator helper. If no cross-architecture toolchain is found this setting has no effect
+		/// </summary>
+		bool bCanExecuteInUBACrossArchitecture { get; }
+
 		/// <summary>
 		/// True if this action is using the GCC compiler.  Some build systems may be able to optimize for this case.
 		/// </summary>
@@ -302,6 +307,11 @@ namespace UnrealBuildTool
 		/// </summary>
 		public bool bCanExecuteInUBA { get; set; } = true;
 
+		/// <summary>
+		/// True if this action can be executed on a cross-architecture Unreal Build Accelerator helper. If no cross-architecture toolchain is found this setting has no effect
+		/// </summary>
+		public bool bCanExecuteInUBACrossArchitecture { get; set; } = true;
+
 		/// <summary>
 		/// True if this action is using the GCC compiler.  Some build systems may be able to optimize for this case.
 		/// </summary>
@@ -378,6 +388,7 @@ namespace UnrealBuildTool
 			bCanExecuteRemotelyWithSNDBS = InOther.bCanExecuteRemotelyWithSNDBS;
 			bCanExecuteRemotelyWithXGE = InOther.bCanExecuteRemotelyWithXGE;
 			bCanExecuteInUBA = InOther.bCanExecuteInUBA;
+			bCanExecuteInUBACrossArchitecture = InOther.bCanExecuteInUBACrossArchitecture;
 			bIsGCCCompiler = InOther.bIsGCCCompiler;
 			bDeleteProducedItemsOnError = InOther.bDeleteProducedItemsOnError;
 			bForceWarningsAsError = InOther.bForceWarningsAsError;
@@ -403,6 +414,7 @@ namespace UnrealBuildTool
 			bCanExecuteRemotelyWithSNDBS = Reader.ReadBool();
 			bCanExecuteRemotelyWithXGE = Reader.ReadBool();
 			bCanExecuteInUBA = Reader.ReadBool();
+			bCanExecuteInUBACrossArchitecture = Reader.ReadBool();
 			bIsGCCCompiler = Reader.ReadBool();
 			bDeleteProducedItemsOnError = Reader.ReadBool();
 			bForceWarningsAsError = Reader.ReadBool();
@@ -437,6 +449,7 @@ namespace UnrealBuildTool
 			Writer.WriteBool(bCanExecuteRemotelyWithSNDBS);
 			Writer.WriteBool(bCanExecuteRemotelyWithXGE);
 			Writer.WriteBool(bCanExecuteInUBA);
+			Writer.WriteBool(bCanExecuteInUBACrossArchitecture);
 			Writer.WriteBool(bIsGCCCompiler);
 			Writer.WriteBool(bDeleteProducedItemsOnError);
 			Writer.WriteBool(bForceWarningsAsError);
@@ -516,6 +529,11 @@ namespace UnrealBuildTool
 				Action.bCanExecuteInUBA = bCanExecuteInUBA;
 			}
 
+			if (Object.TryGetBoolField("bCanExecuteInUBACrossArchitecture", out bool bCanExecuteInUBACrossArchitecture))
+			{
+				Action.bCanExecuteInUBACrossArchitecture = bCanExecuteInUBACrossArchitecture;
+			}
+
 			if (Object.TryGetBoolField("bIsGCCCompiler", out bool bIsGCCCompiler))
 			{
 				Action.bIsGCCCompiler = bIsGCCCompiler;
@@ -637,6 +655,7 @@ namespace UnrealBuildTool
 			Writer.WriteValue("bCanExecuteRemotelyWithSNDBS", Action.bCanExecuteRemotelyWithSNDBS);
 			Writer.WriteValue("bCanExecuteRemotelyWithXGE", Action.bCanExecuteRemotelyWithXGE);
 			Writer.WriteValue("bCanExecuteInUBA", Action.bCanExecuteInUBA);
+			Writer.WriteValue("bCanExecuteInUBACrossArchitecture", Action.bCanExecuteInUBACrossArchitecture);
 			Writer.WriteValue("bIsGCCCompiler", Action.bIsGCCCompiler);
 			Writer.WriteValue("bDeleteProducedItemsOnError", Action.bDeleteProducedItemsOnError);
 			Writer.WriteValue("bForceWarningsAsError", Action.bForceWarningsAsError);			
@@ -770,6 +789,7 @@ namespace UnrealBuildTool
 		public bool bCanExecuteRemotelyWithSNDBS => Inner.bCanExecuteRemotelyWithSNDBS;
 		public bool bCanExecuteRemotelyWithXGE => Inner.bCanExecuteRemotelyWithXGE;
 		public bool bCanExecuteInUBA => Inner.bCanExecuteInUBA;
+		public bool bCanExecuteInUBACrossArchitecture => Inner.bCanExecuteInUBACrossArchitecture;
 		public bool bIsGCCCompiler => Inner.bIsGCCCompiler;
 		public bool bDeleteProducedItemsOnError => Inner.bDeleteProducedItemsOnError;
 		public bool bForceWarningsAsError => Inner.bForceWarningsAsError;
