[UBT]
* Disabled cross architecture for ispc (we don't have the binaries for arm64)

diff --git a/Engine/Source/Programs/UnrealBuildTool/ToolChain/ISPCToolChain.cs b/Engine/Source/Programs/UnrealBuildTool/ToolChain/ISPCToolChain.cs
--- a/Engine/Source/Programs/UnrealBuildTool/ToolChain/ISPCToolChain.cs
+++ b/Engine/Source/Programs/UnrealBuildTool/ToolChain/ISPCToolChain.cs
@@ -456,6 +456,7 @@ namespace UnrealBuildTool
 				CompileAction.ArtifactMode = ArtifactMode.Enabled;
 
 				CompileAction.bCanExecuteRemotely = true;
+				CompileAction.bCanExecuteInUBACrossArchitecture = false;
 
 				// Disable remote execution to workaround mismatched case on XGE
 				CompileAction.bCanExecuteRemotelyWithXGE = false;
@@ -659,6 +660,7 @@ namespace UnrealBuildTool
 				}
 
 				CompileAction.bCanExecuteRemotely = true;
+				CompileAction.bCanExecuteInUBACrossArchitecture = false;
 
 				// Disable remote execution to workaround mismatched case on XGE
 				CompileAction.bCanExecuteRemotelyWithXGE = false;
