[UBA]
* Removed radlinker special options when building import library

diff --git a/Engine/Source/Programs/UnrealBuildTool/Platform/Windows/VCToolChain.cs b/Engine/Source/Programs/UnrealBuildTool/Platform/Windows/VCToolChain.cs
--- a/Engine/Source/Programs/UnrealBuildTool/Platform/Windows/VCToolChain.cs
+++ b/Engine/Source/Programs/UnrealBuildTool/Platform/Windows/VCToolChain.cs
@@ -2710,7 +2710,7 @@ namespace UnrealBuildTool
 				Arguments.Add("/time+");
 			}
 
-			if (Target.WindowsPlatform.bAllowRadLinker)
+			if (!bBuildImportLibraryOnly && Target.WindowsPlatform.bAllowRadLinker)
 			{
 				Arguments.Add("/RAD_SUPPRESS_ERROR:36");
 				Arguments.Add("/RAD_SUPPRESS_ERROR:30");
