[UBA]
* Changed so worker name id counter is not increased if no worker is acquired

diff --git a/Engine/Source/Programs/UnrealBuildTool/Executors/UnrealBuildAccelerator/UBAAgentCoordinatorHorde.cs b/Engine/Source/Programs/UnrealBuildTool/Executors/UnrealBuildAccelerator/UBAAgentCoordinatorHorde.cs
--- a/Engine/Source/Programs/UnrealBuildTool/Executors/UnrealBuildAccelerator/UBAAgentCoordinatorHorde.cs
+++ b/Engine/Source/Programs/UnrealBuildTool/Executors/UnrealBuildAccelerator/UBAAgentCoordinatorHorde.cs
@@ -254,7 +254,7 @@ namespace UnrealBuildTool
 				throw new InvalidOperationException($"Session not initialized. Call {nameof(InitAsync)} first");
 			}
 
-			int workerId = _workerId++;
+			int workerId = _workerId;
 			
 			PrefixLogger workerLogger = new($"[Worker{workerId}]", _logger);
 
@@ -288,6 +288,8 @@ namespace UnrealBuildTool
 					return false;
 				}
 
+				++_workerId;
+
 				workerLogger.LogDebug("Agent properties:");
 
 				string agentName = String.Empty;
