PR #12445: Add new functionalities to Visual Studio workspace generator in UBT

#rnx

diff --git a/Engine/Source/Programs/UnrealBuildTool/ProjectFiles/VisualStudioWorkspace/AvailableConfiguration.cs b/Engine/Source/Programs/UnrealBuildTool/ProjectFiles/VisualStudioWorkspace/AvailableConfiguration.cs
new file mode 100644
--- /dev/null
+++ b/Engine/Source/Programs/UnrealBuildTool/ProjectFiles/VisualStudioWorkspace/AvailableConfiguration.cs
@@ -0,0 +1,19 @@
+// Copyright Epic Games, Inc. All Rights Reserved.
+
+using System.Collections.Generic;
+
+namespace UnrealBuildTool
+{
+	internal class AvailableConfiguration
+	{
+		public HashSet<string> Configurations { get; set; } = new();
+
+		public HashSet<string> Platforms { get; set; } = new();
+
+		public string TargetPath { get; set; } = string.Empty;
+
+		public string ProjectPath { get; set; } = string.Empty;
+
+		public string TargetType { get; set; } = string.Empty;
+	}
+}
diff --git a/Engine/Source/Programs/UnrealBuildTool/ProjectFiles/VisualStudioWorkspace/AvailableTargetsConfigurations.cs b/Engine/Source/Programs/UnrealBuildTool/ProjectFiles/VisualStudioWorkspace/AvailableTargetsConfigurations.cs
new file mode 100644
--- /dev/null
+++ b/Engine/Source/Programs/UnrealBuildTool/ProjectFiles/VisualStudioWorkspace/AvailableTargetsConfigurations.cs
@@ -0,0 +1,11 @@
+// Copyright Epic Games, Inc. All Rights Reserved.
+
+using System.Collections.Generic;
+
+namespace UnrealBuildTool
+{
+	internal class AvailableTargetsConfigurations
+	{
+		public Dictionary<string, AvailableConfiguration> Targets { get; set; } = new();
+	}
+}
diff --git a/Engine/Source/Programs/UnrealBuildTool/ProjectFiles/VisualStudioWorkspace/VSWorkspaceProjectFile.cs b/Engine/Source/Programs/UnrealBuildTool/ProjectFiles/VisualStudioWorkspace/VSWorkspaceProjectFile.cs
--- a/Engine/Source/Programs/UnrealBuildTool/ProjectFiles/VisualStudioWorkspace/VSWorkspaceProjectFile.cs
+++ b/Engine/Source/Programs/UnrealBuildTool/ProjectFiles/VisualStudioWorkspace/VSWorkspaceProjectFile.cs
@@ -4,17 +4,18 @@ using System;
 using System.Collections.Generic;
 using System.IO;
 using System.Linq;
+using System.Text;
 using System.Text.Json;
 using System.Threading.Tasks;
 using EpicGames.Core;
 using Microsoft.Extensions.Logging;
 using UnrealBuildBase;
+using static UnrealBuildTool.VCProjectFile;
 
 namespace UnrealBuildTool
 {
 	internal class VSWorkspaceProjectFile : ProjectFile
 	{
-		private readonly DirectoryReference RootPath;
 		private readonly HashSet<TargetType> TargetTypes;
 		private readonly CommandLineArguments Arguments;
 
@@ -23,6 +24,8 @@ namespace UnrealBuildTool
 		/// </summary>
 		public List<ExportedTargetInfo> ExportedTargetProjects { get; set; } = new();
 
+		public DirectoryReference RootPath { get; }
+
 		public VSWorkspaceProjectFile(FileReference InProjectFilePath, DirectoryReference BaseDir,
 			DirectoryReference RootPath, HashSet<TargetType> TargetTypes, CommandLineArguments Arguments)
 			: base(InProjectFilePath, BaseDir)
@@ -32,6 +35,48 @@ namespace UnrealBuildTool
 			this.Arguments = Arguments;
 		}
 
+		/// <summary>
+		/// Collect available configurations for the project storing it at <c>availableConfigurations</c> dictionary.
+		/// </summary>
+		/// <remarks>
+		/// Keys for the dictionary are available project names.
+		/// </remarks>
+		/// <param name="availableConfigurations"></param>
+		/// <param name="InPlatforms"></param>
+		/// <param name="InConfigurations"></param>
+		/// <param name="PlatformProjectGenerators"></param>
+		/// <param name="Minimize"></param>
+		/// <param name="Logger"></param>
+		public void CollectAvailableConfigurations(Dictionary<string, AvailableConfiguration> availableConfigurations,
+			List<UnrealTargetPlatform> InPlatforms,
+			List<UnrealTargetConfiguration> InConfigurations,
+			PlatformProjectGeneratorCollection PlatformProjectGenerators,
+			JsonWriterStyle Minimize,
+			ILogger Logger)
+		{
+			ProcessProjectFiles(
+				InPlatforms,
+				InConfigurations,
+				PlatformProjectGenerators,
+				// always provide null in here as we want to collect all available data.
+				null,
+				(UnrealTargetConfiguration Configuration, UnrealTargetPlatform Platform, ProjectTarget ProjectTarget) =>
+				{
+					if (!availableConfigurations.TryGetValue(ProjectTarget.Name, out AvailableConfiguration? availableConfiguration))
+					{
+						availableConfiguration = new AvailableConfiguration();
+						availableConfigurations.Add(ProjectTarget.Name, availableConfiguration);
+					}
+					availableConfiguration.TargetType = ProjectTarget.TargetRules?.Type.ToString() ?? string.Empty;
+					availableConfiguration.Configurations.Add(Configuration.ToString());
+					availableConfiguration.Platforms.Add(Platform.ToString());
+					availableConfiguration.TargetPath = ProjectTarget.TargetFilePath.FullName;
+					availableConfiguration.ProjectPath = ProjectTarget.UnrealProjectFilePath?.FullName ?? string.Empty;
+
+					return Task.CompletedTask;
+				});
+		}
+
 		/// <summary>
 		/// Write project file info in JSON file.
 		/// For every combination of <c>UnrealTargetPlatform</c>, <c>UnrealTargetConfiguration</c> and <c>TargetType</c>
@@ -39,25 +84,114 @@ namespace UnrealBuildTool
 		/// </summary>
 		/// <remarks>
 		/// * <c>TargetType.Editor</c> will be generated for current platform only and will ignore <c>UnrealTargetConfiguration.Test</c> and <c>UnrealTargetConfiguration.Shipping</c> configurations
-		/// * <c>TargetType.Program</c>  will be generated for current platform only and <c>UnrealTargetConfiguration.Development</c> configuration only 
+		/// * <c>TargetType.Program</c>  will be generated for current platform only and <c>UnrealTargetConfiguration.Development</c> configuration only
 		/// </remarks>
 		/// <param name="InPlatforms"></param>
 		/// <param name="InConfigurations"></param>
 		/// <param name="PlatformProjectGenerators"></param>
 		/// <param name="Minimize"></param>
 		/// <param name="Logger"></param>
+		/// <param name="ProjectGenerationDescriptionData"></param>
+		/// <param name="VCFormat"></param>
 		/// <returns></returns>
 		public bool WriteProjectFile(List<UnrealTargetPlatform> InPlatforms,
 			List<UnrealTargetConfiguration> InConfigurations,
-			PlatformProjectGeneratorCollection PlatformProjectGenerators, JsonWriterStyle Minimize, ILogger Logger)
+			PlatformProjectGeneratorCollection PlatformProjectGenerators,
+			JsonWriterStyle Minimize,
+			ILogger Logger,
+			Dictionary<string, AvailableConfiguration>? ProjectGenerationDescriptionData,
+			VCProjectFileFormat VCFormat)
 		{
 			DirectoryReference ProjectRootFolder = RootPath;
 
 			HashSet<UnrealTargetPlatform> ServerPlatforms = Utils.GetPlatformsInClass(UnrealPlatformClass.Server).ToHashSet();
+			ProcessProjectFiles(
+				InPlatforms,
+				InConfigurations,
+				PlatformProjectGenerators,
+				ProjectGenerationDescriptionData,
+				(UnrealTargetConfiguration Configuration, UnrealTargetPlatform Platform, ProjectTarget ProjectTarget) => Task.Run(async () =>
+				{
+					bool bBuildByDefault = ShouldBuildByDefaultForSolutionTargets && ProjectTarget.SupportedPlatforms.Contains(Platform);
 
+					UnrealArchitectures ProjectArchitectures = UEBuildPlatform
+						.GetBuildPlatform(Platform)
+						.ArchitectureConfig.ActiveArchitectures(ProjectTarget.UnrealProjectFilePath, ProjectTarget.Name);
+
+					TargetDescriptor TargetDesc = new(ProjectTarget.UnrealProjectFilePath, ProjectTarget.Name,
+						Platform, Configuration, ProjectArchitectures, Arguments);
+					TargetDesc.IntermediateEnvironment = UnrealIntermediateEnvironment.GenerateProjectFiles;
+
+					try
+					{
+						foreach (var Architecture in TargetDesc.Architectures.Architectures)
+						{
+							UEBuildTarget BuildTarget = UEBuildTarget.Create(TargetDesc, false, false, false, UnrealIntermediateEnvironment.GenerateProjectFiles, Logger);
+							BuildTarget.PreBuildSetup(Logger);
+
+							ExportedTargetInfo TargetInfo = ExportTarget(BuildTarget, Architecture, bBuildByDefault, PlatformProjectGenerators, Logger, VCFormat);
+
+							var sb = new StringBuilder();
+							sb.Append(ProjectTarget.TargetFilePath.GetFileNameWithoutAnyExtensions());
+							sb.Append("_");
+							sb.Append(TargetInfo.Configuration);
+							sb.Append("_");
+							sb.Append(TargetInfo.Platform);
+
+							if (!string.IsNullOrEmpty(TargetInfo.Architecture))
+							{
+								sb.Append("_");
+								sb.Append(TargetInfo.Architecture);
+							}
+
+							sb.Append(".json");
+
+							FileReference OutputFile = FileReference.Combine(ProjectRootFolder, sb.ToString());
+
+							DirectoryReference.CreateDirectory(OutputFile.Directory);
+							using FileStream Stream = new(OutputFile.FullName, FileMode.Create, FileAccess.Write);
+							await JsonSerializer.SerializeAsync(Stream, TargetInfo, options: new JsonSerializerOptions()
+							{
+								PropertyNamingPolicy = JsonNamingPolicy.CamelCase,
+								WriteIndented = Minimize == JsonWriterStyle.Readable,
+							});
+
+							lock (ExportedTargetProjects)
+							{
+								ExportedTargetProjects.Add(TargetInfo);
+							}
+						}
+
+					}
+					catch (Exception Ex)
+					{
+						Logger.LogWarning("Exception while generating include data for Target:{0}, Platform: {1}, Configuration: {2}", TargetDesc.Name, Platform.ToString(), Configuration.ToString());
+						Logger.LogWarning("{0}", Ex.ToString());
+					}
+				}));
+
+			return true;
+		}
+
+		/// <summary>
+		/// Filter configurations based on what is supported by the project.
+		/// </summary>
+		/// <param name="InPlatforms"></param>
+		/// <param name="InConfigurations"></param>
+		/// <param name="PlatformProjectGenerators"></param>
+		/// <param name="ProjectGenerationDescriptionData"></param>
+		/// <param name="Handler"></param>
+		private void ProcessProjectFiles(List<UnrealTargetPlatform> InPlatforms,
+			List<UnrealTargetConfiguration> InConfigurations,
+			PlatformProjectGeneratorCollection PlatformProjectGenerators,
+			Dictionary<string, AvailableConfiguration>? ProjectGenerationDescriptionData,
+			Func<UnrealTargetConfiguration, UnrealTargetPlatform, ProjectTarget, Task> Handler)
+		{
+			HashSet<UnrealTargetPlatform> ServerPlatforms = Utils.GetPlatformsInClass(UnrealPlatformClass.Server).ToHashSet();
+			List<Task> Tasks = new List<Task>();
+			List<(UnrealTargetConfiguration, UnrealTargetPlatform, ProjectTarget)> ret = new List<(UnrealTargetConfiguration, UnrealTargetPlatform, ProjectTarget)>();
 			foreach (UnrealTargetConfiguration Configuration in InConfigurations)
 			{
-				List<Task> Tasks = new List<Task>();
 				foreach (UnrealTargetPlatform Platform in InPlatforms)
 				{
 					foreach (ProjectTarget ProjectTarget in ProjectTargets.OfType<ProjectTarget>())
@@ -89,67 +223,30 @@ namespace UnrealBuildTool
 							continue;
 						}
 
-						Task CreateTask = Task.Run(async () =>
+						if (ShouldSkipConfiguration(ProjectGenerationDescriptionData, Configuration, Platform, ProjectTarget))
 						{
-							bool bBuildByDefault = ShouldBuildByDefaultForSolutionTargets && ProjectTarget.SupportedPlatforms.Contains(Platform);
-
-							UnrealArchitectures ProjectArchitectures = UEBuildPlatform
-								.GetBuildPlatform(Platform)
-								.ArchitectureConfig.ActiveArchitectures(ProjectTarget.UnrealProjectFilePath, ProjectTarget.Name);
-
-							TargetDescriptor TargetDesc = new(ProjectTarget.UnrealProjectFilePath, ProjectTarget.Name,
-								Platform, Configuration, ProjectArchitectures, Arguments);
-							TargetDesc.IntermediateEnvironment = UnrealIntermediateEnvironment.GenerateProjectFiles;
-
-							try
-							{
-								FileReference OutputFile = FileReference.Combine(ProjectRootFolder,
-									$"{ProjectTarget.TargetFilePath.GetFileNameWithoutAnyExtensions()}_{Configuration}_{Platform}.json");
-
-								UEBuildTarget BuildTarget = UEBuildTarget.Create(TargetDesc, false, false, false, UnrealIntermediateEnvironment.GenerateProjectFiles, Logger);
-								BuildTarget.PreBuildSetup(Logger);
-
-								ExportedTargetInfo TargetInfo = ExportTarget(BuildTarget, bBuildByDefault, PlatformProjectGenerators, Logger);
-
-								DirectoryReference.CreateDirectory(OutputFile.Directory);
-								using FileStream Stream = new(OutputFile.FullName, FileMode.Create, FileAccess.Write);
-								await JsonSerializer.SerializeAsync(Stream, TargetInfo, options: new JsonSerializerOptions()
-								{
-									PropertyNamingPolicy = JsonNamingPolicy.CamelCase,
-									WriteIndented = Minimize == JsonWriterStyle.Readable,
-								});
-
-								lock (ExportedTargetProjects)
-								{
-									ExportedTargetProjects.Add(TargetInfo);
-								}
-							}
-							catch (Exception Ex)
-							{
-								Logger.LogWarning("Exception while generating include data for Target:{Target}, Platform: {Platform}, Configuration: {Configuration}", TargetDesc.Name, Platform.ToString(), Configuration.ToString());
-								Logger.LogWarning("{Ex}", Ex.ToString());
-							}
-						});
+							continue;
+						}
 
-						Tasks.Add(CreateTask);
+						Tasks.Add(Handler(Configuration, Platform, ProjectTarget));
 					}
 				}
 
 				Task.WaitAll(Tasks.ToArray());
 			}
-
-			return true;
 		}
 
 		/// <summary>
 		/// Write a Target to a JSON writer. Is array is empty, don't write anything
 		/// </summary>
 		/// <param name="Target"></param>
+		/// <param name="Architecture"></param>
 		/// <param name="bBuildByDefault"></param>
 		/// <param name="PlatformProjectGenerators"></param>
 		/// <param name="Logger">Logger for output</param>
+		/// <param name="VCFormat">Project file format.</param>
 		private ExportedTargetInfo ExportTarget(
-			UEBuildTarget Target, bool bBuildByDefault, PlatformProjectGeneratorCollection PlatformProjectGenerators, ILogger Logger)
+			UEBuildTarget Target, UnrealArch? Architecture, bool bBuildByDefault, PlatformProjectGeneratorCollection PlatformProjectGenerators, ILogger Logger, VCProjectFileFormat VCFormat)
 		{
 			ExportedTargetInfo TargetInfo = new()
 			{
@@ -158,8 +255,11 @@ namespace UnrealBuildTool
 				ProjectPath = Target.ProjectFile?.FullName ?? String.Empty,
 				TargetType = Target.TargetType.ToString(),
 				Platform = Target.Platform.ToString(),
+				VCPlatform = GetVCPlatform(Target, Architecture, PlatformProjectGenerators, VCFormat),
 				Configuration = Target.Configuration.ToString(),
-				BuildInfo = ExportBuildInfo(Target, PlatformProjectGenerators, bBuildByDefault, Logger)
+				VCConfiguration = GetVCConfiguration(Target, Architecture, PlatformProjectGenerators, VCFormat),
+				Architecture = GetArchitecture(Target, Architecture, PlatformProjectGenerators, VCFormat),
+				BuildInfo = ExportBuildInfo(Target, Architecture, PlatformProjectGenerators, bBuildByDefault, Logger)
 			};
 
 			UEToolChain TargetToolChain = Target.CreateToolchain(Target.Platform, Logger);
@@ -196,6 +296,26 @@ namespace UnrealBuildTool
 			return TargetInfo;
 		}
 
+		private static bool ShouldSkipConfiguration(Dictionary<string, AvailableConfiguration>? ProjectGenerationDescriptionData,
+			UnrealTargetConfiguration Configuration,
+			UnrealTargetPlatform Platform,
+			ProjectTarget ProjectTarget)
+		{
+			// No restrictions were provided, we should not skip anything.
+			if (ProjectGenerationDescriptionData == null)
+			{
+				return false;
+			}
+
+			// Configs were provided, but we could not find it in the dictionary, so we should not skip it.
+			if (!ProjectGenerationDescriptionData.TryGetValue(ProjectTarget.Name, out AvailableConfiguration? availableConfiguration))
+			{
+				return true;
+			}
+
+			return !availableConfiguration.Configurations.Contains(Configuration.ToString()) || !availableConfiguration.Platforms.Contains(Platform.ToString());
+		}
+
 		/// <summary>
 		/// Write a Module to a JSON writer. If array is empty, don't write anything
 		/// </summary>
@@ -240,6 +360,8 @@ namespace UnrealBuildTool
 			{
 				FileItem IncludeHeader = FileItem.GetItemByFileReference(ModuleCompileEnvironment.PrecompiledHeaderIncludeFilename!);
 				ForcedIncludes.Insert(0, IncludeHeader);
+
+				Result.PrecompiledHeader = IncludeHeader.FullName;
 			}
 
 			Result.ForcedIncludes.AddRange(ForcedIncludes.Select(x => x.FullName));
@@ -252,7 +374,7 @@ namespace UnrealBuildTool
 			return Result;
 		}
 
-		private TargetBuildInfo? ExportBuildInfo(UEBuildTarget Target, PlatformProjectGeneratorCollection PlatformProjectGenerators, bool bBuildByDefault, ILogger Logger)
+		private TargetBuildInfo? ExportBuildInfo(UEBuildTarget Target, UnrealArch? Architecture, PlatformProjectGeneratorCollection PlatformProjectGenerators, bool bBuildByDefault, ILogger Logger)
 		{
 			if (!BuildHostPlatform.Current.Platform.IsInGroup(UnrealPlatformGroup.Windows))
 			{
@@ -277,7 +399,7 @@ namespace UnrealBuildTool
 
 			PlatformProjectGenerator? ProjGenerator = PlatformProjectGenerators.GetPlatformProjectGenerator(Platform, true);
 			VCProjectFile.BuildCommandBuilder BuildCommandBuilder = new(
-				new PlatformProjectGenerator.VSSettings(Platform, Configuration, VCProjectFileFormat.Default, null), ProjectTarget, UProjectPath)
+				new PlatformProjectGenerator.VSSettings(Platform, Configuration, VCProjectFileFormat.Default, Architecture), ProjectTarget, UProjectPath)
 			{
 				ProjectGenerator = ProjGenerator,
 				bIsForeignProject = IsForeignProject
@@ -295,6 +417,67 @@ namespace UnrealBuildTool
 			};
 		}
 
+		private string GetVCPlatform(UEBuildTarget Target, UnrealArch? Architecture, PlatformProjectGeneratorCollection Generators, VCProjectFileFormat VCFormat)
+		{
+			PlatformProjectGenerator? Generator = Generators.GetPlatformProjectGenerator(Target.Platform, true);
+			if (Generator != null)
+			{
+				PlatformProjectGenerator.VSSettings VSSettings = new(Target.Platform, Target.Configuration, VCFormat, Architecture);
+				if (Generator.HasVisualStudioSupport(VSSettings))
+				{
+					return Generator.GetVisualStudioPlatformName(VSSettings);
+				}
+
+				return DefaultPlatformName;
+			}
+
+			return string.Empty;
+		}
+
+		private string GetVCConfiguration(UEBuildTarget Target, UnrealArch? Architecture, PlatformProjectGeneratorCollection Generators, VCProjectFileFormat VCFormat)
+		{
+			PlatformProjectGenerator? Generator = Generators.GetPlatformProjectGenerator(Target.Platform, true);
+			if (Generator != null)
+			{
+				PlatformProjectGenerator.VSSettings VSSettings = new(Target.Platform, Target.Configuration, VCFormat, Architecture);
+				if (Generator.HasVisualStudioSupport(VSSettings))
+				{
+					if (!Generator.RequiresDistinctVisualStudioConfigurationName(VSSettings))
+					{
+						return Target.Configuration.ToString();
+					}
+				}
+
+				return string.Format("{0}{1}_{2}", Target.Platform.ToString(), Architecture != null ? "_" + Architecture.ToString() : string.Empty, Target.Configuration.ToString()); ;
+			}
+
+			return string.Empty;
+		}
+
+		private string GetArchitecture(UEBuildTarget Target, UnrealArch? Architecture, PlatformProjectGeneratorCollection Generators, VCProjectFileFormat VCFormat)
+		{
+			PlatformProjectGenerator? Generator = Generators.GetPlatformProjectGenerator(Target.Platform, true);
+			if (Generator != null)
+			{
+				PlatformProjectGenerator.VSSettings VSSettings = new(Target.Platform, Target.Configuration, VCFormat, Architecture);
+				if (Generator.HasVisualStudioSupport(VSSettings))
+				{
+					if (!Generator.RequiresDistinctVisualStudioConfigurationName(VSSettings))
+					{
+						return string.Empty;
+					}
+				}
+
+				if (Architecture != null)
+				{
+					var result = Architecture.ToString();
+					return result != null ? result : string.Empty;
+				}
+			}
+
+			return string.Empty;
+		}
+
 		internal class ExportedTargetInfo
 		{
 			public string TargetName { get; set; } = String.Empty;
@@ -302,7 +485,10 @@ namespace UnrealBuildTool
 			public string ProjectPath { get; set; } = String.Empty;
 			public string TargetType { get; set; } = String.Empty;
 			public string Platform { get; set; } = String.Empty;
+			public string VCPlatform { get; set; } = String.Empty;
 			public string Configuration { get; set; } = String.Empty;
+			public string VCConfiguration { get; set; } = String.Empty;
+			public string Architecture { get; set; } = String.Empty;
 			public TargetBuildInfo? BuildInfo { get; set; }
 			public Dictionary<string, ExportedModuleInfo> ModuleToCompileSettings { get; set; } = new();
 			public Dictionary<string, string> DirToModule { get; set; } = new();
@@ -317,6 +503,7 @@ namespace UnrealBuildTool
 			public List<string> IncludePaths { get; set; } = new();
 			public List<string> Defines { get; set; } = new();
 			public string? Standard { get; set; }
+			public string? PrecompiledHeader { get; set; }
 			public List<string> ForcedIncludes { get; set; } = new();
 			public string? CompilerPath { get; set; }
 			public List<string> CompilerArgs { get; set; } = new();
diff --git a/Engine/Source/Programs/UnrealBuildTool/ProjectFiles/VisualStudioWorkspace/VSWorkspaceProjectFileGenerator.cs b/Engine/Source/Programs/UnrealBuildTool/ProjectFiles/VisualStudioWorkspace/VSWorkspaceProjectFileGenerator.cs
--- a/Engine/Source/Programs/UnrealBuildTool/ProjectFiles/VisualStudioWorkspace/VSWorkspaceProjectFileGenerator.cs
+++ b/Engine/Source/Programs/UnrealBuildTool/ProjectFiles/VisualStudioWorkspace/VSWorkspaceProjectFileGenerator.cs
@@ -4,7 +4,10 @@ using System;
 using System.Collections.Generic;
 using System.IO;
 using System.Linq;
+using System.Text;
 using System.Text.Json;
+using System.Text.RegularExpressions;
+using System.Xml;
 using EpicGames.Core;
 using Microsoft.Extensions.Logging;
 using UnrealBuildBase;
@@ -19,12 +22,13 @@ namespace UnrealBuildTool
 		// So they must remain constant.
 		private const string VSUnrealWorkspaceFileName = ".vs-unreal-workspace";
 		private const string ProjectFilesFolder = "VisualStudio";
+		private const string WorkspaceGeneratorVersion = "1.0.0";
 
 		private readonly CommandLineArguments Arguments;
 
 		/// <summary>
 		/// List of deprecated platforms.
-		/// Don't generate project model for these platforms unless they are specified in "Platforms" console arguments. 
+		/// Don't generate project model for these platforms unless they are specified in "Platforms" console arguments.
 		/// </summary>
 		/// <returns></returns>
 		private readonly HashSet<UnrealTargetPlatform> DeprecatedPlatforms = new();
@@ -59,6 +63,18 @@ namespace UnrealBuildTool
 		[CommandLine("-Minimize", Value = "Compact")]
 		private JsonWriterStyle Minimize = JsonWriterStyle.Readable;
 
+		[CommandLine("-Query=", Description = "Outputs all the available configurations, platforms and targets to a file provided by the user.")]
+		private string Query = string.Empty;
+
+		[CommandLine("-TargetDescriptionFile=", Description = "")]
+		private string TargetDescriptionFile = string.Empty;
+
+		[CommandLine("-WorkspaceGeneratorVersion", Description = "Show version information and exit.")]
+		private bool ShowVersionInfo = false;
+
+		[CommandLine("-GenerateVsConfigOnly", Description = "Generate .vsconfig file with dependencies and exit")]
+		private bool GenerateVsConfigOnly = false;
+
 		public VSWorkspaceProjectFileGenerator(FileReference? InOnlyGameProject,
 			CommandLineArguments InArguments)
 			: base(InOnlyGameProject)
@@ -89,6 +105,7 @@ namespace UnrealBuildTool
 
 		protected override bool WriteProjectFiles(PlatformProjectGeneratorCollection PlatformProjectGenerators, ILogger Logger)
 		{
+			var ProjectGenerationDescriptionData = ReadProjectGenerateDescriptionFilePath()?.Targets;
 			using ProgressWriter Progress = new("Writing project files...", true, Logger);
 			List<ProjectFile> ProjectsToGenerate = new(GeneratedProjectFiles);
 			if (ProjectNames.Any())
@@ -108,8 +125,7 @@ namespace UnrealBuildTool
 			List<UnrealTargetPlatform> FilteredPlatforms = PlatformsToGenerate.Where(Platform =>
 			{
 				// Skip deprecated unless explicitly specified in the command line.
-				return (!DeprecatedPlatforms.Contains(Platform) || Platforms.Contains(Platform))
-					   && UEBuildPlatform.IsPlatformAvailable(Platform);
+				return (!DeprecatedPlatforms.Contains(Platform) || Platforms.Contains(Platform)) && UEBuildPlatform.IsPlatformAvailable(Platform);
 			}).ToList();
 
 			HashSet<UnrealTargetConfiguration> ConfigurationsToGenerate = new(SupportedConfigurations);
@@ -119,6 +135,8 @@ namespace UnrealBuildTool
 			}
 
 			List<UnrealTargetConfiguration> Configurations = ConfigurationsToGenerate.ToList();
+			VCProjectFileFormat VCFormat = GetVCProjectFileFormat(Logger);
+			var wroteVCProjectFiles = false;
 
 			for (int ProjectFileIndex = 0; ProjectFileIndex < ProjectsToGenerate.Count; ++ProjectFileIndex)
 			{
@@ -127,22 +145,86 @@ namespace UnrealBuildTool
 					return false;
 				}
 
-				if (!CurrentProject.WriteProjectFile(FilteredPlatforms, Configurations, PlatformProjectGenerators, Minimize, Logger))
+				if (!CurrentProject.WriteProjectFile(FilteredPlatforms, Configurations, PlatformProjectGenerators, Minimize, Logger, ProjectGenerationDescriptionData, VCFormat))
 				{
 					return false;
 				}
 
+				foreach (var Target in CurrentProject.ProjectTargets.OfType<ProjectTarget>())
+				{
+					if (!TargetTypes.Any() || (Target.TargetRules != null && TargetTypes.Contains(Target.TargetRules.Type)))
+					{
+						// Ignore errors.
+						WriteVCProjectFileForTarget(CurrentProject.RootPath, Target, VCFormat, FilteredPlatforms, Configurations, PlatformProjectGenerators, ProjectGenerationDescriptionData, Logger);
+						wroteVCProjectFiles = true;
+					}
+				}
+
 				Progress.Write(ProjectFileIndex + 1, TotalProjectFileCount);
 			}
 
+			if (wroteVCProjectFiles)
+			{
+				WriteVCProjectCommonPropsFile(DirectoryReference.Combine(IntermediateProjectFilesPath, ProjectFilesFolder), VCFormat, Logger);
+			}
+
 			Progress.Write(TotalProjectFileCount, TotalProjectFileCount);
 
 			return true;
 		}
 
+		private bool WriteAvailableConfigurations(PlatformProjectGeneratorCollection PlatformProjectGenerators, ILogger Logger)
+		{
+			Dictionary<string, AvailableConfiguration> AvailableConfigurations = new Dictionary<string, AvailableConfiguration>();
+			HashSet<UnrealTargetPlatform> PlatformsToGenerate = new(SupportedPlatforms);
+			List<UnrealTargetPlatform> FilteredPlatforms = PlatformsToGenerate.Where(Platform =>
+			{
+				// Skip deprecated unless explicitly specified in the command line.
+				return (!DeprecatedPlatforms.Contains(Platform) || Platforms.Contains(Platform)) && UEBuildPlatform.IsPlatformAvailable(Platform);
+			}).ToList();
+
+			for (int ProjectFileIndex = 0; ProjectFileIndex < GeneratedProjectFiles.Count; ++ProjectFileIndex)
+			{
+				if (GeneratedProjectFiles[ProjectFileIndex] is not VSWorkspaceProjectFile CurrentProject)
+				{
+					return false;
+				}
+
+				CurrentProject.CollectAvailableConfigurations(AvailableConfigurations, FilteredPlatforms, new HashSet<UnrealTargetConfiguration>(SupportedConfigurations).ToList(), PlatformProjectGenerators, Minimize, Logger);
+			}
+
+			FileReference OutputFile = FileReference.Combine(IntermediateProjectFilesPath, ProjectFilesFolder, Query);
+			DirectoryReference.CreateDirectory(OutputFile.Directory);
+			using FileStream Stream = new(OutputFile.FullName, FileMode.Create, FileAccess.Write);
+
+			var Result = new
+			{
+				Targets = AvailableConfigurations,
+			};
+			JsonSerializer.Serialize(Stream, Result, options: new JsonSerializerOptions()
+			{
+				PropertyNamingPolicy = JsonNamingPolicy.CamelCase,
+				WriteIndented = Minimize == JsonWriterStyle.Readable,
+			});
+
+			return true;
+		}
+
 		public override bool GenerateProjectFiles(PlatformProjectGeneratorCollection PlatformProjectGenerators,
 			string[] Arguments, bool bCacheDataForEditor, ILogger Logger)
 		{
+			if (ShowVersionInfo)
+			{
+				Logger.LogInformation("Workspace Generator Version: {WorkspaceGeneratorVersion}", WorkspaceGeneratorVersion);
+				return true;
+			}
+
+			WriteVsConfigFile(Logger, OnlyGameProject!.Directory);
+			if (GenerateVsConfigOnly)
+			{
+				return true;
+			}
+
 			bool IncludeAllPlatforms = true;
 			ConfigureProjectFileGeneration(Arguments, ref IncludeAllPlatforms, Logger);
 
@@ -195,6 +277,11 @@ namespace UnrealBuildTool
 				AddAllGameProjects(GameProjects);
 			}
 
+			if (!string.IsNullOrEmpty(Query))
+			{
+				return WriteAvailableConfigurations(PlatformProjectGenerators, Logger);
+			}
+
 			WriteProjectFiles(PlatformProjectGenerators, Logger);
 			WritePrimaryProjectFile(UBTProject, PlatformProjectGenerators, Logger);
 
@@ -225,6 +312,8 @@ namespace UnrealBuildTool
 						ProjectPath = g.Select(i => i.ProjectPath).Distinct().Single(),
 						Configurations = g.Select(i => i.Configuration).Distinct().ToList(),
 						Platforms = g.Select(i => i.Platform).Distinct().ToList(),
+						Architectures = g.Select(i => i.Architecture).Where(i => !string.IsNullOrEmpty(i)).Distinct().ToList(),
+						ConfigurationsAndPlatforms = g.Select(i => string.Format("{0}-{1}{2}", i.Configuration, i.Platform, !string.IsNullOrEmpty(i.Architecture) ? "-" + i.Architecture : string.Empty)).Distinct().ToList(),
 					}));
 
 				// The inner Targets object is needed for schema compatibility with the Query Mode API.
@@ -242,7 +331,7 @@ namespace UnrealBuildTool
 			}
 			catch (Exception Ex)
 			{
-				Logger.LogWarning("Exception while writing root project file: {Ex}", Ex.ToString());
+				Logger.LogWarning("Exception while writing root project file: {0}", Ex.ToString());
 				return false;
 			}
 
@@ -254,5 +343,392 @@ namespace UnrealBuildTool
 		{
 			return FileReference.Combine(IntermediateProjectFilesPath, ProjectFilesFolder, BaseName + ProjectFileExtension);
 		}
+
+		private AvailableTargetsConfigurations? ReadProjectGenerateDescriptionFilePath()
+		{
+			if (string.IsNullOrEmpty(TargetDescriptionFile))
+			{
+				return null;
+			}
+
+			using FileStream Stream = new(TargetDescriptionFile, FileMode.Open, FileAccess.Read);
+			return JsonSerializer.Deserialize<AvailableTargetsConfigurations>(Stream, new JsonSerializerOptions
+			{
+				PropertyNamingPolicy = JsonNamingPolicy.CamelCase,
+			});
+		}
+
+		private static readonly Dictionary<string, string> NodesToUpdate = new Dictionary<string, string>
+		{
+			{ "NMakeBuildCommandLine", "$(_UProjectNMakeBuildCommandLine)" },
+			{ "NMakeCleanCommandLine", "$(_UProjectNMakeCleanCommandLine)" },
+			{ "NMakeCompileFileCommandLine", "$(_UProjectNMakeCompileFileCommandLine)" },
+			{ "NMakeOutput", "$(_UProjectNMakeOutput)" },
+			{ "NMakeReBuildCommandLine", "$(_UProjectNMakeReBuildCommandLine)" }
+		};
+
+		private readonly static HashSet<string> NodesAllowList = new HashSet<string>
+		{
+			"CleanDependsOn",
+			"Configuration",
+			"CppCleanDependsOn",
+			"Import",
+			"ImportGroup",
+			"ItemDefinitionGroup",
+			"ItemGroup",
+			"NMakeBuildCommandLine",
+			"NMakeCleanCommandLine",
+			"NMakeCompile",
+			"NMakeCompileFileCommandLine",
+			"NMakeOutput",
+			"NMakeReBuildCommandLine",
+			"Platform",
+			"Project",
+			"ProjectConfiguration",
+			"ProjectGuid",
+			"PropertyGroup",
+			"RootNamespace",
+			"xml"
+		};
+
+		private readonly static HashSet<string> NodesBlockedList = new HashSet<string>
+		{
+			"AdditionalOptions",
+			"ClCompile",
+			"ClInclude",
+			"IncludePath",
+			"Manifest",
+			"NMakeAssemblySearchPath",
+			"NMakeForcedIncludes",
+			"NMakeForcedUsingAssemblies",
+			"NMakePreprocessorDefinitions",
+			"None",
+			"ProjectForcedIncludeFiles",
+			"PropertyPageSchema",
+			"ResourceCompile",
+			"SourcePath"
+		};
+
+		private static readonly List<Regex> BlockedNodesListRegEx = new List<Regex>()
+		{
+			new Regex("ClCompile_AdditionalIncludeDirectories*"),
+			new Regex("ClCompile_AdditionalOptions*"),
+			new Regex("ClCompile_ForcedIncludeFiles*"),
+			new Regex("ProjectAdditionalIncludeDirectories*"),
+			new Regex("ProjectForcedIncludeFiles*")
+		};
+
+		private VCProjectFileFormat GetVCProjectFileFormat(ILogger Logger)
+		{
+			// Enumerate all the valid installations. This list is already sorted by preference.
+			List<VisualStudioInstallation> Installations = MicrosoftPlatformSDK.FindVisualStudioInstallations(Logger).Where(x => WindowsPlatform.HasCompiler(x.Compiler, UnrealArch.X64, Logger)).ToList();
+
+			// Get the corresponding project file format
+			VCProjectFileFormat Format = VCProjectFileFormat.Default;
+			foreach (VisualStudioInstallation Installation in Installations)
+			{
+				if (Installation.Compiler == WindowsCompiler.VisualStudio2022)
+				{
+					Format = VCProjectFileFormat.VisualStudio2022;
+					break;
+				}
+			}
+
+			return Format;
+		}
+
+		private bool IncludePlatform(AvailableConfiguration? AvailableConfiguration, UnrealTargetPlatform Platform, IList<UnrealTargetPlatform> FilteredPlatforms)
+		{
+			if (AvailableConfiguration == null || AvailableConfiguration.Platforms.Contains(Platform.ToString()))
+			{
+				FilteredPlatforms.Add(Platform);
+				return true;
+			}
+
+			return false;
+		}
+
+		private bool IncludeConfiguration(AvailableConfiguration? AvailableConfiguration, UnrealTargetConfiguration Configuration, IList<UnrealTargetConfiguration> FilteredConfigurations)
+		{
+			if (AvailableConfiguration == null || AvailableConfiguration.Configurations.Contains(Configuration.ToString()))
+			{
+				FilteredConfigurations.Add(Configuration);
+				return true;
+			}
+
+			return false;
+		}
+
+		private bool WriteVCProjectFileForTarget(
+			DirectoryReference RootPath,
+			ProjectTarget Target,
+			VCProjectFileFormat Format,
+			List<UnrealTargetPlatform> Platforms,
+			List<UnrealTargetConfiguration> Configurations,
+			PlatformProjectGeneratorCollection PlatformProjectGenerators,
+			Dictionary<string, AvailableConfiguration>? ProjectGenerationDescriptionData,
+			ILogger Logger)
+		{
+			AvailableConfiguration? AvailableConfiguration = null;
+			List<UnrealTargetPlatform> FilteredPlatforms = Platforms;
+			List<UnrealTargetConfiguration> FilteredConfigurations = Configurations;
+			FileReference Output = FileReference.Combine(RootPath, $"{Target.TargetFilePath.GetFileNameWithoutAnyExtensions()}.proj");
+
+			try
+			{
+				if (ProjectGenerationDescriptionData != null)
+				{
+					if (!ProjectGenerationDescriptionData.TryGetValue(Target.Name, out AvailableConfiguration))
+					{
+						AvailableConfiguration = new AvailableConfiguration();
+					}
+				}
+
+				if (Target.TargetRules != null)
+				{
+					if (Target.TargetRules.Type == TargetType.Program)
+					{
+						return true;
+					}
+
+					if (Target.TargetRules.Type == TargetType.Editor)
+					{
+						FilteredPlatforms = new List<UnrealTargetPlatform>();
+						IncludePlatform(AvailableConfiguration, BuildHostPlatform.Current.Platform, FilteredPlatforms);
+
+						FilteredConfigurations = new List<UnrealTargetConfiguration>();
+						IncludeConfiguration(AvailableConfiguration, UnrealTargetConfiguration.Debug, FilteredConfigurations);
+						IncludeConfiguration(AvailableConfiguration, UnrealTargetConfiguration.DebugGame, FilteredConfigurations);
+						IncludeConfiguration(AvailableConfiguration, UnrealTargetConfiguration.Development, FilteredConfigurations);
+					}
+				}
+
+				if (AvailableConfiguration != null && FilteredPlatforms == Platforms && FilteredConfigurations == Configurations)
+				{
+					FilteredPlatforms = new List<UnrealTargetPlatform>();
+					FilteredConfigurations = new List<UnrealTargetConfiguration>();
+
+					foreach (var Platform in Platforms)
+					{
+						IncludePlatform(AvailableConfiguration, Platform, FilteredPlatforms);
+					}
+
+					foreach (var Configuration in Configurations)
+					{
+						IncludeConfiguration(AvailableConfiguration, Configuration, FilteredConfigurations);
+					}
+				}
+
+				if (FilteredPlatforms.Count == 0 || FilteredConfigurations.Count == 0)
+				{
+					return true;
+				}
+
+				var Settings = new VCProjectFileSettings();
+				Settings.MaxSharedIncludePaths = 0;
+
+				var VCProject = new VCProjectFile(Output, RootPath, Format, false, true, null, Settings);
+				VCProject.ProjectTargets.Add(Target);
+				if (VCProject.WriteProjectFile(FilteredPlatforms, FilteredConfigurations, PlatformProjectGenerators, Logger))
+				{
+					XmlDocument XmlDoc = new XmlDocument();
+					XmlDoc.Load(Output.FullName);
+
+					List<XmlNode> RemovedChildNodes = new List<XmlNode>();
+
+					Queue<XmlNode> Nodes = new Queue<XmlNode>();
+					Nodes.Enqueue(XmlDoc);
+
+					while (Nodes.Count != 0)
+					{
+						XmlNode Node = Nodes.Dequeue();
+						using (var NodeList = Node.ChildNodes)
+						{
+							if (NodeList != null)
+							{
+								foreach (XmlNode ChildNode in NodeList)
+								{
+									if (ChildNode != null && ChildNode.NodeType != XmlNodeType.Text)
+									{
+										if (NodesAllowList.Contains(ChildNode.Name))
+										{
+											Nodes.Enqueue(ChildNode);
+											continue;
+										}
+										else if (NodesBlockedList.Contains(ChildNode.Name))
+										{
+											RemovedChildNodes.Add(ChildNode);
+											continue;
+										}
+										else
+										{
+											foreach (var RegEx in BlockedNodesListRegEx)
+											{
+												if (RegEx.Match(ChildNode.Name).Success)
+												{
+													RemovedChildNodes.Add(ChildNode);
+													continue;
+												}
+											}
+
+											Nodes.Enqueue(ChildNode);
+											continue;
+										}
+									}
+									else
+									{
+										if (NodesToUpdate.TryGetValue(Node.Name, out string? NewNodeValue) && !string.IsNullOrEmpty(NewNodeValue))
+										{
+											Node.InnerText = NewNodeValue;
+										}
+									}
+								}
+							}
+						}
+
+						foreach (var ChildNode in RemovedChildNodes)
+						{
+							Node.RemoveChild(ChildNode);
+						}
+
+						if (Node.ChildNodes.Count == 0 && Node.Attributes?.Count == 0)
+						{
+							Node.ParentNode?.RemoveChild(Node);
+						}
+
+						RemovedChildNodes.Clear();
+					}
+
+					XmlDoc.Save(Output.FullName);
+				}
+				else
+				{
+					throw new BuildException("Failed to create workspace project file {0}", Output.FullName);
+				}
+			}
+			catch (Exception)
+			{
+				Logger.LogWarning("Failed to create workspace project files. Disabling Visual Studio workspace MSBuild integration for target {0}.", Target.Name);
+				if (File.Exists(Output.FullName))
+				{
+					File.Delete(Output.FullName);
+				}
+
+				return false;
+			}
+
+			return true;
+		}
+
+		private void WriteVCProjectCommonPropsFile(DirectoryReference RootPath, VCProjectFileFormat Format, ILogger Logger)
+		{
+			FileReference Output = FileReference.Combine(RootPath, "UECommon.props");
+
+			try
+			{
+				StringBuilder VCCommonTargetFileContent = new StringBuilder();
+				VCCommonTargetFileContent.AppendLine("<?xml version=\"1.0\" encoding=\"utf-8\"?>");
+				VCCommonTargetFileContent.AppendLine("<Project xmlns=\"http://schemas.microsoft.com/developer/msbuild/2003\">");
+
+				// Project globals (project GUID, project type, SCC bindings, etc)
+				{
+					string ToolVersionString = VCProjectFileGenerator.GetProjectFileToolVersionString(Format);
+					VCCommonTargetFileContent.AppendLine("  <PropertyGroup Label=\"Globals\">");
+					VCCommonTargetFileContent.AppendLine("    <Keyword>MakeFileProj</Keyword>");
+					VCProjectFileGenerator.AppendPlatformToolsetProperty(VCCommonTargetFileContent, Format);
+					VCCommonTargetFileContent.AppendLine("    <MinimumVisualStudioVersion>{0}</MinimumVisualStudioVersion>", ToolVersionString);
+					VCCommonTargetFileContent.AppendLine("    <VCProjectVersion>{0}</VCProjectVersion>", ToolVersionString);
+					VCCommonTargetFileContent.AppendLine("    <NMakeUseOemCodePage>true</NMakeUseOemCodePage>"); // Fixes mojibake with non-Latin character sets (UE-102825)
+					VCCommonTargetFileContent.AppendLine("    <TargetRuntime>Native</TargetRuntime>");
+					VCCommonTargetFileContent.AppendLine("  </PropertyGroup>");
+				}
+
+				// Write the default configuration info
+				VCCommonTargetFileContent.AppendLine("  <PropertyGroup Label=\"Configuration\">");
+				VCCommonTargetFileContent.AppendLine($"    <ConfigurationType>{PlatformProjectGenerator.DefaultPlatformConfigurationType}</ConfigurationType>");
+				VCProjectFileGenerator.AppendPlatformToolsetProperty(VCCommonTargetFileContent, Format);
+				VCCommonTargetFileContent.AppendLine("  </PropertyGroup>");
+
+				VCCommonTargetFileContent.AppendLine("  <Import Project=\"$(VCTargetsPath)\\Microsoft.Cpp.Default.props\" />");
+				VCCommonTargetFileContent.AppendLine("  <Import Project=\"$(VCTargetsPath)\\Microsoft.Cpp.props\" />");
+
+				// Write the common and invalid configuration values
+				{
+					const string InvalidMessage = "echo The selected platform/configuration is not valid for this target.";
+
+					string ProjectRelativeUnusedDirectory = ProjectFile.NormalizeProjectPath(DirectoryReference.Combine(Unreal.EngineDirectory, "Intermediate", "Build", "Unused"));
+
+					VCCommonTargetFileContent.AppendLine("  <PropertyGroup>");
+
+					DirectoryReference BatchFilesDirectory = DirectoryReference.Combine(Unreal.EngineDirectory, "Build", "BatchFiles");
+					VCCommonTargetFileContent.AppendLine("    <BuildBatchScript>{0}</BuildBatchScript>", ProjectFile.EscapePath(ProjectFile.NormalizeProjectPath(FileReference.Combine(BatchFilesDirectory, "Build.bat"))));
+					VCCommonTargetFileContent.AppendLine("    <RebuildBatchScript>{0}</RebuildBatchScript>", ProjectFile.EscapePath(ProjectFile.NormalizeProjectPath(FileReference.Combine(BatchFilesDirectory, "Rebuild.bat"))));
+					VCCommonTargetFileContent.AppendLine("    <CleanBatchScript>{0}</CleanBatchScript>", ProjectFile.EscapePath(ProjectFile.NormalizeProjectPath(FileReference.Combine(BatchFilesDirectory, "Clean.bat"))));
+					VCCommonTargetFileContent.AppendLine("    <NMakeBuildCommandLine>{0}</NMakeBuildCommandLine>", InvalidMessage);
+					VCCommonTargetFileContent.AppendLine("    <NMakeReBuildCommandLine>{0}</NMakeReBuildCommandLine>", InvalidMessage);
+					VCCommonTargetFileContent.AppendLine("    <NMakeCleanCommandLine>{0}</NMakeCleanCommandLine>", InvalidMessage);
+					VCCommonTargetFileContent.AppendLine("    <NMakeOutput>Invalid Output</NMakeOutput>", InvalidMessage);
+					VCCommonTargetFileContent.AppendLine("    <OutDir>{0}{1}</OutDir>", ProjectRelativeUnusedDirectory, Path.DirectorySeparatorChar);
+					VCCommonTargetFileContent.AppendLine("    <IntDir>{0}{1}</IntDir>", ProjectRelativeUnusedDirectory, Path.DirectorySeparatorChar);
+					// NOTE: We are intentionally overriding defaults for these paths with empty strings.  We never want Visual Studio's
+					//       defaults for these fields to be propagated, since they are version-sensitive paths that may not reflect
+					//       the environment that UBT is building in.  We'll set these environment variables ourselves!
+					// NOTE: We don't touch 'ExecutablePath' because that would result in Visual Studio clobbering the system "Path"
+					//       environment variable
+					VCCommonTargetFileContent.AppendLine("    <IncludePath />");
+					VCCommonTargetFileContent.AppendLine("    <ReferencePath />");
+					VCCommonTargetFileContent.AppendLine("    <LibraryPath />");
+					VCCommonTargetFileContent.AppendLine("    <LibraryWPath />");
+					VCCommonTargetFileContent.AppendLine("    <SourcePath />");
+					VCCommonTargetFileContent.AppendLine("    <ExcludePath />");
+
+					// Add all the default system include paths
+					if (OperatingSystem.IsWindows())
+					{
+						if (SupportedPlatforms.Contains(UnrealTargetPlatform.Win64))
+						{
+							VCCommonTargetFileContent.AppendLine("    <DefaultSystemIncludePaths>{0}</DefaultSystemIncludePaths>", VCToolChain.GetVCIncludePaths(UnrealTargetPlatform.Win64, VCProjectFileGenerator.GetCompilerForIntellisense(Format), null, null, Logger));
+						}
+					}
+					else
+					{
+						VCCommonTargetFileContent.AppendLine("    <DefaultSystemIncludePaths />");
+					}
+
+					VCCommonTargetFileContent.AppendLine("  </PropertyGroup>");
+
+				}
+
+				// Write default import group
+				VCCommonTargetFileContent.AppendLine("  <ImportGroup Label=\"PropertySheets\">");
+				VCCommonTargetFileContent.AppendLine("    <Import Project=\"$(UserRootDir)\\Microsoft.Cpp.$(Platform).user.props\" Condition=\"exists('$(UserRootDir)\\Microsoft.Cpp.$(Platform).user.props')\" Label=\"LocalAppDataPlatform\" />");
+				VCCommonTargetFileContent.AppendLine("  </ImportGroup>");
+
+				VCCommonTargetFileContent.AppendLine("</Project>");
+
+				Utils.WriteFileIfChanged(Output, VCCommonTargetFileContent.ToString(), Logger);
+			}
+			catch (Exception)
+			{
+				Logger.LogWarning("Failed to create workspace project file {0}.", Output.FullName);
+			}
+		}
+
+		private bool WriteVsConfigFile(ILogger Logger, DirectoryReference primaryProjectPath)
+		{
+			StringBuilder VsConfigFileContent = new StringBuilder();
+
+			VsConfigFileContent.AppendLine("{");
+			VsConfigFileContent.AppendLine("  \"version\": \"1.0\",");
+			VsConfigFileContent.AppendLine("  \"components\": [");
+			IEnumerable<string> Components = MicrosoftPlatformSDK.GetVisualStudioSuggestedComponents(VCProjectFileFormat.VisualStudio2022);
+			string ComponentsString = String.Join($",{Environment.NewLine}    ", Components.Select(x => $"\"{x}\""));
+			VsConfigFileContent.AppendLine($"    {ComponentsString}");
+			VsConfigFileContent.AppendLine("  ]");
+			VsConfigFileContent.AppendLine("}");
+
+			FileReference VsConfigFileName = FileReference.Combine(primaryProjectPath, ".vsconfig");
+			return WriteFileIfChanged(VsConfigFileName.FullName, VsConfigFileContent.ToString(), Logger);
+		}
 	}
 }
