Refactor Compiler warnings controlled via module rules or target rules

diff --git a/Engine/Source/Programs/Shared/UnrealBuildTool.Tests/CppCompileWarningsTests.cs b/Engine/Source/Programs/Shared/UnrealBuildTool.Tests/CppCompileWarningsTests.cs
new file mode 100644
--- /dev/null
+++ b/Engine/Source/Programs/Shared/UnrealBuildTool.Tests/CppCompileWarningsTests.cs
@@ -0,0 +1,155 @@
+// Copyright Epic Games, Inc. All Rights Reserved.
+
+using Microsoft.VisualStudio.TestTools.UnitTesting;
+
+namespace UnrealBuildTool.Tests
+{
+	[TestClass]
+	public class CppCompileWarningsTests
+	{
+		internal class TestBuildSettingsProvider : IBuildContextProvider
+		{
+			internal BuildSettingsVersion _testBuildSettings = BuildSettingsVersion.V1;
+
+			public BuildSettingsVersion GetBuildSettings()
+			{
+				return _testBuildSettings;
+			}
+		}
+
+		[TestMethod]
+		public void TestContextBuildSettingsOverride()
+		{
+			TestBuildSettingsProvider buildContextProvider = new TestBuildSettingsProvider();
+			buildContextProvider._testBuildSettings = BuildSettingsVersion.V1;
+
+			CppCompileWarnings newWarnings = new CppCompileWarnings(buildContextProvider, null);
+
+			// Shadow Variable build setting defaults
+			{
+				Assert.AreEqual(WarningLevel.Default, newWarnings.ShadowVariableWarningLevel);
+
+				// Boundary analysis
+				buildContextProvider._testBuildSettings = BuildSettingsVersion.V2;
+				Assert.AreEqual(WarningLevel.Error, newWarnings.ShadowVariableWarningLevel);
+
+				buildContextProvider._testBuildSettings = BuildSettingsVersion.V3;
+				Assert.AreEqual(WarningLevel.Error, newWarnings.ShadowVariableWarningLevel);
+			}
+		}
+
+		[TestMethod]
+		public void TestParentSettings()
+		{
+			TestBuildSettingsProvider buildContextProvider = new TestBuildSettingsProvider();
+			buildContextProvider._testBuildSettings = BuildSettingsVersion.V1;
+			CppCompileWarnings parentWarnings = new CppCompileWarnings(buildContextProvider, null);
+			CppCompileWarnings newWarnings = new CppCompileWarnings(buildContextProvider, null, new ReadOnlyCppCompileWarnings(parentWarnings));
+
+			// Assert invariants
+			{
+				Assert.AreEqual(WarningLevel.Default, newWarnings.ShadowVariableWarningLevel);
+				Assert.AreEqual(WarningLevel.Default, newWarnings.UnsafeTypeCastWarningLevel);
+				Assert.AreEqual(WarningLevel.Default, newWarnings.UndefinedIdentifierWarningLevel);
+			}
+
+			// Assert new parent values
+			{
+				parentWarnings.ShadowVariableWarningLevel = WarningLevel.Error;
+				parentWarnings.UndefinedIdentifierWarningLevel = WarningLevel.Error;
+				parentWarnings.UnsafeTypeCastWarningLevel = WarningLevel.Error;
+
+				Assert.AreEqual(WarningLevel.Error, newWarnings.ShadowVariableWarningLevel);
+				Assert.AreEqual(WarningLevel.Error, newWarnings.UnsafeTypeCastWarningLevel);
+				Assert.AreEqual(WarningLevel.Error, newWarnings.UndefinedIdentifierWarningLevel);
+			}
+		}
+
+		[TestMethod]
+		public void TestReparentOperations()
+		{
+			TestBuildSettingsProvider buildContextProvider = new TestBuildSettingsProvider();
+			buildContextProvider._testBuildSettings = BuildSettingsVersion.V1;
+			CppCompileWarnings parentWarnings = new CppCompileWarnings(buildContextProvider, null);
+			CppCompileWarnings newWarnings = new CppCompileWarnings(buildContextProvider, null, new ReadOnlyCppCompileWarnings(parentWarnings));
+
+			// Assert invariant parent values
+			{
+				parentWarnings.ShadowVariableWarningLevel = WarningLevel.Error;
+				parentWarnings.UndefinedIdentifierWarningLevel = WarningLevel.Error;
+				parentWarnings.UnsafeTypeCastWarningLevel = WarningLevel.Error;
+
+				Assert.AreEqual(WarningLevel.Error, newWarnings.ShadowVariableWarningLevel);
+				Assert.AreEqual(WarningLevel.Error, newWarnings.UnsafeTypeCastWarningLevel);
+				Assert.AreEqual(WarningLevel.Error, newWarnings.UndefinedIdentifierWarningLevel);
+			}
+
+			// Assert reparenting
+			CppCompileWarnings parentWarnings2 = new CppCompileWarnings(buildContextProvider, null);
+			{
+				newWarnings.AddParent(parentWarnings2);
+				parentWarnings2.ShadowVariableWarningLevel = WarningLevel.Warning;
+				parentWarnings2.UndefinedIdentifierWarningLevel = WarningLevel.Warning;
+				parentWarnings2.UnsafeTypeCastWarningLevel = WarningLevel.Warning;
+
+				Assert.AreEqual(WarningLevel.Warning, newWarnings.ShadowVariableWarningLevel);
+				Assert.AreEqual(WarningLevel.Warning, newWarnings.UnsafeTypeCastWarningLevel);
+				Assert.AreEqual(WarningLevel.Warning, newWarnings.UndefinedIdentifierWarningLevel);
+			}
+
+			// Assert transitive to eldest ancestor
+			{
+				parentWarnings2.ShadowVariableWarningLevel = WarningLevel.Default;
+				Assert.AreEqual(parentWarnings.ShadowVariableWarningLevel, newWarnings.ShadowVariableWarningLevel);
+			}
+
+			// Assert removal of parentWarnings2 - UnsafeTypeCastWarningLevel & UndefinedIdentifierWarningLevel will now be the reoslved items
+			{
+				newWarnings.RemoveParent();
+				Assert.AreEqual(parentWarnings.UnsafeTypeCastWarningLevel, newWarnings.UnsafeTypeCastWarningLevel);
+				Assert.AreEqual(parentWarnings.UndefinedIdentifierWarningLevel, newWarnings.UndefinedIdentifierWarningLevel);
+
+			}
+
+			// Removal of only parent; should only be default values
+			{
+				newWarnings.RemoveParent();
+				Assert.AreEqual(WarningLevel.Default, newWarnings.ShadowVariableWarningLevel);
+				Assert.AreEqual(WarningLevel.Default, newWarnings.UnsafeTypeCastWarningLevel);
+				Assert.AreEqual(WarningLevel.Default, newWarnings.UndefinedIdentifierWarningLevel);
+			}
+		}
+
+		[TestMethod]
+		public void TestChildOverrides()
+		{
+			TestBuildSettingsProvider buildContextProvider = new TestBuildSettingsProvider();
+			buildContextProvider._testBuildSettings = BuildSettingsVersion.V1;
+			CppCompileWarnings parentWarnings = new CppCompileWarnings(buildContextProvider, null);
+			CppCompileWarnings newWarnings = new CppCompileWarnings(buildContextProvider, null, new ReadOnlyCppCompileWarnings(parentWarnings));
+
+			parentWarnings.ShadowVariableWarningLevel = WarningLevel.Error;
+			parentWarnings.UndefinedIdentifierWarningLevel = WarningLevel.Error;
+			parentWarnings.UnsafeTypeCastWarningLevel = WarningLevel.Error;
+
+			// Assert invariants
+			{
+				Assert.AreEqual(WarningLevel.Error, newWarnings.ShadowVariableWarningLevel);
+				Assert.AreEqual(WarningLevel.Error, newWarnings.UnsafeTypeCastWarningLevel);
+				Assert.AreEqual(WarningLevel.Error, newWarnings.UndefinedIdentifierWarningLevel);
+			}
+
+			// Child overrides;
+			newWarnings.ShadowVariableWarningLevel = WarningLevel.Warning;
+			newWarnings.UndefinedIdentifierWarningLevel = WarningLevel.Warning;
+			newWarnings.UnsafeTypeCastWarningLevel = WarningLevel.Warning;
+
+			// Assert new parent values
+			{
+				Assert.AreEqual(WarningLevel.Warning, newWarnings.ShadowVariableWarningLevel);
+				Assert.AreEqual(WarningLevel.Warning, newWarnings.UnsafeTypeCastWarningLevel);
+				Assert.AreEqual(WarningLevel.Warning, newWarnings.UndefinedIdentifierWarningLevel);
+			}
+		}
+	}
+}
\ No newline at end of file
diff --git a/Engine/Source/Programs/UnrealBuildTool/Configuration/BuildSystemContext.cs b/Engine/Source/Programs/UnrealBuildTool/Configuration/BuildSystemContext.cs
new file mode 100644
--- /dev/null
+++ b/Engine/Source/Programs/UnrealBuildTool/Configuration/BuildSystemContext.cs
@@ -0,0 +1,72 @@
+namespace UnrealBuildTool
+{
+	/// <summary>
+	/// Interface that defines standard build context parameters that can influence the resolution of build warnings.
+	/// </summary>
+	public interface IBuildContextProvider
+	{
+		/// <summary>
+		/// Obtains the <see cref="BuildSettingsVersion"/> associated with the provided context provider.
+		/// </summary>
+		/// <returns>The <see cref="BuildSettingsVersion"/>.</returns>
+		public BuildSettingsVersion GetBuildSettings();
+	}
+
+	#region -- Adapters for IBuildContextProvider -- 
+
+	internal class ModuleRulesBuildSettingsProvider : IBuildContextProvider
+	{
+		private readonly ModuleRules _moduleRules;
+
+		public static implicit operator ModuleRulesBuildSettingsProvider(ModuleRules moduleRules)
+		{
+			return new ModuleRulesBuildSettingsProvider(moduleRules);
+		}
+
+		public ModuleRulesBuildSettingsProvider(ModuleRules moduleRules)
+		{
+			_moduleRules = moduleRules;
+		}
+
+		public BuildSettingsVersion GetBuildSettings()
+		{
+			return _moduleRules != null ? _moduleRules.DefaultBuildSettings : (default);
+		}
+	}
+
+	internal class TargetRulesBuildSettingsProvider : IBuildContextProvider
+	{
+		private readonly TargetRules _targetRules;
+
+		public static implicit operator TargetRulesBuildSettingsProvider(TargetRules targetRules)
+		{
+			return new TargetRulesBuildSettingsProvider(targetRules);
+		}
+
+		public TargetRulesBuildSettingsProvider(TargetRules targetRules)
+		{
+			_targetRules = targetRules;
+		}
+
+		public BuildSettingsVersion GetBuildSettings()
+		{
+			return _targetRules.DefaultBuildSettings;
+		}
+	}
+
+	#endregion -- Adapters for IBuildContextProvider --
+
+	internal class BuildSystemContext
+	{
+		internal BuildSystemContext(IBuildContextProvider context, TargetRules? targetRules = null, ModuleRules? moduleRules = null)
+		{
+			_targetRulesPrivate = targetRules;
+			_moduleRulesPrivate = moduleRules;
+			_buildContext = context;
+		}
+
+		internal TargetRules? _targetRulesPrivate;
+		internal ModuleRules? _moduleRulesPrivate;
+		internal IBuildContextProvider _buildContext;
+	}
+}
\ No newline at end of file
diff --git a/Engine/Source/Programs/UnrealBuildTool/Configuration/CppCompileWarnings.cs b/Engine/Source/Programs/UnrealBuildTool/Configuration/CppCompileWarnings.cs
new file mode 100644
--- /dev/null
+++ b/Engine/Source/Programs/UnrealBuildTool/Configuration/CppCompileWarnings.cs
@@ -0,0 +1,200 @@
+using EpicGames.Core;
+using Microsoft.Extensions.Logging;
+using System;
+
+namespace UnrealBuildTool
+{
+	internal static class WarningLevelExtensions
+	{
+		/// <summary>
+		/// Applies a warning on top of another, given the new <paramref name="appliedWarningLevel"/> has a value, and it is not the <see cref="WarningLevel.Default"/>.
+		/// </summary>
+		/// <param name="inSourceWarningLevel">The source warning level to override.</param>
+		/// <param name="appliedWarningLevel">The warning level to apply to the source.</param>
+		/// <remarks>If <paramref name="appliedWarningLevel"/> with value of null or <see cref="WarningLevel.Default"/> will be ignored as an override.</remarks>
+		internal static void ApplyWarning(this ref WarningLevel inSourceWarningLevel, WarningLevel? appliedWarningLevel)
+		{
+			if (appliedWarningLevel.HasValue && appliedWarningLevel.Value != WarningLevel.Default)
+			{
+				inSourceWarningLevel = appliedWarningLevel.Value;
+			}
+		}
+	}
+
+	/// <summary>
+	/// Container class used for C++ compiler warning settings.
+	/// </summary>
+	public class CppCompileWarnings
+	{
+		#region -- Compiler Warning Settings --
+
+		/// <summary>
+		/// Forces shadow variable warnings to be treated as errors on platforms that support it.
+		/// </summary>
+		[CommandLine("-ShadowVariableErrors", Value = nameof(WarningLevel.Error))]
+		[RequiresUniqueBuildEnvironment]
+		public WarningLevel ShadowVariableWarningLevel
+		{
+			get => ResolveWarning((BuildSettingsVersion x) => (x >= BuildSettingsVersion.V2) ? WarningLevel.Error : (ParentCppCompileWarnings != null ? ParentCppCompileWarnings.ShadowVariableWarningLevel : WarningLevel.Default), ParentCppCompileWarnings?.ShadowVariableWarningLevel, _shadowVariableWarningLevelPrivate);
+			set => _shadowVariableWarningLevelPrivate = value;
+		}
+
+		private WarningLevel _shadowVariableWarningLevelPrivate;
+
+		/// <summary>
+		/// Indicates what warning/error level to treat unsafe type casts as on platforms that support it (e.g., double->float or int64->int32)
+		/// </summary>
+		[XmlConfigFile(Category = "BuildConfiguration")]
+		[RequiresUniqueBuildEnvironment]
+		public WarningLevel UnsafeTypeCastWarningLevel
+		{
+			get => ResolveWarning(null, ParentCppCompileWarnings?.UnsafeTypeCastWarningLevel, _unsafeTypeCastWarningLevelPrivate);
+			set => _unsafeTypeCastWarningLevelPrivate = value;
+		}
+
+		private WarningLevel _unsafeTypeCastWarningLevelPrivate;
+
+		/// <summary>
+		/// Indicates what warning/error level to treat undefined identifiers in conditional expressions.
+		/// </summary>
+		[XmlConfigFile(Category = "BuildConfiguration")]
+		[CommandLine("-UndefinedIdentifierWarningLevel=")]
+		[RequiresUniqueBuildEnvironment]
+		public WarningLevel UndefinedIdentifierWarningLevel
+		{
+			get => ResolveWarning(null, ParentCppCompileWarnings?.UndefinedIdentifierWarningLevel, _undefinedIdentifierWarningLevelPrivate);
+			set => _undefinedIdentifierWarningLevelPrivate = value;
+		}
+
+		private WarningLevel _undefinedIdentifierWarningLevelPrivate;
+
+
+		#endregion  -- Compiler Warning Settings --
+
+		#region -- Constructors --
+
+		/// <summary>
+		/// CppCompilerWarnings constructor based off of <see cref="TargetRules"/> context.
+		/// </summary>
+		/// <param name="targetRulesContextObject">The <see cref="TargetRules"/> context to construct the warnings object from.</param>
+		/// <param name="logger">Logger for output.</param>
+		/// <remarks>The <see cref="TargetRules"/> implicitly has no default parent context.</remarks>
+		public CppCompileWarnings(TargetRules targetRulesContextObject, ILogger? logger)
+		{
+			BuildContextProvider = new BuildSystemContext(new TargetRulesBuildSettingsProvider(targetRulesContextObject), targetRulesContextObject, null);
+			ParentCppCompileWarnings = null;
+			_logger = logger;
+
+			ApplyTargetDefaults();
+		}
+
+		/// <summary>
+		/// CppCompilerWarnings constructor based off of <see cref="ModuleRules"/> context.
+		/// </summary>
+		/// <param name="moduleRulesContextObject">The <see cref="ModuleRules"/> context to construct the warnings object from.</param>
+		/// <param name="logger">Logger for output.</param>
+		/// <remarks>The provided module implicitly has a child relationship to it's <see cref="ModuleRules.Target"/>, and will yield to it's warning specifications unless overridden.</remarks>
+		public CppCompileWarnings(ModuleRules moduleRulesContextObject, ILogger? logger)
+		{
+			BuildContextProvider = new BuildSystemContext(new ModuleRulesBuildSettingsProvider(moduleRulesContextObject), null, moduleRulesContextObject);
+			ParentCppCompileWarnings = moduleRulesContextObject.Target.CppCompileWarningSettings;
+			_logger = logger;
+		}
+
+		/// <summary>
+		/// CppCompilerWarnings constructor based off of the explicit relationship between the newly constructed <see cref="CppCompileWarnings"/> object, and the provided parent <see cref="IBuildContextProvider"/> object.
+		/// </summary>
+		/// <param name="contextObject">The parent context which will be used as the default.</param>
+		/// <param name="logger">Logger for output.</param>
+		/// <param name="parentCompilerWarnings">The parent compiler warnings to associated with this new compiler warnings instance.</param>
+		public CppCompileWarnings(IBuildContextProvider contextObject, ILogger? logger, ReadOnlyCppCompileWarnings? parentCompilerWarnings = null)
+		{
+			BuildContextProvider = new BuildSystemContext(contextObject);
+			ParentCppCompileWarnings = parentCompilerWarnings;
+			_logger = logger;
+		}
+
+		#endregion -- Constructors --
+
+		#region -- Public API --
+
+		/// <summary>
+		/// Adds the provided <see cref="CppCompileWarnings"/> as a direct ancestor, promoting previous parent if it exists.
+		/// </summary>
+		/// <param name="newParentCppCompilerWarnings">The new direct parent.</param>
+		public void AddParent(CppCompileWarnings newParentCppCompilerWarnings)
+		{
+			if (ParentCppCompileWarnings == null)
+			{
+				ParentCppCompileWarnings = new ReadOnlyCppCompileWarnings(newParentCppCompilerWarnings);
+				return;
+			}
+			
+			_logger?.LogTrace("Promoting existing parent CppCompileWarnings, and inserting new immediate parent.");
+			newParentCppCompilerWarnings.ParentCppCompileWarnings = ParentCppCompileWarnings;
+
+			ParentCppCompileWarnings = new ReadOnlyCppCompileWarnings(newParentCppCompilerWarnings);
+		}
+
+		/// <summary>
+		/// Removes the immediate parent <see cref="CppCompileWarnings"/>, using the next ancestor.
+		/// </summary>
+		public void RemoveParent()
+		{
+			if (ParentCppCompileWarnings == null)
+			{
+				return;
+			}
+
+			_logger?.LogTrace("Removing parent CppCompileWarnings.");
+			ReadOnlyCppCompileWarnings? ancestor = ParentCppCompileWarnings.ParentCppCompileWarnings;
+			ParentCppCompileWarnings = ancestor;
+		}
+
+		#endregion -- Public API --
+
+		#region -- Private API --
+
+		/// <summary>
+		/// General purpose resolver that will resolve a warning based off of the precedence of 1. locally set in the <see cref="CppCompileWarnings"/> object; 2. utilize the parent if it exists; 3. a default as evaluated by the provided build settings delegate.
+		/// </summary>
+		/// <param name="buildSettingsEvaluator">The build settings evaluator to use to ascibe a build default if there is no instance setting.</param>
+		/// <param name="parentWarningLevel">The parent <see cref="CppCompileWarnings"/> <see cref="WarningLevel"/>.</param>
+		/// <param name="instanceWarningLevel">The instance <see cref="WarningLevel"/>.</param>
+		/// <returns>The resolved warning level.</returns>
+		private WarningLevel ResolveWarning(Func<BuildSettingsVersion, WarningLevel>? buildSettingsEvaluator, WarningLevel? parentWarningLevel, WarningLevel instanceWarningLevel)
+		{
+			WarningLevel resolvedWarning = WarningLevel.Default;
+
+			if (buildSettingsEvaluator != null && instanceWarningLevel == WarningLevel.Default)
+			{
+				resolvedWarning.ApplyWarning(buildSettingsEvaluator(BuildContextProvider?._buildContext != null ? BuildContextProvider._buildContext.GetBuildSettings() : BuildSettingsVersion.Latest));
+			}
+			else
+			{
+				resolvedWarning.ApplyWarning(parentWarningLevel);
+				resolvedWarning.ApplyWarning(instanceWarningLevel);
+			}
+
+			return resolvedWarning;
+		}
+
+		private void ApplyTargetDefaults()
+		{
+			_logger?.LogTrace("Applying default C++ Compiler warnings for the Target({TargetName}) context.", BuildContextProvider?._targetRulesPrivate?.Name);
+			_shadowVariableWarningLevelPrivate = WarningLevel.Warning;
+			_unsafeTypeCastWarningLevelPrivate = WarningLevel.Off;
+			_undefinedIdentifierWarningLevelPrivate = WarningLevel.Error;
+		}
+
+		#endregion -- Private API --
+
+		#region -- Internal & Private Instance State --
+
+		private readonly ILogger? _logger;
+		private BuildSystemContext BuildContextProvider { get; set; }
+		internal ReadOnlyCppCompileWarnings? ParentCppCompileWarnings { get; private set; }
+
+		#endregion -- Internal & Private Instance State --
+	}
+}
\ No newline at end of file
diff --git a/Engine/Source/Programs/UnrealBuildTool/Configuration/ModuleRules.cs b/Engine/Source/Programs/UnrealBuildTool/Configuration/ModuleRules.cs
--- a/Engine/Source/Programs/UnrealBuildTool/Configuration/ModuleRules.cs
+++ b/Engine/Source/Programs/UnrealBuildTool/Configuration/ModuleRules.cs
@@ -936,18 +936,20 @@ namespace UnrealBuildTool
 		/// <inheritdoc cref="DeterministicWarningLevelPrivate"/>
 		private WarningLevel DeterministicWarningLevelPrivate;
 
+		/// <summary>
+		/// Settings object used by the module.
+		/// </summary>
+		public CppCompileWarnings CppCompileWarningSettings { get; init; }
+
 		/// <summary>
 		/// How to treat shadow variable warnings
 		/// </summary>
 		public WarningLevel ShadowVariableWarningLevel
 		{
-			get => (ShadowVariableWarningLevelPrivate == WarningLevel.Default) ? ((DefaultBuildSettings >= BuildSettingsVersion.V2) ? WarningLevel.Error : Target.ShadowVariableWarningLevel) : ShadowVariableWarningLevelPrivate;
-			set => ShadowVariableWarningLevelPrivate = value;
+			get => CppCompileWarningSettings.ShadowVariableWarningLevel;
+			set => CppCompileWarningSettings.ShadowVariableWarningLevel = value;
 		}
 
-		/// <inheritdoc cref="ShadowVariableWarningLevelPrivate"/>
-		private WarningLevel ShadowVariableWarningLevelPrivate;
-
 		/// <summary>
 		/// Whether to enable all warnings as errors. UE enables most warnings as errors already, but disables a few (such as deprecation warnings).
 		/// </summary>`
@@ -958,25 +960,20 @@ namespace UnrealBuildTool
 		/// </summary>
 		public WarningLevel UnsafeTypeCastWarningLevel
 		{
-			get => (UnsafeTypeCastWarningLevelPrivate == WarningLevel.Default) ? Target.UnsafeTypeCastWarningLevel : UnsafeTypeCastWarningLevelPrivate;
-			set => UnsafeTypeCastWarningLevelPrivate = value;
+			get => CppCompileWarningSettings.UnsafeTypeCastWarningLevel;
+			set => CppCompileWarningSettings.UnsafeTypeCastWarningLevel = value;
 		}
 
-		/// <inheritdoc cref="UnsafeTypeCastWarningLevel"/>
-		private WarningLevel UnsafeTypeCastWarningLevelPrivate;
 
 		/// <summary>
 		/// Indicates what warning/error level to treat undefined identifiers in conditional expressions.
 		/// </summary>
 		public WarningLevel UndefinedIdentifierWarningLevel
 		{
-			get => (_undefinedIdentifierWarningLevelPrivate == WarningLevel.Default) ? Target.UndefinedIdentifierWarningLevel : _undefinedIdentifierWarningLevelPrivate;
-			set => _undefinedIdentifierWarningLevelPrivate = value;
+			get => CppCompileWarningSettings.UndefinedIdentifierWarningLevel;
+			set => CppCompileWarningSettings.UndefinedIdentifierWarningLevel = value;
 		}
 
-		/// <inheritdoc cref="UndefinedIdentifierWarningLevel"/>
-		private WarningLevel _undefinedIdentifierWarningLevelPrivate;
-
 		/// <summary>
 		/// Enable warnings for using undefined identifiers in #if expressions
 		/// </summary>
@@ -1516,6 +1513,7 @@ namespace UnrealBuildTool
 		public ModuleRules(ReadOnlyTargetRules target)
 		{
 			Target = target;
+			CppCompileWarningSettings = new CppCompileWarnings(this, Logger);
 		}
 #nullable restore
 
diff --git a/Engine/Source/Programs/UnrealBuildTool/Configuration/ReadOnlyCppCompileWarnings.cs b/Engine/Source/Programs/UnrealBuildTool/Configuration/ReadOnlyCppCompileWarnings.cs
new file mode 100644
--- /dev/null
+++ b/Engine/Source/Programs/UnrealBuildTool/Configuration/ReadOnlyCppCompileWarnings.cs
@@ -0,0 +1,35 @@
+namespace UnrealBuildTool
+{
+	/// <summary>
+	/// Read-only wrapper around an existing CppCompileWarnings instance. This exposes CppCompileWarnings settings without letting them be modified.
+	/// </summary>
+	public class ReadOnlyCppCompileWarnings
+	{
+		/// <summary>
+		/// The writeable CppCompileWarnings instance
+		/// </summary>
+		readonly CppCompileWarnings _inner;
+		internal ReadOnlyCppCompileWarnings? ParentCppCompileWarnings { get; init; }
+
+		/// <summary>
+		/// Constructor
+		/// </summary>
+		/// <param name="inner">The CppCompileWarnings instance to wrap around</param>
+		public ReadOnlyCppCompileWarnings(CppCompileWarnings inner)
+		{
+			_inner = inner;
+			ParentCppCompileWarnings = inner.ParentCppCompileWarnings;
+		}
+
+		/// <summary>
+		/// Accessors for fields on the inner CppCompileWarnings instance
+		/// </summary>
+		#region Read-only accessor properties
+#pragma warning disable CS1591
+		public WarningLevel ShadowVariableWarningLevel => _inner.ShadowVariableWarningLevel;
+		public WarningLevel UndefinedIdentifierWarningLevel => _inner.UndefinedIdentifierWarningLevel;
+		public WarningLevel UnsafeTypeCastWarningLevel => _inner.UnsafeTypeCastWarningLevel;
+#pragma warning restore C1591
+		#endregion
+	}
+}
diff --git a/Engine/Source/Programs/UnrealBuildTool/Configuration/ReadOnlyTargetRules.cs b/Engine/Source/Programs/UnrealBuildTool/Configuration/ReadOnlyTargetRules.cs
--- a/Engine/Source/Programs/UnrealBuildTool/Configuration/ReadOnlyTargetRules.cs
+++ b/Engine/Source/Programs/UnrealBuildTool/Configuration/ReadOnlyTargetRules.cs
@@ -32,6 +32,7 @@ namespace UnrealBuildTool
 			LinuxPlatform = new ReadOnlyLinuxTargetRules(inner.LinuxPlatform);
 			MacPlatform = new ReadOnlyMacTargetRules(inner.MacPlatform);
 			WindowsPlatform = new ReadOnlyWindowsTargetRules(inner.WindowsPlatform);
+			CppCompileWarningSettings = new ReadOnlyCppCompileWarnings(inner.CppCompileWarningSettings);
 		}
 
 		/// <summary>
@@ -358,6 +359,8 @@ namespace UnrealBuildTool
 
 		public WarningLevel DeprecationWarningLevel => Inner.DeprecationWarningLevel;
 
+		public ReadOnlyCppCompileWarnings CppCompileWarningSettings { get; init; }
+
 		public WarningLevel ShadowVariableWarningLevel => Inner.ShadowVariableWarningLevel;
 
 		public WarningLevel UnsafeTypeCastWarningLevel => Inner.UnsafeTypeCastWarningLevel;
diff --git a/Engine/Source/Programs/UnrealBuildTool/Configuration/TargetRules.cs b/Engine/Source/Programs/UnrealBuildTool/Configuration/TargetRules.cs
--- a/Engine/Source/Programs/UnrealBuildTool/Configuration/TargetRules.cs
+++ b/Engine/Source/Programs/UnrealBuildTool/Configuration/TargetRules.cs
@@ -1258,7 +1258,7 @@ namespace UnrealBuildTool
 				{
 					return bForceEnableTracePrivate.Value;
 				}
-				
+
 				return false;
 			}
 			set => bForceEnableTracePrivate = value;
@@ -1826,12 +1826,20 @@ namespace UnrealBuildTool
 		[XmlConfigFile(Category = "BuildConfiguration", Name = nameof(DeprecationWarningLevel))]
 		private WarningLevel DeprecationWarningLevelPrivate;
 
+		/// <summary>
+		/// C++ Warnings settings object used by the target
+		/// </summary>
+		[ConfigSubObject]
+		public CppCompileWarnings CppCompileWarningSettings {get; init;}
+
 		/// <summary>
 		/// Forces shadow variable warnings to be treated as errors on platforms that support it.
 		/// </summary>
-		[CommandLine("-ShadowVariableErrors", Value = nameof(WarningLevel.Error))]
-		[RequiresUniqueBuildEnvironment]
-		public WarningLevel ShadowVariableWarningLevel { get; set; } = WarningLevel.Warning;
+		public WarningLevel ShadowVariableWarningLevel
+		{
+			get => CppCompileWarningSettings.ShadowVariableWarningLevel;
+			set => CppCompileWarningSettings.ShadowVariableWarningLevel = value;
+		}
 
 		/// <summary>
 		/// Whether to enable all warnings as errors. UE enables most warnings as errors already, but disables a few (such as deprecation warnings).
@@ -1844,9 +1852,11 @@ namespace UnrealBuildTool
 		/// <summary>
 		/// Indicates what warning/error level to treat unsafe type casts as on platforms that support it (e.g., double->float or int64->int32)
 		/// </summary>
-		[XmlConfigFile(Category = "BuildConfiguration")]
-		[RequiresUniqueBuildEnvironment]
-		public WarningLevel UnsafeTypeCastWarningLevel { get; set; } = WarningLevel.Off;
+		public WarningLevel UnsafeTypeCastWarningLevel 
+		{
+			get => CppCompileWarningSettings.UnsafeTypeCastWarningLevel;
+			set => CppCompileWarningSettings.UnsafeTypeCastWarningLevel = value;
+		}
 
 		/// <summary>
 		/// Forces the use of undefined identifiers in conditional expressions to be treated as errors.
@@ -1858,10 +1868,11 @@ namespace UnrealBuildTool
 		/// <summary>
 		/// Indicates what warning/error level to treat undefined identifiers in conditional expressions.
 		/// </summary>
-		[XmlConfigFile(Category = "BuildConfiguration")]
-		[CommandLine("-UndefinedIdentifierWarningLevel=")]
-		[RequiresUniqueBuildEnvironment]
-		public WarningLevel UndefinedIdentifierWarningLevel { get; set; } = WarningLevel.Error;
+		public WarningLevel UndefinedIdentifierWarningLevel 
+		{
+			get => CppCompileWarningSettings.UndefinedIdentifierWarningLevel;
+			set => CppCompileWarningSettings.UndefinedIdentifierWarningLevel = value;
+		}
 
 		/// <summary>
 		/// Indicates what warning/error level to treat potential PCH performance issues.
@@ -3061,6 +3072,7 @@ namespace UnrealBuildTool
 			ProjectFile = target.ProjectFile;
 			Version = target.Version;
 			WindowsPlatform = new WindowsTargetRules(this);
+			CppCompileWarningSettings = new CppCompileWarnings(this, Logger);
 
 			// Make sure the logger was initialized by the caller
 			if (Logger == null)
diff --git a/Engine/Source/Programs/UnrealBuildTool/Configuration/UEBuildModuleCPP.cs b/Engine/Source/Programs/UnrealBuildTool/Configuration/UEBuildModuleCPP.cs
--- a/Engine/Source/Programs/UnrealBuildTool/Configuration/UEBuildModuleCPP.cs
+++ b/Engine/Source/Programs/UnrealBuildTool/Configuration/UEBuildModuleCPP.cs
@@ -2179,9 +2179,9 @@ namespace UnrealBuildTool
 			Result.bEnableObjCExceptions |= Rules.bEnableObjCExceptions;
 			Result.bEnableObjCAutomaticReferenceCounting = Rules.bEnableObjCAutomaticReferenceCounting;
 			Result.bWarningsAsErrors |= Rules.bWarningsAsErrors;
-			Result.ShadowVariableWarningLevel = Rules.ShadowVariableWarningLevel;
-			Result.UnsafeTypeCastWarningLevel = Rules.UnsafeTypeCastWarningLevel;
-			Result.UndefinedIdentifierWarningLevel = Rules.UndefinedIdentifierWarningLevel;
+			Result.ShadowVariableWarningLevel = Rules.CppCompileWarningSettings.ShadowVariableWarningLevel;
+			Result.UnsafeTypeCastWarningLevel = Rules.CppCompileWarningSettings.UnsafeTypeCastWarningLevel;
+			Result.UndefinedIdentifierWarningLevel = Rules.CppCompileWarningSettings.UndefinedIdentifierWarningLevel;
 			Result.bDisableStaticAnalysis = Rules.bDisableStaticAnalysis || (Target.bStaticAnalyzerProjectOnly && Rules.bTreatAsEngineModule);
 			Result.bStaticAnalyzerExtensions = Rules.bStaticAnalyzerExtensions;
 			Result.StaticAnalyzerRulesets = Rules.StaticAnalyzerRulesets;
diff --git a/Engine/Source/Programs/UnrealBuildTool/Configuration/UEBuildTarget.cs b/Engine/Source/Programs/UnrealBuildTool/Configuration/UEBuildTarget.cs
--- a/Engine/Source/Programs/UnrealBuildTool/Configuration/UEBuildTarget.cs
+++ b/Engine/Source/Programs/UnrealBuildTool/Configuration/UEBuildTarget.cs
@@ -5225,9 +5225,9 @@ namespace UnrealBuildTool
 			GlobalCompileEnvironment.bEnableObjCExceptions = Rules.bForceEnableObjCExceptions || Rules.bCompileAgainstEditor;
 			GlobalCompileEnvironment.DefaultWarningLevel = Rules.DefaultWarningLevel;
 			GlobalCompileEnvironment.DeprecationWarningLevel = Rules.DeprecationWarningLevel;
-			GlobalCompileEnvironment.ShadowVariableWarningLevel = Rules.ShadowVariableWarningLevel;
-			GlobalCompileEnvironment.UnsafeTypeCastWarningLevel = Rules.UnsafeTypeCastWarningLevel;
-			GlobalCompileEnvironment.UndefinedIdentifierWarningLevel = Rules.UndefinedIdentifierWarningLevel;
+			GlobalCompileEnvironment.ShadowVariableWarningLevel = Rules.CppCompileWarningSettings.ShadowVariableWarningLevel;
+			GlobalCompileEnvironment.UnsafeTypeCastWarningLevel = Rules.CppCompileWarningSettings.UnsafeTypeCastWarningLevel;
+			GlobalCompileEnvironment.UndefinedIdentifierWarningLevel = Rules.CppCompileWarningSettings.UndefinedIdentifierWarningLevel;
 			GlobalCompileEnvironment.bRetainFramePointers = Rules.bRetainFramePointers;
 			GlobalCompileEnvironment.bWarningsAsErrors = Rules.bWarningsAsErrors;
 			GlobalCompileEnvironment.OptimizationLevel = Rules.OptimizationLevel;
