Missing Copyright boilerplate

diff --git a/Engine/Source/Programs/UnrealBuildTool/Configuration/BuildSystemContext.cs b/Engine/Source/Programs/UnrealBuildTool/Configuration/BuildSystemContext.cs
--- a/Engine/Source/Programs/UnrealBuildTool/Configuration/BuildSystemContext.cs
+++ b/Engine/Source/Programs/UnrealBuildTool/Configuration/BuildSystemContext.cs
@@ -1,3 +1,5 @@
+// Copyright Epic Games, Inc. All Rights Reserved.
+
 namespace UnrealBuildTool
 {
 	/// <summary>
diff --git a/Engine/Source/Programs/UnrealBuildTool/Configuration/CppCompileWarnings.cs b/Engine/Source/Programs/UnrealBuildTool/Configuration/CppCompileWarnings.cs
--- a/Engine/Source/Programs/UnrealBuildTool/Configuration/CppCompileWarnings.cs
+++ b/Engine/Source/Programs/UnrealBuildTool/Configuration/CppCompileWarnings.cs
@@ -1,3 +1,5 @@
+// Copyright Epic Games, Inc. All Rights Reserved.
+
 using EpicGames.Core;
 using Microsoft.Extensions.Logging;
 using System;
diff --git a/Engine/Source/Programs/UnrealBuildTool/Configuration/ReadOnlyCppCompileWarnings.cs b/Engine/Source/Programs/UnrealBuildTool/Configuration/ReadOnlyCppCompileWarnings.cs
--- a/Engine/Source/Programs/UnrealBuildTool/Configuration/ReadOnlyCppCompileWarnings.cs
+++ b/Engine/Source/Programs/UnrealBuildTool/Configuration/ReadOnlyCppCompileWarnings.cs
@@ -1,3 +1,5 @@
+// Copyright Epic Games, Inc. All Rights Reserved.
+
 namespace UnrealBuildTool
 {
 	/// <summary>
