Add Android's onTopResumedActivityChanged to upl

diff --git a/Engine/Build/Android/Java/src/com/epicgames/unreal/GameActivity.java.template b/Engine/Build/Android/Java/src/com/epicgames/unreal/GameActivity.java.template
--- a/Engine/Build/Android/Java/src/com/epicgames/unreal/GameActivity.java.template
+++ b/Engine/Build/Android/Java/src/com/epicgames/unreal/GameActivity.java.template
@@ -3642,6 +3642,17 @@ public class GameActivity extends $${gameActivitySuperClass}$$ implements Surfac
 		Log.debug("==============> GameActive.onResume complete!");
 	}
 
+	@TargetApi(Build.VERSION_CODES.Q)
+	@Override
+	public void onTopResumedActivityChanged( boolean isTopResumedActivity )
+	{
+		if ( Build.VERSION.SDK_INT >= Build.VERSION_CODES.Q )
+		{
+			Log.debug("onTopResumedActivityChanged: "+isTopResumedActivity);
+			//$${gameActivityOnTopResumeActivityAdditions}$$
+		}
+	}
+
 	@Override
 	public void onResume()
 	{
diff --git a/Engine/Source/Programs/UnrealBuildTool/Platform/Android/UEDeployAndroid.cs b/Engine/Source/Programs/UnrealBuildTool/Platform/Android/UEDeployAndroid.cs
--- a/Engine/Source/Programs/UnrealBuildTool/Platform/Android/UEDeployAndroid.cs
+++ b/Engine/Source/Programs/UnrealBuildTool/Platform/Android/UEDeployAndroid.cs
@@ -6157,6 +6157,7 @@ import java.util.Collection;
 				{ "//$${gameActivityOnRequestPermissionsResultAdditions}$$", UPL.ProcessPluginNode(NDKArch, "gameActivityOnRequestPermissionsResultAdditions", "")},
 				{ "//$${gameActivityOnPauseAdditions}$$", UPL.ProcessPluginNode(NDKArch, "gameActivityOnPauseAdditions", "")},
 				{ "//$${gameActivityOnResumeAdditions}$$", UPL.ProcessPluginNode(NDKArch, "gameActivityOnResumeAdditions", "")},
+				{ "//$${gameActivityOnTopResumeActivityAdditions}$$", UPL.ProcessPluginNode(NDKArch, "gameActivityOnTopResumeActivityAdditions", "")},
 				{ "//$${gameActivityOnNewIntentAdditions}$$", UPL.ProcessPluginNode(NDKArch, "gameActivityOnNewIntentAdditions", "")},
   				{ "//$${gameActivityOnActivityResultAdditions}$$", UPL.ProcessPluginNode(NDKArch, "gameActivityOnActivityResultAdditions", "")},
   				{ "//$${gameActivityPreConfigRulesParseAdditions}$$", UPL.ProcessPluginNode(NDKArch, "gameActivityPreConfigRulesParseAdditions", "")},
diff --git a/Engine/Source/Programs/UnrealBuildTool/System/UnrealPluginLanguage.cs b/Engine/Source/Programs/UnrealBuildTool/System/UnrealPluginLanguage.cs
--- a/Engine/Source/Programs/UnrealBuildTool/System/UnrealPluginLanguage.cs
+++ b/Engine/Source/Programs/UnrealBuildTool/System/UnrealPluginLanguage.cs
@@ -424,6 +424,9 @@ namespace UnrealBuildTool
 	 * 	<!-- optional additions to GameActivity onResume in GameActivity.java -->
 	 *	<gameActivityOnResumeAdditions>	</gameActivityOnResumeAdditions>
 	 *	
+	 * 	<!-- optional additions to GameActivity onTopResumeActivity in GameActivity.java -->
+	 *	<gameActivityOnTopResumeActivityAdditions>	</gameActivityOnTopResumeActivityAdditions>
+	 *	
 	 * 	<!-- optional additions to GameActivity onNewIntent in GameActivity.java -->
 	 *	<gameActivityOnNewIntentAdditions> </gameActivityOnNewIntentAdditions>
 	 *	
