UnrealBuildTool: Use all cores when running clang thin lto

#rnx

diff --git a/Engine/Source/Programs/UnrealBuildTool/Platform/Android/AndroidToolChain.cs b/Engine/Source/Programs/UnrealBuildTool/Platform/Android/AndroidToolChain.cs
--- a/Engine/Source/Programs/UnrealBuildTool/Platform/Android/AndroidToolChain.cs
+++ b/Engine/Source/Programs/UnrealBuildTool/Platform/Android/AndroidToolChain.cs
@@ -784,6 +784,7 @@ namespace UnrealBuildTool
 				if ((Options & ClangToolChainOptions.EnableThinLTO) != 0)
 				{
 					Result += " -flto=thin";
+					Result += " -Wl,--thinlto-jobs=all";
 				}
 				else
 				{
@@ -1124,6 +1125,15 @@ namespace UnrealBuildTool
 				LinkAction.CommandArguments += String.Format(" -o \"{0}\"", OutputFile.AbsolutePath);
 			}
 
+			if (LinkEnvironment.bAllowLTCG && Options.HasFlag(ClangToolChainOptions.EnableThinLTO))
+			{
+				// Set the weight to number of logical cores as lld can max out the available cores
+				LinkAction.Weight = Utils.GetLogicalProcessorCount();
+
+				// Disallow remote to prevent this long running action from running on an agent if remote linking is enabled
+				LinkAction.bCanExecuteRemotely = false;
+			}
+
 			// Add the input files to a response file, and pass the response file on the command-line.
 			List<string> InputFileNames = new List<string>();
 			foreach (FileItem InputFile in LinkEnvironment.InputFiles)
diff --git a/Engine/Source/Programs/UnrealBuildTool/Platform/Linux/LinuxToolChain.cs b/Engine/Source/Programs/UnrealBuildTool/Platform/Linux/LinuxToolChain.cs
--- a/Engine/Source/Programs/UnrealBuildTool/Platform/Linux/LinuxToolChain.cs
+++ b/Engine/Source/Programs/UnrealBuildTool/Platform/Linux/LinuxToolChain.cs
@@ -310,7 +310,7 @@ namespace UnrealBuildTool
 				{
 					if (bForLinker)
 					{
-						Arguments.Add(String.Format("-flto=thin -Wl,--thinlto-jobs={0}, -Wl,-mllvm,-disable-auto-upgrade-debug-info, -Wl,-mllvm,-enable-ext-tsp-block-placement=1", Utils.GetPhysicalProcessorCount()));
+						Arguments.Add( "-flto=thin -Wl,--thinlto-jobs=all, -Wl,-mllvm,-disable-auto-upgrade-debug-info, -Wl,-mllvm,-enable-ext-tsp-block-placement=1");
 					}
 					else
 					{
@@ -940,9 +940,16 @@ namespace UnrealBuildTool
 				Logger.LogInformation("Using PGI (profile guided instrumentation).");
 			}
 
-			if (CompileEnvironment.bAllowLTCG)
+			if (Options.HasFlag(ClangToolChainOptions.EnableLinkTimeOptimization))
 			{
-				Logger.LogInformation("Using LTO (link-time optimization).");
+				if (Options.HasFlag(ClangToolChainOptions.EnableThinLTO))
+				{
+					Logger.LogInformation("Using ThinLTO (link-time optimization).");
+				}
+				else
+				{
+					Logger.LogInformation("Using LTO (link-time optimization).");
+				}
 			}
 
 			if (CompileEnvironment.bUsePIE)
@@ -1160,6 +1167,15 @@ namespace UnrealBuildTool
 				LinkAction.CommandDescription = "Link";
 			}
 
+			if (Options.HasFlag(ClangToolChainOptions.EnableLinkTimeOptimization) && Options.HasFlag(ClangToolChainOptions.EnableThinLTO))
+			{
+				// Set the weight to number of logical cores as lld can max out the available cores
+				LinkAction.Weight = Utils.GetLogicalProcessorCount();
+
+				// Disallow remote to prevent this long running action from running on an agent if remote linking is enabled
+				LinkAction.bCanExecuteRemotely = false;
+			}
+
 			if (!OperatingSystem.IsWindows())
 			{
 				LinkAction.bCanExecuteInUBA = false; // Linker on native linux uses vfork/exec which is not handled in uba right now
diff --git a/Engine/Source/Programs/UnrealBuildTool/Platform/Windows/VCToolChain.cs b/Engine/Source/Programs/UnrealBuildTool/Platform/Windows/VCToolChain.cs
--- a/Engine/Source/Programs/UnrealBuildTool/Platform/Windows/VCToolChain.cs
+++ b/Engine/Source/Programs/UnrealBuildTool/Platform/Windows/VCToolChain.cs
@@ -2982,6 +2982,15 @@ namespace UnrealBuildTool
 				LinkAction.bCanExecuteRemotely = true;
 			}
 
+			if (Target.WindowsPlatform.Compiler.IsClang() && Target.WindowsPlatform.bAllowClangLinker && LinkEnvironment.bAllowLTCG)
+			{
+				// Set the weight to number of logical cores as lld can max out the available cores
+				LinkAction.Weight = Utils.GetLogicalProcessorCount();
+
+				// Disallow remote to prevent this long running action from running on an agent if remote linking is enabled
+				LinkAction.bCanExecuteRemotely = false;
+			}
+
 			if (LinkEnvironment.bPGOOptimize || LinkEnvironment.bPGOProfile)
 			{
 				LinkAction.bCanExecuteInUBA = false; // Disabled for now. Should revisit to see why it is not working
@@ -3232,6 +3241,9 @@ namespace UnrealBuildTool
 					{
 						Log.TraceInformationOnce("Enabling Link-time optimization. Linking will take a while.");
 
+						// lld should consider logical cores when determining how many threads to use
+						Arguments.Add("/opt:lldltojobs=all");
+
 						// ThinLTO incremental cache
 						DirectoryReference? ThinLTOCacheDir = DirectoryReference.FromString(LinkEnvironment.ThinLTOCacheDirectory);
 						if (ThinLTOCacheDir != null)
