UnrealBuildTool: Support .natjmc

#rnx

diff --git a/Engine/Source/Programs/UnrealBuildTool/Configuration/ModuleRules.cs b/Engine/Source/Programs/UnrealBuildTool/Configuration/ModuleRules.cs
--- a/Engine/Source/Programs/UnrealBuildTool/Configuration/ModuleRules.cs
+++ b/Engine/Source/Programs/UnrealBuildTool/Configuration/ModuleRules.cs
@@ -1277,7 +1277,7 @@ namespace UnrealBuildTool
 		public List<string> PublicAdditionalLibraries = new();
 
 		/// <summary>
-		/// List of additional debug visualizers (.natvis and .natstepfilter) exposed to other modules - typically used for External (third party) modules
+		/// List of additional debug visualizers (.natvis, .natstepfilter, and .natjmc) exposed to other modules - typically used for External (third party) modules
 		/// </summary>
 		public List<string> PublicDebugVisualizerPaths = new();
 
diff --git a/Engine/Source/Programs/UnrealBuildTool/Configuration/UEBuildModule.cs b/Engine/Source/Programs/UnrealBuildTool/Configuration/UEBuildModule.cs
--- a/Engine/Source/Programs/UnrealBuildTool/Configuration/UEBuildModule.cs
+++ b/Engine/Source/Programs/UnrealBuildTool/Configuration/UEBuildModule.cs
@@ -306,9 +306,9 @@ namespace UnrealBuildTool
 			foreach (string natVisPath in HashSetFromOptionalEnumerableStringParameter(Rules.PublicDebugVisualizerPaths))
 			{
 				FileItem natVisItem = FileItem.GetItemByPath(natVisPath);
-				if (!natVisItem.HasExtension(".natvis") && !natVisItem.HasExtension(".natstepfilter"))
+				if (!natVisItem.HasExtension(".natvis") && !natVisItem.HasExtension(".natstepfilter") && !natVisItem.HasExtension(".natjmc"))
 				{
-					Log.TraceWarningTask(RulesFile, $"Referenced Debug Visualizer '{natVisItem}' is not a .natvis or .natstepfilter file");
+					Log.TraceWarningTask(RulesFile, $"Referenced Debug Visualizer '{natVisItem}' is not a .natvis, .natstepfilter, or .natjmc file");
 				}
 				else if (!natVisItem.Exists)
 				{
@@ -323,7 +323,7 @@ namespace UnrealBuildTool
 			// Add any debug visualizers found in the module directories
 			foreach (DirectoryItem directory in ModuleDirectories.Select(x => DirectoryItem.GetItemByDirectoryReference(x)))
 			{
-				NatvisFiles.UnionWith(directory.EnumerateFiles().Where(x => x.HasExtension(".natvis") || x.HasExtension(".natstepfilter")));
+				NatvisFiles.UnionWith(directory.EnumerateFiles().Where(x => x.HasExtension(".natvis") || x.HasExtension(".natstepfilter") || x.HasExtension(".natjmc")));
 			}
 		}
 
diff --git a/Engine/Source/Programs/UnrealBuildTool/Platform/Windows/UEBuildWindows.cs b/Engine/Source/Programs/UnrealBuildTool/Platform/Windows/UEBuildWindows.cs
--- a/Engine/Source/Programs/UnrealBuildTool/Platform/Windows/UEBuildWindows.cs
+++ b/Engine/Source/Programs/UnrealBuildTool/Platform/Windows/UEBuildWindows.cs
@@ -1589,7 +1589,8 @@ namespace UnrealBuildTool
 				|| IsBuildProductName(FileName, NamePrefixes, NameSuffixes, ".map")
 				|| IsBuildProductName(FileName, NamePrefixes, NameSuffixes, ".objpaths")
 				|| IsBuildProductName(FileName, NamePrefixes, NameSuffixes, ".natvis")
-				|| IsBuildProductName(FileName, NamePrefixes, NameSuffixes, ".natstepfilter");
+				|| IsBuildProductName(FileName, NamePrefixes, NameSuffixes, ".natstepfilter")
+				|| IsBuildProductName(FileName, NamePrefixes, NameSuffixes, ".natjmc");
 		}
 
 		/// <summary>
diff --git a/Engine/Source/Programs/UnrealBuildTool/Platform/Windows/VCToolChain.cs b/Engine/Source/Programs/UnrealBuildTool/Platform/Windows/VCToolChain.cs
--- a/Engine/Source/Programs/UnrealBuildTool/Platform/Windows/VCToolChain.cs
+++ b/Engine/Source/Programs/UnrealBuildTool/Platform/Windows/VCToolChain.cs
@@ -1830,7 +1830,7 @@ namespace UnrealBuildTool
 
 		public override FileItem? CopyDebuggerVisualizer(FileItem SourceFile, DirectoryReference IntermediateDirectory, IActionGraphBuilder Graph)
 		{
-			if (SourceFile.HasExtension(".natvis") || SourceFile.HasExtension(".natstepfilter"))
+			if (SourceFile.HasExtension(".natvis") || SourceFile.HasExtension(".natstepfilter") || SourceFile.HasExtension(".natjmc"))
 			{
 				FileReference IntermediateFile = FileReference.Combine(IntermediateDirectory, SourceFile.Name);
 				if (!Unreal.IsFileInstalled(IntermediateFile))
diff --git a/Engine/Source/Programs/UnrealBuildTool/ProjectFiles/VisualStudio/VCProjectFileGenerator.cs b/Engine/Source/Programs/UnrealBuildTool/ProjectFiles/VisualStudio/VCProjectFileGenerator.cs
--- a/Engine/Source/Programs/UnrealBuildTool/ProjectFiles/VisualStudio/VCProjectFileGenerator.cs
+++ b/Engine/Source/Programs/UnrealBuildTool/ProjectFiles/VisualStudio/VCProjectFileGenerator.cs
@@ -794,12 +794,14 @@ namespace UnrealBuildTool
 				// Get the path to the visualizers file. Try to make it relative to the solution directory, but fall back to a full path if it's a foreign project.
 				FileReference VisualizersFile = FileReference.Combine(Unreal.EngineDirectory, "Extras", "VisualStudioDebugging", "Unreal.natvis");
 				FileReference VisualizersStepFile = FileReference.Combine(Unreal.EngineDirectory, "Extras", "VisualStudioDebugging", "Unreal.natstepfilter");
+				FileReference VisualizersJustMyCodeFile = FileReference.Combine(Unreal.EngineDirectory, "Extras", "VisualStudioDebugging", "Unreal.natjmc");
 
 				// Add the visualizers at the solution level. Doesn't seem to be picked up from a makefile project in VS2017 15.8.5.
 				VCSolutionFileContent.AppendLine(String.Format("Project(\"{0}\") = \"Visualizers\", \"Visualizers\", \"{{1CCEC849-CC72-4C59-8C36-2F7C38706D4C}}\"", SolutionFolderEntryGUID));
 				VCSolutionFileContent.AppendLine("\tProjectSection(SolutionItems) = preProject");
 				VCSolutionFileContent.AppendLine("\t\t{0} = {0}", VisualizersFile.MakeRelativeTo(PrimaryProjectPath));
 				VCSolutionFileContent.AppendLine("\t\t{0} = {0}", VisualizersStepFile.MakeRelativeTo(PrimaryProjectPath));
+				VCSolutionFileContent.AppendLine("\t\t{0} = {0}", VisualizersJustMyCodeFile.MakeRelativeTo(PrimaryProjectPath));
 				VCSolutionFileContent.AppendLine("\tEndProjectSection");
 				VCSolutionFileContent.AppendLine("EndProject");
 			}
