UnrealBuildTool: Missed a change from last cl

#rnx

diff --git a/Engine/Source/Programs/UnrealBuildTool/Platform/Windows/VCToolChain.cs b/Engine/Source/Programs/UnrealBuildTool/Platform/Windows/VCToolChain.cs
--- a/Engine/Source/Programs/UnrealBuildTool/Platform/Windows/VCToolChain.cs
+++ b/Engine/Source/Programs/UnrealBuildTool/Platform/Windows/VCToolChain.cs
@@ -1845,7 +1845,7 @@ namespace UnrealBuildTool
 
 		public override FileItem? LinkDebuggerVisualizer(FileItem SourceFile, DirectoryReference IntermediateDirectory)
 		{
-			if (SourceFile.HasExtension(".natvis") || SourceFile.HasExtension(".natstepfilter"))
+			if (SourceFile.HasExtension(".natvis") || SourceFile.HasExtension(".natstepfilter") || SourceFile.HasExtension(".natjmc"))
 			{
 				FileItem IntermediateFile = FileItem.GetItemByFileReference(FileReference.Combine(IntermediateDirectory, SourceFile.Name));
 				return IntermediateFile;
