.natjmc for skipping over certain Core operations during debugging:

- Delegates
- TFunction and variants
- Containers
- Platform thread abstractions
- STL atomics

diff --git a/Engine/Source/Programs/UnrealBuildTool/Platform/Windows/VCToolChain.cs b/Engine/Source/Programs/UnrealBuildTool/Platform/Windows/VCToolChain.cs
--- a/Engine/Source/Programs/UnrealBuildTool/Platform/Windows/VCToolChain.cs
+++ b/Engine/Source/Programs/UnrealBuildTool/Platform/Windows/VCToolChain.cs
@@ -934,6 +934,12 @@ namespace UnrealBuildTool
 			// If enabled, create debug information.
 			if (CompileEnvironment.bCreateDebugInfo && !(Target.WindowsPlatform.Compiler.IsIntel() && (CompileEnvironment.bPGOOptimize || CompileEnvironment.bPGOProfile || CompileEnvironment.bAllowLTCG)))
 			{
+				// Enable Just-My-Code in Debug or Development, unless on Clang
+				if ((CompileEnvironment.Configuration == CppConfiguration.Debug || CompileEnvironment.Configuration == CppConfiguration.Development) && !Target.WindowsPlatform.Compiler.IsClang())
+				{
+					Arguments.Add("/JMC");
+				}
+
 				// Store debug info in .pdb files.
 				// @todo clang: PDB files are emited from Clang but do not fully work with Visual Studio yet (breakpoints won't hit due to "symbol read error")
 				// @todo clang (update): as of clang 3.9 breakpoints work with PDBs, and the callstack is correct, so could be used for crash dumps. However debugging is still impossible due to the large amount of unreadable variables and unpredictable breakpoint stepping behaviour
