Add UBT option to enable warning for switches on enumerations that don't have a case for all enumerators. It is off by default, and can be enabled per-module or per-target.
Add macros to disable and restore the previous warning behavior for specific code regions.

diff --git a/Engine/Source/Programs/UnrealBuildTool/Configuration/CppCompileWarnings.cs b/Engine/Source/Programs/UnrealBuildTool/Configuration/CppCompileWarnings.cs
--- a/Engine/Source/Programs/UnrealBuildTool/Configuration/CppCompileWarnings.cs
+++ b/Engine/Source/Programs/UnrealBuildTool/Configuration/CppCompileWarnings.cs
@@ -70,6 +70,19 @@ namespace UnrealBuildTool
 
 		private WarningLevel _undefinedIdentifierWarningLevelPrivate;
 
+		/// <summary>
+		/// Indicates what warning/error level to treat unhandled enumerators in switches on enumeration-typed values.
+		/// </summary>
+		[XmlConfigFile(Category = "BuildConfiguration")]
+		[CommandLine("-SwitchUnhandledEnumeratorWarningLevel=")]
+		[RequiresUniqueBuildEnvironment]
+		public WarningLevel SwitchUnhandledEnumeratorWarningLevel
+		{
+			get => ResolveWarning(null, ParentCppCompileWarnings?.SwitchUnhandledEnumeratorWarningLevel, _switchUnhandledEnumeratorWarningLevelPrivate);
+			set => _switchUnhandledEnumeratorWarningLevelPrivate = value;
+		}
+
+		private WarningLevel _switchUnhandledEnumeratorWarningLevelPrivate;
 
 		#endregion  -- Compiler Warning Settings --
 
@@ -187,6 +200,7 @@ namespace UnrealBuildTool
 			_shadowVariableWarningLevelPrivate = WarningLevel.Warning;
 			_unsafeTypeCastWarningLevelPrivate = WarningLevel.Off;
 			_undefinedIdentifierWarningLevelPrivate = WarningLevel.Error;
+			_switchUnhandledEnumeratorWarningLevelPrivate = WarningLevel.Off;
 		}
 
 		#endregion -- Private API --
diff --git a/Engine/Source/Programs/UnrealBuildTool/Configuration/ReadOnlyCppCompileWarnings.cs b/Engine/Source/Programs/UnrealBuildTool/Configuration/ReadOnlyCppCompileWarnings.cs
--- a/Engine/Source/Programs/UnrealBuildTool/Configuration/ReadOnlyCppCompileWarnings.cs
+++ b/Engine/Source/Programs/UnrealBuildTool/Configuration/ReadOnlyCppCompileWarnings.cs
@@ -31,6 +31,7 @@ namespace UnrealBuildTool
 		public WarningLevel ShadowVariableWarningLevel => _inner.ShadowVariableWarningLevel;
 		public WarningLevel UndefinedIdentifierWarningLevel => _inner.UndefinedIdentifierWarningLevel;
 		public WarningLevel UnsafeTypeCastWarningLevel => _inner.UnsafeTypeCastWarningLevel;
+		public WarningLevel SwitchUnhandledEnumeratorWarningLevel => _inner.SwitchUnhandledEnumeratorWarningLevel;
 #pragma warning restore C1591
 		#endregion
 	}
diff --git a/Engine/Source/Programs/UnrealBuildTool/Configuration/UEBuildModuleCPP.cs b/Engine/Source/Programs/UnrealBuildTool/Configuration/UEBuildModuleCPP.cs
--- a/Engine/Source/Programs/UnrealBuildTool/Configuration/UEBuildModuleCPP.cs
+++ b/Engine/Source/Programs/UnrealBuildTool/Configuration/UEBuildModuleCPP.cs
@@ -2182,6 +2182,7 @@ namespace UnrealBuildTool
 			Result.ShadowVariableWarningLevel = Rules.CppCompileWarningSettings.ShadowVariableWarningLevel;
 			Result.UnsafeTypeCastWarningLevel = Rules.CppCompileWarningSettings.UnsafeTypeCastWarningLevel;
 			Result.UndefinedIdentifierWarningLevel = Rules.CppCompileWarningSettings.UndefinedIdentifierWarningLevel;
+			Result.SwitchUnhandledEnumeratorWarningLevel = Rules.CppCompileWarningSettings.SwitchUnhandledEnumeratorWarningLevel;
 			Result.bDisableStaticAnalysis = Rules.bDisableStaticAnalysis || (Target.bStaticAnalyzerProjectOnly && Rules.bTreatAsEngineModule);
 			Result.bStaticAnalyzerExtensions = Rules.bStaticAnalyzerExtensions;
 			Result.StaticAnalyzerRulesets = Rules.StaticAnalyzerRulesets;
diff --git a/Engine/Source/Programs/UnrealBuildTool/Configuration/UEBuildTarget.cs b/Engine/Source/Programs/UnrealBuildTool/Configuration/UEBuildTarget.cs
--- a/Engine/Source/Programs/UnrealBuildTool/Configuration/UEBuildTarget.cs
+++ b/Engine/Source/Programs/UnrealBuildTool/Configuration/UEBuildTarget.cs
@@ -5228,6 +5228,7 @@ namespace UnrealBuildTool
 			GlobalCompileEnvironment.ShadowVariableWarningLevel = Rules.CppCompileWarningSettings.ShadowVariableWarningLevel;
 			GlobalCompileEnvironment.UnsafeTypeCastWarningLevel = Rules.CppCompileWarningSettings.UnsafeTypeCastWarningLevel;
 			GlobalCompileEnvironment.UndefinedIdentifierWarningLevel = Rules.CppCompileWarningSettings.UndefinedIdentifierWarningLevel;
+			GlobalCompileEnvironment.SwitchUnhandledEnumeratorWarningLevel = Rules.CppCompileWarningSettings.SwitchUnhandledEnumeratorWarningLevel;
 			GlobalCompileEnvironment.bRetainFramePointers = Rules.bRetainFramePointers;
 			GlobalCompileEnvironment.bWarningsAsErrors = Rules.bWarningsAsErrors;
 			GlobalCompileEnvironment.OptimizationLevel = Rules.OptimizationLevel;
diff --git a/Engine/Source/Programs/UnrealBuildTool/Platform/Windows/VCToolChain.cs b/Engine/Source/Programs/UnrealBuildTool/Platform/Windows/VCToolChain.cs
--- a/Engine/Source/Programs/UnrealBuildTool/Platform/Windows/VCToolChain.cs
+++ b/Engine/Source/Programs/UnrealBuildTool/Platform/Windows/VCToolChain.cs
@@ -1109,6 +1109,18 @@ namespace UnrealBuildTool
 				}
 			}
 
+			if (CompileEnvironment.SwitchUnhandledEnumeratorWarningLevel != WarningLevel.Off)
+			{
+				if (CompileEnvironment.UndefinedIdentifierWarningLevel == WarningLevel.Error)
+				{
+					Arguments.Add("/we4061");
+				}
+				else
+				{
+					Arguments.Add("/w44061");
+				}
+			}
+
 			// The unsafe type cast warnings setting controls the following warnings currently:
 			//   4244: conversion from 'type1' to 'type2', possible loss of data
 			//   4838: conversion from 'type1' to 'type2' requires a narrowing conversion
diff --git a/Engine/Source/Programs/UnrealBuildTool/System/CppCompileEnvironment.cs b/Engine/Source/Programs/UnrealBuildTool/System/CppCompileEnvironment.cs
--- a/Engine/Source/Programs/UnrealBuildTool/System/CppCompileEnvironment.cs
+++ b/Engine/Source/Programs/UnrealBuildTool/System/CppCompileEnvironment.cs
@@ -371,6 +371,11 @@ namespace UnrealBuildTool
 		/// </summary>
 		public WarningLevel UndefinedIdentifierWarningLevel = WarningLevel.Off;
 
+		/// <summary>
+		/// Indicates what warning/error level to treat unhandled enumerators in switches on enumeration-typed values.
+		/// </summary>
+		public WarningLevel SwitchUnhandledEnumeratorWarningLevel = WarningLevel.Off;
+
 		/// <summary>
 		/// Whether to treat all warnings as errors
 		/// </summary>
@@ -784,6 +789,7 @@ namespace UnrealBuildTool
 			DeprecationWarningLevel = Other.DeprecationWarningLevel;
 			ShadowVariableWarningLevel = Other.ShadowVariableWarningLevel;
 			UnsafeTypeCastWarningLevel = Other.UnsafeTypeCastWarningLevel;
+			SwitchUnhandledEnumeratorWarningLevel = Other.SwitchUnhandledEnumeratorWarningLevel;
 			UndefinedIdentifierWarningLevel = Other.UndefinedIdentifierWarningLevel;
 			bWarningsAsErrors = Other.bWarningsAsErrors;
 			bDisableStaticAnalysis = Other.bDisableStaticAnalysis;
diff --git a/Engine/Source/Programs/UnrealBuildTool/ToolChain/ClangWarnings.cs b/Engine/Source/Programs/UnrealBuildTool/ToolChain/ClangWarnings.cs
--- a/Engine/Source/Programs/UnrealBuildTool/ToolChain/ClangWarnings.cs
+++ b/Engine/Source/Programs/UnrealBuildTool/ToolChain/ClangWarnings.cs
@@ -134,6 +134,11 @@ namespace UnrealBuildTool
 				Arguments.Add("-Wundef" + (CompileEnvironment.UndefinedIdentifierWarningLevel == WarningLevel.Error ? "" : " -Wno-error=undef"));
 			}
 
+			if (CompileEnvironment.SwitchUnhandledEnumeratorWarningLevel != WarningLevel.Off)
+			{
+				Arguments.Add("-Wswitch-enum" + (CompileEnvironment.SwitchUnhandledEnumeratorWarningLevel == WarningLevel.Error ? "" : " -Wno-error=switch-enum"));
+			}
+
 			// Note: This should be kept in sync with PRAGMA_DISABLE_UNSAFE_TYPECAST_WARNINGS in ClangPlatformCompilerPreSetup.h
 			string[] UnsafeTypeCastWarningList = {
 				"float-conversion",
diff --git a/Engine/Source/Runtime/Core/Public/Clang/ClangPlatformCompilerPreSetup.h b/Engine/Source/Runtime/Core/Public/Clang/ClangPlatformCompilerPreSetup.h
--- a/Engine/Source/Runtime/Core/Public/Clang/ClangPlatformCompilerPreSetup.h
+++ b/Engine/Source/Runtime/Core/Public/Clang/ClangPlatformCompilerPreSetup.h
@@ -113,6 +113,17 @@
 	_Pragma("clang diagnostic pop")
 #endif // PRAGMA_RESTORE_UNSAFE_TYPECAST_WARNINGS
 
+#ifndef PRAGMA_DISABLE_SWITCH_UNHANDLED_ENUM_CASE_WARNINGS
+	#define PRAGMA_DISABLE_SWITCH_UNHANDLED_ENUM_CASE_WARNINGS \
+		_Pragma("clang diagnostic push") \
+		_Pragma("clang diagnostic ignored \"-Wswitch-enum\"")
+#endif // PRAGMA_DISABLE_UNSAFE_TYPECAST_WARNINGS
+
+#ifndef PRAGMA_RESTORE_SWITCH_UNHANDLED_ENUM_CASE_WARNINGS
+	#define PRAGMA_RESTORE_SWITCH_UNHANDLED_ENUM_CASE_WARNINGS \
+		_Pragma("clang diagnostic pop")
+#endif // PRAGMA_RESTORE_UNSAFE_TYPECAST_WARNINGS
+
 #if __has_warning("-Wordered-compare-function-pointers")
 #define DISABLE_ORDERED_COMPARE_FUNCTION_POINTERS _Pragma("clang diagnostic ignored \"-Wordered-compare-function-pointers\"")
 #else
diff --git a/Engine/Source/Runtime/Core/Public/MSVC/MSVCPlatformCompilerPreSetup.h b/Engine/Source/Runtime/Core/Public/MSVC/MSVCPlatformCompilerPreSetup.h
--- a/Engine/Source/Runtime/Core/Public/MSVC/MSVCPlatformCompilerPreSetup.h
+++ b/Engine/Source/Runtime/Core/Public/MSVC/MSVCPlatformCompilerPreSetup.h
@@ -101,6 +101,17 @@
 			__pragma(warning(pop))
 	#endif // PRAGMA_RESTORE_UNSAFE_TYPECAST_WARNINGS
 
+	#ifndef PRAGMA_DISABLE_SWITCH_UNHANDLED_ENUM_CASE_WARNINGS
+		#define PRAGMA_DISABLE_SWITCH_UNHANDLED_ENUM_CASE_WARNINGS \
+			__pragma (warning(push)) \
+			__pragma (warning(disable: 4061)) /*  enumerator 'enumator' in switch of enum 'enumeration' is not explicitly handled by a case label */
+	#endif // PRAGMA_DISABLE_UNSAFE_TYPECAST_WARNINGS
+
+	#ifndef PRAGMA_RESTORE_SWITCH_UNHANDLED_ENUM_CASE_WARNINGS
+	#define PRAGMA_RESTORE_SWITCH_UNHANDLED_ENUM_CASE_WARNINGS \
+				__pragma(warning(pop))
+	#endif // PRAGMA_RESTORE_UNSAFE_TYPECAST_WARNINGS
+
 	#ifndef PRAGMA_DISABLE_UNDEFINED_IDENTIFIER_WARNINGS
 		#define PRAGMA_DISABLE_UNDEFINED_IDENTIFIER_WARNINGS \
 			__pragma(warning(push)) \
