UnrealBuildTool: If bin-woa64 exists, use it for arm64 hosts

#rnx

diff --git a/Engine/Source/Programs/UnrealBuildTool/Platform/Windows/MicrosoftPlatformSDK.cs b/Engine/Source/Programs/UnrealBuildTool/Platform/Windows/MicrosoftPlatformSDK.cs
--- a/Engine/Source/Programs/UnrealBuildTool/Platform/Windows/MicrosoftPlatformSDK.cs
+++ b/Engine/Source/Programs/UnrealBuildTool/Platform/Windows/MicrosoftPlatformSDK.cs
@@ -865,7 +865,15 @@ namespace UnrealBuildTool
 								EpicGames.UBA.Utils.RegisterCrossArchitecturePath(HostBinDir.FullName, CrossBinDir.FullName);
 							}
 						}
-
+						else if (Compiler.IsClang())
+						{
+							DirectoryReference HostBinDir = DirectoryReference.Combine(ToolChain.BaseDir, MSVCHostUnrealArch.bIsX64 ? "bin" : "bin-woa64");
+							DirectoryReference CrossBinDir = DirectoryReference.Combine(ToolChain.BaseDir, !MSVCHostUnrealArch.bIsX64 ? "bin" : "bin-woa64");
+							if (FileReference.Exists(HostBinDir) && FileReference.Exists(CrossBinDir))
+							{
+								EpicGames.UBA.Utils.RegisterCrossArchitecturePath(HostBinDir.FullName, CrossBinDir.FullName);
+							}
+						}
 					}
 				}
 				return ToolChains;
diff --git a/Engine/Source/Programs/UnrealBuildTool/Platform/Windows/VCEnvironment.cs b/Engine/Source/Programs/UnrealBuildTool/Platform/Windows/VCEnvironment.cs
--- a/Engine/Source/Programs/UnrealBuildTool/Platform/Windows/VCEnvironment.cs
+++ b/Engine/Source/Programs/UnrealBuildTool/Platform/Windows/VCEnvironment.cs
@@ -9,6 +9,7 @@ using System.Runtime.InteropServices;
 using System.Runtime.Versioning;
 using EpicGames.Core;
 using Microsoft.Extensions.Logging;
+using UnrealBuildBase;
 
 namespace UnrealBuildTool
 {
@@ -183,7 +184,7 @@ namespace UnrealBuildTool
 		}
 
 		/// <summary>
-		/// Gets the path to the tool binaries.
+		/// Gets the path to the VC tool binaries.
 		/// </summary>
 		/// <param name="VCToolChainDir">Base directory for the VC toolchain</param>
 		/// <param name="Architecture">Target Architecture</param>
@@ -199,6 +200,21 @@ namespace UnrealBuildTool
 			throw new BuildException("No required {0} compiler toolchain found in {1}", Architecture, VCToolChainDir);
 		}
 
+		/// <summary>
+		/// Gets the path to the Clang tool binaries.
+		/// </summary>
+		/// <param name="ClangToolChainDir">Base directory for the Clang toolchain</param>
+		/// <returns>Directory containing the 64-bit toolchain binaries</returns>
+		public static DirectoryReference GetClangToolPath(DirectoryReference ClangToolChainDir)
+		{
+			// If on an arm64 host, use arm64 binaries if available, otherwise use x64 
+			if (!MicrosoftPlatformSDK.MSVCHostUnrealArch.bIsX64 && DirectoryReference.Exists(DirectoryReference.Combine(ClangToolChainDir, "bin-woa64")))
+			{
+				return DirectoryReference.Combine(ClangToolChainDir, "bin-woa64");
+			}
+			return DirectoryReference.Combine(ClangToolChainDir, "bin");
+		}
+
 		/// <summary>
 		/// Gets the path to the compiler.
 		/// </summary>
@@ -206,11 +222,11 @@ namespace UnrealBuildTool
 		{
 			if (Compiler == WindowsCompiler.Clang)
 			{
-				return FileReference.Combine(CompilerDir, "bin", "clang-cl.exe");
+				return FileReference.Combine(GetClangToolPath(CompilerDir), "clang-cl.exe");
 			}
 			else if (Compiler == WindowsCompiler.ClangRTFM)
 			{
-				return FileReference.Combine(CompilerDir, "bin", "verse-clang-cl.exe");
+				return FileReference.Combine(GetClangToolPath(CompilerDir), "verse-clang-cl.exe");
 			}
 			else if (Compiler == WindowsCompiler.Intel)
 			{
@@ -230,15 +246,13 @@ namespace UnrealBuildTool
 			}
 			else if ((Compiler == WindowsCompiler.Clang || Compiler == WindowsCompiler.ClangRTFM) && bAllowClangLinker)
 			{
-				return FileReference.Combine(CompilerDir, "bin", "lld-link.exe");
+				return FileReference.Combine(GetClangToolPath(CompilerDir), "lld-link.exe");
 			}
 			else if (Compiler == WindowsCompiler.Intel && bAllowClangLinker)
 			{
 				return FileReference.Combine(CompilerDir, "bin", "compiler", "lld-link.exe");
-				//return FileReference.Combine(CompilerDir, "bin", "xilink.exe");
 			}
 			return FileReference.Combine(GetVCToolPath(ToochainDir, Architecture), "link.exe");
-
 		}
 
 		/// <summary>
@@ -275,11 +289,20 @@ namespace UnrealBuildTool
 		/// </summary>
 		void RegisterCrossArchitectureToolPaths(DirectoryReference WindowsSdkDir, VersionNumber WindowsSdkVersion)
 		{
-			DirectoryReference HostBinDir = DirectoryReference.Combine(WindowsSdkDir, "bin", WindowsSdkVersion.ToString(), RuntimeInformation.ProcessArchitecture == System.Runtime.InteropServices.Architecture.Arm64 ? "arm64" : "x64");
-			DirectoryReference CrossBinDir = DirectoryReference.Combine(WindowsSdkDir, "bin", WindowsSdkVersion.ToString(), RuntimeInformation.ProcessArchitecture != System.Runtime.InteropServices.Architecture.Arm64 ? "arm64" : "x64");
-			if (FileReference.Exists(HostBinDir) && FileReference.Exists(CrossBinDir))
+			IEnumerable<DirectoryReference> ToolPaths = [
+				 DirectoryReference.Combine(WindowsSdkDir, "bin", WindowsSdkVersion.ToString()),
+				 DirectoryReference.Combine(Unreal.EngineDirectory, "Binaries", "Win64", "RadDebugger"),
+				 DirectoryReference.Combine(Unreal.EngineDirectory, "Binaries", "Win64", "UnrealBuildAccelerator")
+			];
+
+			foreach (DirectoryReference ToolPath in ToolPaths)
 			{
-				EpicGames.UBA.Utils.RegisterCrossArchitecturePath(HostBinDir.FullName, CrossBinDir.FullName);
+				DirectoryReference HostBinDir = DirectoryReference.Combine(ToolPath, MicrosoftPlatformSDK.MSVCHostUnrealArch.bIsX64 ? "x64" : "arm64");
+				DirectoryReference CrossBinDir = DirectoryReference.Combine(ToolPath, !MicrosoftPlatformSDK.MSVCHostUnrealArch.bIsX64 ? "arm64" : "x64");
+				if (FileReference.Exists(HostBinDir) && FileReference.Exists(CrossBinDir))
+				{
+					EpicGames.UBA.Utils.RegisterCrossArchitecturePath(HostBinDir.FullName, CrossBinDir.FullName);
+				}
 			}
 		}
 
