Disabling /JMC for msvc debug and development builds until we can have a more explicit opt-in for users or narrow the scope of what is built with /JMC as the option adds too much runtime overhead for it to be on by default (25% average for editor iteration, but invidivual functions can be 4x slower)

#rnx

diff --git a/Engine/Source/Programs/UnrealBuildTool/Platform/Windows/VCToolChain.cs b/Engine/Source/Programs/UnrealBuildTool/Platform/Windows/VCToolChain.cs
--- a/Engine/Source/Programs/UnrealBuildTool/Platform/Windows/VCToolChain.cs
+++ b/Engine/Source/Programs/UnrealBuildTool/Platform/Windows/VCToolChain.cs
@@ -934,11 +934,15 @@ namespace UnrealBuildTool
 			// If enabled, create debug information.
 			if (CompileEnvironment.bCreateDebugInfo && !(Target.WindowsPlatform.Compiler.IsIntel() && (CompileEnvironment.bPGOOptimize || CompileEnvironment.bPGOProfile || CompileEnvironment.bAllowLTCG)))
 			{
+				/*
+				Disabled until we can find a better means to limit the cost of /JMC as it can add considerable runtime overhead
+
 				// Enable Just-My-Code in Debug or Development, unless on Clang
 				if ((CompileEnvironment.Configuration == CppConfiguration.Debug || CompileEnvironment.Configuration == CppConfiguration.Development) && !Target.WindowsPlatform.Compiler.IsClang())
 				{
 					Arguments.Add("/JMC");
 				}
+				*/
 
 				// Store debug info in .pdb files.
 				// @todo clang: PDB files are emited from Clang but do not fully work with Visual Studio yet (breakpoints won't hit due to "symbol read error")
