Enable support for Android NDK r28

diff --git a/Engine/Config/Android/Android_SDK.json b/Engine/Config/Android/Android_SDK.json
--- a/Engine/Config/Android/Android_SDK.json
+++ b/Engine/Config/Android/Android_SDK.json
@@ -1,7 +1,7 @@
 {
 	"MainVersion": "r25b",
 	"MinVersion": "r25b",
-	"MaxVersion": "r27",
+	"MaxVersion": "r28",
 
 	"AutoSDKDirectory": "-27",
 
diff --git a/Engine/Source/Programs/UnrealBuildTool/Platform/Android/AndroidToolChain.cs b/Engine/Source/Programs/UnrealBuildTool/Platform/Android/AndroidToolChain.cs
--- a/Engine/Source/Programs/UnrealBuildTool/Platform/Android/AndroidToolChain.cs
+++ b/Engine/Source/Programs/UnrealBuildTool/Platform/Android/AndroidToolChain.cs
@@ -471,6 +471,8 @@ namespace UnrealBuildTool
 			Arguments.Add("-Wno-reorder");                  // member initialization order
 			Arguments.Add("-Wno-logical-op-parentheses");   // needed for external headers we can't change
 			Arguments.Add("-Wno-nonportable-include-path"); // not all of these are real
+			Arguments.Add("-Wno-extra-qualification");      // metasound DECLARE_METASOUND_xxx sometimes have redundant ::Metasound prefix
+			Arguments.Add("-Wno-shorten-64-to-32");         // too many right now, same on other platforms 
 		}
 
 		/// <inheritdoc/>
