In case of ASIS specifies that the activity should not be listed in the system's application launcher.

#rnx

diff --git a/Engine/Source/Programs/UnrealBuildTool/Platform/Android/UEDeployAndroid.cs b/Engine/Source/Programs/UnrealBuildTool/Platform/Android/UEDeployAndroid.cs
--- a/Engine/Source/Programs/UnrealBuildTool/Platform/Android/UEDeployAndroid.cs
+++ b/Engine/Source/Programs/UnrealBuildTool/Platform/Android/UEDeployAndroid.cs
@@ -2744,6 +2744,9 @@ namespace UnrealBuildTool
 			bool bEnableMulticastSupport = false;
 			Ini.GetBool("/Script/AndroidRuntimeSettings.AndroidRuntimeSettings", "bEnableMulticastSupport", out bEnableMulticastSupport);
 
+			bool bASISEnabled = false;
+			Ini.GetBool("/Script/AndroidSingleInstanceServiceEditor.AndroidSingleInstanceServiceRuntimeSettings", "bEnableASISPlugin", out bASISEnabled);
+
 			// only apply density to configChanges if using android-24 or higher and minimum sdk is 17
 			bool bAddDensity = (SDKLevelInt >= 24) && (MinSDKVersion >= 17);
 
@@ -2891,10 +2894,13 @@ namespace UnrealBuildTool
 				}
 				Text.AppendLine(String.Format("\t\t          android:screenOrientation=\"{0}\"", Orientation));
 				Text.AppendLine(String.Format("\t\t          android:debuggable=\"{0}\">", bIsForDistribution ? "false" : "true"));
-				Text.AppendLine("\t\t\t<intent-filter>");
-				Text.AppendLine("\t\t\t\t<action android:name=\"android.intent.action.MAIN\" />");
-				Text.AppendLine(String.Format("\t\t\t\t<category android:name=\"android.intent.category.LAUNCHER\" />"));
-				Text.AppendLine("\t\t\t</intent-filter>");
+				if (!bASISEnabled)
+				{
+					Text.AppendLine("\t\t\t<intent-filter>");
+					Text.AppendLine("\t\t\t\t<action android:name=\"android.intent.action.MAIN\" />");
+					Text.AppendLine(String.Format("\t\t\t\t<category android:name=\"android.intent.category.LAUNCHER\" />"));
+					Text.AppendLine("\t\t\t</intent-filter>");
+				}
 				Text.AppendLine("\t\t</activity>");
 				Text.AppendLine("\t\t<activity android:name=\"com.epicgames.unreal.GameActivity\"");
 				Text.AppendLine("\t\t          android:exported=\"true\"");
@@ -2928,7 +2934,7 @@ namespace UnrealBuildTool
 			}
 			Text.AppendLine(String.Format("\t\t          android:debuggable=\"{0}\">", bIsForDistribution ? "false" : "true"));
 			Text.AppendLine("\t\t\t<meta-data android:name=\"android.app.lib_name\" android:value=\"Unreal\"/>");
-			if (!bShowLaunchImage)
+			if (!bShowLaunchImage && !bASISEnabled)
 			{
 				Text.AppendLine("\t\t\t<intent-filter>");
 				Text.AppendLine("\t\t\t\t<action android:name=\"android.intent.action.MAIN\" />");
