UnrealBuildTool: Fix error getting WindowsName for unsupported host architecture binary

#rnx

diff --git a/Engine/Source/Programs/UnrealBuildTool/Platform/Windows/MicrosoftPlatformSDK.cs b/Engine/Source/Programs/UnrealBuildTool/Platform/Windows/MicrosoftPlatformSDK.cs
--- a/Engine/Source/Programs/UnrealBuildTool/Platform/Windows/MicrosoftPlatformSDK.cs
+++ b/Engine/Source/Programs/UnrealBuildTool/Platform/Windows/MicrosoftPlatformSDK.cs
@@ -666,7 +666,7 @@ namespace UnrealBuildTool
 				.ThenBy(x => x.FamilyRank)
 				.ThenByDescending(x => x.IsAutoSdk)
 				.ThenByDescending(x => x.Version)
-				.ThenBy(x => x.HostArchitecture.WindowsName);
+				.ThenBy(x => x.HostArchitecture.ToString());
 
 			if (SortedToolChains.Any())
 			{
@@ -1133,7 +1133,7 @@ namespace UnrealBuildTool
 				{
 					// Unknown binary architecture
 					HostArchitecture = UnrealArch.Deprecated;
-					Error = $"Unable to determine architecture of clang binary.";
+					Error = $"Clang binary is not x64 or arm64.";
 				}
 
 				Logger.LogDebug("Found Clang toolchain: {ToolChainDir} (Version={Version}, HostArchitecture={HostArchitecture}, Rank={Rank}, Error={Error})", ToolChainDir, Version, HostArchitecture, Rank, Error != null);
@@ -1187,7 +1187,7 @@ namespace UnrealBuildTool
 				{
 					// Unknown binary architecture
 					HostArchitecture = UnrealArch.Deprecated;
-					Error = $"Unable to determine architecture of clang binary.";
+					Error = $"Clang binary is not x64 or arm64.";
 				}
 
 				Logger.LogDebug("Found Intel OneAPI toolchain: {ToolChainDir} (Version={Version}, HostArchitecture={HostArchitecture}, Rank={Rank}, Error={Error})", ToolChainDir, Version, Is64Bit, Rank, Error != null);
