bWarnAboutMonolithicHeadersIncluded added to other header entry points to avoid redefinition errors & invariance tests to catch for the most simplistic COMPILE_WARNING() COMPILE_ERROR() macro usage.

diff --git a/Engine/Source/Programs/AutomationTool.Tests/LogParserTests.cs b/Engine/Source/Programs/AutomationTool.Tests/LogParserTests.cs
--- a/Engine/Source/Programs/AutomationTool.Tests/LogParserTests.cs
+++ b/Engine/Source/Programs/AutomationTool.Tests/LogParserTests.cs
@@ -419,6 +419,27 @@ namespace AutomationTool.Tests
 			}
 		}
 
+		[TestMethod]
+		public void CompileWarningErrorMacroMatcher()
+		{
+			// Visual C++ error & warning emitted via COMPILE_WARNING & COMPILE_ERROR
+			{
+				string[] lines =
+				{
+					MakeAbsolutePath(@"Engine\Plugins\Experimental\VirtualCamera\Source\VirtualCamera\Private\VCamBlueprintFunctionLibrary.cpp") + @"(249): warning: Some developer's custom warning message.",
+					MakeAbsolutePath(@"Engine\Plugins\Experimental\VirtualCamera\Source\VirtualCamera\Private\VCamBlueprintFunctionLibrary.cpp") + @"(250): error: Some developer's custom error message.",
+				};
+
+				List<LogEvent> logEvents = Parse(String.Join("\n", lines));
+				Assert.AreEqual(2, logEvents.Count);
+
+				Assert.AreEqual(LogLevel.Warning, logEvents[0].Level);
+				Assert.AreEqual(KnownLogEvents.Compiler, logEvents[0].Id);
+				Assert.AreEqual(LogLevel.Error, logEvents[1].Level);
+				Assert.AreEqual(KnownLogEvents.Compiler, logEvents[1].Id);
+			}
+		}
+
 		[TestMethod]
 		public void WarningsAsErrorsEventMatcher()
 		{
diff --git a/Engine/Source/Programs/UnrealBuildTool/Configuration/UEBuildModuleCPP.cs b/Engine/Source/Programs/UnrealBuildTool/Configuration/UEBuildModuleCPP.cs
--- a/Engine/Source/Programs/UnrealBuildTool/Configuration/UEBuildModuleCPP.cs
+++ b/Engine/Source/Programs/UnrealBuildTool/Configuration/UEBuildModuleCPP.cs
@@ -1270,7 +1270,7 @@ namespace UnrealBuildTool
 					Definitions.RemoveAll(x => x.Contains("UE_DEPRECATED_FORGAME", StringComparison.Ordinal));
 					Definitions.RemoveAll(x => x.Contains("UE_DEPRECATED_FORENGINE", StringComparison.Ordinal));
 
-					Definitions.Add($"UE_IS_ENGINE_MODULE={(Rules.bTreatAsEngineModule ? "1" : "0")}");
+					Definitions.Add($"UE_IS_ENGINE_MODULE={(Rules.bTreatAsEngineModule || Rules.Target.bWarnAboutMonolithicHeadersIncluded ? "1" : "0")}");
 					Definitions.Add($"UE_DEPRECATED_FORGAME={(Rules.bTreatAsEngineModule ? "PREPROCESSOR_NOTHING_FUNCTION" : "UE_DEPRECATED")}");
 					Definitions.Add($"UE_DEPRECATED_FORENGINE={(Rules.bTreatAsEngineModule || Rules.Target.bDisableEngineDeprecations ? "PREPROCESSOR_NOTHING_FUNCTION" : "UE_DEPRECATED")}");
 				}
@@ -2389,7 +2389,7 @@ namespace UnrealBuildTool
 			CompileEnvironment.bIsBuildingLibrary = false;
 
 			// Add a macro for when we're compiling an engine module, to enable additional compiler diagnostics through code.
-			CompileEnvironment.Definitions.Add($"UE_IS_ENGINE_MODULE={(Rules.bTreatAsEngineModule ? "1" : "0")}");
+			CompileEnvironment.Definitions.Add($"UE_IS_ENGINE_MODULE={(Rules.bTreatAsEngineModule || Rules.Target.bWarnAboutMonolithicHeadersIncluded ? "1" : "0")}");
 
 			if (!Rules.bTreatAsEngineModule)
 			{
