[UBT]
* Fixed linker flags to make dll builds linked with radlink.exe work

diff --git a/Engine/Source/Programs/UnrealBuildTool/Platform/Windows/VCToolChain.cs b/Engine/Source/Programs/UnrealBuildTool/Platform/Windows/VCToolChain.cs
--- a/Engine/Source/Programs/UnrealBuildTool/Platform/Windows/VCToolChain.cs
+++ b/Engine/Source/Programs/UnrealBuildTool/Platform/Windows/VCToolChain.cs
@@ -2732,12 +2732,20 @@ namespace UnrealBuildTool
 				Arguments.Add("/time+");
 			}
 
-			if (!bBuildImportLibraryOnly && Target.WindowsPlatform.bAllowRadLinker)
+			// Rad linker should not be used for import libs
+			if (Target.WindowsPlatform.bAllowRadLinker && !bBuildImportLibraryOnly)
 			{
-				Arguments.Add("/RAD_SUPPRESS_ERROR:36");
-				Arguments.Add("/RAD_SUPPRESS_ERROR:30");
-				Arguments.Add("/FIXED");
-				//Arguments.Add("/RAD_LARGE_PAGES");
+				// Saves space in type table in pdb
+				Arguments.Add("/rad_pdb_hash_type_names:lenient");
+
+				// Speeds up linking - TODO: disabled because of windows issues
+				//Arguments.Add("/rad_large_pages");
+
+				// Speeds up linking. We don't want this enabled on shipping since it will affect security
+				if (!LinkEnvironment.bIsBuildingDLL && !LinkEnvironment.bIsBuildingLibrary && LinkEnvironment.Configuration != CppConfiguration.Shipping)
+				{
+					Arguments.Add("/fixed");
+				}
 			}
 
 			// If we're only building an import library, add the '/DEF' option that tells the LIB utility
