[UBT]
* Enable caret in error messages from msvc cl.exe.. This makes it easier to figure out where in the line the error is

diff --git a/Engine/Source/Programs/UnrealBuildTool/Platform/Windows/VCToolChain.cs b/Engine/Source/Programs/UnrealBuildTool/Platform/Windows/VCToolChain.cs
--- a/Engine/Source/Programs/UnrealBuildTool/Platform/Windows/VCToolChain.cs
+++ b/Engine/Source/Programs/UnrealBuildTool/Platform/Windows/VCToolChain.cs
@@ -480,6 +480,9 @@ namespace UnrealBuildTool
 			else if (Target.WindowsPlatform.Compiler.IsMSVC())
 			{
 				Arguments.Add("/FC");
+
+				// Add column (in addition to line) to error messages
+				Arguments.Add("/diagnostics:caret");
 			}
 
 			if (Target.WindowsPlatform.Compiler.IsClang() && Target.Platform.IsInGroup(UnrealPlatformGroup.Windows))
