[UBT] CppCompileEnvironment to use CppCompileWarning instead of raw properties

diff --git a/Engine/Source/Programs/Shared/UnrealBuildTool.Tests/CppCompileWarningsTests.cs b/Engine/Source/Programs/Shared/UnrealBuildTool.Tests/CppCompileWarningsTests.cs
--- a/Engine/Source/Programs/Shared/UnrealBuildTool.Tests/CppCompileWarningsTests.cs
+++ b/Engine/Source/Programs/Shared/UnrealBuildTool.Tests/CppCompileWarningsTests.cs
@@ -17,6 +17,51 @@ namespace UnrealBuildTool.Tests
 			}
 		}
 
+		[TestMethod]
+		public void TestDefaultBuildSettings()
+		{
+			CppCompileWarnings defaultBuildSettings = new CppCompileWarnings();
+
+			Assert.AreEqual(WarningLevel.Warning, defaultBuildSettings.ShadowVariableWarningLevel);
+			Assert.AreEqual(WarningLevel.Off, defaultBuildSettings.UndefinedIdentifierWarningLevel);
+			Assert.AreEqual(WarningLevel.Off, defaultBuildSettings.UnsafeTypeCastWarningLevel);
+			Assert.AreEqual(WarningLevel.Off, defaultBuildSettings.SwitchUnhandledEnumeratorWarningLevel);
+		}
+
+		[TestMethod]
+		public void TestClonedMutableReturn()
+		{
+			CppCompileWarnings defaultBuildSettings = new CppCompileWarnings();
+			defaultBuildSettings.ShadowVariableWarningLevel = WarningLevel.Error;
+			defaultBuildSettings.UndefinedIdentifierWarningLevel = WarningLevel.Warning;
+			defaultBuildSettings.UnsafeTypeCastWarningLevel = WarningLevel.Error;
+			defaultBuildSettings.SwitchUnhandledEnumeratorWarningLevel= WarningLevel.Warning;
+
+			ReadOnlyCppCompileWarnings wrapper = new ReadOnlyCppCompileWarnings(defaultBuildSettings);
+
+			CppCompileWarnings cloned = wrapper.CloneAsWriteable();
+
+			{
+				Assert.AreEqual(defaultBuildSettings.ShadowVariableWarningLevel, cloned.ShadowVariableWarningLevel);
+				Assert.AreEqual(defaultBuildSettings.UndefinedIdentifierWarningLevel, cloned.UndefinedIdentifierWarningLevel);
+				Assert.AreEqual(defaultBuildSettings.UnsafeTypeCastWarningLevel, cloned.UnsafeTypeCastWarningLevel);
+				Assert.AreEqual(defaultBuildSettings.SwitchUnhandledEnumeratorWarningLevel, cloned.SwitchUnhandledEnumeratorWarningLevel);
+			}
+
+			// Modify and assert != to each other - the same underyling reference has not been returned
+			{
+				cloned.ShadowVariableWarningLevel = WarningLevel.Warning;
+				cloned.UndefinedIdentifierWarningLevel = WarningLevel.Error;
+				cloned.UnsafeTypeCastWarningLevel = WarningLevel.Warning;
+				cloned.SwitchUnhandledEnumeratorWarningLevel = WarningLevel.Error;
+
+				Assert.AreNotEqual(defaultBuildSettings.ShadowVariableWarningLevel, cloned.ShadowVariableWarningLevel);
+				Assert.AreNotEqual(defaultBuildSettings.UndefinedIdentifierWarningLevel, cloned.UndefinedIdentifierWarningLevel);
+				Assert.AreNotEqual(defaultBuildSettings.UnsafeTypeCastWarningLevel, cloned.UnsafeTypeCastWarningLevel);
+				Assert.AreNotEqual(defaultBuildSettings.SwitchUnhandledEnumeratorWarningLevel, cloned.SwitchUnhandledEnumeratorWarningLevel);
+			}
+		}
+
 		[TestMethod]
 		public void TestContextBuildSettingsOverride()
 		{
diff --git a/Engine/Source/Programs/UnrealBuildTool/Configuration/BuildSystemContext.cs b/Engine/Source/Programs/UnrealBuildTool/Configuration/BuildSystemContext.cs
--- a/Engine/Source/Programs/UnrealBuildTool/Configuration/BuildSystemContext.cs
+++ b/Engine/Source/Programs/UnrealBuildTool/Configuration/BuildSystemContext.cs
@@ -14,6 +14,18 @@ namespace UnrealBuildTool
 		public BuildSettingsVersion GetBuildSettings();
 	}
 
+	/// <summary>
+	/// The default provider that will just specify the <see cref="BuildSettingsVersion.V5"/> as a default.
+	/// </summary>
+	public class DefaultBuildContextProvider : IBuildContextProvider
+	{
+		/// <inheritdoc/>
+		public BuildSettingsVersion GetBuildSettings()
+		{
+			return BuildSettingsVersion.V5;
+		}
+	}
+
 	#region -- Adapters for IBuildContextProvider -- 
 
 	internal class ModuleRulesBuildSettingsProvider : IBuildContextProvider
diff --git a/Engine/Source/Programs/UnrealBuildTool/Configuration/CppCompileWarnings.cs b/Engine/Source/Programs/UnrealBuildTool/Configuration/CppCompileWarnings.cs
--- a/Engine/Source/Programs/UnrealBuildTool/Configuration/CppCompileWarnings.cs
+++ b/Engine/Source/Programs/UnrealBuildTool/Configuration/CppCompileWarnings.cs
@@ -126,9 +126,29 @@ namespace UnrealBuildTool
 		{
 			BuildContextProvider = new BuildSystemContext(contextObject);
 			ParentCppCompileWarnings = parentCompilerWarnings;
+
 			_logger = logger;
 		}
 
+		/// <summary>
+		/// Default constructor that uses the <see cref="DefaultBuildContextProvider"/>, and uses <see cref="CppCompileEnvironment"/> default settings 
+		/// </summary>
+		public CppCompileWarnings()
+		{
+			BuildContextProvider = new BuildSystemContext(new DefaultBuildContextProvider());
+			ApplyDefaults();
+		}
+
+		/// <summary>
+		/// Creates a memberwise shallow-copy of the provided object.
+		/// </summary>
+		/// <param name="other">The instance to clone.</param>
+		/// <returns>A clone of the provided <see cref="CppCompileWarnings"/>.</returns>
+		internal static CppCompileWarnings CreateShallowCopy(CppCompileWarnings other)
+		{
+			return (CppCompileWarnings)other.MemberwiseClone();
+		}
+
 		#endregion -- Constructors --
 
 		#region -- Public API --
@@ -194,6 +214,9 @@ namespace UnrealBuildTool
 			return resolvedWarning;
 		}
 
+		/// <summary>
+		/// Default settings for the target context.
+		/// </summary>
 		private void ApplyTargetDefaults()
 		{
 			_logger?.LogTrace("Applying default C++ Compiler warnings for the Target({TargetName}) context.", BuildContextProvider?._targetRulesPrivate?.Name);
@@ -203,6 +226,19 @@ namespace UnrealBuildTool
 			_switchUnhandledEnumeratorWarningLevelPrivate = WarningLevel.Off;
 		}
 
+		/// <summary>
+		/// Default settings for the no-arg constructor <see cref="CppCompileWarnings()"/>.
+		/// </summary>
+		/// <remarks>Defaults for the default compile environment state at <see cref="BuildSettingsVersion.V5"/> from <see cref="CppCompileEnvironment"/>.</remarks>
+		private void ApplyDefaults()
+		{
+			_logger?.LogTrace("Applying default C++ Compiler warnings for the default construction context.");
+			_shadowVariableWarningLevelPrivate = WarningLevel.Warning;
+			_unsafeTypeCastWarningLevelPrivate = WarningLevel.Off;
+			_undefinedIdentifierWarningLevelPrivate = WarningLevel.Off;
+			_switchUnhandledEnumeratorWarningLevelPrivate = WarningLevel.Off;
+		}
+
 		#endregion -- Private API --
 
 		#region -- Internal & Private Instance State --
diff --git a/Engine/Source/Programs/UnrealBuildTool/Configuration/ReadOnlyCppCompileWarnings.cs b/Engine/Source/Programs/UnrealBuildTool/Configuration/ReadOnlyCppCompileWarnings.cs
--- a/Engine/Source/Programs/UnrealBuildTool/Configuration/ReadOnlyCppCompileWarnings.cs
+++ b/Engine/Source/Programs/UnrealBuildTool/Configuration/ReadOnlyCppCompileWarnings.cs
@@ -23,6 +23,15 @@ namespace UnrealBuildTool
 			ParentCppCompileWarnings = inner.ParentCppCompileWarnings;
 		}
 
+		/// <summary>
+		/// Clones the <see cref="ReadOnlyCppCompileWarnings"/> as a writeable <see cref="CppCompileWarnings"/>.
+		/// </summary>
+		/// <returns>A writeable clone of the <see cref="ReadOnlyCppCompileWarnings"/>.</returns>
+		internal CppCompileWarnings CloneAsWriteable()
+		{
+			return CppCompileWarnings.CreateShallowCopy(_inner);
+		}
+
 		/// <summary>
 		/// Accessors for fields on the inner CppCompileWarnings instance
 		/// </summary>
diff --git a/Engine/Source/Programs/UnrealBuildTool/Configuration/UEBuildModuleCPP.cs b/Engine/Source/Programs/UnrealBuildTool/Configuration/UEBuildModuleCPP.cs
--- a/Engine/Source/Programs/UnrealBuildTool/Configuration/UEBuildModuleCPP.cs
+++ b/Engine/Source/Programs/UnrealBuildTool/Configuration/UEBuildModuleCPP.cs
@@ -2179,10 +2179,7 @@ namespace UnrealBuildTool
 			Result.bEnableObjCExceptions |= Rules.bEnableObjCExceptions;
 			Result.bEnableObjCAutomaticReferenceCounting = Rules.bEnableObjCAutomaticReferenceCounting;
 			Result.bWarningsAsErrors |= Rules.bWarningsAsErrors;
-			Result.ShadowVariableWarningLevel = Rules.CppCompileWarningSettings.ShadowVariableWarningLevel;
-			Result.UnsafeTypeCastWarningLevel = Rules.CppCompileWarningSettings.UnsafeTypeCastWarningLevel;
-			Result.UndefinedIdentifierWarningLevel = Rules.CppCompileWarningSettings.UndefinedIdentifierWarningLevel;
-			Result.SwitchUnhandledEnumeratorWarningLevel = Rules.CppCompileWarningSettings.SwitchUnhandledEnumeratorWarningLevel;
+			Result.CppCompileWarnings = CppCompileWarnings.CreateShallowCopy(Rules.CppCompileWarningSettings);
 			Result.bDisableStaticAnalysis = Rules.bDisableStaticAnalysis || (Target.bStaticAnalyzerProjectOnly && Rules.bTreatAsEngineModule);
 			Result.bStaticAnalyzerExtensions = Rules.bStaticAnalyzerExtensions;
 			Result.StaticAnalyzerRulesets = Rules.StaticAnalyzerRulesets;
diff --git a/Engine/Source/Programs/UnrealBuildTool/Configuration/UEBuildTarget.cs b/Engine/Source/Programs/UnrealBuildTool/Configuration/UEBuildTarget.cs
--- a/Engine/Source/Programs/UnrealBuildTool/Configuration/UEBuildTarget.cs
+++ b/Engine/Source/Programs/UnrealBuildTool/Configuration/UEBuildTarget.cs
@@ -4082,7 +4082,7 @@ namespace UnrealBuildTool
 						foreach (PrecompiledHeaderInstance UnsafeTypeCastWarningInstance in SharedPCH.Instances.Where(instance => instance.Modules.First().Rules.UnsafeTypeCastWarningLevel == WarningLevel.Warning))
 						{
 							CppCompileEnvironment UpdatedCppCompileEnvironment = new CppCompileEnvironment(UnsafeTypeCastWarningInstance.CompileEnvironment);
-							UpdatedCppCompileEnvironment.UnsafeTypeCastWarningLevel = WarningLevel.Error;
+							UpdatedCppCompileEnvironment.CppCompileWarnings.UnsafeTypeCastWarningLevel = WarningLevel.Error;
 							foreach (PrecompiledHeaderInstance UnsafeTypeCastErrorInstance in SharedPCH.Instances.Where(instance => instance.Modules.First().Rules.UnsafeTypeCastWarningLevel == WarningLevel.Error))
 							{
 								if (UEBuildModuleCPP.IsCompatibleForSharedPCH(UnsafeTypeCastErrorInstance.CompileEnvironment, UpdatedCppCompileEnvironment))
@@ -5225,10 +5225,7 @@ namespace UnrealBuildTool
 			GlobalCompileEnvironment.bEnableObjCExceptions = Rules.bForceEnableObjCExceptions || Rules.bCompileAgainstEditor;
 			GlobalCompileEnvironment.DefaultWarningLevel = Rules.DefaultWarningLevel;
 			GlobalCompileEnvironment.DeprecationWarningLevel = Rules.DeprecationWarningLevel;
-			GlobalCompileEnvironment.ShadowVariableWarningLevel = Rules.CppCompileWarningSettings.ShadowVariableWarningLevel;
-			GlobalCompileEnvironment.UnsafeTypeCastWarningLevel = Rules.CppCompileWarningSettings.UnsafeTypeCastWarningLevel;
-			GlobalCompileEnvironment.UndefinedIdentifierWarningLevel = Rules.CppCompileWarningSettings.UndefinedIdentifierWarningLevel;
-			GlobalCompileEnvironment.SwitchUnhandledEnumeratorWarningLevel = Rules.CppCompileWarningSettings.SwitchUnhandledEnumeratorWarningLevel;
+			GlobalCompileEnvironment.CppCompileWarnings = Rules.CppCompileWarningSettings.CloneAsWriteable();
 			GlobalCompileEnvironment.bRetainFramePointers = Rules.bRetainFramePointers;
 			GlobalCompileEnvironment.bWarningsAsErrors = Rules.bWarningsAsErrors;
 			GlobalCompileEnvironment.OptimizationLevel = Rules.OptimizationLevel;
diff --git a/Engine/Source/Programs/UnrealBuildTool/Platform/Windows/PVSToolChain.cs b/Engine/Source/Programs/UnrealBuildTool/Platform/Windows/PVSToolChain.cs
--- a/Engine/Source/Programs/UnrealBuildTool/Platform/Windows/PVSToolChain.cs
+++ b/Engine/Source/Programs/UnrealBuildTool/Platform/Windows/PVSToolChain.cs
@@ -794,7 +794,7 @@ namespace UnrealBuildTool
 			// Preprocess the source files with the regular toolchain
 			CppCompileEnvironment PreprocessCompileEnvironment = new CppCompileEnvironment(CompileEnvironment);
 			PreprocessCompileEnvironment.bPreprocessOnly = true;
-			PreprocessCompileEnvironment.UndefinedIdentifierWarningLevel = WarningLevel.Off; // Not sure why THIRD_PARTY_INCLUDES_START doesn't pick this up; the _Pragma appears in the preprocessed output. Perhaps in preprocess-only mode the compiler doesn't respect these?
+			PreprocessCompileEnvironment.CppCompileWarnings.UndefinedIdentifierWarningLevel = WarningLevel.Off; // Not sure why THIRD_PARTY_INCLUDES_START doesn't pick this up; the _Pragma appears in the preprocessed output. Perhaps in preprocess-only mode the compiler doesn't respect these?
 			PreprocessCompileEnvironment.AdditionalArguments += " /wd4005 /wd4828 /wd5105";
 			PreprocessCompileEnvironment.Definitions.Add("PVS_STUDIO");
 
diff --git a/Engine/Source/Programs/UnrealBuildTool/Platform/Windows/VCToolChain.cs b/Engine/Source/Programs/UnrealBuildTool/Platform/Windows/VCToolChain.cs
--- a/Engine/Source/Programs/UnrealBuildTool/Platform/Windows/VCToolChain.cs
+++ b/Engine/Source/Programs/UnrealBuildTool/Platform/Windows/VCToolChain.cs
@@ -1090,13 +1090,13 @@ namespace UnrealBuildTool
 			if (Target.WindowsPlatform.Compiler.IsMSVC())
 			{
 				// Disable shadow variable warnings
-				if (CompileEnvironment.ShadowVariableWarningLevel == WarningLevel.Off)
+				if (CompileEnvironment.CppCompileWarnings.ShadowVariableWarningLevel == WarningLevel.Off)
 				{
 					Arguments.Add("/wd4456"); // 4456 - declaration of 'LocalVariable' hides previous local declaration
 					Arguments.Add("/wd4458"); // 4458 - declaration of 'parameter' hides class member
 					Arguments.Add("/wd4459"); // 4459 - declaration of 'LocalVariable' hides global declaration
 				}
-				else if (CompileEnvironment.ShadowVariableWarningLevel == WarningLevel.Error)
+				else if (CompileEnvironment.CppCompileWarnings.ShadowVariableWarningLevel == WarningLevel.Error)
 				{
 					Arguments.Add("/we4456"); // 4456 - declaration of 'LocalVariable' hides previous local declaration
 					Arguments.Add("/we4458"); // 4458 - declaration of 'parameter' hides class member
@@ -1104,9 +1104,9 @@ namespace UnrealBuildTool
 				}
 			}
 
-			if (CompileEnvironment.UndefinedIdentifierWarningLevel != WarningLevel.Off && !CompileEnvironment.bPreprocessOnly)
+			if (CompileEnvironment.CppCompileWarnings.UndefinedIdentifierWarningLevel != WarningLevel.Off && !CompileEnvironment.bPreprocessOnly)
 			{
-				if (CompileEnvironment.UndefinedIdentifierWarningLevel == WarningLevel.Error)
+				if (CompileEnvironment.CppCompileWarnings.UndefinedIdentifierWarningLevel == WarningLevel.Error)
 				{
 					Arguments.Add("/we4668");
 				}
@@ -1116,9 +1116,9 @@ namespace UnrealBuildTool
 				}
 			}
 
-			if (CompileEnvironment.SwitchUnhandledEnumeratorWarningLevel != WarningLevel.Off)
+			if (CompileEnvironment.CppCompileWarnings.SwitchUnhandledEnumeratorWarningLevel != WarningLevel.Off)
 			{
-				if (CompileEnvironment.UndefinedIdentifierWarningLevel == WarningLevel.Error)
+				if (CompileEnvironment.CppCompileWarnings.UndefinedIdentifierWarningLevel == WarningLevel.Error)
 				{
 					Arguments.Add("/we4061");
 				}
@@ -1134,7 +1134,7 @@ namespace UnrealBuildTool
 			//@TODO: FLOATPRECISION: Consider doing the following as well:
 			//   4267: 'var' : conversion from 'size_t' to 'type', possible loss of data
 			//   4305: 'identifier' : truncation from 'type1' to 'type2'
-			WarningLevel EffectiveCastWarningLevel = (Target.Platform == UnrealTargetPlatform.Win64) ? CompileEnvironment.UnsafeTypeCastWarningLevel : WarningLevel.Off;
+			WarningLevel EffectiveCastWarningLevel = (Target.Platform == UnrealTargetPlatform.Win64) ? CompileEnvironment.CppCompileWarnings.UnsafeTypeCastWarningLevel : WarningLevel.Off;
 			if (EffectiveCastWarningLevel == WarningLevel.Error)
 			{
 				Arguments.Add("/we4244");
diff --git a/Engine/Source/Programs/UnrealBuildTool/System/CppCompileEnvironment.cs b/Engine/Source/Programs/UnrealBuildTool/System/CppCompileEnvironment.cs
--- a/Engine/Source/Programs/UnrealBuildTool/System/CppCompileEnvironment.cs
+++ b/Engine/Source/Programs/UnrealBuildTool/System/CppCompileEnvironment.cs
@@ -357,24 +357,10 @@ namespace UnrealBuildTool
 		public WarningLevel DeprecationWarningLevel = WarningLevel.Warning;
 
 		/// <summary>
-		/// Whether to warn about the use of shadow variables
+		/// The compiler warnings associacted with the compile environment.
 		/// </summary>
-		public WarningLevel ShadowVariableWarningLevel = WarningLevel.Warning;
+		public CppCompileWarnings CppCompileWarnings;
 
-		/// <summary>
-		/// How to treat unsafe implicit type cast warnings (e.g., double->float or int64->int32)
-		/// </summary>
-		public WarningLevel UnsafeTypeCastWarningLevel = WarningLevel.Off;
-
-		/// <summary>
-		/// Indicates what warning/error level to treat undefined identifiers in conditional expressions.
-		/// </summary>
-		public WarningLevel UndefinedIdentifierWarningLevel = WarningLevel.Off;
-
-		/// <summary>
-		/// Indicates what warning/error level to treat unhandled enumerators in switches on enumeration-typed values.
-		/// </summary>
-		public WarningLevel SwitchUnhandledEnumeratorWarningLevel = WarningLevel.Off;
 
 		/// <summary>
 		/// Whether to treat all warnings as errors
@@ -755,6 +741,7 @@ namespace UnrealBuildTool
 			SharedUserIncludePaths = new HashSet<DirectoryReference>();
 			SharedSystemIncludePaths = new HashSet<DirectoryReference>();
 			ModuleInterfacePaths = new HashSet<DirectoryReference>();
+			CppCompileWarnings = new CppCompileWarnings();
 		}
 
 		/// <summary>
@@ -787,10 +774,7 @@ namespace UnrealBuildTool
 			bEnableObjCAutomaticReferenceCounting = Other.bEnableObjCAutomaticReferenceCounting;
 			DefaultWarningLevel = Other.DefaultWarningLevel;
 			DeprecationWarningLevel = Other.DeprecationWarningLevel;
-			ShadowVariableWarningLevel = Other.ShadowVariableWarningLevel;
-			UnsafeTypeCastWarningLevel = Other.UnsafeTypeCastWarningLevel;
-			SwitchUnhandledEnumeratorWarningLevel = Other.SwitchUnhandledEnumeratorWarningLevel;
-			UndefinedIdentifierWarningLevel = Other.UndefinedIdentifierWarningLevel;
+			CppCompileWarnings = CppCompileWarnings.CreateShallowCopy(Other.CppCompileWarnings);
 			bWarningsAsErrors = Other.bWarningsAsErrors;
 			bDisableStaticAnalysis = Other.bDisableStaticAnalysis;
 			StaticAnalyzerCheckers = new HashSet<string>(Other.StaticAnalyzerCheckers);
diff --git a/Engine/Source/Programs/UnrealBuildTool/ToolChain/ClangWarnings.cs b/Engine/Source/Programs/UnrealBuildTool/ToolChain/ClangWarnings.cs
--- a/Engine/Source/Programs/UnrealBuildTool/ToolChain/ClangWarnings.cs
+++ b/Engine/Source/Programs/UnrealBuildTool/ToolChain/ClangWarnings.cs
@@ -122,21 +122,21 @@ namespace UnrealBuildTool
 			else
 			{
 				// https://clang.llvm.org/docs/DiagnosticsReference.html#wshadow
-				if (CompileEnvironment.ShadowVariableWarningLevel != WarningLevel.Off)
+				if (CompileEnvironment.CppCompileWarnings.ShadowVariableWarningLevel != WarningLevel.Off)
 				{
-					Arguments.Add("-Wshadow" + ((CompileEnvironment.ShadowVariableWarningLevel == WarningLevel.Error) ? "" : " -Wno-error=shadow"));
+					Arguments.Add("-Wshadow" + ((CompileEnvironment.CppCompileWarnings.ShadowVariableWarningLevel == WarningLevel.Error) ? "" : " -Wno-error=shadow"));
 				}
 			}
 
 			// https://clang.llvm.org/docs/DiagnosticsReference.html#wundef
-			if (CompileEnvironment.UndefinedIdentifierWarningLevel != WarningLevel.Off)
+			if (CompileEnvironment.CppCompileWarnings.UndefinedIdentifierWarningLevel != WarningLevel.Off)
 			{
-				Arguments.Add("-Wundef" + (CompileEnvironment.UndefinedIdentifierWarningLevel == WarningLevel.Error ? "" : " -Wno-error=undef"));
+				Arguments.Add("-Wundef" + (CompileEnvironment.CppCompileWarnings.UndefinedIdentifierWarningLevel == WarningLevel.Error ? "" : " -Wno-error=undef"));
 			}
 
-			if (CompileEnvironment.SwitchUnhandledEnumeratorWarningLevel != WarningLevel.Off)
+			if (CompileEnvironment.CppCompileWarnings.SwitchUnhandledEnumeratorWarningLevel != WarningLevel.Off)
 			{
-				Arguments.Add("-Wswitch-enum" + (CompileEnvironment.SwitchUnhandledEnumeratorWarningLevel == WarningLevel.Error ? "" : " -Wno-error=switch-enum"));
+				Arguments.Add("-Wswitch-enum" + (CompileEnvironment.CppCompileWarnings.SwitchUnhandledEnumeratorWarningLevel == WarningLevel.Error ? "" : " -Wno-error=switch-enum"));
 			}
 
 			// Note: This should be kept in sync with PRAGMA_DISABLE_UNSAFE_TYPECAST_WARNINGS in ClangPlatformCompilerPreSetup.h
@@ -149,14 +149,14 @@ namespace UnrealBuildTool
 				//"sign-conversion",	<-- too many hits right now, probably want it eventually
 			};
 
-			if (CompileEnvironment.UnsafeTypeCastWarningLevel == WarningLevel.Error)
+			if (CompileEnvironment.CppCompileWarnings.UnsafeTypeCastWarningLevel == WarningLevel.Error)
 			{
 				foreach (string Warning in UnsafeTypeCastWarningList)
 				{
 					Arguments.Add("-W" + Warning);
 				}
 			}
-			else if (CompileEnvironment.UnsafeTypeCastWarningLevel == WarningLevel.Warning)
+			else if (CompileEnvironment.CppCompileWarnings.UnsafeTypeCastWarningLevel == WarningLevel.Warning)
 			{
 				foreach (string Warning in UnsafeTypeCastWarningList)
 				{
diff --git a/Engine/Source/Programs/UnrealBuildTool/ToolChain/IWYUToolChain.cs b/Engine/Source/Programs/UnrealBuildTool/ToolChain/IWYUToolChain.cs
--- a/Engine/Source/Programs/UnrealBuildTool/ToolChain/IWYUToolChain.cs
+++ b/Engine/Source/Programs/UnrealBuildTool/ToolChain/IWYUToolChain.cs
@@ -190,7 +190,7 @@ namespace UnrealBuildTool
 
 		protected override void GetCompileArguments_WarningsAndErrors(CppCompileEnvironment CompileEnvironment, List<string> Arguments)
 		{
-			CompileEnvironment.ShadowVariableWarningLevel = WarningLevel.Off;
+			CompileEnvironment.CppCompileWarnings.ShadowVariableWarningLevel = WarningLevel.Off;
 			base.GetCompileArguments_WarningsAndErrors(CompileEnvironment, Arguments);
 			Arguments.Add("-Wno-undefined-bool-conversion");
 			Arguments.Add("-Wno-deprecated-anon-enum-enum-conversion");
