- Cleaned up versioning and other SDK things for Apple plaftforms, mostly related to:
  - Runtime min OS version
  - Moved things from Mac/IOS specific classes into Apple classes
  - Deleted the other ApplePlatformSDK.cs in EpicGames.Build and rolled into ApplePlaformSDK in UBT
  - Added stuff, like DeveloperDir, and GetToolchainPath() to ApplePlaformSDK
  - Removed the -miphoneos-min-version=13 type params to the new -target arm64-apple-iphoneos13 standard param
  - Now allow for MinSoftwareVersion to be overridden on a per-project basis (could be used instead of .ini setting)

diff --git a/Engine/Platforms/VisionOS/Config/VisionOS_SDK.json b/Engine/Platforms/VisionOS/Config/VisionOS_SDK.json
--- a/Engine/Platforms/VisionOS/Config/VisionOS_SDK.json
+++ b/Engine/Platforms/VisionOS/Config/VisionOS_SDK.json
@@ -2,5 +2,9 @@
 	// inherit the IOS MinSoftwareVersion
 	"ParentSDKFile": "../../../Config/IOS/IOS_SDK.json",
 
+	// xcode version
 	"MinVersion": "15.3.0",
+	
+	// runtime OS version
+	"MinSoftwareVersion":  "1.0",
 }
diff --git a/Engine/Platforms/VisionOS/Source/Programs/UnrealBuildTool/UEBuildVisionOS.cs b/Engine/Platforms/VisionOS/Source/Programs/UnrealBuildTool/UEBuildVisionOS.cs
--- a/Engine/Platforms/VisionOS/Source/Programs/UnrealBuildTool/UEBuildVisionOS.cs
+++ b/Engine/Platforms/VisionOS/Source/Programs/UnrealBuildTool/UEBuildVisionOS.cs
@@ -9,7 +9,7 @@ namespace UnrealBuildTool
 	{
 		// @todo get these from ini or whatever
 		public override string RuntimeDevices => "7";
-		public override string RuntimeVersion => "1.0";
+		public override string RuntimeVersion => UEBuildPlatformSDK.GetSDKForPlatform("VisionOS")!.GetSoftwareInfo()!.Min!;
 
 		public VisionOSProjectSettings(FileReference? ProjectFile, string? Bundle)
 			: base(ProjectFile, UnrealTargetPlatform.VisionOS, Bundle)
@@ -119,7 +119,7 @@ namespace UnrealBuildTool
 		{
 			VisionOSProjectSettings ProjectSettings = ((VisionOSPlatform)UEBuildPlatform.GetBuildPlatform(UnrealTargetPlatform.VisionOS)).ReadProjectSettings(Target.ProjectFile);
 			VisionOSToolChain NewToolChain = new VisionOSToolChain(Target, ProjectSettings, Logger);
-			SDKVersionFloat = NewToolChain.GetSDKVersionFloat();
+			SDKVersionFloat = NewToolChain.SDKVersionFloat;
 			return NewToolChain;
 		}
 
diff --git a/Engine/Platforms/VisionOS/Source/Programs/UnrealBuildTool/VisionOSToolChain.cs b/Engine/Platforms/VisionOS/Source/Programs/UnrealBuildTool/VisionOSToolChain.cs
--- a/Engine/Platforms/VisionOS/Source/Programs/UnrealBuildTool/VisionOSToolChain.cs
+++ b/Engine/Platforms/VisionOS/Source/Programs/UnrealBuildTool/VisionOSToolChain.cs
@@ -1,33 +1,23 @@
 // Copyright Epic Games, Inc. All Rights Reserved.
 
 using Microsoft.Extensions.Logging;
+using EpicGames.Core;
 
 namespace UnrealBuildTool
 {
 	class VisionOSToolChainSettings : IOSToolChainSettings
 	{
-		public VisionOSToolChainSettings(ILogger Logger) : base("XROS", "XRSimulator", "xros", Logger)
+		public VisionOSToolChainSettings(IOSProjectSettings ProjectSettings, ILogger Logger) 
+			: base("XROS", "XRSimulator", "xros", ProjectSettings, Logger)
 		{
 		}
 	}
 
 	class VisionOSToolChain : IOSToolChain
 	{
-		private VisionOSToolChainSettings Settings => (VisionOSToolChainSettings)ToolChainSettings.Value;
-
 		public VisionOSToolChain(ReadOnlyTargetRules InTarget, VisionOSProjectSettings InProjectSettings, ILogger InLogger)
-			: base(InTarget, InProjectSettings, () => new VisionOSToolChainSettings(InLogger), ClangToolChainOptions.None, InLogger)
-		{
-		}
-
-		public override string GetXcodeMinVersionParam(UnrealArch Architecture)
-		{
-			return "";
-		}
-
-		public float GetSDKVersionFloat()
+			: base(InTarget, InProjectSettings, () => new VisionOSToolChainSettings(InProjectSettings, InLogger), ClangToolChainOptions.None, InLogger)
 		{
-			return Settings.SDKVersionFloat;
 		}
 	}
 }
diff --git a/Engine/Platforms/VisionOS/Source/Runtime/Launch/Launch_VisionOS.Build.cs b/Engine/Platforms/VisionOS/Source/Runtime/Launch/Launch_VisionOS.Build.cs
--- a/Engine/Platforms/VisionOS/Source/Runtime/Launch/Launch_VisionOS.Build.cs
+++ b/Engine/Platforms/VisionOS/Source/Runtime/Launch/Launch_VisionOS.Build.cs
@@ -8,5 +8,6 @@ public class Launch_VisionOS : Launch
 {
 	public Launch_VisionOS(ReadOnlyTargetRules Target) : base(Target)
 	{
+		PrivateDependencyModuleNames.Add("Swift");
 	}
 }
diff --git a/Engine/Source/Programs/AutomationTool/Scripts/CheckXcodeVersion.cs b/Engine/Source/Programs/AutomationTool/Scripts/CheckXcodeVersion.cs
deleted file mode 100644
diff --git a/Engine/Source/Programs/Shared/EpicGames.Build/Platform/Mac/ApplePlatformSDK.cs b/Engine/Source/Programs/Shared/EpicGames.Build/Platform/Mac/ApplePlatformSDK.cs
deleted file mode 100644
diff --git a/Engine/Source/Programs/UnrealBuildTool/Executors/Experimental/FASTBuild.cs b/Engine/Source/Programs/UnrealBuildTool/Executors/Experimental/FASTBuild.cs
--- a/Engine/Source/Programs/UnrealBuildTool/Executors/Experimental/FASTBuild.cs
+++ b/Engine/Source/Programs/UnrealBuildTool/Executors/Experimental/FASTBuild.cs
@@ -1074,8 +1074,8 @@ namespace UnrealBuildTool
 
 			if (BuildHostPlatform.Current.Platform == UnrealTargetPlatform.Mac)
 			{
-				AddText($".MacBaseSDKDir = '{MacToolChain.Settings.GetSDKPath()}'\n");
-				AddText($".MacToolchainDir = '{MacToolChain.Settings.ToolchainDir}'\n");
+				AddText($".MacBaseSDKDir = '{ApplePlatformSDK.GetPlatformSDKDirectory("MacOSX")}'\n");
+				AddText($".MacToolchainDir = '{ApplePlatformSDK.GetToolchainDirectory()}'\n");
 				AddText($"Compiler('UEAppleCompiler') \n{{\n");
 				AddText($"\t.Executable = '$MacToolchainDir$/clang++'\n");
 				AddText($"\t.ClangRewriteIncludes = false\n"); // This is to fix an issue with iOS clang builds, __has_include, and Objective-C #imports
diff --git a/Engine/Source/Programs/UnrealBuildTool/Platform/Apple/AppleProjectGenerator.cs b/Engine/Source/Programs/UnrealBuildTool/Platform/Apple/AppleProjectGenerator.cs
--- a/Engine/Source/Programs/UnrealBuildTool/Platform/Apple/AppleProjectGenerator.cs
+++ b/Engine/Source/Programs/UnrealBuildTool/Platform/Apple/AppleProjectGenerator.cs
@@ -146,11 +146,10 @@ namespace UnrealBuildTool
 			{
 				if (Platform == UnrealTargetPlatform.Mac)
 				{
-					return MacToolChain.Settings.GetSDKPath().FullName;
+					return ApplePlatformSDK.GetPlatformSDKDirectory("MacOSX").FullName;
 				}
 
-				// Resolves to AppleToolChainSettings.GetSDKPath() which is not overridden
-				return new IOSToolChainSettings(Logger).GetSDKPath(Architecture).FullName;
+				return ApplePlatformSDK.GetPlatformSDKDirectory(Architecture == UnrealArch.IOSSimulator ? "iPhoneSimulator" : "iPhoneOS").FullName;
 			}
 		}
 	}
diff --git a/Engine/Source/Programs/UnrealBuildTool/Platform/IOS/IOSPostBuildSyncMode.cs b/Engine/Source/Programs/UnrealBuildTool/Platform/IOS/IOSPostBuildSyncMode.cs
--- a/Engine/Source/Programs/UnrealBuildTool/Platform/IOS/IOSPostBuildSyncMode.cs
+++ b/Engine/Source/Programs/UnrealBuildTool/Platform/IOS/IOSPostBuildSyncMode.cs
@@ -75,8 +75,6 @@ namespace UnrealBuildTool
 			Arguments.ApplyTo(this);
 			Arguments.CheckAllArgumentsUsed();
 
-			IOSToolChainSettings.SelectXcode(false, Logger);
-
 			// Run the PostBuildSync command
 			IOSPostBuildSyncTarget Target = BinaryFormatterUtils.Load<IOSPostBuildSyncTarget>(InputFile!);
 			IOSToolChain.PostBuildSync(Target, Logger);
diff --git a/Engine/Source/Programs/UnrealBuildTool/Platform/IOS/IOSToolChain.cs b/Engine/Source/Programs/UnrealBuildTool/Platform/IOS/IOSToolChain.cs
--- a/Engine/Source/Programs/UnrealBuildTool/Platform/IOS/IOSToolChain.cs
+++ b/Engine/Source/Programs/UnrealBuildTool/Platform/IOS/IOSToolChain.cs
@@ -17,13 +17,13 @@ namespace UnrealBuildTool
 {
 	class IOSToolChainSettings : AppleToolChainSettings
 	{
-		public IOSToolChainSettings(ILogger Logger)
-			: this("iPhoneOS", "iPhoneSimulator", "ios", Logger)
+		public IOSToolChainSettings(IOSProjectSettings ProjectSettings, ILogger Logger)
+			: this("iPhoneOS", "iPhoneSimulator", "ios", ProjectSettings, Logger)
 		{
 		}
 
-		protected IOSToolChainSettings(string DevicePlatformName, string SimulatorPlatformName, string TargetOSName, ILogger Logger)
-			: base(DevicePlatformName, SimulatorPlatformName, TargetOSName, true, Logger)
+		protected IOSToolChainSettings(string DevicePlatformName, string SimulatorPlatformName, string TargetOSName, IOSProjectSettings ProjectSettings, ILogger Logger)
+			: base(DevicePlatformName, SimulatorPlatformName, TargetOSName, ProjectSettings.RuntimeVersion, true, Logger)
 		{
 		}
 	}
@@ -35,7 +35,7 @@ namespace UnrealBuildTool
 		protected IOSProjectSettings ProjectSettings;
 
 		public IOSToolChain(ReadOnlyTargetRules? Target, IOSProjectSettings InProjectSettings, ClangToolChainOptions ToolchainOptions, ILogger InLogger)
-			: this(Target, InProjectSettings, () => new IOSToolChainSettings(InLogger), ToolchainOptions, InLogger)
+			: this(Target, InProjectSettings, () => new IOSToolChainSettings(InProjectSettings, InLogger), ToolchainOptions, InLogger)
 		{
 		}
 
@@ -70,18 +70,11 @@ namespace UnrealBuildTool
 		/// </summary>
 		private const string IOSArchiver = "libtool";
 
-		private IOSToolChainSettings Settings => (IOSToolChainSettings)ToolChainSettings.Value;
-
 		protected override ClangToolChainInfo GetToolChainInfo()
 		{
 			FileReference CompilerPath = FileReference.Combine(Settings.ToolchainDir, IOSCompiler);
 			FileReference ArchiverPath = FileReference.Combine(Settings.ToolchainDir, IOSArchiver);
-			return new AppleToolChainInfo(UnrealTargetPlatform.IOS, IOSToolChainSettings.XcodeDeveloperDir, CompilerPath, ArchiverPath, Logger);
-		}
-
-		public override string GetSDKVersion()
-		{
-			return Settings.SDKVersionFloat.ToString();
+			return new AppleToolChainInfo(UnrealTargetPlatform.IOS, ApplePlatformSDK.DeveloperDir, CompilerPath, ArchiverPath, Logger);
 		}
 
 		public override void ModifyBuildProducts(ReadOnlyTargetRules Target, UEBuildBinary Binary, IEnumerable<string> Libraries, IEnumerable<UEBuildBundleResource> BundleResources, Dictionary<FileReference, BuildProductType> BuildProducts)
@@ -227,12 +220,6 @@ namespace UnrealBuildTool
 
 			Arguments.Add($"-isysroot \"{Settings.GetSDKPath(CompileEnvironment.Architecture)}\"");
 
-			string MinVersionParam = GetXcodeMinVersionParam(CompileEnvironment.Architecture);
-			if (MinVersionParam != "")
-			{
-				Arguments.Add($"-m{MinVersionParam}={ProjectSettings.RuntimeVersion}");
-			}
-
 			// Add additional frameworks so that their headers can be found
 			foreach (UEBuildFramework Framework in CompileEnvironment.AdditionalFrameworks)
 			{
@@ -301,11 +288,6 @@ namespace UnrealBuildTool
 			RunExecutableAndWait("mkdir", String.Format("-p \"{0}\"", Path), out ResultsText);
 		}
 
-		public virtual string GetXcodeMinVersionParam(UnrealArch Architecture)
-		{
-			return (Architecture != UnrealArch.IOSSimulator) ? "iphoneos-version-min" : "iphonesimulator-version-min";
-		}
-
 		public string GetAdditionalLinkerFlags(CppConfiguration InConfiguration)
 		{
 			if (InConfiguration != CppConfiguration.Shipping)
@@ -327,11 +309,7 @@ namespace UnrealBuildTool
 			Arguments.Add($" -isysroot \"{SDKPath}\"");
 
 			Arguments.Add("-dead_strip");
-			string MinVersionParam = GetXcodeMinVersionParam(LinkEnvironment.Architecture);
-			if (MinVersionParam != "")
-			{
-				Arguments.Add($"-m{MinVersionParam}={ProjectSettings.RuntimeVersion}");
-			}
+			Arguments.Add($"-target {Settings.GetTargetTuple(LinkEnvironment.Architecture, Target!.Platform, Target!.Type)}");
 			Arguments.Add("-Wl-no_pie");
 			Arguments.Add("-stdlib=libc++");
 			Arguments.Add("-ObjC");
@@ -470,7 +448,7 @@ namespace UnrealBuildTool
 
 				// allow for auto-linked swift libs to be found
 				// @todo Mac will need this too - a lot of this code should probably be moved up to AppleToolChain and shared
-				LinkCommandArguments += $" -L{ToolChainSettings.Value.GetSDKPath(LinkEnvironment.Architecture)}/usr/lib/swift";
+				LinkCommandArguments += $" -L{Settings.GetSDKPath(LinkEnvironment.Architecture)}/usr/lib/swift";
 			}
 
 			// Handle additional framework assets that might need to be shadowed
@@ -1293,7 +1271,7 @@ namespace UnrealBuildTool
 			}
 
 			string FabricPath = Unreal.EngineDirectory + "/Intermediate/UnzippedFrameworks/Crashlytics/Fabric.embeddedframework";
-			if (Directory.Exists(FabricPath) && Unreal.IsBuildMachine())
+			if (Directory.Exists(FabricPath) && UnrealBuildBase.Unreal.IsBuildMachine())
 			{
 				//string PlistFile = ProjectDir + "/Intermediate/IOS/" + ProjectName + "-Info.plist";
 				Process FabricProcess = new Process();
@@ -1603,7 +1581,7 @@ namespace UnrealBuildTool
 						Writer.WriteLine(String.Format("rm -rf \"{0}\"", FrameworkPayloadDirectory));
 
 						// Build the framework wrapper
-						CmdLine = FileReference.Combine(AppleToolChainSettings.XcodeDeveloperDir, "usr/bin/xcodebuild").FullName +
+						CmdLine = FileReference.Combine(ApplePlatformSDK.DeveloperDir, "usr/bin/xcodebuild").FullName +
 							" -project \"" + WrapperProject + "\"" +
 								" -configuration \"" + ConfigName + "\"" +
 							" -scheme '" + SchemeName + "'" +
@@ -1616,7 +1594,7 @@ namespace UnrealBuildTool
 					else
 					{
 						// code sign the project
-						CmdLine = FileReference.Combine(AppleToolChainSettings.XcodeDeveloperDir, "usr/bin/xcodebuild").FullName +
+						CmdLine = FileReference.Combine(ApplePlatformSDK.DeveloperDir, "usr/bin/xcodebuild").FullName +
 							" -workspace \"" + XcodeWorkspaceDir + "\"" +
 								" -configuration \"" + ConfigName + "\"" +
 							" -scheme '" + SchemeName + "'" +
diff --git a/Engine/Source/Programs/UnrealBuildTool/Platform/IOS/UEBuildIOS.cs b/Engine/Source/Programs/UnrealBuildTool/Platform/IOS/UEBuildIOS.cs
--- a/Engine/Source/Programs/UnrealBuildTool/Platform/IOS/UEBuildIOS.cs
+++ b/Engine/Source/Programs/UnrealBuildTool/Platform/IOS/UEBuildIOS.cs
@@ -369,15 +369,16 @@ namespace UnrealBuildTool
 			{
 				switch (MinimumIOSVersion)
 				{
-					case "IOS_Minimum":
 					case "IOS_15":
 						return "15.0";
 					case "IOS_16":
 						return "16.0";
 					case "IOS_17":
 						return "17.0";
+					case "IOS_Minimum":
 					default:
-						return "15.0";
+						// default to the UE min in IOS_SDK.json
+						return UEBuildPlatformSDK.GetSDKForPlatform("IOS")!.GetSoftwareInfo()!.Min!;
 				}
 			}
 		}
diff --git a/Engine/Source/Programs/UnrealBuildTool/Platform/Mac/ApplePlatformSDK.cs b/Engine/Source/Programs/UnrealBuildTool/Platform/Mac/ApplePlatformSDK.cs
--- a/Engine/Source/Programs/UnrealBuildTool/Platform/Mac/ApplePlatformSDK.cs
+++ b/Engine/Source/Programs/UnrealBuildTool/Platform/Mac/ApplePlatformSDK.cs
@@ -1,7 +1,14 @@
 // Copyright Epic Games, Inc. All Rights Reserved.
 
+using System;
+using System.IO;
+using System.Diagnostics;
+using System.Text;
+using System.Text.RegularExpressions;
+using System.Runtime.Versioning;
 using EpicGames.Core;
 using Microsoft.Extensions.Logging;
+using Microsoft.Win32;
 using UnrealBuildBase;
 
 ///////////////////////////////////////////////////////////////////
@@ -19,14 +26,241 @@ namespace UnrealBuildTool
 		{
 		}
 
-		public override bool TryConvertVersionToInt(string? StringValue, out ulong OutValue, string? Hint)
+		public static Lazy<DirectoryReference?> _DeveloperDir = new(() => GetDeveloperDir());
+
+		// this shouldn't get called if there's no SDK installed
+		public static DirectoryReference DeveloperDir => _DeveloperDir.Value!;
+
+		// the vrsion of Xcode that is installed, not any of the SDKs inside
+		public static Lazy<string?> InstalledXcodeVersion => new(() => _GetInstalledSDKVersion());
+
+		public static DirectoryReference GetToolchainDirectory()
+		{
+			return DirectoryReference.Combine(DeveloperDir, "Toolchains/XcodeDefault.xctoolchain/usr/bin");
+		}
+
+		public static DirectoryReference GetPlatformSDKDirectory(string AppleOSName)
+		{
+			return DirectoryReference.Combine(DeveloperDir, $"Platforms/{AppleOSName}.platform/Developer/SDKs/{AppleOSName}.sdk");
+		}
+
+		public static string? GetPlatformSDKVersion(string AppleOSName)
+		{
+			return GetVersionFromSDKDir(AppleOSName);
+		}
+
+		public static float GetPlatformSDKVersionFloat(string AppleOSName)
 		{
-			return UnrealBuildBase.ApplePlatformSDK.TryConvertVersionToInt(StringValue, out OutValue);
+			return Single.Parse(GetPlatformSDKVersion(AppleOSName) ?? "0.0", System.Globalization.CultureInfo.InvariantCulture);
 		}
 
+		// 8 bits per component, with high getting extra from high 32
+		[GeneratedRegex(@"^(\d+).(\d+)(.(\d+))?(.(\d+))?(.(\d+))?$")]
+		private static partial Regex VersionToIntRegex();
+
+		public override bool TryConvertVersionToInt(string? StringValue, out ulong OutValue, string? Hint)
+		{
+			OutValue = 0;
+
+			if (StringValue == null)
+			{
+				return false;
+			}
+
+			Match Result = VersionToIntRegex().Match(StringValue);
+			if (Result.Success)
+			{
+				OutValue = UInt64.Parse(Result.Groups[1].Value) << 24 | UInt64.Parse(Result.Groups[2].Value) << 16;
+				if (Result.Groups[4].Success)
+				{
+					OutValue |= UInt64.Parse(Result.Groups[4].Value) << 8;
+				}
+				if (Result.Groups[6].Success)
+				{
+					OutValue |= UInt64.Parse(Result.Groups[6].Value) << 0;
+				}
+				return true;
+			}
+
+			return false;
+		}
+		
 		protected override string? GetInstalledSDKVersion()
 		{
-			return UnrealBuildBase.ApplePlatformSDK.InstalledSDKVersion;
+			return _GetInstalledSDKVersion();
+		}
+
+		private static string? _GetInstalledSDKVersion()
+		{
+			// get Xcode version on Mac
+			if (OperatingSystem.IsMacOS())
+			{
+				if (_DeveloperDir.Value == null)
+				{
+					return null;
+				}
+				FileReference Plist = FileReference.Combine(DeveloperDir.ParentDirectory!, "Info.plist");
+				// Find out the version number in Xcode.app/Contents/Info.plist
+				int ExitCode;
+				string Output = Utils.RunLocalProcessAndReturnStdOut("/bin/sh",
+						$"-c 'plutil -extract CFBundleShortVersionString raw {Plist}'", out ExitCode);
+				if (ExitCode == 0)
+				{
+					return Output;
+				}
+				return null;
+			}
+
+			if (OperatingSystem.IsWindows())
+			{
+				// otherwise, get iTunes "Version"
+				string? DllPath =
+					Registry.GetValue(
+						"HKEY_LOCAL_MACHINE\\SOFTWARE\\Wow6432Node\\Apple Inc.\\Apple Mobile Device Support\\Shared",
+						"iTunesMobileDeviceDLL", null) as string;
+				if (String.IsNullOrEmpty(DllPath) || !File.Exists(DllPath))
+				{
+					DllPath = Registry.GetValue(
+						"HKEY_LOCAL_MACHINE\\SOFTWARE\\Wow6432Node\\Apple Inc.\\Apple Mobile Device Support\\Shared",
+						"MobileDeviceDLL", null) as string;
+					if (String.IsNullOrEmpty(DllPath) || !File.Exists(DllPath))
+					{
+						// iTunes >= 12.7 doesn't have a key specifying the 32-bit DLL but it does have a ASMapiInterfaceDLL key and MobileDevice.dll is in usually in the same directory
+						DllPath = Registry.GetValue(
+							"HKEY_LOCAL_MACHINE\\SOFTWARE\\Wow6432Node\\Apple Inc.\\Apple Mobile Device Support\\Shared",
+							"ASMapiInterfaceDLL", null) as string;
+						DllPath = String.IsNullOrEmpty(DllPath)
+							? null
+							: DllPath.Substring(0, DllPath.LastIndexOf('\\') + 1) + "MobileDevice.dll";
+
+						if (String.IsNullOrEmpty(DllPath) || !File.Exists(DllPath))
+						{
+							DllPath = FindWindowsStoreITunesDLL();
+						}
+					}
+				}
+
+				if (!String.IsNullOrEmpty(DllPath) && File.Exists(DllPath))
+				{
+					string? DllVersion = FileVersionInfo.GetVersionInfo(DllPath).FileVersion;
+					// Only return the DLL version as the SDK version if we can correctly parse it
+					if (UEBuildPlatformSDK.GetSDKForPlatform("Mac")!.TryConvertVersionToInt(DllVersion, out _, null))
+					{
+						return DllVersion;
+					}
+				}
+			}
+
+			return null;
+		}
+
+		[SupportedOSPlatform("windows")]
+		private static string? FindWindowsStoreITunesDLL()
+		{
+			string? InstallPath = null;
+
+			string PackagesKeyName = "Software\\Classes\\Local Settings\\Software\\Microsoft\\Windows\\CurrentVersion\\AppModel\\PackageRepository\\Packages";
+
+			RegistryKey? PackagesKey = Registry.LocalMachine.OpenSubKey(PackagesKeyName);
+			if (PackagesKey != null)
+			{
+				string[] PackageSubKeyNames = PackagesKey.GetSubKeyNames();
+
+				foreach (string PackageSubKeyName in PackageSubKeyNames)
+				{
+					if (PackageSubKeyName.Contains("AppleInc.iTunes") && (PackageSubKeyName.Contains("_x64") || PackageSubKeyName.Contains("_x86")))
+					{
+						string FullPackageSubKeyName = PackagesKeyName + "\\" + PackageSubKeyName;
+
+						RegistryKey? iTunesKey = Registry.LocalMachine.OpenSubKey(FullPackageSubKeyName);
+						if (iTunesKey != null)
+						{
+							object? Value = iTunesKey.GetValue("Path");
+							if (Value != null)
+							{
+								InstallPath = (string)Value + "\\AMDS32\\MobileDevice.dll";
+							}
+							break;
+						}
+					}
+				}
+			}
+
+			return InstallPath;
+		}
+
+		private static DirectoryReference? GetDeveloperDir()
+		{
+			int ExitCode;
+			// xcode-select -p gives the currently selected Xcode location (xcodebuild -version may fail if Xcode.app is broken)
+			// Example output: /Applications/Xcode.app/Contents/Developer
+			string Output = Utils.RunLocalProcessAndReturnStdOut("/bin/sh", "-c 'xcode-select -p'", out ExitCode);
+
+			if (ExitCode == 0)
+			{
+				return new DirectoryReference(Output);
+			}
+
+			return null;
+		}
+
+		private static string? GetVersionFromSDKDir(string AppleOSName)
+		{
+			string? PlatformSDKVersion = null;
+			try
+			{
+				// loop over the subdirs and parse out the version
+				int MaxSDKVersionMajor = 0;
+				int MaxSDKVersionMinor = 0;
+				string? MaxSDKVersionString = null;
+				foreach (DirectoryReference SubDir in DirectoryReference.EnumerateDirectories(GetPlatformSDKDirectory(AppleOSName).ParentDirectory!))
+				{
+					string SubDirName = Path.GetFileNameWithoutExtension(SubDir.GetDirectoryName());
+					if (SubDirName.StartsWith(AppleOSName))
+					{
+						// get the SDK version from the directory name
+						string SDKString = SubDirName.Replace(AppleOSName, "");
+						int Major = 0;
+						int Minor = 0;
+
+						// parse it into whole and fractional parts (since 10.10 > 10.9 in versions, but not in math)
+						try
+						{
+							string[] Tokens = SDKString.Split(".".ToCharArray());
+							if (Tokens.Length == 2)
+							{
+								Major = Int32.Parse(Tokens[0]);
+								Minor = Int32.Parse(Tokens[1]);
+							}
+						}
+						catch (Exception)
+						{
+							// weirdly formatted SDKs
+							continue;
+						}
+
+						// update largest SDK version number
+						if (Major > MaxSDKVersionMajor || (Major == MaxSDKVersionMajor && Minor > MaxSDKVersionMinor))
+						{
+							MaxSDKVersionString = SDKString;
+							MaxSDKVersionMajor = Major;
+							MaxSDKVersionMinor = Minor;
+						}
+					}
+				}
+
+				// use the largest version
+				if (MaxSDKVersionString != null)
+				{
+					PlatformSDKVersion = MaxSDKVersionString;
+				}
+			}
+			catch (Exception)
+			{
+				return null;
+			}
+
+			return PlatformSDKVersion;
 		}
 
 		protected override SDKStatus HasRequiredManualSDKInternal()
diff --git a/Engine/Source/Programs/UnrealBuildTool/Platform/Mac/MacToolChain.cs b/Engine/Source/Programs/UnrealBuildTool/Platform/Mac/MacToolChain.cs
--- a/Engine/Source/Programs/UnrealBuildTool/Platform/Mac/MacToolChain.cs
+++ b/Engine/Source/Programs/UnrealBuildTool/Platform/Mac/MacToolChain.cs
@@ -13,10 +13,12 @@ namespace UnrealBuildTool
 {
 	class MacToolChainSettings : AppleToolChainSettings
 	{
+		public override string GetTargetVersionForTargetType(TargetType Type) => ((ApplePlatformSDK)UEBuildPlatformSDK.GetSDKForPlatform("Mac")!).GetBuildTargetVersion(Type);
+
 		/// <summary>
 		/// Which version of the Mac OS X to allow at run time
 		/// </summary>
-		public string MinMacBuildVersion(TargetType TargetType)
+		public static string MinMacBuildVersion(TargetType TargetType)
 		{
 			return ((ApplePlatformSDK)UEBuildPlatformSDK.GetSDKForPlatform("Mac")!).GetBuildTargetVersion(TargetType);
 		}
@@ -24,7 +26,7 @@ namespace UnrealBuildTool
 		/// <summary>
 		/// Minimum version of Mac OS X to actually run on, running on earlier versions will display the system minimum version error dialog and exit.
 		/// </summary>
-		public string MinMacDeploymentVersion(TargetType TargetType)
+		public static string MinMacDeploymentVersion(TargetType TargetType)
 		{
 			return ((ApplePlatformSDK)UEBuildPlatformSDK.GetSDKForPlatform("Mac")!).GetDeploymentTargetVersion(TargetType);
 		}
@@ -35,14 +37,10 @@ namespace UnrealBuildTool
 		/// <param name="bVerbose">Whether to output verbose logging</param>
 		/// <param name="Logger">Logger for output</param>
 		public MacToolChainSettings(bool bVerbose, ILogger Logger)
-			: base("MacOSX", null, "macos", bVerbose, Logger)
+			// passing 0.0 for target version because we have per-target versions when making the Tuple
+			: base("MacOSX", null, "macos", "0.0", bVerbose, Logger)
 		{
 		}
-
-		public DirectoryReference GetSDKPath()
-		{
-			return GetSDKPath(UnrealArch.Host.Value);
-		}
 	}
 
 	/// <summary>
@@ -51,12 +49,10 @@ namespace UnrealBuildTool
 	class MacToolChain : AppleToolChain
 	{
 		public MacToolChain(ReadOnlyTargetRules? Target, ClangToolChainOptions InOptions, ILogger InLogger)
-			: base(Target, () => new MacToolChainSettings(false, InLogger), InOptions, InLogger)
+			: base(Target, () => new MacToolChainSettings(true, InLogger), InOptions, InLogger)
 		{
 		}
 
-		public static MacToolChainSettings Settings => new MacToolChainSettings(false, Log.Logger);
-
 		/// <summary>
 		/// Which compiler\linker frontend to use
 		/// </summary>
@@ -71,7 +67,7 @@ namespace UnrealBuildTool
 		{
 			FileReference CompilerPath = FileReference.Combine(Settings.ToolchainDir, MacCompiler);
 			FileReference ArchiverPath = FileReference.Combine(Settings.ToolchainDir, MacArchiver);
-			return new AppleToolChainInfo(UnrealTargetPlatform.Mac, MacToolChainSettings.XcodeDeveloperDir, CompilerPath, ArchiverPath, Logger);
+			return new AppleToolChainInfo(UnrealTargetPlatform.Mac, ApplePlatformSDK.DeveloperDir, CompilerPath, ArchiverPath, Logger);
 		}
 
 		public static DirectoryReference FindProductDirectory(FileReference? ProjectFile, DirectoryReference BinaryDir, string? NameIfProgram)
@@ -203,8 +199,7 @@ namespace UnrealBuildTool
 
 			// Pass through architecture and OS info
 			Arguments.Add("" + FormatArchitectureArg(CompileEnvironment.Architectures));
-			Arguments.Add($"-isysroot \"{Settings.GetSDKPath()}\"");
-			Arguments.Add("-mmacosx-version-min=" + (CompileEnvironment.bEnableOSX109Support ? "10.9" : Settings.MinMacBuildVersion(Target!.Type)));
+			Arguments.Add($"-isysroot \"{Settings.GetSDKPath(CompileEnvironment.Architecture)}\"");
 
 			List<string> FrameworksSearchPaths = new List<string>();
 			foreach (UEBuildFramework Framework in CompileEnvironment.AdditionalFrameworks)
@@ -235,8 +230,8 @@ namespace UnrealBuildTool
 			base.GetLinkArguments_Global(LinkEnvironment, Arguments);
 			// Pass through architecture and OS info		
 			Arguments.Add(FormatArchitectureArg(LinkEnvironment.Architectures));
-			Arguments.Add(String.Format("-isysroot \"{0}\"", ToolChainSettings.Value.GetSDKPath(LinkEnvironment.Architecture)));
-			Arguments.Add("-mmacosx-version-min=" + Settings.MinMacBuildVersion(Target!.Type));
+			Arguments.Add(String.Format("-isysroot \"{0}\"", Settings.GetSDKPath(LinkEnvironment.Architecture)));
+			Arguments.Add($"-target {Settings.GetTargetTuple(LinkEnvironment.Architecture, Target!.Platform, Target!.Type)}");
 			Arguments.Add("-dead_strip");
 
 			// Temporary workaround for linker warning with Xcode 14:
@@ -629,7 +624,7 @@ namespace UnrealBuildTool
 				// Fix contents of Info.plist
 				AppendMacLine(FinalizeAppBundleScript, "/usr/bin/sed -i \"\" -e \"s/\\${0}/{1}/g\" \"{2}\"", "{EXECUTABLE_NAME}", ExeName, TempInfoPlist);
 				AppendMacLine(FinalizeAppBundleScript, "/usr/bin/sed -i \"\" -e \"s/\\${0}/{1}/g\" \"{2}\"", "{APP_NAME}", bBuildingEditor ? ("com.epicgames." + GameName) : (BundleIdentifier.Replace("[PROJECT_NAME]", GameName).Replace("_", "")), TempInfoPlist);
-				AppendMacLine(FinalizeAppBundleScript, "/usr/bin/sed -i \"\" -e \"s/\\${0}/{1}/g\" \"{2}\"", "{MACOSX_DEPLOYMENT_TARGET}", Settings.MinMacDeploymentVersion(Target!.Type), TempInfoPlist);
+				AppendMacLine(FinalizeAppBundleScript, "/usr/bin/sed -i \"\" -e \"s/\\${0}/{1}/g\" \"{2}\"", "{MACOSX_DEPLOYMENT_TARGET}", MacToolChainSettings.MinMacDeploymentVersion(Target!.Type), TempInfoPlist);
 				AppendMacLine(FinalizeAppBundleScript, "/usr/bin/sed -i \"\" -e \"s/\\${0}/{1}/g\" \"{2}\"", "{ICON_NAME}", GameName, TempInfoPlist);
 				AppendMacLine(FinalizeAppBundleScript, "/usr/bin/sed -i \"\" -e \"s/\\${0}/{1}/g\" \"{2}\"", "{BUNDLE_VERSION}", BundleVersion, TempInfoPlist);
 
diff --git a/Engine/Source/Programs/UnrealBuildTool/Platform/Mac/UEBuildMac.cs b/Engine/Source/Programs/UnrealBuildTool/Platform/Mac/UEBuildMac.cs
--- a/Engine/Source/Programs/UnrealBuildTool/Platform/Mac/UEBuildMac.cs
+++ b/Engine/Source/Programs/UnrealBuildTool/Platform/Mac/UEBuildMac.cs
@@ -280,7 +280,7 @@ namespace UnrealBuildTool
 		{
 			base.GetExternalBuildMetadata(ProjectFile, Metadata);
 
-			Metadata.AppendLine("xcode-select: {0}", AppleToolChainSettings.XcodeDeveloperDir);
+			Metadata.AppendLine("xcode-select: {0}", ApplePlatformSDK.DeveloperDir);
 		}
 	}
 
@@ -339,8 +339,7 @@ namespace UnrealBuildTool
 				Target.DisablePlugins.AddRange(Array.Empty<string>());
 			}
 
-			// Needs OS X 10.11 for Metal. The remote toolchain has not been initialized yet, so just assume it's a recent SDK.
-			if ((BuildHostPlatform.Current.Platform != UnrealTargetPlatform.Mac || MacToolChain.Settings.SDKVersionFloat >= 10.11f) && Target.bCompileAgainstEngine)
+			if (Target.bCompileAgainstEngine)
 			{
 				Target.GlobalDefinitions.Add("HAS_METAL=1");
 				Target.ExtraModuleNames.Add("MetalRHI");
diff --git a/Engine/Source/Programs/UnrealBuildTool/Platform/Mac/UEDeployMac.cs b/Engine/Source/Programs/UnrealBuildTool/Platform/Mac/UEDeployMac.cs
--- a/Engine/Source/Programs/UnrealBuildTool/Platform/Mac/UEDeployMac.cs
+++ b/Engine/Source/Programs/UnrealBuildTool/Platform/Mac/UEDeployMac.cs
@@ -57,7 +57,7 @@ namespace UnrealBuildTool
 			PListData = PListData.Replace("${EXECUTABLE_NAME}", ExeName).
 				Replace("${APP_NAME}", FinalBundleID).
 				Replace("${ICON_NAME}", GameName).
-				Replace("${MACOSX_DEPLOYMENT_TARGET}", MacToolChain.Settings.MinMacDeploymentVersion(TargetType)).
+				Replace("${MACOSX_DEPLOYMENT_TARGET}", MacToolChainSettings.MinMacDeploymentVersion(TargetType)).
 				Replace("${BUNDLE_VERSION}", BundleVersion);
 
 			if (!Directory.Exists(IntermediateDirectory))
diff --git a/Engine/Source/Programs/UnrealBuildTool/Platform/TVOS/TVOSToolChain.cs b/Engine/Source/Programs/UnrealBuildTool/Platform/TVOS/TVOSToolChain.cs
--- a/Engine/Source/Programs/UnrealBuildTool/Platform/TVOS/TVOSToolChain.cs
+++ b/Engine/Source/Programs/UnrealBuildTool/Platform/TVOS/TVOSToolChain.cs
@@ -6,7 +6,8 @@ namespace UnrealBuildTool
 {
 	class TVOSToolChainSettings : IOSToolChainSettings
 	{
-		public TVOSToolChainSettings(ILogger Logger) : base("AppleTVOS", "AppleTVSimulator", "tvos", Logger)
+		public TVOSToolChainSettings(IOSProjectSettings ProjectSettings, ILogger Logger) 
+			: base("AppleTVOS", "AppleTVSimulator", "tvos", ProjectSettings, Logger)
 		{
 		}
 	}
@@ -14,13 +15,8 @@ namespace UnrealBuildTool
 	class TVOSToolChain : IOSToolChain
 	{
 		public TVOSToolChain(ReadOnlyTargetRules InTarget, TVOSProjectSettings InProjectSettings, ILogger InLogger)
-			: base(InTarget, InProjectSettings, () => new TVOSToolChainSettings(InLogger), ClangToolChainOptions.None, InLogger)
+			: base(InTarget, InProjectSettings, () => new TVOSToolChainSettings(InProjectSettings, InLogger), ClangToolChainOptions.None, InLogger)
 		{
 		}
-
-		public override string GetXcodeMinVersionParam(UnrealArch Architecture)
-		{
-			return "tvos-version-min";
-		}
 	}
 }
diff --git a/Engine/Source/Programs/UnrealBuildTool/ProjectFiles/VisualStudioCode/VSCodeProjectFileGenerator.cs b/Engine/Source/Programs/UnrealBuildTool/ProjectFiles/VisualStudioCode/VSCodeProjectFileGenerator.cs
--- a/Engine/Source/Programs/UnrealBuildTool/ProjectFiles/VisualStudioCode/VSCodeProjectFileGenerator.cs
+++ b/Engine/Source/Programs/UnrealBuildTool/ProjectFiles/VisualStudioCode/VSCodeProjectFileGenerator.cs
@@ -353,7 +353,7 @@ namespace UnrealBuildTool
 			{
 				MacToolChainSettings Settings = new MacToolChainSettings(false, Logger);
 				CompilerPath = FileReference.Combine(Settings.ToolchainDir, "clang++");
-				SysRootPath = Settings.GetSDKPath();
+				SysRootPath = Settings.GetSDKPath(UnrealArch.Host.Value);
 			}
 			else
 			{
@@ -725,14 +725,10 @@ namespace UnrealBuildTool
 					throw new BuildException($"Unsupported C++ standard type set: {CppStandard}");
 			}
 
-			if (HostPlatform == UnrealTargetPlatform.Win64)
-			{
-				OutFile.AddField("intelliSenseMode", "msvc-x64");
-			}
-			else
-			{
-				OutFile.AddField("intelliSenseMode", "clang-x64");
-			}
+			string CompilerPrefix = HostPlatform == UnrealTargetPlatform.Win64 ? "msvc" : "clang";
+			string ArchSuffix = UnrealArch.Host.Value.bIsX64 ? "x64" : "arm64";
+			// make msvc-x64, clang-arm64, etc
+			OutFile.AddField("intelliSenseMode", $"{CompilerPrefix}-{ArchSuffix}");
 
 			if (bUseVSCodeExtension)
 			{
diff --git a/Engine/Source/Programs/UnrealBuildTool/ProjectFiles/Xcode/XcodeProject.cs b/Engine/Source/Programs/UnrealBuildTool/ProjectFiles/Xcode/XcodeProject.cs
--- a/Engine/Source/Programs/UnrealBuildTool/ProjectFiles/Xcode/XcodeProject.cs
+++ b/Engine/Source/Programs/UnrealBuildTool/ProjectFiles/Xcode/XcodeProject.cs
@@ -1856,7 +1856,7 @@ namespace UnrealBuildTool.XcodeProjectXcconfig
 				SDKRoot = "macosx";
 				SupportedPlatforms = "macosx";
 				DeploymentTargetKey = "MACOSX_DEPLOYMENT_TARGET";
-				DeploymentTarget = MacToolChain.Settings.MinMacDeploymentVersion(UnrealData.TargetRules.Type);
+				DeploymentTarget = MacToolChainSettings.MinMacDeploymentVersion(UnrealData.TargetRules.Type);
 				BundleIdentifier = bIsEditor ? "com.epicgames.UnrealEditor" : UnrealData.BundleIdentifier;
 
 				// @todo: get a version for  games, like IOS has
@@ -2403,7 +2403,7 @@ namespace UnrealBuildTool.XcodeProjectXcconfig
 				}
 
 				UnrealData.InitializeUProjectFileLocation(this);
-				if (!AppleExports.UseModernXcode(ProjectFilePath))
+				if (!AppleExports.UseModernXcode(UnrealData.UProjectFileLocation))
 				{
 					LegacyProjectFile = new XcodeProjectLegacy.XcodeProjectFile(ProjectFilePath, BaseDir, UnrealData.bForDistribution, UnrealData.BundleIdentifier, UnrealData.AppName, UnrealData.bMakeProjectPerTarget);
 					LegacyProjectFile.ProjectTargets.AddRange(ProjectTargets);
@@ -2761,7 +2761,7 @@ namespace UnrealBuildTool.XcodeProjectXcconfig
 
 					// the Build target used some ini settings to compile, and Run target must match, so we override a few settings, at
 					// whatever level they were already specified at (Projet and/or Target)
-					XcodeUtils.PlistSetUpdate($":objects:{ConfigGuid}:buildSettings:MACOSX_DEPLOYMENT_TARGET", MacToolChain.Settings.MinMacDeploymentVersion(UnrealData.TargetRules.Type));
+					XcodeUtils.PlistSetUpdate($":objects:{ConfigGuid}:buildSettings:MACOSX_DEPLOYMENT_TARGET", MacToolChainSettings.MinMacDeploymentVersion(UnrealData.TargetRules.Type));
 					if (UnrealData.IOSProjectSettings != null)
 					{
 						XcodeUtils.PlistSetUpdate($":objects:{ConfigGuid}:buildSettings:IPHONEOS_DEPLOYMENT_TARGET", UnrealData.IOSProjectSettings.RuntimeVersion);
diff --git a/Engine/Source/Programs/UnrealBuildTool/ProjectFiles/Xcode/XcodeProjectLegacy.cs b/Engine/Source/Programs/UnrealBuildTool/ProjectFiles/Xcode/XcodeProjectLegacy.cs
--- a/Engine/Source/Programs/UnrealBuildTool/ProjectFiles/Xcode/XcodeProjectLegacy.cs
+++ b/Engine/Source/Programs/UnrealBuildTool/ProjectFiles/Xcode/XcodeProjectLegacy.cs
@@ -1266,7 +1266,7 @@ namespace UnrealBuildTool.XcodeProjectLegacy
 			// #jira UE-143619: Pre Monterey macOS requires this option for a packaged app to run on iOS15 due to new code signature format. Could be removed once Monterey is minimum OS.
 			Content.Append("\t\t\t\tOTHER_CODE_SIGN_FLAGS = \"--generate-entitlement-der --deep\";" + ProjectFileGenerator.NewLine);
 
-			Content.Append("\t\t\t\tMACOSX_DEPLOYMENT_TARGET = " + MacToolChain.Settings.MinMacDeploymentVersion(Config.ProjectTarget!.TargetRules!.Type) + ";" + ProjectFileGenerator.NewLine);
+			Content.Append("\t\t\t\tMACOSX_DEPLOYMENT_TARGET = " + MacToolChainSettings.MinMacDeploymentVersion(Config.ProjectTarget!.TargetRules!.Type) + ";" + ProjectFileGenerator.NewLine);
 			Content.Append("\t\t\t\tINFOPLIST_OUTPUT_FORMAT = xml;" + ProjectFileGenerator.NewLine);
 			Content.Append("\t\t\t\tCOMBINE_HIDPI_IMAGES = YES;" + ProjectFileGenerator.NewLine);
 
diff --git a/Engine/Source/Programs/UnrealBuildTool/ToolChain/AppleToolChain.cs b/Engine/Source/Programs/UnrealBuildTool/ToolChain/AppleToolChain.cs
--- a/Engine/Source/Programs/UnrealBuildTool/ToolChain/AppleToolChain.cs
+++ b/Engine/Source/Programs/UnrealBuildTool/ToolChain/AppleToolChain.cs
@@ -21,17 +21,19 @@ namespace UnrealBuildTool
 	public abstract class AppleToolChainSettings
 	{
 		/// <summary>
-		/// Which developer directory to root from? If this is "xcode-select", UBT will query for the currently selected Xcode
+		/// Cached copy of ApplePlatformSDK.GetToolchainDirectory, for existing code to work
 		/// </summary>
-		public static DirectoryReference XcodeDeveloperDir => _XcodeDeveloperDir!;
-
-		// This is not expected to be used before we've made any Settings options, so XcodeDeveloperDir will assert if it's used too early
-		private static DirectoryReference? _XcodeDeveloperDir = null;
+		public DirectoryReference ToolchainDir => ApplePlatformSDK.GetToolchainDirectory();
 
 		/// <summary>
-		/// Directory for the developer binaries
+		/// Name for Xcode plaform directory under Toolchains
+		/// </summary>
+		public string PlatformDirName;
+		
+		/// <summary>
+		/// Name for Xcode simulator platform directory under Toolchains
 		/// </summary>
-		public DirectoryReference ToolchainDir;
+		public string SimulatorPlatformDirName = "";
 
 		/// <summary>
 		/// A portion of the target "tuple"
@@ -39,14 +41,36 @@ namespace UnrealBuildTool
 		private string TargetOSName;
 
 		/// <summary>
-		/// The version of the iOS SDK to target at build time.
+		/// The name apple uses for SDK directories
+		/// </summary>
+		private string TargetSDKName;
+
+		/// <summary>
+		/// The version of the SDK being used to build with
 		/// </summary>
 		public string SDKVersion;
 
 		/// <summary>
-		/// The version in a floating point value for easy comparison
+		/// The version of the iOS SDK to target at build time.
+		/// </summary>
+		public string MinTargetVersion;
+
+		/// <summary>
+		/// The build version in a floating point value for easy comparison
+		/// </summary>
+		public readonly float SDKVersionFloat;
+
+		/// <summary>
+		/// The target version in a floating point value for easy comparison
 		/// </summary>
-		public readonly float SDKVersionFloat = 0.0f;
+		public readonly float MinTargetVersionFloat;
+
+		/// <summary>
+		/// For platforms that need a different min version for editor vs game, this can be overridden
+		/// </summary>
+		/// <param name="Type"></param>
+		/// <returns></returns>
+		public virtual string GetTargetVersionForTargetType(TargetType Type) => MinTargetVersion;
 
 		/// <summary>
 		/// Cache SDK dir for device and simulator (for non-Mac)
@@ -60,33 +84,34 @@ namespace UnrealBuildTool
 		/// <param name="OSPrefix">SDK name (like "MacOSX")</param>
 		/// <param name="SimulatorOSPrefix">Sinulator SDK name (like "iPhoneOSSimulator")</param>
 		/// <param name="TargetOSName">Platform name used in the -target parameter</param>
+		/// <param name="TargetOSVersion">min OS version this project wants to target (not what it's being built with)</param>
 		/// <param name="bVerbose"></param>
 		/// <param name="Logger"></param>
-		protected AppleToolChainSettings(string OSPrefix, string? SimulatorOSPrefix, string TargetOSName, bool bVerbose, ILogger Logger)
+		protected AppleToolChainSettings(string OSPrefix, string? SimulatorOSPrefix, string TargetOSName, string TargetOSVersion, bool bVerbose, ILogger Logger)
 		{
 			this.TargetOSName = TargetOSName;
-
-			if (_XcodeDeveloperDir == null)
-			{
-				SelectXcode(bVerbose, Logger);
-			}
+			this.TargetSDKName = OSPrefix;
+			PlatformDirName = OSPrefix.ToLower();
 
 			// set up directories
-			ToolchainDir = DirectoryReference.Combine(XcodeDeveloperDir, "Toolchains/XcodeDefault.xctoolchain/usr/bin");
-			SDKDir = DirectoryReference.Combine(XcodeDeveloperDir, $"Platforms/{OSPrefix}.platform/Developer/SDKs/{OSPrefix}.sdk");
+			SDKDir = ApplePlatformSDK.GetPlatformSDKDirectory(OSPrefix);
 			if (SimulatorOSPrefix != null)
 			{
-				SimulatorSDKDir = DirectoryReference.Combine(XcodeDeveloperDir, $"Platforms/{SimulatorOSPrefix}.platform/Developer/SDKs/{SimulatorOSPrefix}.sdk");
+				SimulatorSDKDir = ApplePlatformSDK.GetPlatformSDKDirectory(SimulatorOSPrefix);
+				SimulatorPlatformDirName = SimulatorOSPrefix.ToLower();
 			}
 
-			// use one up fro the SDK dir to look in for the version'd name directories next to the SDK dir
-			SDKVersion = SelectSDK(SDKDir.ParentDirectory!, OSPrefix, bVerbose, Logger);
+			MinTargetVersion = TargetOSVersion;
+			MinTargetVersionFloat = Single.Parse(MinTargetVersion, System.Globalization.CultureInfo.InvariantCulture);
 
-			// convert to float for easy comparison
-			SDKVersionFloat = Single.Parse(SDKVersion, System.Globalization.CultureInfo.InvariantCulture);
+			// cache some info for this OS
+			SDKVersion = ApplePlatformSDK.GetPlatformSDKVersion(OSPrefix) ?? "";
+			SDKVersionFloat = ApplePlatformSDK.GetPlatformSDKVersionFloat(OSPrefix);
 
 			// Xcode sdk files are not allowed to be transferred over the network, ensure UBA will not do so
-			EpicGames.UBA.Utils.RegisterDisallowedPaths(XcodeDeveloperDir.FullName);
+			EpicGames.UBA.Utils.RegisterDisallowedPaths(ApplePlatformSDK.DeveloperDir.FullName);
+
+			TestXcode(bVerbose, Logger);
 		}
 
 		/// <summary>
@@ -108,13 +133,16 @@ namespace UnrealBuildTool
 		/// Gets the string used by xcode to taget a platform and version (will return something like "arm64-apple-ios17.0-simulator"
 		/// </summary>
 		/// <param name="Architecture"></param>
+		/// <param name="Platform"></param>
+		/// <param name="TargetType"></param>
 		/// <returns></returns>
-		public string GetTargetTuple(UnrealArch Architecture)
+		public virtual string GetTargetTuple(UnrealArch Architecture, UnrealTargetPlatform Platform, TargetType TargetType)
 		{
 			string Prefix = Architecture.AppleName;
 			string Suffix = (Architecture == UnrealArch.IOSSimulator || Architecture == UnrealArch.TVOSSimulator) ? "-simulator" : "";
+			string TargetVersion = GetTargetVersionForTargetType(TargetType);
 
-			return $"{Prefix}-apple-{TargetOSName}{SDKVersion}{Suffix}";
+			return $"{Prefix}-apple-{TargetOSName}{TargetVersion}{Suffix}";
 		}
 
 		/// <summary>
@@ -123,19 +151,9 @@ namespace UnrealBuildTool
 		/// <param name="bVerbose"></param>
 		/// <param name="Logger"></param>
 		/// <exception cref="BuildException"></exception>
-		public static void SelectXcode(bool bVerbose, ILogger Logger)
+		private static void TestXcode(bool bVerbose, ILogger Logger)
 		{
-			// on the Mac, run xcode-select directly.
-			int ReturnCode;
-			string XcodeSelectResult = Utils.RunLocalProcessAndReturnStdOut("xcode-select", "--print-path", null, out ReturnCode);
-			if (ReturnCode != 0)
-			{
-				string? MinVersion = UEBuildPlatform.GetSDK(UnrealTargetPlatform.Mac)!.GetSDKInfo("Sdk")!.Min;
-				throw new BuildException($"We were unable to find your build tools (via 'xcode-select --print-path'). Please install Xcode, version {MinVersion} or later");
-			}
-
-			_XcodeDeveloperDir = new DirectoryReference(XcodeSelectResult);
-
+			DirectoryReference XcodeDeveloperDir = ApplePlatformSDK.DeveloperDir;
 			// make sure we get a full path
 			if (DirectoryReference.Exists(XcodeDeveloperDir) == false)
 			{
@@ -158,7 +176,7 @@ namespace UnrealBuildTool
 			// Installed engine requires Xcode 13
 			if (Unreal.IsEngineInstalled())
 			{
-				string? InstalledSdkVersion = UnrealBuildBase.ApplePlatformSDK.InstalledSDKVersion;
+				string? InstalledSdkVersion = ApplePlatformSDK.InstalledXcodeVersion.Value;
 				if (String.IsNullOrEmpty(InstalledSdkVersion))
 				{
 					throw new BuildException("Unable to get xcode version");
@@ -168,73 +186,7 @@ namespace UnrealBuildTool
 					throw new BuildException("Building for macOS, iOS and tvOS requires Xcode 13.4.1 or newer, Xcode " + InstalledSdkVersion + " detected");
 				}
 			}
-		}
-
-		private static string SelectSDK(DirectoryReference BaseSDKDir, string OSPrefix, bool bVerbose, ILogger Logger)
-		{
-			string PlatformSDKVersion = "";
-			try
-			{
-				// loop over the subdirs and parse out the version
-				int MaxSDKVersionMajor = 0;
-				int MaxSDKVersionMinor = 0;
-				string? MaxSDKVersionString = null;
-				foreach (DirectoryReference SubDir in DirectoryReference.EnumerateDirectories(BaseSDKDir))
-				{
-					string SubDirName = Path.GetFileNameWithoutExtension(SubDir.GetDirectoryName());
-					if (SubDirName.StartsWith(OSPrefix))
-					{
-						// get the SDK version from the directory name
-						string SDKString = SubDirName.Replace(OSPrefix, "");
-						int Major = 0;
-						int Minor = 0;
-
-						// parse it into whole and fractional parts (since 10.10 > 10.9 in versions, but not in math)
-						try
-						{
-							string[] Tokens = SDKString.Split(".".ToCharArray());
-							if (Tokens.Length == 2)
-							{
-								Major = Int32.Parse(Tokens[0]);
-								Minor = Int32.Parse(Tokens[1]);
-							}
-						}
-						catch (Exception)
-						{
-							// weirdly formatted SDKs
-							continue;
-						}
-
-						// update largest SDK version number
-						if (Major > MaxSDKVersionMajor || (Major == MaxSDKVersionMajor && Minor > MaxSDKVersionMinor))
-						{
-							MaxSDKVersionString = SDKString;
-							MaxSDKVersionMajor = Major;
-							MaxSDKVersionMinor = Minor;
-						}
-					}
-				}
-
-				// use the largest version
-				if (MaxSDKVersionString != null)
-				{
-					PlatformSDKVersion = MaxSDKVersionString;
-				}
-			}
-			catch (Exception Ex)
-			{
-				// on any exception, just use the backup version
-				Logger.LogInformation("Triggered an exception while looking for SDK directory in Xcode.app");
-				Logger.LogInformation("{Ex}", Ex.ToString());
-			}
-
-			if (bVerbose && !ProjectFileGenerator.bGenerateProjectFiles)
-			{
-				Logger.LogInformation("Compiling with {Os} SDK {Sdk}", OSPrefix, PlatformSDKVersion);
-			}
-
-			return PlatformSDKVersion;
-		}
+		}		
 	}
 
 	abstract class AppleToolChain : ClangToolChain
@@ -280,7 +232,15 @@ namespace UnrealBuildTool
 			public VersionNumber AppleClangVersion => VersionNumber.Parse(base.QueryClangVersion().ToString());
 		}
 
-		public Lazy<AppleToolChainSettings> ToolChainSettings;
+		private Lazy<AppleToolChainSettings> ToolChainSettings;
+		
+		protected AppleToolChainSettings Settings => ToolChainSettings.Value;
+
+
+		/// <summary>
+		/// The SDK version being used to compile with the toolchain
+		/// </summary>
+		public float SDKVersionFloat => Settings.SDKVersionFloat;
 
 		public readonly ReadOnlyTargetRules? Target;
 
@@ -310,6 +270,11 @@ namespace UnrealBuildTool
 			return "-arch " + String.Join('+', InArchitectures.Architectures.Select(x => x.AppleName));
 		}
 
+		public override string GetSDKVersion()
+		{
+			return Settings.SDKVersion;
+		}
+
 		protected DirectoryReference GetMacDevSrcRoot()
 		{
 			return Unreal.EngineSourceDirectory;
@@ -322,6 +287,9 @@ namespace UnrealBuildTool
 				return new CPPOutput();
 			}
 
+			// log some iseful info
+			Log.TraceInformationOnce($"Compiling with {CompileEnvironment.Platform} SDK {Settings.SDKVersion}, minimum Target OS version {Settings.GetTargetVersionForTargetType(Target!.Type)}");
+
 			List<string> GlobalArguments = new();
 
 			GetCompileArguments_Global(CompileEnvironment, GlobalArguments);
@@ -606,6 +574,8 @@ namespace UnrealBuildTool
 
 			Arguments.Add("-fmessage-length=0");
 			Arguments.Add("-fpascal-strings");
+
+			Arguments.Add($"-target {ToolChainSettings.Value.GetTargetTuple(CompileEnvironment.Architecture, Target!.Platform, Target!.Type)}");
 		}
 
 		public override CppCompileEnvironment CreateSharedResponseFile(CppCompileEnvironment CompileEnvironment, FileReference OutResponseFile, IActionGraphBuilder Graph)
@@ -662,8 +632,8 @@ namespace UnrealBuildTool
 			Arguments.Add($"-o \"{OutputFile}\"");
 
 			// platform settings
-			Arguments.Add($"-target {ToolChainSettings.Value.GetTargetTuple(CompileEnvironment.Architecture)}");
-			Arguments.Add($"-sdk {ToolChainSettings.Value.GetSDKPath(CompileEnvironment.Architecture)}");
+			Arguments.Add($"-target {Settings.GetTargetTuple(CompileEnvironment.Architecture, Target!.Platform, Target!.Type)}");
+			Arguments.Add($"-sdk {Settings.GetSDKPath(CompileEnvironment.Architecture)}");
 
 			// misc settings copied from Xcode
 			Arguments.Add("-Xllvm");
@@ -721,7 +691,7 @@ namespace UnrealBuildTool
 				Arguments.Add($"-import-objc-header \"{Unreal.EngineDirectory}/Source/Runtime/Launch/Private/IOS/UECppToSwift.h\"");
 
 				// platform settings
-				Arguments.Add($"-target {ToolChainSettings.Value.GetTargetTuple(CompileEnvironment.Architecture)}");
+				Arguments.Add($"-target {ToolChainSettings.Value.GetTargetTuple(CompileEnvironment.Architecture, Target!.Platform, Target!.Type)}");
 				Arguments.Add($"-sdk {ToolChainSettings.Value.GetSDKPath(CompileEnvironment.Architecture)}");
 
 				Arguments.Add($"-swift-version 5");
@@ -952,7 +922,7 @@ namespace UnrealBuildTool
 				List<string> Arguments = new()
 				{
 					"UBT_NO_POST_DEPLOY=true",
-					FileReference.Combine(AppleToolChainSettings.XcodeDeveloperDir, "usr/bin/xcodebuild").FullName,
+					FileReference.Combine(ApplePlatformSDK.DeveloperDir, "usr/bin/xcodebuild").FullName,
 					Action,
 					$"-workspace \"{XcodeProject.FullName}\"",
 					$"-scheme \"{SchemeName}\"",
