UnrealBuildTool: Move more warnings into CppCompileWarnings

#rnx

diff --git a/Engine/Source/Programs/UnrealBuildTool/Configuration/CppCompileWarnings.cs b/Engine/Source/Programs/UnrealBuildTool/Configuration/CppCompileWarnings.cs
--- a/Engine/Source/Programs/UnrealBuildTool/Configuration/CppCompileWarnings.cs
+++ b/Engine/Source/Programs/UnrealBuildTool/Configuration/CppCompileWarnings.cs
@@ -37,11 +37,10 @@ namespace UnrealBuildTool
 		[RequiresUniqueBuildEnvironment]
 		public WarningLevel ShadowVariableWarningLevel
 		{
-			get => ResolveWarning((BuildSettingsVersion x) => (x >= BuildSettingsVersion.V2) ? WarningLevel.Error : (ParentCppCompileWarnings != null ? ParentCppCompileWarnings.ShadowVariableWarningLevel : WarningLevel.Default), ParentCppCompileWarnings?.ShadowVariableWarningLevel, _shadowVariableWarningLevelPrivate);
-			set => _shadowVariableWarningLevelPrivate = value;
+			get => ResolveWarning((BuildSettingsVersion x) => (x >= BuildSettingsVersion.V2) ? WarningLevel.Error : (ParentCppCompileWarnings != null ? ParentCppCompileWarnings.ShadowVariableWarningLevel : WarningLevel.Default), ParentCppCompileWarnings?.ShadowVariableWarningLevel, _shadowVariableWarningLevel);
+			set => _shadowVariableWarningLevel = value;
 		}
-
-		private WarningLevel _shadowVariableWarningLevelPrivate;
+		private WarningLevel _shadowVariableWarningLevel;
 
 		/// <summary>
 		/// Indicates what warning/error level to treat unsafe type casts as on platforms that support it (e.g., double->float or int64->int32)
@@ -50,11 +49,10 @@ namespace UnrealBuildTool
 		[RequiresUniqueBuildEnvironment]
 		public WarningLevel UnsafeTypeCastWarningLevel
 		{
-			get => ResolveWarning(null, ParentCppCompileWarnings?.UnsafeTypeCastWarningLevel, _unsafeTypeCastWarningLevelPrivate);
-			set => _unsafeTypeCastWarningLevelPrivate = value;
+			get => ResolveWarning(null, ParentCppCompileWarnings?.UnsafeTypeCastWarningLevel, _unsafeTypeCastWarningLevel);
+			set => _unsafeTypeCastWarningLevel = value;
 		}
-
-		private WarningLevel _unsafeTypeCastWarningLevelPrivate;
+		private WarningLevel _unsafeTypeCastWarningLevel;
 
 		/// <summary>
 		/// Indicates what warning/error level to treat undefined identifiers in conditional expressions.
@@ -64,11 +62,10 @@ namespace UnrealBuildTool
 		[RequiresUniqueBuildEnvironment]
 		public WarningLevel UndefinedIdentifierWarningLevel
 		{
-			get => ResolveWarning(null, ParentCppCompileWarnings?.UndefinedIdentifierWarningLevel, _undefinedIdentifierWarningLevelPrivate);
-			set => _undefinedIdentifierWarningLevelPrivate = value;
+			get => ResolveWarning(null, ParentCppCompileWarnings?.UndefinedIdentifierWarningLevel, _undefinedIdentifierWarningLevel);
+			set => _undefinedIdentifierWarningLevel = value;
 		}
-
-		private WarningLevel _undefinedIdentifierWarningLevelPrivate;
+		private WarningLevel _undefinedIdentifierWarningLevel;
 
 		/// <summary>
 		/// Indicates what warning/error level to treat unhandled enumerators in switches on enumeration-typed values.
@@ -78,13 +75,68 @@ namespace UnrealBuildTool
 		[RequiresUniqueBuildEnvironment]
 		public WarningLevel SwitchUnhandledEnumeratorWarningLevel
 		{
-			get => ResolveWarning(null, ParentCppCompileWarnings?.SwitchUnhandledEnumeratorWarningLevel, _switchUnhandledEnumeratorWarningLevelPrivate);
-			set => _switchUnhandledEnumeratorWarningLevelPrivate = value;
+			get => ResolveWarning(null, ParentCppCompileWarnings?.SwitchUnhandledEnumeratorWarningLevel, _switchUnhandledEnumeratorWarningLevel);
+			set => _switchUnhandledEnumeratorWarningLevel = value;
+		}
+		private WarningLevel _switchUnhandledEnumeratorWarningLevel;
+
+		#endregion -- Compiler Warning Settings --
+
+		#region -- Performance Warning Settings --
+
+		/// <summary>
+		/// Indicates what warning/error level to treat potential PCH performance issues.
+		/// </summary>
+		[XmlConfigFile(Category = "BuildConfiguration")]
+		[CommandLine("-PCHPerformanceIssueWarningLevel=")]
+		public WarningLevel PCHPerformanceIssueWarningLevel
+		{
+			get => ResolveWarning(null, ParentCppCompileWarnings?.PCHPerformanceIssueWarningLevel, _pchPerformanceIssueWarningLevel);
+			set => _pchPerformanceIssueWarningLevel = value;
+		}
+		private WarningLevel _pchPerformanceIssueWarningLevel;
+
+		#endregion -- Performance Warning Settings --
+
+		#region -- Rules Validation Warning Settings --
+
+		/// <summary>
+		/// How to treat general module include path validation messages
+		/// </summary>
+		[XmlConfigFile(Category = "BuildConfiguration")]
+		[CommandLine("-ModuleIncludePathWarningLevel=")]
+		public WarningLevel ModuleIncludePathWarningLevel
+		{
+			get => ResolveWarning(null, ParentCppCompileWarnings?.ModuleIncludePathWarningLevel, _moduleIncludePathWarningLevel);
+			set => _moduleIncludePathWarningLevel = value;
 		}
+		private WarningLevel _moduleIncludePathWarningLevel;
 
-		private WarningLevel _switchUnhandledEnumeratorWarningLevelPrivate;
+		/// <summary>
+		/// How to treat private module include path validation messages, where a module is adding an include path that exposes private headers
+		/// </summary>
+		[XmlConfigFile(Category = "BuildConfiguration")]
+		[CommandLine("-ModuleIncludePrivateWarningLevel=")]
+		public WarningLevel ModuleIncludePrivateWarningLevel
+		{
+			get => ResolveWarning(null, ParentCppCompileWarnings?.ModuleIncludePrivateWarningLevel, _moduleIncludePrivateWarningLevel);
+			set => _moduleIncludePrivateWarningLevel = value;
+		}
+		private WarningLevel _moduleIncludePrivateWarningLevel;
+
+		/// <summary>
+		/// How to treat unnecessary module sub-directory include path validation messages
+		/// </summary>
+		[XmlConfigFile(Category = "BuildConfiguration")]
+		[CommandLine("-ModuleIncludeSubdirectoryWarningLevel=")]
+		public WarningLevel ModuleIncludeSubdirectoryWarningLevel
+		{
+			get => ResolveWarning(null, ParentCppCompileWarnings?.ModuleIncludeSubdirectoryWarningLevel, _moduleIncludeSubdirectoryWarningLevel);
+			set => _moduleIncludePrivateWarningLevel = value;
+		}
+		private WarningLevel _moduleIncludeSubdirectoryWarningLevel;
 
-		#endregion  -- Compiler Warning Settings --
+		#endregion -- Rules Validation Warning Settings --
 
 		#region -- Constructors --
 
@@ -220,10 +272,20 @@ namespace UnrealBuildTool
 		private void ApplyTargetDefaults()
 		{
 			_logger?.LogTrace("Applying default C++ Compiler warnings for the Target({TargetName}) context.", BuildContextProvider?._targetRulesPrivate?.Name);
-			_shadowVariableWarningLevelPrivate = WarningLevel.Warning;
-			_unsafeTypeCastWarningLevelPrivate = WarningLevel.Off;
-			_undefinedIdentifierWarningLevelPrivate = WarningLevel.Error;
-			_switchUnhandledEnumeratorWarningLevelPrivate = WarningLevel.Off;
+
+			// Compiler Warning Settings
+			_shadowVariableWarningLevel = WarningLevel.Warning;
+			_unsafeTypeCastWarningLevel = WarningLevel.Off;
+			_undefinedIdentifierWarningLevel = WarningLevel.Error;
+			_switchUnhandledEnumeratorWarningLevel = WarningLevel.Off;
+
+			// Performance Warning Settings
+			_pchPerformanceIssueWarningLevel = WarningLevel.Off;
+
+			// Rules Validation Warning Settings
+			_moduleIncludePathWarningLevel = WarningLevel.Off;
+			_moduleIncludePrivateWarningLevel = WarningLevel.Off;
+			_moduleIncludeSubdirectoryWarningLevel = WarningLevel.Off;
 		}
 
 		/// <summary>
@@ -233,10 +295,20 @@ namespace UnrealBuildTool
 		private void ApplyDefaults()
 		{
 			_logger?.LogTrace("Applying default C++ Compiler warnings for the default construction context.");
-			_shadowVariableWarningLevelPrivate = WarningLevel.Warning;
-			_unsafeTypeCastWarningLevelPrivate = WarningLevel.Off;
-			_undefinedIdentifierWarningLevelPrivate = WarningLevel.Off;
-			_switchUnhandledEnumeratorWarningLevelPrivate = WarningLevel.Off;
+
+			// Compiler Warning Settings
+			_shadowVariableWarningLevel = WarningLevel.Warning;
+			_unsafeTypeCastWarningLevel = WarningLevel.Off;
+			_undefinedIdentifierWarningLevel = WarningLevel.Off;
+			_switchUnhandledEnumeratorWarningLevel = WarningLevel.Off;
+
+			// Performance Warning Settings
+			_pchPerformanceIssueWarningLevel = WarningLevel.Off;
+
+			// Rules Validation Warning Settings
+			_moduleIncludePathWarningLevel = WarningLevel.Off;
+			_moduleIncludePrivateWarningLevel = WarningLevel.Off;
+			_moduleIncludeSubdirectoryWarningLevel = WarningLevel.Off;
 		}
 
 		#endregion -- Private API --
diff --git a/Engine/Source/Programs/UnrealBuildTool/Configuration/ModuleRules.cs b/Engine/Source/Programs/UnrealBuildTool/Configuration/ModuleRules.cs
--- a/Engine/Source/Programs/UnrealBuildTool/Configuration/ModuleRules.cs
+++ b/Engine/Source/Programs/UnrealBuildTool/Configuration/ModuleRules.cs
@@ -985,37 +985,28 @@ namespace UnrealBuildTool
 		/// </summary>
 		public WarningLevel ModuleIncludePathWarningLevel
 		{
-			get => (ModuleIncludePathWarningLevelPrivate == WarningLevel.Default) ? Target.ModuleIncludePathWarningLevel : ModuleIncludePathWarningLevelPrivate;
-			set => ModuleIncludePathWarningLevelPrivate = value;
+			get => CppCompileWarningSettings.ModuleIncludePathWarningLevel;
+			set => CppCompileWarningSettings.ModuleIncludePathWarningLevel = value;
 		}
 
-		/// <inheritdoc cref="ModuleIncludePathWarningLevel"/>
-		private WarningLevel ModuleIncludePathWarningLevelPrivate;
-
 		/// <summary>
 		/// How to treat private module include path validation messages, where a module is adding an include path that exposes private headers
 		/// </summary>
 		public WarningLevel ModuleIncludePrivateWarningLevel
 		{
-			get => (ModuleIncludePrivateWarningLevelPrivate == WarningLevel.Default) ? Target.ModuleIncludePrivateWarningLevel : ModuleIncludePrivateWarningLevelPrivate;
-			set => ModuleIncludePrivateWarningLevelPrivate = value;
+			get => CppCompileWarningSettings.ModuleIncludePrivateWarningLevel;
+			set => CppCompileWarningSettings.ModuleIncludePrivateWarningLevel = value;
 		}
 
-		/// <inheritdoc cref="ModuleIncludePrivateWarningLevel"/>
-		private WarningLevel ModuleIncludePrivateWarningLevelPrivate;
-
 		/// <summary>
 		/// How to treat unnecessary module sub-directory include path validation messages
 		/// </summary>
 		public WarningLevel ModuleIncludeSubdirectoryWarningLevel
 		{
-			get => (ModuleIncludeSubdirectoryWarningLevelPrivate == WarningLevel.Default) ? Target.ModuleIncludeSubdirectoryWarningLevel : ModuleIncludeSubdirectoryWarningLevelPrivate;
-			set => ModuleIncludeSubdirectoryWarningLevelPrivate = value;
+			get => CppCompileWarningSettings.ModuleIncludeSubdirectoryWarningLevel;
+			set => CppCompileWarningSettings.ModuleIncludeSubdirectoryWarningLevel = value;
 		}
 
-		/// <inheritdoc cref="ModuleIncludeSubdirectoryWarningLevel"/>
-		private WarningLevel ModuleIncludeSubdirectoryWarningLevelPrivate;
-
 		/// <summary>
 		/// Disable all static analysis - clang, msvc, pvs-studio.
 		/// </summary>
diff --git a/Engine/Source/Programs/UnrealBuildTool/Configuration/ReadOnlyCppCompileWarnings.cs b/Engine/Source/Programs/UnrealBuildTool/Configuration/ReadOnlyCppCompileWarnings.cs
--- a/Engine/Source/Programs/UnrealBuildTool/Configuration/ReadOnlyCppCompileWarnings.cs
+++ b/Engine/Source/Programs/UnrealBuildTool/Configuration/ReadOnlyCppCompileWarnings.cs
@@ -37,10 +37,28 @@ namespace UnrealBuildTool
 		/// </summary>
 		#region Read-only accessor properties
 #pragma warning disable CS1591
+		#region -- Compiler Warning Settings --
+
 		public WarningLevel ShadowVariableWarningLevel => _inner.ShadowVariableWarningLevel;
 		public WarningLevel UndefinedIdentifierWarningLevel => _inner.UndefinedIdentifierWarningLevel;
 		public WarningLevel UnsafeTypeCastWarningLevel => _inner.UnsafeTypeCastWarningLevel;
 		public WarningLevel SwitchUnhandledEnumeratorWarningLevel => _inner.SwitchUnhandledEnumeratorWarningLevel;
+
+		#endregion -- Compiler Warning Settings --
+
+		#region -- Performance Warning Settings --
+
+		public WarningLevel PCHPerformanceIssueWarningLevel => _inner.PCHPerformanceIssueWarningLevel;
+
+		#endregion -- Performance Warning Settings --
+
+		#region -- Rules Validation Warning Settings --
+
+		public WarningLevel ModuleIncludePathWarningLevel => _inner.ModuleIncludePathWarningLevel;
+		public WarningLevel ModuleIncludePrivateWarningLevel => _inner.ModuleIncludePrivateWarningLevel;
+		public WarningLevel ModuleIncludeSubdirectoryWarningLevel => _inner.ModuleIncludeSubdirectoryWarningLevel;
+
+		#endregion -- Rules Validation Warning Settings --
 #pragma warning restore C1591
 		#endregion
 	}
diff --git a/Engine/Source/Programs/UnrealBuildTool/Configuration/TargetRules.cs b/Engine/Source/Programs/UnrealBuildTool/Configuration/TargetRules.cs
--- a/Engine/Source/Programs/UnrealBuildTool/Configuration/TargetRules.cs
+++ b/Engine/Source/Programs/UnrealBuildTool/Configuration/TargetRules.cs
@@ -1830,7 +1830,7 @@ namespace UnrealBuildTool
 		/// C++ Warnings settings object used by the target
 		/// </summary>
 		[ConfigSubObject]
-		public CppCompileWarnings CppCompileWarningSettings {get; init;}
+		public CppCompileWarnings CppCompileWarningSettings { get; init; }
 
 		/// <summary>
 		/// Forces shadow variable warnings to be treated as errors on platforms that support it.
@@ -1877,30 +1877,38 @@ namespace UnrealBuildTool
 		/// <summary>
 		/// Indicates what warning/error level to treat potential PCH performance issues.
 		/// </summary>
-		[XmlConfigFile(Category = "BuildConfiguration")]
-		[CommandLine("-PCHPerformanceIssueWarningLevel=")]
-		public WarningLevel PCHPerformanceIssueWarningLevel { get; set; } = WarningLevel.Off;
+		public WarningLevel PCHPerformanceIssueWarningLevel
+		{
+			get => CppCompileWarningSettings.PCHPerformanceIssueWarningLevel;
+			set => CppCompileWarningSettings.PCHPerformanceIssueWarningLevel = value;
+		}
 
 		/// <summary>
 		/// How to treat general module include path validation messages
 		/// </summary>
-		[XmlConfigFile(Category = "BuildConfiguration")]
-		[CommandLine("-ModuleIncludePathWarningLevel=")]
-		public WarningLevel ModuleIncludePathWarningLevel { get; set; } = WarningLevel.Off;
+		public WarningLevel ModuleIncludePathWarningLevel
+		{
+			get => CppCompileWarningSettings.ModuleIncludePathWarningLevel;
+			set => CppCompileWarningSettings.ModuleIncludePathWarningLevel = value;
+		}
 
 		/// <summary>
 		/// How to treat private module include path validation messages, where a module is adding an include path that exposes private headers
 		/// </summary>
-		[XmlConfigFile(Category = "BuildConfiguration")]
-		[CommandLine("-ModuleIncludePrivateWarningLevel=")]
-		public WarningLevel ModuleIncludePrivateWarningLevel { get; set; } = WarningLevel.Off;
+		public WarningLevel ModuleIncludePrivateWarningLevel
+		{
+			get => CppCompileWarningSettings.ModuleIncludePrivateWarningLevel;
+			set => CppCompileWarningSettings.ModuleIncludePrivateWarningLevel = value;
+		}
 
 		/// <summary>
 		/// How to treat unnecessary module sub-directory include path validation messages
 		/// </summary>
-		[XmlConfigFile(Category = "BuildConfiguration")]
-		[CommandLine("-ModuleIncludeSubdirectoryWarningLevel=")]
-		public WarningLevel ModuleIncludeSubdirectoryWarningLevel { get; set; } = WarningLevel.Off;
+		public WarningLevel ModuleIncludeSubdirectoryWarningLevel
+		{
+			get => CppCompileWarningSettings.ModuleIncludeSubdirectoryWarningLevel;
+			set => CppCompileWarningSettings.ModuleIncludeSubdirectoryWarningLevel = value;
+		}
 
 		/// <summary>
 		/// Print out files that are included by each source file
diff --git a/Engine/Source/Programs/UnrealBuildTool/Configuration/UEBuildPlatform.cs b/Engine/Source/Programs/UnrealBuildTool/Configuration/UEBuildPlatform.cs
--- a/Engine/Source/Programs/UnrealBuildTool/Configuration/UEBuildPlatform.cs
+++ b/Engine/Source/Programs/UnrealBuildTool/Configuration/UEBuildPlatform.cs
@@ -817,9 +817,10 @@ namespace UnrealBuildTool
 		/// <param name="AllModules">Other modules to validate against, if needed</param>
 		public virtual bool ValidateModuleIncludePaths(UEBuildModule Module, ReadOnlyTargetRules Target, IEnumerable<UEBuildModule> AllModules)
 		{
-			if (Module.Rules.ModuleIncludePathWarningLevel <= WarningLevel.Off
-				&& Module.Rules.ModuleIncludePrivateWarningLevel <= WarningLevel.Off
-				&& Module.Rules.ModuleIncludeSubdirectoryWarningLevel <= WarningLevel.Off)
+			CppCompileWarnings ModuleWarningSettings = Module.Rules.CppCompileWarningSettings;
+			if (ModuleWarningSettings.ModuleIncludePathWarningLevel <= WarningLevel.Off
+				&& ModuleWarningSettings.ModuleIncludePrivateWarningLevel <= WarningLevel.Off
+				&& ModuleWarningSettings.ModuleIncludeSubdirectoryWarningLevel <= WarningLevel.Off)
 			{
 				return false;
 			}
@@ -827,11 +828,11 @@ namespace UnrealBuildTool
 			bool AnyErrors = false;
 			void LoggerFunc(string? message, params object?[] args)
 			{
-				if (Module.Rules.ModuleIncludePathWarningLevel == WarningLevel.Warning)
+				if (ModuleWarningSettings.ModuleIncludePathWarningLevel == WarningLevel.Warning)
 				{
 					Logger.LogWarning($"Warning: {message}", args);
 				}
-				else if (Module.Rules.ModuleIncludePathWarningLevel == WarningLevel.Error)
+				else if (ModuleWarningSettings.ModuleIncludePathWarningLevel == WarningLevel.Error)
 				{
 					Logger.LogError($"Error: {message}", args);
 					AnyErrors = true;
@@ -840,11 +841,11 @@ namespace UnrealBuildTool
 
 			void LoggerFuncPrivate(string? message, params object?[] args)
 			{
-				if (Module.Rules.ModuleIncludePrivateWarningLevel == WarningLevel.Warning)
+				if (ModuleWarningSettings.ModuleIncludePrivateWarningLevel == WarningLevel.Warning)
 				{
 					Logger.LogWarning($"Warning: {message}", args);
 				}
-				else if (Module.Rules.ModuleIncludePrivateWarningLevel == WarningLevel.Error)
+				else if (ModuleWarningSettings.ModuleIncludePrivateWarningLevel == WarningLevel.Error)
 				{
 					Logger.LogError($"Error: {message}", args);
 					AnyErrors = true;
@@ -853,11 +854,11 @@ namespace UnrealBuildTool
 
 			void LoggerFuncSubDir(string? message, params object?[] args)
 			{
-				if (Module.Rules.ModuleIncludeSubdirectoryWarningLevel == WarningLevel.Warning)
+				if (ModuleWarningSettings.ModuleIncludeSubdirectoryWarningLevel == WarningLevel.Warning)
 				{
 					Logger.LogWarning($"Warning: {message}", args);
 				}
-				else if (Module.Rules.ModuleIncludeSubdirectoryWarningLevel == WarningLevel.Error)
+				else if (ModuleWarningSettings.ModuleIncludeSubdirectoryWarningLevel == WarningLevel.Error)
 				{
 					Logger.LogError($"Error: {message}", args);
 					AnyErrors = true;
diff --git a/Engine/Source/Programs/UnrealBuildTool/Configuration/UEBuildTarget.cs b/Engine/Source/Programs/UnrealBuildTool/Configuration/UEBuildTarget.cs
--- a/Engine/Source/Programs/UnrealBuildTool/Configuration/UEBuildTarget.cs
+++ b/Engine/Source/Programs/UnrealBuildTool/Configuration/UEBuildTarget.cs
@@ -4067,11 +4067,11 @@ namespace UnrealBuildTool
 			}
 
 			// Throw warnings/errors if there are PCH performance issues
-			if (Target.PCHPerformanceIssueWarningLevel != WarningLevel.Off)
+			if (Target.CppCompileWarningSettings.PCHPerformanceIssueWarningLevel != WarningLevel.Off)
 			{
 				bool bFoundIssues = false;
-				LogLevel PCHLoggerLevel = Target.PCHPerformanceIssueWarningLevel == WarningLevel.Warning ? LogLevel.Warning : LogLevel.Error;
-				string Prefix = Target.PCHPerformanceIssueWarningLevel == WarningLevel.Warning ? "Warning:" : "Error:";
+				LogLevel PCHLoggerLevel = Target.CppCompileWarningSettings.PCHPerformanceIssueWarningLevel == WarningLevel.Warning ? LogLevel.Warning : LogLevel.Error;
+				string Prefix = Target.CppCompileWarningSettings.PCHPerformanceIssueWarningLevel == WarningLevel.Warning ? "Warning:" : "Error:";
 				for (int SharedPCHIndex = NumSharedPCHs - 1; SharedPCHIndex >= 0; SharedPCHIndex--)
 				{
 					// Throw warnings or errors if there are module settings that could result in PCH performance issues
