[UBT]
* Fixed crash when not having mac platform sdk

diff --git a/Engine/Source/Programs/UnrealBuildTool/Platform/Mac/ApplePlatformSDK.cs b/Engine/Source/Programs/UnrealBuildTool/Platform/Mac/ApplePlatformSDK.cs
--- a/Engine/Source/Programs/UnrealBuildTool/Platform/Mac/ApplePlatformSDK.cs
+++ b/Engine/Source/Programs/UnrealBuildTool/Platform/Mac/ApplePlatformSDK.cs
@@ -144,7 +144,8 @@ namespace UnrealBuildTool
 				{
 					string? DllVersion = FileVersionInfo.GetVersionInfo(DllPath).FileVersion;
 					// Only return the DLL version as the SDK version if we can correctly parse it
-					if (UEBuildPlatformSDK.GetSDKForPlatform("Mac")!.TryConvertVersionToInt(DllVersion, out _, null))
+					UEBuildPlatformSDK? Sdk = UEBuildPlatformSDK.GetSDKForPlatform("Mac");
+					if (Sdk != null && Sdk.TryConvertVersionToInt(DllVersion, out _, null))
 					{
 						return DllVersion;
 					}
