UnrealBuildTool: Show all batch files in project

#rnx

diff --git a/Engine/Source/Programs/UnrealBuildTool/UnrealBuildTool.csproj b/Engine/Source/Programs/UnrealBuildTool/UnrealBuildTool.csproj
--- a/Engine/Source/Programs/UnrealBuildTool/UnrealBuildTool.csproj
+++ b/Engine/Source/Programs/UnrealBuildTool/UnrealBuildTool.csproj
@@ -14,11 +14,8 @@
     <RootNamespace>UnrealBuildTool</RootNamespace>
     <SatelliteResourceLanguages>en</SatelliteResourceLanguages> <!-- remove non english resource languages -->
     <AppendRuntimeIdentifierToOutputPath>false</AppendRuntimeIdentifierToOutputPath>
-    <IsWindows Condition="'$([System.Runtime.InteropServices.RuntimeInformation]::IsOSPlatform($([System.Runtime.InteropServices.OSPlatform]::Windows)))' == 'true'">true</IsWindows> 
-    <IsOSX Condition="'$([System.Runtime.InteropServices.RuntimeInformation]::IsOSPlatform($([System.Runtime.InteropServices.OSPlatform]::OSX)))' == 'true'">true</IsOSX> 
-    <IsLinux Condition="'$([System.Runtime.InteropServices.RuntimeInformation]::IsOSPlatform($([System.Runtime.InteropServices.OSPlatform]::Linux)))' == 'true'">true</IsLinux> 
-    <DebugType>pdbonly</DebugType> 
-    <DebugType Condition="'$([System.Runtime.InteropServices.RuntimeInformation]::IsOSPlatform($([System.Runtime.InteropServices.OSPlatform]::Windows)))' == 'true' And '$([System.Runtime.InteropServices.RuntimeInformation]::OSArchitecture)' == 'Arm64'">portable</DebugType>
+    <DebugType>pdbonly</DebugType>
+    <DebugType Condition="'$(IsWindows)' == 'true' And '$(IsArm64)' == 'true'">portable</DebugType>
     <Nullable>enable</Nullable>
     <RunAnalyzersDuringBuild>false</RunAnalyzersDuringBuild>
     <GenerateDocumentationFile>True</GenerateDocumentationFile>
@@ -32,7 +29,7 @@
   <PropertyGroup Condition=" '$(Configuration)|$(Platform)' == 'Debug|AnyCPU' ">
     <DebugSymbols>true</DebugSymbols>
     <DebugType>full</DebugType>
-    <DebugType Condition="'$([System.Runtime.InteropServices.RuntimeInformation]::IsOSPlatform($([System.Runtime.InteropServices.OSPlatform]::Windows)))' == 'true' And '$([System.Runtime.InteropServices.RuntimeInformation]::OSArchitecture)' == 'Arm64'">portable</DebugType>
+    <DebugType Condition="'$(IsWindows)' == 'true' And '$(IsArm64)' == 'true'">portable</DebugType>
   </PropertyGroup>
   <PropertyGroup Condition=" '$(Configuration)|$(Platform)' == 'Analyze|AnyCPU' ">
     <RunAnalyzersDuringBuild>True</RunAnalyzersDuringBuild>
@@ -43,13 +40,7 @@
       <HintPath>..\..\..\Binaries\DotNET\Ionic.Zip.Reduced.dll</HintPath>
     </Reference>
   </ItemGroup>
-  <Choose>
-    <When Condition="Exists('..\..\..\Restricted\NotForLicensees\Source\Programs\UnrealBuildTool\Executors\FASTBuild.cs')">
-      <PropertyGroup>
-        <DefineConstants>$(DefineConstants);__FASTBUILD_AVAILABLE__</DefineConstants>
-      </PropertyGroup>
-    </When>
-  </Choose>
+
   <Choose>
     <When Condition="Exists('..\..\..\Restricted\NotForLicensees\Source\Programs\UnrealBuildTool\ProjectFiles\VProject\VProjectFileGenerator.cs')">
       <PropertyGroup>
@@ -126,22 +117,11 @@
     <PackageReference Include="Newtonsoft.Json" Version="13.0.3" />
   </ItemGroup>
 
-  <!-- Various batch files that build or run UnrealBuildTool -->
+  <!-- Various batch files and shell scripts that build or run UnrealBuildTool -->
   <ItemGroup>
-    <None Include="..\..\..\Build\BatchFiles\BuildUAT.bat"><Link>BatchFiles\BuildUAT.bat</Link></None>
-    <None Include="..\..\..\Build\BatchFiles\BuildUAT.sh"><Link>BatchFiles\BuildUAT.sh</Link></None>
-    <None Include="..\..\..\Build\BatchFiles\BuildUBT.bat"><Link>BatchFiles\BuildUBT.bat</Link></None>
-    <None Include="..\..\..\Build\BatchFiles\BuildUBT.sh"><Link>BatchFiles\BuildUBT.sh</Link></None>
-    <None Include="..\..\..\Build\BatchFiles\GenerateProjectFiles.bat"><Link>BatchFiles\GenerateProjectFiles.bat</Link></None>
-    <None Include="..\..\..\Build\BatchFiles\RunUAT.bat"><Link>BatchFiles\RunUAT.bat</Link></None>
-    <None Include="..\..\..\Build\BatchFiles\RunUAT.sh"><Link>BatchFiles\RunUAT.sh</Link></None>
-    <None Include="..\..\..\Build\BatchFiles\RunUBT.bat"><Link>BatchFiles\RunUBT.bat</Link></None>
-    <None Include="..\..\..\Build\BatchFiles\RunUBT.sh"><Link>BatchFiles\RunUBT.sh</Link></None>
-
-    <!-- Visual Studio batch files -->
-    <None Include="..\..\..\Build\BatchFiles\Build.bat"><Link>BatchFiles\Build.bat</Link></None>
-    <None Include="..\..\..\Build\BatchFiles\Clean.bat"><Link>BatchFiles\Clean.bat</Link></None>
-    <None Include="..\..\..\Build\BatchFiles\Rebuild.bat"><Link>BatchFiles\Rebuild.bat</Link></None>
+    <None Include="../../../Build/BatchFiles/**/*">
+      <Link>$([System.Text.RegularExpressions.Regex]::Replace(%(FullPath), '^.+?[\\/]BatchFiles[\\/]', 'BatchFiles/'))</Link>
+    </None>
   </ItemGroup>
 
   <!-- Various BuildConfiguration files -->
