UnrealBuildTool: Move action root paths to CppCompileEnvironment & LinkEnvironment

#rnx

diff --git a/Engine/Source/Programs/UnrealBuildTool/Configuration/UEBuildBinary.cs b/Engine/Source/Programs/UnrealBuildTool/Configuration/UEBuildBinary.cs
--- a/Engine/Source/Programs/UnrealBuildTool/Configuration/UEBuildBinary.cs
+++ b/Engine/Source/Programs/UnrealBuildTool/Configuration/UEBuildBinary.cs
@@ -856,6 +856,12 @@ namespace UnrealBuildTool
 						}
 					}
 
+					// Project root path should only be added if the binary needs it, to prevent violating shared environment actions
+					if (Target.ProjectFile != null && BinaryLinkEnvironment.InputFiles.Any(x => x.Location.IsUnderDirectory(Target.ProjectFile.Directory)))
+					{
+						BinaryLinkEnvironment.RootPaths[CppRootPathFolder.Project] = Target.ProjectFile.Directory;
+					}
+
 					// Force a reference to initialize module for this binary
 					if (Module.Rules.bRequiresImplementModule ?? true)
 					{
diff --git a/Engine/Source/Programs/UnrealBuildTool/Configuration/UEBuildModuleCPP.cs b/Engine/Source/Programs/UnrealBuildTool/Configuration/UEBuildModuleCPP.cs
--- a/Engine/Source/Programs/UnrealBuildTool/Configuration/UEBuildModuleCPP.cs
+++ b/Engine/Source/Programs/UnrealBuildTool/Configuration/UEBuildModuleCPP.cs
@@ -2360,6 +2360,12 @@ namespace UnrealBuildTool
 				}
 			}
 
+			// Project root path should only be added if the module needs it, to prevent violating shared environment actions
+			if (Target.ProjectFile != null && (!Result.bUseSharedBuildEnvironment || Result.AllIncludePath.Any(x => x.IsUnderDirectory(Target.ProjectFile.Directory))))
+			{
+				Result.RootPaths[CppRootPathFolder.Project] = Target.ProjectFile.Directory;
+			}
+
 			return Result;
 		}
 
diff --git a/Engine/Source/Programs/UnrealBuildTool/Configuration/UEBuildTarget.cs b/Engine/Source/Programs/UnrealBuildTool/Configuration/UEBuildTarget.cs
--- a/Engine/Source/Programs/UnrealBuildTool/Configuration/UEBuildTarget.cs
+++ b/Engine/Source/Programs/UnrealBuildTool/Configuration/UEBuildTarget.cs
@@ -5208,7 +5208,7 @@ namespace UnrealBuildTool
 		{
 			UEBuildPlatform BuildPlatform = UEBuildPlatform.GetBuildPlatform(Platform);
 
-			ToolChain.SetUpGlobalEnvironment(Rules);
+			ToolChain.SetUpGlobalEnvironment(Rules, GlobalCompileEnvironment, GlobalLinkEnvironment);
 
 			// @Hack: This to prevent UHT from listing CoreUObject.init.gen.cpp as its dependency.
 			// We flag the compile environment when we build UHT so that we don't need to check
diff --git a/Engine/Source/Programs/UnrealBuildTool/Platform/Android/AndroidToolChain.cs b/Engine/Source/Programs/UnrealBuildTool/Platform/Android/AndroidToolChain.cs
--- a/Engine/Source/Programs/UnrealBuildTool/Platform/Android/AndroidToolChain.cs
+++ b/Engine/Source/Programs/UnrealBuildTool/Platform/Android/AndroidToolChain.cs
@@ -1082,7 +1082,7 @@ namespace UnrealBuildTool
 
 			// Create an action that invokes the linker.
 			Action LinkAction = Graph.CreateAction(ActionType.Link);
-			LinkAction.RootPaths.AddRange(GetEnvironmentBasePaths(LinkEnvironment));
+			LinkAction.RootPaths.AddRange(LinkEnvironment.RootPaths.Values.Select(x => DirectoryItem.GetItemByDirectoryReference(x)));
 			LinkAction.WorkingDirectory = Unreal.EngineSourceDirectory;
 
 			if (LinkEnvironment.bIsBuildingLibrary)
diff --git a/Engine/Source/Programs/UnrealBuildTool/Platform/IOS/IOSToolChain.cs b/Engine/Source/Programs/UnrealBuildTool/Platform/IOS/IOSToolChain.cs
--- a/Engine/Source/Programs/UnrealBuildTool/Platform/IOS/IOSToolChain.cs
+++ b/Engine/Source/Programs/UnrealBuildTool/Platform/IOS/IOSToolChain.cs
@@ -392,7 +392,7 @@ namespace UnrealBuildTool
 
 			// Create an action that invokes the linker.
 			Action LinkAction = Graph.CreateAction(ActionType.Link);
-			LinkAction.RootPaths.AddRange(GetEnvironmentBasePaths(LinkEnvironment));
+			LinkAction.RootPaths.AddRange(LinkEnvironment.RootPaths.Values.Select(x => DirectoryItem.GetItemByDirectoryReference(x)));
 
 			// RPC utility parameters are in terms of the Mac side
 			LinkAction.WorkingDirectory = GetMacDevSrcRoot();
diff --git a/Engine/Source/Programs/UnrealBuildTool/Platform/Linux/LinuxToolChain.cs b/Engine/Source/Programs/UnrealBuildTool/Platform/Linux/LinuxToolChain.cs
--- a/Engine/Source/Programs/UnrealBuildTool/Platform/Linux/LinuxToolChain.cs
+++ b/Engine/Source/Programs/UnrealBuildTool/Platform/Linux/LinuxToolChain.cs
@@ -1131,7 +1131,7 @@ namespace UnrealBuildTool
 
 			// Create an action that invokes the linker.
 			Action LinkAction = Graph.CreateAction(ActionType.Link);
-			LinkAction.RootPaths.AddRange(GetEnvironmentBasePaths(LinkEnvironment));
+			LinkAction.RootPaths.AddRange(LinkEnvironment.RootPaths.Values.Select(x => DirectoryItem.GetItemByDirectoryReference(x)));
 			LinkAction.WorkingDirectory = Unreal.EngineSourceDirectory;
 
 			string LinkCommandString;
@@ -1589,7 +1589,7 @@ namespace UnrealBuildTool
 				{
 					// Create the action to relink the library. This actions does not overwrite the source file so it can be executed in parallel
 					Action RelinkAction = Graph.CreateAction(ActionType.Link);
-					RelinkAction.RootPaths.AddRange(GetEnvironmentBasePaths(LinkEnvironment));
+					RelinkAction.RootPaths.AddRange(LinkEnvironment.RootPaths.Values.Select(x => DirectoryItem.GetItemByDirectoryReference(x)));
 					RelinkAction.WorkingDirectory = LinkAction.WorkingDirectory;
 					RelinkAction.StatusDescription = LinkAction.StatusDescription;
 					RelinkAction.CommandDescription = "Relink";
diff --git a/Engine/Source/Programs/UnrealBuildTool/Platform/Mac/MacToolChain.cs b/Engine/Source/Programs/UnrealBuildTool/Platform/Mac/MacToolChain.cs
--- a/Engine/Source/Programs/UnrealBuildTool/Platform/Mac/MacToolChain.cs
+++ b/Engine/Source/Programs/UnrealBuildTool/Platform/Mac/MacToolChain.cs
@@ -680,7 +680,7 @@ namespace UnrealBuildTool
 
 			// Create an action that invokes the linker.
 			Action LinkAction = Graph.CreateAction(ActionType.Link);
-			LinkAction.RootPaths.AddRange(GetEnvironmentBasePaths(LinkEnvironment));
+			LinkAction.RootPaths.AddRange(LinkEnvironment.RootPaths.Values.Select(x => DirectoryItem.GetItemByDirectoryReference(x)));
 
 			FileReference LinkerPath = LinkEnvironment.bIsBuildingLibrary ? Info.Archiver : Info.Clang;
 
diff --git a/Engine/Source/Programs/UnrealBuildTool/Platform/Windows/PVSToolChain.cs b/Engine/Source/Programs/UnrealBuildTool/Platform/Windows/PVSToolChain.cs
--- a/Engine/Source/Programs/UnrealBuildTool/Platform/Windows/PVSToolChain.cs
+++ b/Engine/Source/Programs/UnrealBuildTool/Platform/Windows/PVSToolChain.cs
@@ -640,15 +640,13 @@ namespace UnrealBuildTool
 			Lines.Add(String.Format("Using PVS-Studio {0} at {1} with analysis mode {2} ({3})", AnalyzerVersion, AnalyzerFile, (uint)Settings.ModeFlags, Settings.ModeFlags.ToString()));
 		}
 
-		protected override IEnumerable<DirectoryItem> GetEnvironmentBasePaths(CppCompileEnvironment CompileEnvironment)
+		public override void SetUpGlobalEnvironment(ReadOnlyTargetRules Target, CppCompileEnvironment GlobalCompileEnvironment, LinkEnvironment GlobalLinkEnvironment)
 		{
-			yield return DirectoryItem.GetItemByDirectoryReference(Unreal.EngineDirectory);
-			if (ProjectFile != null && (!CompileEnvironment.bUseSharedBuildEnvironment || CompileEnvironment.AllIncludePath.Any(x => x.IsUnderDirectory(ProjectFile.Directory))))
-			{
-				yield return DirectoryItem.GetItemByDirectoryReference(ProjectFile.Directory);
-			}
-			yield return DirectoryItem.GetItemByDirectoryReference(Unreal.RootDirectory);
-			yield return DirectoryItem.GetItemByDirectoryReference(AnalyzerFile.Directory);
+			base.SetUpGlobalEnvironment(Target, GlobalCompileEnvironment, GlobalLinkEnvironment);
+			InnerToolChain.SetUpGlobalEnvironment(Target, GlobalCompileEnvironment, GlobalLinkEnvironment);
+
+			GlobalCompileEnvironment.RootPaths[CppRootPathFolder.Extra] = AnalyzerFile.Directory;
+			GlobalLinkEnvironment.RootPaths[CppRootPathFolder.Extra] = AnalyzerFile.Directory;
 		}
   
 		public override void SetEnvironmentVariables()
diff --git a/Engine/Source/Programs/UnrealBuildTool/Platform/Windows/VCToolChain.cs b/Engine/Source/Programs/UnrealBuildTool/Platform/Windows/VCToolChain.cs
--- a/Engine/Source/Programs/UnrealBuildTool/Platform/Windows/VCToolChain.cs
+++ b/Engine/Source/Programs/UnrealBuildTool/Platform/Windows/VCToolChain.cs
@@ -1701,43 +1701,11 @@ namespace UnrealBuildTool
 			}
 		}
 
-		protected override IEnumerable<DirectoryItem> GetEnvironmentBasePaths(CppCompileEnvironment CompileEnvironment)
-		{
-			yield return DirectoryItem.GetItemByDirectoryReference(Unreal.EngineDirectory);
-			if (Target.ProjectFile != null && (!CompileEnvironment.bUseSharedBuildEnvironment || CompileEnvironment.AllIncludePath.Any(x => x.IsUnderDirectory(Target.ProjectFile.Directory))))
-			{
-				yield return DirectoryItem.GetItemByDirectoryReference(Target.ProjectFile.Directory);
-			}
-			yield return DirectoryItem.GetItemByDirectoryReference(Unreal.RootDirectory);
-			yield return DirectoryItem.GetItemByDirectoryReference(EnvVars.CompilerDir);
-			if (EnvVars.CompilerDir != EnvVars.ToolChainDir)
-			{
-				yield return DirectoryItem.GetItemByDirectoryReference(EnvVars.ToolChainDir);
-			}
-			yield return DirectoryItem.GetItemByDirectoryReference(EnvVars.WindowsSdkDir);
-		}
-
-		protected override IEnumerable<DirectoryItem> GetEnvironmentBasePaths(LinkEnvironment LinkEnvironment)
-		{
-			yield return DirectoryItem.GetItemByDirectoryReference(Unreal.EngineDirectory);
-			if (Target.ProjectFile != null && LinkEnvironment.InputFiles.Any(x => x.Location.IsUnderDirectory(Target.ProjectFile.Directory)))
-			{
-				yield return DirectoryItem.GetItemByDirectoryReference(Target.ProjectFile.Directory);
-			}
-			yield return DirectoryItem.GetItemByDirectoryReference(Unreal.RootDirectory);
-			yield return DirectoryItem.GetItemByDirectoryReference(EnvVars.CompilerDir);
-			if (EnvVars.CompilerDir != EnvVars.ToolChainDir)
-			{
-				yield return DirectoryItem.GetItemByDirectoryReference(EnvVars.ToolChainDir);
-			}
-			yield return DirectoryItem.GetItemByDirectoryReference(EnvVars.WindowsSdkDir);
-		}
-
 		private VCCompileAction CreateBaseCompileAction(CppCompileEnvironment CompileEnvironment)
 		{
 			VCCompileAction BaseCompileAction = new VCCompileAction(EnvVars);
 
-			IEnumerable<DirectoryItem> rootPaths = GetEnvironmentBasePaths(CompileEnvironment);
+			IEnumerable<DirectoryItem> rootPaths = CompileEnvironment.RootPaths.Values.Select(x => DirectoryItem.GetItemByDirectoryReference(x));
 
 			// TODO: Revisit this code. We want to use d1trimfile to make outputs machine independent.
 			// but want to make sure we're not causing any frustration for devs (since __FILE__ will show a relative path with lines below)
@@ -2392,7 +2360,7 @@ namespace UnrealBuildTool
 			foreach (FileItem RCFile in InputFiles)
 			{
 				Action CompileAction = Graph.CreateAction(ActionType.Compile);
-				CompileAction.RootPaths.AddRange(GetEnvironmentBasePaths(CompileEnvironment));
+				CompileAction.RootPaths.AddRange(CompileEnvironment.RootPaths.Values.Select(x => DirectoryItem.GetItemByDirectoryReference(x)));
 				CompileAction.CommandDescription = "Resource";
 				CompileAction.WorkingDirectory = Unreal.EngineSourceDirectory;
 				CompileAction.CommandPath = EnvVars.ResourceCompilerPath;
@@ -2604,7 +2572,7 @@ namespace UnrealBuildTool
 			CompileAction.CommandVersion = EnvVars.ToolChainVersion.ToString();
 			CompileAction.bShouldOutputStatusDescription = false;
 			CompileAction.bCanExecuteRemotely = false; // Incompatible with remote distribution
-			CompileAction.RootPaths.AddRange(GetEnvironmentBasePaths(CompileEnvironment));
+			CompileAction.RootPaths.AddRange(CompileEnvironment.RootPaths.Values.Select(x => DirectoryItem.GetItemByDirectoryReference(x)));
 			CompileAction.ArtifactMode = ArtifactMode.Enabled;
 		}
 
@@ -2973,7 +2941,7 @@ namespace UnrealBuildTool
 
 			// Create an action that invokes the linker.
 			Action LinkAction = Graph.CreateAction(ActionType.Link);
-			LinkAction.RootPaths.AddRange(GetEnvironmentBasePaths(LinkEnvironment));
+			LinkAction.RootPaths.AddRange(LinkEnvironment.RootPaths.Values.Select(x => DirectoryItem.GetItemByDirectoryReference(x)));
 			string ReadableArch = UnrealArchitectureConfig.ForPlatform(LinkEnvironment.Platform).ConvertToReadableArchitecture(LinkEnvironment.Architecture);
 			LinkAction.CommandDescription += $"Link [{ReadableArch}]";
 			LinkAction.WorkingDirectory = Unreal.EngineSourceDirectory;
@@ -3380,6 +3348,21 @@ namespace UnrealBuildTool
 			return Target.WindowsPlatform.Compiler.IsClang() ? "dynlist" : "obj";
 		}
 
+		public override void SetUpGlobalEnvironment(ReadOnlyTargetRules Target, CppCompileEnvironment GlobalCompileEnvironment, LinkEnvironment GlobalLinkEnvironment)
+		{
+			base.SetUpGlobalEnvironment(Target, GlobalCompileEnvironment, GlobalLinkEnvironment);
+
+			GlobalCompileEnvironment.RootPaths[CppRootPathFolder.Compiler] = EnvVars.CompilerDir;
+			GlobalLinkEnvironment.RootPaths[CppRootPathFolder.Compiler] = EnvVars.CompilerDir;
+			if (EnvVars.CompilerDir != EnvVars.ToolChainDir)
+			{
+				GlobalCompileEnvironment.RootPaths[CppRootPathFolder.Toolchain] = EnvVars.ToolChainDir;
+				GlobalLinkEnvironment.RootPaths[CppRootPathFolder.Toolchain] = EnvVars.ToolChainDir;
+			}
+			GlobalCompileEnvironment.RootPaths[CppRootPathFolder.WinSDK] = EnvVars.WindowsSdkDir;
+			GlobalLinkEnvironment.RootPaths[CppRootPathFolder.WinSDK] = EnvVars.WindowsSdkDir;
+		}
+
 		public override void ModifyBuildProducts(ReadOnlyTargetRules Target, UEBuildBinary Binary, IEnumerable<string> Libraries, IEnumerable<UEBuildBundleResource> BundleResources, Dictionary<FileReference, BuildProductType> BuildProducts)
 		{
 			if (Binary.Type == UEBuildBinaryType.DynamicLinkLibrary)
diff --git a/Engine/Source/Programs/UnrealBuildTool/System/CppCompileEnvironment.cs b/Engine/Source/Programs/UnrealBuildTool/System/CppCompileEnvironment.cs
--- a/Engine/Source/Programs/UnrealBuildTool/System/CppCompileEnvironment.cs
+++ b/Engine/Source/Programs/UnrealBuildTool/System/CppCompileEnvironment.cs
@@ -149,6 +149,24 @@ namespace UnrealBuildTool
 		Create
 	}
 
+	/// <summary>
+	/// Named root folders
+	/// </summary>
+	enum CppRootPathFolder
+	{
+		Engine,
+		Project,
+		Root,
+
+		Compiler,
+		Toolchain,
+		WinSDK,
+
+		PlatformSDK,
+
+		Extra,
+	}
+
 	/// <summary>
 	/// Encapsulates the compilation output of compiling a set of C++ files.
 	/// </summary>
@@ -216,6 +234,11 @@ namespace UnrealBuildTool
 		/// </summary>
 		public readonly SourceFileMetadataCache MetadataCache;
 
+		/// <summary>
+		/// Root paths for this environment. Used for Clang VFS and UBA
+		/// </summary>
+		public readonly SortedDictionary<CppRootPathFolder, DirectoryReference> RootPaths;
+
 		/// <summary>
 		/// Templates for shared precompiled headers
 		/// </summary>
@@ -735,6 +758,7 @@ namespace UnrealBuildTool
 			this.Configuration = Configuration;
 			this.Architectures = Architectures;
 			this.MetadataCache = MetadataCache;
+			RootPaths = new();
 			SharedPCHs = new List<PrecompiledHeaderTemplate>();
 			UserIncludePaths = new HashSet<DirectoryReference>();
 			SystemIncludePaths = new HashSet<DirectoryReference>();
@@ -754,6 +778,7 @@ namespace UnrealBuildTool
 			Configuration = Other.Configuration;
 			Architectures = new(Other.Architectures);
 			MetadataCache = Other.MetadataCache;
+			RootPaths = new(Other.RootPaths, null);
 			SharedPCHs = Other.SharedPCHs;
 			PrecompiledHeaderIncludeFilename = Other.PrecompiledHeaderIncludeFilename;
 			PrecompiledHeaderAction = Other.PrecompiledHeaderAction;
diff --git a/Engine/Source/Programs/UnrealBuildTool/System/LinkEnvironment.cs b/Engine/Source/Programs/UnrealBuildTool/System/LinkEnvironment.cs
--- a/Engine/Source/Programs/UnrealBuildTool/System/LinkEnvironment.cs
+++ b/Engine/Source/Programs/UnrealBuildTool/System/LinkEnvironment.cs
@@ -33,6 +33,11 @@ namespace UnrealBuildTool
 		/// </summary>
 		public UnrealArch Architecture => Architectures.SingleArchitecture;
 
+		/// <summary>
+		/// Root paths for this environment. Used for Clang VFS and UBA
+		/// </summary>
+		public readonly SortedDictionary<CppRootPathFolder, DirectoryReference> RootPaths;
+
 		/// <summary>
 		/// On Mac, indicates the path to the target's application bundle
 		/// </summary>
@@ -375,6 +380,7 @@ namespace UnrealBuildTool
 			this.Platform = Platform;
 			this.Configuration = Configuration;
 			this.Architectures = new(Architectures);
+			RootPaths = new();
 		}
 
 		/// <summary>
@@ -385,6 +391,7 @@ namespace UnrealBuildTool
 			Platform = Other.Platform;
 			Configuration = Other.Configuration;
 			Architectures = Other.Architectures;
+			RootPaths = new(Other.RootPaths, null);
 			BundleDirectory = Other.BundleDirectory;
 			OutputDirectory = Other.OutputDirectory;
 			IntermediateDirectory = Other.IntermediateDirectory;
diff --git a/Engine/Source/Programs/UnrealBuildTool/ToolChain/ClangToolChain.cs b/Engine/Source/Programs/UnrealBuildTool/ToolChain/ClangToolChain.cs
--- a/Engine/Source/Programs/UnrealBuildTool/ToolChain/ClangToolChain.cs
+++ b/Engine/Source/Programs/UnrealBuildTool/ToolChain/ClangToolChain.cs
@@ -233,9 +233,15 @@ namespace UnrealBuildTool
 			return LazyInfo.Value.Clang;
 		}
 
-		public override void SetUpGlobalEnvironment(ReadOnlyTargetRules Target)
+		public override void SetUpGlobalEnvironment(ReadOnlyTargetRules Target, CppCompileEnvironment GlobalCompileEnvironment, LinkEnvironment GlobalLinkEnvironment)
 		{
-			base.SetUpGlobalEnvironment(Target);
+			base.SetUpGlobalEnvironment(Target, GlobalCompileEnvironment, GlobalLinkEnvironment);
+
+			if (LazyInfo.Value.BasePath != null)
+			{
+				GlobalCompileEnvironment.RootPaths[CppRootPathFolder.PlatformSDK] = LazyInfo.Value.BasePath;
+				GlobalLinkEnvironment.RootPaths[CppRootPathFolder.PlatformSDK] = LazyInfo.Value.BasePath;
+			}
 
 			PreprocessDepends = Target.bPreprocessDepends;
 			ShowIncludes = Target.bShowIncludes;
@@ -1183,38 +1189,10 @@ namespace UnrealBuildTool
 			Graph.AddAction(new ClangSpecificFileAction(SourceDir, OutputDir, Action, GraphBuilder.ContentLines));
 		}
 
-		protected override IEnumerable<DirectoryItem> GetEnvironmentBasePaths(CppCompileEnvironment CompileEnvironment)
-		{
-			yield return DirectoryItem.GetItemByDirectoryReference(Unreal.EngineDirectory);
-			if (ProjectFile != null && (!CompileEnvironment.bUseSharedBuildEnvironment || CompileEnvironment.AllIncludePath.Any(x => x.IsUnderDirectory(ProjectFile.Directory))))
-			{
-				yield return DirectoryItem.GetItemByDirectoryReference(ProjectFile.Directory);
-			}
-			yield return DirectoryItem.GetItemByDirectoryReference(Unreal.RootDirectory);
-			if (GetToolChainInfo().BasePath != null)
-			{
-				yield return DirectoryItem.GetItemByDirectoryReference(GetToolChainInfo().BasePath!);
-			}
-		}
-
-		protected override IEnumerable<DirectoryItem> GetEnvironmentBasePaths(LinkEnvironment LinkEnvironment)
-		{
-			yield return DirectoryItem.GetItemByDirectoryReference(Unreal.EngineDirectory);
-			if (ProjectFile != null && LinkEnvironment.InputFiles.Any(x => x.Location.IsUnderDirectory(ProjectFile.Directory)))
-			{
-				yield return DirectoryItem.GetItemByDirectoryReference(ProjectFile.Directory);
-			}
-			yield return DirectoryItem.GetItemByDirectoryReference(Unreal.RootDirectory);
-			if (GetToolChainInfo().BasePath != null)
-			{
-				yield return DirectoryItem.GetItemByDirectoryReference(GetToolChainInfo().BasePath!);
-			}
-		}
-
 		protected virtual Action CompileCPPFile(CppCompileEnvironment CompileEnvironment, FileItem SourceFile, DirectoryReference OutputDir, string ModuleName, IActionGraphBuilder Graph, IReadOnlyCollection<string> GlobalArguments, CPPOutput Result)
 		{
 			Action CompileAction = Graph.CreateAction(ActionType.Compile);
-			CompileAction.RootPaths.AddRange(GetEnvironmentBasePaths(CompileEnvironment));
+			CompileAction.RootPaths.AddRange(CompileEnvironment.RootPaths.Values.Select(x => DirectoryItem.GetItemByDirectoryReference(x)));
 
 			// If we are using the AutoRTFM compiler, we make the compile action depend on the version of the compiler itself.
 			// This lets us update the compiler (which might not cause a version update of the compiler, which instead tracks
diff --git a/Engine/Source/Programs/UnrealBuildTool/ToolChain/ISPCToolChain.cs b/Engine/Source/Programs/UnrealBuildTool/ToolChain/ISPCToolChain.cs
--- a/Engine/Source/Programs/UnrealBuildTool/ToolChain/ISPCToolChain.cs
+++ b/Engine/Source/Programs/UnrealBuildTool/ToolChain/ISPCToolChain.cs
@@ -21,9 +21,9 @@ namespace UnrealBuildTool
 		protected bool bMergeModules = false;
 		protected bool bAllowUbaCompression = false;
 
-		public override void SetUpGlobalEnvironment(ReadOnlyTargetRules Target)
+		public override void SetUpGlobalEnvironment(ReadOnlyTargetRules Target, CppCompileEnvironment GlobalCompileEnvironment, LinkEnvironment GlobalLinkEnvironment)
 		{
-			base.SetUpGlobalEnvironment(Target);
+			base.SetUpGlobalEnvironment(Target, GlobalCompileEnvironment, GlobalLinkEnvironment);
 			ProjectFile = Target.ProjectFile;
 			bMergeModules = Target.bMergeModules;
 			bAllowUbaCompression = Target.bAllowUbaCompression;
@@ -359,26 +359,6 @@ namespace UnrealBuildTool
 			return NormalizeCommandLinePath(Item.Location);
 		}
 
-		protected virtual IEnumerable<DirectoryItem> GetEnvironmentBasePaths(CppCompileEnvironment CompileEnvironment)
-		{
-			yield return DirectoryItem.GetItemByDirectoryReference(Unreal.EngineDirectory);
-			if (ProjectFile != null && (!CompileEnvironment.bUseSharedBuildEnvironment || CompileEnvironment.AllIncludePath.Any(x => x.IsUnderDirectory(ProjectFile.Directory))))
-			{
-				yield return DirectoryItem.GetItemByDirectoryReference(ProjectFile.Directory);
-			}
-			yield return DirectoryItem.GetItemByDirectoryReference(Unreal.RootDirectory);
-		}
-
-		protected virtual IEnumerable<DirectoryItem> GetEnvironmentBasePaths(LinkEnvironment LinkEnvironment)
-		{
-			yield return DirectoryItem.GetItemByDirectoryReference(Unreal.EngineDirectory);
-			if (ProjectFile != null && LinkEnvironment.InputFiles.Any(x => x.Location.IsUnderDirectory(ProjectFile.Directory)))
-			{
-				yield return DirectoryItem.GetItemByDirectoryReference(ProjectFile.Directory);
-			}
-			yield return DirectoryItem.GetItemByDirectoryReference(Unreal.RootDirectory);
-		}
-
 		protected override CPPOutput GenerateISPCHeaders(CppCompileEnvironment CompileEnvironment, IEnumerable<FileItem> InputFiles, DirectoryReference OutputDir, IActionGraphBuilder Graph)
 		{
 			CPPOutput Result = new CPPOutput();
@@ -441,7 +421,7 @@ namespace UnrealBuildTool
 				}
 			}
 
-			List<DirectoryItem> RootPaths = new(GetEnvironmentBasePaths(CompileEnvironment));
+			IEnumerable<DirectoryItem> RootPaths = CompileEnvironment.RootPaths.Values.Select(x => DirectoryItem.GetItemByDirectoryReference(x));
 
 			foreach (FileItem ISPCFile in InputFiles)
 			{
@@ -643,7 +623,7 @@ namespace UnrealBuildTool
 				}
 			}
 
-			List<DirectoryItem> RootPaths = new(GetEnvironmentBasePaths(CompileEnvironment));
+			IEnumerable<DirectoryItem> RootPaths = CompileEnvironment.RootPaths.Values.Select(x => DirectoryItem.GetItemByDirectoryReference(x));
 
 			foreach (FileItem ISPCFile in InputFiles)
 			{
diff --git a/Engine/Source/Programs/UnrealBuildTool/ToolChain/UEToolChain.cs b/Engine/Source/Programs/UnrealBuildTool/ToolChain/UEToolChain.cs
--- a/Engine/Source/Programs/UnrealBuildTool/ToolChain/UEToolChain.cs
+++ b/Engine/Source/Programs/UnrealBuildTool/ToolChain/UEToolChain.cs
@@ -272,8 +272,13 @@ namespace UnrealBuildTool
 			return Executables.SelectMany(x => PostBuild(Target, x, ExecutableLinkEnvironment, Graph)).ToList();
 		}
 
-		public virtual void SetUpGlobalEnvironment(ReadOnlyTargetRules Target)
+		public virtual void SetUpGlobalEnvironment(ReadOnlyTargetRules Target, CppCompileEnvironment GlobalCompileEnvironment, LinkEnvironment GlobalLinkEnvironment)
 		{
+			GlobalCompileEnvironment.RootPaths[CppRootPathFolder.Engine] = Unreal.EngineDirectory;
+			GlobalLinkEnvironment.RootPaths[CppRootPathFolder.Engine] = Unreal.EngineDirectory;
+
+			GlobalCompileEnvironment.RootPaths[CppRootPathFolder.Root] = Unreal.RootDirectory;
+			GlobalLinkEnvironment.RootPaths[CppRootPathFolder.Root] = Unreal.RootDirectory;
 		}
 
 		public virtual void ModifyBuildProducts(ReadOnlyTargetRules Target, UEBuildBinary Binary, IEnumerable<string> Libraries, IEnumerable<UEBuildBundleResource> BundleResources, Dictionary<FileReference, BuildProductType> BuildProducts)
