Fix directory checked for profile for Android variables on Mac

#android

diff --git a/Engine/Source/Programs/UnrealBuildTool/Platform/Android/AndroidPlatformSDK.cs b/Engine/Source/Programs/UnrealBuildTool/Platform/Android/AndroidPlatformSDK.cs
--- a/Engine/Source/Programs/UnrealBuildTool/Platform/Android/AndroidPlatformSDK.cs
+++ b/Engine/Source/Programs/UnrealBuildTool/Platform/Android/AndroidPlatformSDK.cs
@@ -22,30 +22,56 @@ namespace UnrealBuildTool
 		{
 		}
 
-		public static string? GetNDKRoot()
+		private static void CheckEnvironmentVariables(bool bUseConfig)
 		{
-			string? NDKPath = Environment.GetEnvironmentVariable("NDKROOT");
-
-			// don't register if we don't have an NDKROOT specified
-			if (!String.IsNullOrEmpty(NDKPath))
+			Dictionary<string, string> EnvVarNames = new Dictionary<string, string> {
+													 {"ANDROID_HOME", "SDKPath"},
+													 {"NDKROOT", "NDKPath"},
+													 {"JAVA_HOME", "JavaPath"}
+													 };
+			Dictionary<string, string> AndroidEnv = new Dictionary<string, string>();
+
+			// first check for configuration overrides
+			if (bUseConfig)
 			{
-				return NDKPath.Replace("\"", "");
+				ConfigHierarchy configCacheIni = ConfigCache.ReadHierarchy(ConfigHierarchyType.Engine, null, BuildHostPlatform.Current.Platform);
+				string path;
+				foreach (KeyValuePair<string, string> kvp in EnvVarNames)
+				{
+					if (GetPath(configCacheIni, "/Script/AndroidPlatformEditor.AndroidSDKSettings", kvp.Value, out path) && !String.IsNullOrEmpty(path))
+					{
+						AndroidEnv.Add(kvp.Key, path);
+					}
+				}
 			}
 
-			if (OperatingSystem.IsMacOS())
+			// fill in from environment variables not already set
+			foreach (KeyValuePair<string, string> kvp in EnvVarNames)
+			{
+				if (!AndroidEnv.ContainsKey(kvp.Key))
+				{
+					string? envValue = Environment.GetEnvironmentVariable(kvp.Key);
+					if (!String.IsNullOrEmpty(envValue))
+					{
+						AndroidEnv.Add(kvp.Key, envValue);
+					}
+				}
+			}
+				
+			// If we are not running on Windows and we are still missing a key then go and find it from the .bash_profile
+			if (!RuntimePlatform.IsWindows && !EnvVarNames.All(s => AndroidEnv.ContainsKey(s.Key)))
 			{
-				Dictionary<string, string> AndroidEnv = new Dictionary<string, string>();
-				Dictionary<string, string> EnvVarNames = new Dictionary<string, string> {
-														 {"ANDROID_HOME", "SDKPath"},
-														 {"NDKROOT", "NDKPath"},
-														 {"JAVA_HOME", "JavaPath"}
-														 };
-
-				string BashProfilePath = Path.Combine(Environment.GetFolderPath(Environment.SpecialFolder.Personal), ".bash_profile");
+				string BashProfilePath = Path.Combine(Environment.GetFolderPath(Environment.SpecialFolder.UserProfile), ".bash_profile");
 				if (!File.Exists(BashProfilePath))
 				{
 					// Try .bashrc if didn't fine .bash_profile
-					BashProfilePath = Path.Combine(Environment.GetFolderPath(Environment.SpecialFolder.Personal), ".bashrc");
+					BashProfilePath = Path.Combine(Environment.GetFolderPath(Environment.SpecialFolder.UserProfile), ".bashrc");
+
+					// Lastly try .zprofile
+					if (!File.Exists(BashProfilePath))
+					{
+						BashProfilePath = Path.Combine(Environment.GetFolderPath(Environment.SpecialFolder.UserProfile), ".zprofile");
+					}
 				}
 				if (File.Exists(BashProfilePath))
 				{
@@ -60,6 +86,7 @@ namespace UnrealBuildTool
 							{
 								continue;
 							}
+
 							if (BashProfileContents[LineIndex].StartsWith("export " + kvp.Key + "="))
 							{
 								string PathVar = BashProfileContents[LineIndex].Split('=')[1].Replace("\"", "");
@@ -68,20 +95,34 @@ namespace UnrealBuildTool
 						}
 					}
 				}
+			}
 
-				// Set for the process
-				foreach (KeyValuePair<string, string> kvp in AndroidEnv)
-				{
-					Environment.SetEnvironmentVariable(kvp.Key, kvp.Value);
-				}
+			// Set for the process
+			foreach (KeyValuePair<string, string> kvp in AndroidEnv)
+			{
+				Environment.SetEnvironmentVariable(kvp.Key, kvp.Value);
+			}
+		}
 
-				NDKPath = Environment.GetEnvironmentVariable("NDKROOT");
-				if (!String.IsNullOrEmpty(NDKPath))
-				{
-					return NDKPath.Replace("\"", "");
-				}
+		public static string? GetNDKRoot()
+		{
+			string? NDKPath = Environment.GetEnvironmentVariable("NDKROOT");
+
+			// don't register if we don't have an NDKROOT specified
+			if (!String.IsNullOrEmpty(NDKPath))
+			{
+				return NDKPath.Replace("\"", "");
 			}
 
+			// set environment variables (don't check config)
+			CheckEnvironmentVariables(false);
+
+			// try NDKROOT again
+			NDKPath = Environment.GetEnvironmentVariable("NDKROOT");
+			if (!String.IsNullOrEmpty(NDKPath))
+			{
+				return NDKPath.Replace("\"", "");
+			}
 			return null;
 		}
 
@@ -95,8 +136,6 @@ namespace UnrealBuildTool
 				return null;
 			}
 
-			NDKPath = NDKPath.Replace("\"", "");
-
 			// figure out the NDK version
 			string? NDKToolchainVersion = null;
 			string SourcePropFilename = Path.Combine(NDKPath, "source.properties");
@@ -244,78 +283,11 @@ namespace UnrealBuildTool
 				return true;
 			}
 
+			// Set environment variables as needed (allow config overrides)
+			CheckEnvironmentVariables(true);
+			
 			string? NDKPath = Environment.GetEnvironmentVariable("NDKROOT");
 			string? JavaPath = Environment.GetEnvironmentVariable("JAVA_HOME");
-			{
-				ConfigHierarchy configCacheIni = ConfigCache.ReadHierarchy(ConfigHierarchyType.Engine, null, BuildHostPlatform.Current.Platform);
-				Dictionary<string, string> AndroidEnv = new Dictionary<string, string>();
-
-				Dictionary<string, string> EnvVarNames = new Dictionary<string, string> {
-														 {"ANDROID_HOME", "SDKPath"},
-														 {"NDKROOT", "NDKPath"},
-														 {"JAVA_HOME", "JavaPath"}
-														 };
-
-				string path;
-				foreach (KeyValuePair<string, string> kvp in EnvVarNames)
-				{
-					if (GetPath(configCacheIni, "/Script/AndroidPlatformEditor.AndroidSDKSettings", kvp.Value, out path) && !String.IsNullOrEmpty(path))
-					{
-						AndroidEnv.Add(kvp.Key, path);
-					}
-					else
-					{
-						string? envValue = Environment.GetEnvironmentVariable(kvp.Key);
-						if (!String.IsNullOrEmpty(envValue))
-						{
-							AndroidEnv.Add(kvp.Key, envValue);
-						}
-					}
-				}
-
-				// If we are not running on Windows and we are still missing a key then go and find it from the .bash_profile
-				if (!RuntimePlatform.IsWindows && !EnvVarNames.All(s => AndroidEnv.ContainsKey(s.Key)))
-				{
-					string BashProfilePath = Path.Combine(Environment.GetFolderPath(Environment.SpecialFolder.Personal), ".bash_profile");
-					if (!File.Exists(BashProfilePath))
-					{
-						// Try .bashrc if didn't fine .bash_profile
-						BashProfilePath = Path.Combine(Environment.GetFolderPath(Environment.SpecialFolder.Personal), ".bashrc");
-					}
-					if (File.Exists(BashProfilePath))
-					{
-						string[] BashProfileContents = File.ReadAllLines(BashProfilePath);
-
-						// Walk backwards so we keep the last export setting instead of the first
-						for (int LineIndex = BashProfileContents.Length - 1; LineIndex >= 0; --LineIndex)
-						{
-							foreach (KeyValuePair<string, string> kvp in EnvVarNames)
-							{
-								if (AndroidEnv.ContainsKey(kvp.Key))
-								{
-									continue;
-								}
-
-								if (BashProfileContents[LineIndex].StartsWith("export " + kvp.Key + "="))
-								{
-									string PathVar = BashProfileContents[LineIndex].Split('=')[1].Replace("\"", "");
-									AndroidEnv.Add(kvp.Key, PathVar);
-								}
-							}
-						}
-					}
-				}
-
-				// Set for the process
-				foreach (KeyValuePair<string, string> kvp in AndroidEnv)
-				{
-					Environment.SetEnvironmentVariable(kvp.Key, kvp.Value);
-				}
-
-				// See if we have an NDK path now...
-				AndroidEnv.TryGetValue("NDKROOT", out NDKPath);
-				AndroidEnv.TryGetValue("JAVA_HOME", out JavaPath);
-			}
 
 			// we don't have an NDKROOT specified
 			if (String.IsNullOrEmpty(NDKPath))
