Obsolete raw cppcompilewarningsettings properties.

diff --git a/Engine/Source/Programs/UnrealBuildTool/Configuration/ModuleRules.cs b/Engine/Source/Programs/UnrealBuildTool/Configuration/ModuleRules.cs
--- a/Engine/Source/Programs/UnrealBuildTool/Configuration/ModuleRules.cs
+++ b/Engine/Source/Programs/UnrealBuildTool/Configuration/ModuleRules.cs
@@ -944,6 +944,7 @@ namespace UnrealBuildTool
 		/// <summary>
 		/// How to treat shadow variable warnings
 		/// </summary>
+		[Obsolete("Deprecated in UE5.6 - Replace with ModuleRules.CppCompileWarningSettings.ShadowVariableWarningLevel")]
 		public WarningLevel ShadowVariableWarningLevel
 		{
 			get => CppCompileWarningSettings.ShadowVariableWarningLevel;
@@ -958,6 +959,7 @@ namespace UnrealBuildTool
 		/// <summary>
 		/// How to treat unsafe implicit type cast warnings (e.g., double->float or int64->int32)
 		/// </summary>
+		[Obsolete("Deprecated in UE5.6 - Replace with ModuleRules.CppCompileWarningSettings.UnsafeTypeCastWarningLevel")]
 		public WarningLevel UnsafeTypeCastWarningLevel
 		{
 			get => CppCompileWarningSettings.UnsafeTypeCastWarningLevel;
@@ -968,6 +970,7 @@ namespace UnrealBuildTool
 		/// <summary>
 		/// Indicates what warning/error level to treat undefined identifiers in conditional expressions.
 		/// </summary>
+		[Obsolete("Deprecated in UE5.6 - Replace with ModuleRules.CppCompileWarningSettings.UndefinedIdentifierWarningLevel")]
 		public WarningLevel UndefinedIdentifierWarningLevel
 		{
 			get => CppCompileWarningSettings.UndefinedIdentifierWarningLevel;
@@ -983,6 +986,7 @@ namespace UnrealBuildTool
 		/// <summary>
 		/// How to treat general module include path validation messages
 		/// </summary>
+		[Obsolete("Deprecated in UE5.6 - Replace with ModuleRules.CppCompileWarningSettings.ModuleIncludePathWarningLevel")]
 		public WarningLevel ModuleIncludePathWarningLevel
 		{
 			get => CppCompileWarningSettings.ModuleIncludePathWarningLevel;
@@ -992,6 +996,7 @@ namespace UnrealBuildTool
 		/// <summary>
 		/// How to treat private module include path validation messages, where a module is adding an include path that exposes private headers
 		/// </summary>
+		[Obsolete("Deprecated in UE5.6 - Replace with ModuleRules.CppCompileWarningSettings.ModuleIncludePrivateWarningLevel")]
 		public WarningLevel ModuleIncludePrivateWarningLevel
 		{
 			get => CppCompileWarningSettings.ModuleIncludePrivateWarningLevel;
@@ -1001,6 +1006,7 @@ namespace UnrealBuildTool
 		/// <summary>
 		/// How to treat unnecessary module sub-directory include path validation messages
 		/// </summary>
+		[Obsolete("Deprecated in UE5.6 - Replace with ModuleRules.CppCompileWarningSettings.ModuleIncludeSubdirectoryWarningLevel")]
 		public WarningLevel ModuleIncludeSubdirectoryWarningLevel
 		{
 			get => CppCompileWarningSettings.ModuleIncludeSubdirectoryWarningLevel;
diff --git a/Engine/Source/Programs/UnrealBuildTool/Configuration/ReadOnlyTargetRules.cs b/Engine/Source/Programs/UnrealBuildTool/Configuration/ReadOnlyTargetRules.cs
--- a/Engine/Source/Programs/UnrealBuildTool/Configuration/ReadOnlyTargetRules.cs
+++ b/Engine/Source/Programs/UnrealBuildTool/Configuration/ReadOnlyTargetRules.cs
@@ -361,22 +361,29 @@ namespace UnrealBuildTool
 
 		public ReadOnlyCppCompileWarnings CppCompileWarningSettings { get; init; }
 
-		public WarningLevel ShadowVariableWarningLevel => Inner.ShadowVariableWarningLevel;
+		[Obsolete("Deprecated in UE5.6 - Replace with ReadOnlyTargetRules.CppCompileWarningSettings.ShadowVariableWarningLevel")]
+		public WarningLevel ShadowVariableWarningLevel => Inner.CppCompileWarningSettings.ShadowVariableWarningLevel;
 
-		public WarningLevel UnsafeTypeCastWarningLevel => Inner.UnsafeTypeCastWarningLevel;
-		
-		public WarningLevel UndefinedIdentifierWarningLevel => Inner.UndefinedIdentifierWarningLevel;
+		[Obsolete("Deprecated in UE5.6 - Replace with ReadOnlyTargetRules.CppCompileWarningSettings.UnsafeTypeCastWarningLevel")]
+		public WarningLevel UnsafeTypeCastWarningLevel => Inner.CppCompileWarningSettings.UnsafeTypeCastWarningLevel;
+
+		[Obsolete("Deprecated in UE5.6 - Replace with ReadOnlyTargetRules.CppCompileWarningSettings.UndefinedIdentifierWarningLevel")]
+		public WarningLevel UndefinedIdentifierWarningLevel => Inner.CppCompileWarningSettings.UndefinedIdentifierWarningLevel;
 
 		[Obsolete("Deprecated in UE5.5 - Replace with ModuleRules.UndefinedIdentifierWarningLevel")]
 		public bool bUndefinedIdentifierErrors => Inner.bUndefinedIdentifierErrors;
 
-		public WarningLevel PCHPerformanceIssueWarningLevel => Inner.PCHPerformanceIssueWarningLevel;
+		[Obsolete("Deprecated in UE5.6 - Replace with ReadOnlyTargetRules.CppCompileWarningSettings.PCHPerformanceIssueWarningLevel")]
+		public WarningLevel PCHPerformanceIssueWarningLevel => Inner.CppCompileWarningSettings.PCHPerformanceIssueWarningLevel;
 
-		public WarningLevel ModuleIncludePathWarningLevel => Inner.ModuleIncludePathWarningLevel;
+		[Obsolete("Deprecated in UE5.6 - Replace with ReadOnlyTargetRules.CppCompileWarningSettings.ModuleIncludePathWarningLevel")]
+		public WarningLevel ModuleIncludePathWarningLevel => Inner.CppCompileWarningSettings.ModuleIncludePathWarningLevel;
 
-		public WarningLevel ModuleIncludePrivateWarningLevel => Inner.ModuleIncludePrivateWarningLevel;
+		[Obsolete("Deprecated in UE5.6 - Replace with ReadOnlyTargetRules.CppCompileWarningSettings.ModuleIncludePrivateWarningLevel")]
+		public WarningLevel ModuleIncludePrivateWarningLevel => Inner.CppCompileWarningSettings.ModuleIncludePrivateWarningLevel;
 
-		public WarningLevel ModuleIncludeSubdirectoryWarningLevel => Inner.ModuleIncludeSubdirectoryWarningLevel;
+		[Obsolete("Deprecated in UE5.6 - Replace with ReadOnlyTargetRules.CppCompileWarningSettings.ModuleIncludeSubdirectoryWarningLevel")]
+		public WarningLevel ModuleIncludeSubdirectoryWarningLevel => Inner.CppCompileWarningSettings.ModuleIncludeSubdirectoryWarningLevel;
 
 		public bool bShowIncludes => Inner.bShowIncludes;
 
diff --git a/Engine/Source/Programs/UnrealBuildTool/Configuration/TargetRules.cs b/Engine/Source/Programs/UnrealBuildTool/Configuration/TargetRules.cs
--- a/Engine/Source/Programs/UnrealBuildTool/Configuration/TargetRules.cs
+++ b/Engine/Source/Programs/UnrealBuildTool/Configuration/TargetRules.cs
@@ -1835,6 +1835,7 @@ namespace UnrealBuildTool
 		/// <summary>
 		/// Forces shadow variable warnings to be treated as errors on platforms that support it.
 		/// </summary>
+		[Obsolete("Deprecated in UE5.6 - Replace with TargetRules.CppCompileWarningSettings.ShadowVariableWarningLevel")]
 		public WarningLevel ShadowVariableWarningLevel
 		{
 			get => CppCompileWarningSettings.ShadowVariableWarningLevel;
@@ -1852,6 +1853,7 @@ namespace UnrealBuildTool
 		/// <summary>
 		/// Indicates what warning/error level to treat unsafe type casts as on platforms that support it (e.g., double->float or int64->int32)
 		/// </summary>
+		[Obsolete("Deprecated in UE5.6 - Replace with TargetRules.CppCompileWarningSettings.UnsafeTypeCastWarningLevel")]
 		public WarningLevel UnsafeTypeCastWarningLevel 
 		{
 			get => CppCompileWarningSettings.UnsafeTypeCastWarningLevel;
@@ -1868,6 +1870,7 @@ namespace UnrealBuildTool
 		/// <summary>
 		/// Indicates what warning/error level to treat undefined identifiers in conditional expressions.
 		/// </summary>
+		[Obsolete("Deprecated in UE5.6 - Replace with TargetRules.CppCompileWarningSettings.UndefinedIdentifierWarningLevel")]
 		public WarningLevel UndefinedIdentifierWarningLevel 
 		{
 			get => CppCompileWarningSettings.UndefinedIdentifierWarningLevel;
@@ -1877,6 +1880,7 @@ namespace UnrealBuildTool
 		/// <summary>
 		/// Indicates what warning/error level to treat potential PCH performance issues.
 		/// </summary>
+		[Obsolete("Deprecated in UE5.6 - Replace with TargetRules.CppCompileWarningSettings.PCHPerformanceIssueWarningLevel")]
 		public WarningLevel PCHPerformanceIssueWarningLevel
 		{
 			get => CppCompileWarningSettings.PCHPerformanceIssueWarningLevel;
@@ -1886,6 +1890,7 @@ namespace UnrealBuildTool
 		/// <summary>
 		/// How to treat general module include path validation messages
 		/// </summary>
+		[Obsolete("Deprecated in UE5.6 - Replace with TargetRules.CppCompileWarningSettings.ModuleIncludePathWarningLevel")]
 		public WarningLevel ModuleIncludePathWarningLevel
 		{
 			get => CppCompileWarningSettings.ModuleIncludePathWarningLevel;
@@ -1895,6 +1900,7 @@ namespace UnrealBuildTool
 		/// <summary>
 		/// How to treat private module include path validation messages, where a module is adding an include path that exposes private headers
 		/// </summary>
+		[Obsolete("Deprecated in UE5.6 - Replace with TargetRules.CppCompileWarningSettings.ModuleIncludePrivateWarningLevel")]
 		public WarningLevel ModuleIncludePrivateWarningLevel
 		{
 			get => CppCompileWarningSettings.ModuleIncludePrivateWarningLevel;
@@ -1904,6 +1910,7 @@ namespace UnrealBuildTool
 		/// <summary>
 		/// How to treat unnecessary module sub-directory include path validation messages
 		/// </summary>
+		[Obsolete("Deprecated in UE5.6 - Replace with TargetRules.CppCompileWarningSettings.ModuleIncludeSubdirectoryWarningLevel")]
 		public WarningLevel ModuleIncludeSubdirectoryWarningLevel
 		{
 			get => CppCompileWarningSettings.ModuleIncludeSubdirectoryWarningLevel;
@@ -3414,7 +3421,7 @@ namespace UnrealBuildTool
 				if (BuildSettingsVersion.V2 <= latestVersion && DefaultBuildSettings < BuildSettingsVersion.V2)
 				{
 					modifiedSettings.Add(Tuple.Create(String.Format("{0} = false", nameof(bLegacyPublicIncludePaths)), "Omits subfolders from public include paths to reduce compiler command line length. (Previously: true)."));
-					modifiedSettings.Add(Tuple.Create(String.Format("{0} = WarningLevel.Error", nameof(ShadowVariableWarningLevel)), "Treats shadowed variable warnings as errors. (Previously: WarningLevel.Warning)."));
+					modifiedSettings.Add(Tuple.Create(String.Format("{0} = WarningLevel.Error", nameof(CppCompileWarningSettings.ShadowVariableWarningLevel)), "Treats shadowed variable warnings as errors. (Previously: WarningLevel.Warning)."));
 					modifiedSettings.Add(Tuple.Create(String.Format("{0} = PCHUsageMode.UseExplicitOrSharedPCHs", nameof(ModuleRules.PCHUsage)), "Set in build.cs files to enables IWYU-style PCH model. See https://docs.unrealengine.com/en-US/Programming/BuildTools/UnrealBuildTool/IWYU/index.html. (Previously: PCHUsageMode.UseSharedPCHs)."));
 				}
 
diff --git a/Engine/Source/Programs/UnrealBuildTool/Configuration/UEBuildTarget.cs b/Engine/Source/Programs/UnrealBuildTool/Configuration/UEBuildTarget.cs
--- a/Engine/Source/Programs/UnrealBuildTool/Configuration/UEBuildTarget.cs
+++ b/Engine/Source/Programs/UnrealBuildTool/Configuration/UEBuildTarget.cs
@@ -4079,11 +4079,11 @@ namespace UnrealBuildTool
 					if (SharedPCH.Instances.Any())
 					{
 						// UnsafeTypeCastWarningLevel
-						foreach (PrecompiledHeaderInstance UnsafeTypeCastWarningInstance in SharedPCH.Instances.Where(instance => instance.Modules.First().Rules.UnsafeTypeCastWarningLevel == WarningLevel.Warning))
+						foreach (PrecompiledHeaderInstance UnsafeTypeCastWarningInstance in SharedPCH.Instances.Where(instance => instance.Modules.First().Rules.CppCompileWarningSettings.UnsafeTypeCastWarningLevel == WarningLevel.Warning))
 						{
 							CppCompileEnvironment UpdatedCppCompileEnvironment = new CppCompileEnvironment(UnsafeTypeCastWarningInstance.CompileEnvironment);
 							UpdatedCppCompileEnvironment.CppCompileWarnings.UnsafeTypeCastWarningLevel = WarningLevel.Error;
-							foreach (PrecompiledHeaderInstance UnsafeTypeCastErrorInstance in SharedPCH.Instances.Where(instance => instance.Modules.First().Rules.UnsafeTypeCastWarningLevel == WarningLevel.Error))
+							foreach (PrecompiledHeaderInstance UnsafeTypeCastErrorInstance in SharedPCH.Instances.Where(instance => instance.Modules.First().Rules.CppCompileWarningSettings.UnsafeTypeCastWarningLevel == WarningLevel.Error))
 							{
 								if (UEBuildModuleCPP.IsCompatibleForSharedPCH(UnsafeTypeCastErrorInstance.CompileEnvironment, UpdatedCppCompileEnvironment))
 								{
