UnrealBuildTool: Initial experimental Clang VFS overlay support

#rnx

diff --git a/Engine/Source/Programs/UnrealBuildTool/Configuration/ReadOnlyTargetRules.cs b/Engine/Source/Programs/UnrealBuildTool/Configuration/ReadOnlyTargetRules.cs
--- a/Engine/Source/Programs/UnrealBuildTool/Configuration/ReadOnlyTargetRules.cs
+++ b/Engine/Source/Programs/UnrealBuildTool/Configuration/ReadOnlyTargetRules.cs
@@ -423,6 +423,9 @@ namespace UnrealBuildTool
 		public bool bUsePCHFiles => Inner.bUsePCHFiles;
 
 		public bool bDeterministic => Inner.bDeterministic;
+
+		public bool bCompilerVFS => Inner.bCompilerVFS;
+
 		public bool bChainPCHs => Inner.bChainPCHs;
 
 		public bool bForceIncludePCHHeadersForGenCppFilesWhenPCHIsDisabled => Inner.bForceIncludePCHHeadersForGenCppFilesWhenPCHIsDisabled;
diff --git a/Engine/Source/Programs/UnrealBuildTool/Configuration/TargetRules.cs b/Engine/Source/Programs/UnrealBuildTool/Configuration/TargetRules.cs
--- a/Engine/Source/Programs/UnrealBuildTool/Configuration/TargetRules.cs
+++ b/Engine/Source/Programs/UnrealBuildTool/Configuration/TargetRules.cs
@@ -2079,6 +2079,14 @@ namespace UnrealBuildTool
 		}
 		private bool? bDeterministicPrivate;
 
+		/// <summary>
+		/// Experimental work in progress: Set flags to use toolchain virtual file system support, to generate consistent pathing in outputs
+		/// </summary>
+		[CommandLine("-CompilerVFS", Value = "true")]
+		[XmlConfigFile(Category = "BuildConfiguration")]
+		[RequiresUniqueBuildEnvironment]
+		public bool bCompilerVFS { get; set; }
+
 		/// <summary>
 		/// Whether PCHs should be chained when compiling with clang.
 		/// </summary>
diff --git a/Engine/Source/Programs/UnrealBuildTool/Configuration/UEBuildModuleCPP.cs b/Engine/Source/Programs/UnrealBuildTool/Configuration/UEBuildModuleCPP.cs
--- a/Engine/Source/Programs/UnrealBuildTool/Configuration/UEBuildModuleCPP.cs
+++ b/Engine/Source/Programs/UnrealBuildTool/Configuration/UEBuildModuleCPP.cs
@@ -6,6 +6,7 @@ using System.Collections.Generic;
 using System.IO;
 using System.Linq;
 using System.Text;
+using System.Text.Json;
 using EpicGames.Core;
 using Microsoft.Extensions.Logging;
 using UnrealBuildBase;
@@ -357,6 +358,34 @@ namespace UnrealBuildTool
 			base.AddModuleToCompileEnvironment(SourceModule, SourceBinary, IncludePaths, SystemIncludePaths, ModuleInterfacePaths, Definitions, AdditionalFrameworks, AdditionalPrerequisites, bLegacyPublicIncludePaths, bLegacyParentIncludePaths);
 		}
 
+		private FileItem CreateVfsOverlayFile(CppCompileEnvironment compileEnvironment, DirectoryReference outputDir, IActionGraphBuilder graph)
+		{
+			FileItem OverlayFile = FileItem.GetItemByFileReference(FileReference.Combine(outputDir, $"vfs.json"));
+
+			List<Dictionary<string, object>> jsonRoots = new();
+			foreach (KeyValuePair<CppRootPathFolder, DirectoryReference> item in compileEnvironment.RootPaths)
+			{
+				jsonRoots.Add(new()
+				{
+					["name"] = $"{CppCompileEnvironment.ClangVfsRootPath}/{item.Key}",
+					["type"] = "directory-remap",
+					["external-contents"] = item.Value.FullName.Replace('\\', '/')
+				});
+			}
+
+			Dictionary<string, object> json = new()
+			{
+				["version"] = 0,
+				["use-external-names"] = false,
+				["roots"] = jsonRoots
+			};
+
+			JsonSerializerOptions options = new() { WriteIndented = true };
+			string jsonStr = JsonSerializer.Serialize(json, options);
+			graph.CreateIntermediateTextFile(OverlayFile, jsonStr);
+			return OverlayFile;
+		}
+
 		// UEBuildModule interface.
 		public override List<FileItem> Compile(ReadOnlyTargetRules Target, UEToolChain ToolChain, CppCompileEnvironment BinaryCompileEnvironment, ISourceFileWorkingSet WorkingSet, IActionGraphBuilder Graph, ILogger Logger)
 		{
@@ -366,6 +395,14 @@ namespace UnrealBuildTool
 
 			CppCompileEnvironment ModuleCompileEnvironment = CreateModuleCompileEnvironment(Target, BinaryCompileEnvironment, Logger);
 
+			// Setup clang virtual file system overlay file if enabled
+			// TODO: Support clang-cl
+			if (Target.bCompilerVFS && !Target.Platform.IsInGroup(UnrealPlatformGroup.Microsoft))
+			{
+				ModuleCompileEnvironment.VfsOverlayFile = CreateVfsOverlayFile(ModuleCompileEnvironment, IntermediateDirectory, Graph);
+				ModuleCompileEnvironment.AdditionalPrerequisites.Add(ModuleCompileEnvironment.VfsOverlayFile);
+			}
+
 			// If the module is precompiled, read the object files from the manifest
 			if (Rules.bUsePrecompiled && Target.LinkType == TargetLinkType.Monolithic)
 			{
diff --git a/Engine/Source/Programs/UnrealBuildTool/Platform/Android/AndroidToolChain.cs b/Engine/Source/Programs/UnrealBuildTool/Platform/Android/AndroidToolChain.cs
--- a/Engine/Source/Programs/UnrealBuildTool/Platform/Android/AndroidToolChain.cs
+++ b/Engine/Source/Programs/UnrealBuildTool/Platform/Android/AndroidToolChain.cs
@@ -448,10 +448,10 @@ namespace UnrealBuildTool
 			// remove paths that are not meant for this architecture
 			// @todo can remove this when we only add paths properly for architecture
 			IEnumerable<DirectoryReference> FilteredPaths = CompileEnvironment.UserIncludePaths.Where(x => IsDirectoryForArch(x.FullName, CompileEnvironment.Architecture));
-			Arguments.AddRange(FilteredPaths.Select(IncludePath => GetUserIncludePathArgument(IncludePath)));
+			Arguments.AddRange(FilteredPaths.Select(IncludePath => GetUserIncludePathArgument(IncludePath, CompileEnvironment)));
 
 			FilteredPaths = CompileEnvironment.SystemIncludePaths.Where(x => IsDirectoryForArch(x.FullName, CompileEnvironment.Architecture));
-			Arguments.AddRange(FilteredPaths.Select(IncludePath => GetSystemIncludePathArgument(IncludePath)));
+			Arguments.AddRange(FilteredPaths.Select(IncludePath => GetSystemIncludePathArgument(IncludePath, CompileEnvironment)));
 		}
 
 		/// <inheritdoc/>
diff --git a/Engine/Source/Programs/UnrealBuildTool/Platform/Linux/LinuxToolChain.cs b/Engine/Source/Programs/UnrealBuildTool/Platform/Linux/LinuxToolChain.cs
--- a/Engine/Source/Programs/UnrealBuildTool/Platform/Linux/LinuxToolChain.cs
+++ b/Engine/Source/Programs/UnrealBuildTool/Platform/Linux/LinuxToolChain.cs
@@ -519,8 +519,8 @@ namespace UnrealBuildTool
 			if (ShouldUseLibcxx())
 			{
 				Arguments.Add("-nostdinc++");
-				Arguments.Add(GetSystemIncludePathArgument(DirectoryReference.Combine(Unreal.EngineSourceDirectory, "ThirdParty", "Unix", "LibCxx", "include")));
-				Arguments.Add(GetSystemIncludePathArgument(DirectoryReference.Combine(Unreal.EngineSourceDirectory, "ThirdParty", "Unix", "LibCxx", "include", "c++", "v1")));
+				Arguments.Add(GetSystemIncludePathArgument(DirectoryReference.Combine(Unreal.EngineSourceDirectory, "ThirdParty", "Unix", "LibCxx", "include"), CompileEnvironment));
+				Arguments.Add(GetSystemIncludePathArgument(DirectoryReference.Combine(Unreal.EngineSourceDirectory, "ThirdParty", "Unix", "LibCxx", "include", "c++", "v1"), CompileEnvironment));
 			}
 
 			if (CompilerVersionGreaterOrEqual(12, 0, 0))
@@ -585,7 +585,7 @@ namespace UnrealBuildTool
 			if (CrossCompiling())
 			{
 				Arguments.Add($"-target {CompileEnvironment.Architecture.LinuxName}");        // Set target triple
-				Arguments.Add($"--sysroot=\"{NormalizeCommandLinePath(LinuxInfo.BaseLinuxPath!)}\"");
+				Arguments.Add($"--sysroot=\"{NormalizeCommandLinePath(LinuxInfo.BaseLinuxPath!, CompileEnvironment)}\"");
 			}
 		}
 
diff --git a/Engine/Source/Programs/UnrealBuildTool/Platform/Mac/MacToolChain.cs b/Engine/Source/Programs/UnrealBuildTool/Platform/Mac/MacToolChain.cs
--- a/Engine/Source/Programs/UnrealBuildTool/Platform/Mac/MacToolChain.cs
+++ b/Engine/Source/Programs/UnrealBuildTool/Platform/Mac/MacToolChain.cs
@@ -207,7 +207,7 @@ namespace UnrealBuildTool
 				FileReference FrameworkPath = new FileReference(Path.GetFullPath(Framework.Name));
 				if (!FrameworksSearchPaths.Contains(FrameworkPath.Directory.FullName))
 				{
-					Arguments.Add($"-F \"{NormalizeCommandLinePath(FrameworkPath.Directory)}\"");
+					Arguments.Add($"-F \"{NormalizeCommandLinePath(FrameworkPath.Directory, CompileEnvironment)}\"");
 					FrameworksSearchPaths.Add(FrameworkPath.Directory.FullName);
 				}
 			}
diff --git a/Engine/Source/Programs/UnrealBuildTool/System/CppCompileEnvironment.cs b/Engine/Source/Programs/UnrealBuildTool/System/CppCompileEnvironment.cs
--- a/Engine/Source/Programs/UnrealBuildTool/System/CppCompileEnvironment.cs
+++ b/Engine/Source/Programs/UnrealBuildTool/System/CppCompileEnvironment.cs
@@ -5,6 +5,8 @@ using System.Linq;
 using System.IO.Enumeration;
 using EpicGames.Core;
 using UnrealBuildBase;
+using System;
+using System.Diagnostics.CodeAnalysis;
 
 namespace UnrealBuildTool
 {
@@ -591,6 +593,16 @@ namespace UnrealBuildTool
 		/// </summary>
 		public List<FileItem> ForceIncludeFiles = new List<FileItem>();
 
+		/// <summary>
+		/// Path to a virtual file system overlay file if in use
+		/// </summary>
+		public FileItem? VfsOverlayFile = null;
+
+		/// <summary>
+		/// The root of the clang vfs path
+		/// </summary>
+		public static readonly string ClangVfsRootPath = OperatingSystem.IsWindows() ? "Z:/UE" : "/UE";
+
 		/// <summary>
 		/// List of files that need to be up to date before compile can proceed
 		/// </summary>
@@ -844,6 +856,7 @@ namespace UnrealBuildTool
 			ModuleInterfacePaths = new HashSet<DirectoryReference>(Other.ModuleInterfacePaths);
 			bCheckSystemHeadersForModification = Other.bCheckSystemHeadersForModification;
 			ForceIncludeFiles.AddRange(Other.ForceIncludeFiles);
+			VfsOverlayFile = Other.VfsOverlayFile;
 			AdditionalPrerequisites.AddRange(Other.AdditionalPrerequisites);
 			ExtraGeneratedCPPFileTypes = Other.ExtraGeneratedCPPFileTypes;
 			CollidingNames = Other.CollidingNames;
@@ -894,5 +907,29 @@ namespace UnrealBuildTool
 
 			return false;
 		}
+
+		/// <summary>
+		/// Get a virtual filesystem overlay path if in use for a local path
+		/// </summary>
+		/// <param name="reference">The path to convert</param>
+		/// <param name="vfsPath">Out param, the converted path</param>
+		/// <returns>True if the local path can be mapped to a virtual path when enabled, false otherwise</returns>
+		public bool GetVfsOverlayPath(FileSystemReference reference, [NotNullWhen(true)] out string? vfsPath)
+		{
+			if (VfsOverlayFile != null)
+			{
+				foreach (KeyValuePair<CppRootPathFolder, DirectoryReference> item in RootPaths)
+				{
+					if (!reference.IsUnderDirectory(item.Value))
+					{
+						continue;
+					}
+					vfsPath = $"{ClangVfsRootPath}/{item.Key}/{reference.MakeRelativeTo(item.Value).Replace("\\", "/")}";
+					return true;
+				}
+			}
+			vfsPath = null;
+			return false;
+		}
 	}
 }
diff --git a/Engine/Source/Programs/UnrealBuildTool/System/Unity.cs b/Engine/Source/Programs/UnrealBuildTool/System/Unity.cs
--- a/Engine/Source/Programs/UnrealBuildTool/System/Unity.cs
+++ b/Engine/Source/Programs/UnrealBuildTool/System/Unity.cs
@@ -279,7 +279,11 @@ namespace UnrealBuildTool
 				foreach (FileItem CPPFile in UnityFile.Files)
 				{
 					string CPPFileString = CPPFile.AbsolutePath;
-					if (CPPFile.Location.IsUnderDirectory(Unreal.RootDirectory))
+					if (CompileEnvironment.GetVfsOverlayPath(CPPFile.Location, out string? vfsPath))
+					{
+						CPPFileString = vfsPath;
+					}
+					else if (CPPFile.Location.IsUnderDirectory(Unreal.RootDirectory))
 					{
 						CPPFileString = CPPFile.Location.MakeRelativeTo(Unreal.EngineSourceDirectory);
 					}
diff --git a/Engine/Source/Programs/UnrealBuildTool/ToolChain/AppleToolChain.cs b/Engine/Source/Programs/UnrealBuildTool/ToolChain/AppleToolChain.cs
--- a/Engine/Source/Programs/UnrealBuildTool/ToolChain/AppleToolChain.cs
+++ b/Engine/Source/Programs/UnrealBuildTool/ToolChain/AppleToolChain.cs
@@ -611,7 +611,7 @@ namespace UnrealBuildTool
 
 			if (CompileEnvironment.ForceIncludeFiles.Count() > 0)
 			{
-				Arguments.Add($"-import-objc-header \"{CompileEnvironment.ForceIncludeFiles.First()}\"");
+				Arguments.Add($"-import-objc-header \"{NormalizeCommandLinePath(CompileEnvironment.ForceIncludeFiles.First(), CompileEnvironment)}\"");
 			}
 
 			// tell swift to define PLATFORM_FOO, and tell C++ to define PLATFORM_FOO=1
diff --git a/Engine/Source/Programs/UnrealBuildTool/ToolChain/ClangToolChain.cs b/Engine/Source/Programs/UnrealBuildTool/ToolChain/ClangToolChain.cs
--- a/Engine/Source/Programs/UnrealBuildTool/ToolChain/ClangToolChain.cs
+++ b/Engine/Source/Programs/UnrealBuildTool/ToolChain/ClangToolChain.cs
@@ -502,20 +502,20 @@ namespace UnrealBuildTool
 			return CompileEnvironment.bUseRTTI ? "-frtti" : "-fno-rtti";
 		}
 
-		protected virtual string GetUserIncludePathArgument(DirectoryReference IncludePath)
+		protected virtual string GetUserIncludePathArgument(DirectoryReference IncludePath, CppCompileEnvironment CompileEnvironment)
 		{
-			return $"-I\"{NormalizeCommandLinePath(IncludePath)}\"";
+			return $"-I\"{NormalizeCommandLinePath(IncludePath, CompileEnvironment)}\"";
 		}
 
-		protected virtual string GetSystemIncludePathArgument(DirectoryReference IncludePath)
+		protected virtual string GetSystemIncludePathArgument(DirectoryReference IncludePath, CppCompileEnvironment? CompileEnvironment)
 		{
-			return $"-isystem\"{NormalizeCommandLinePath(IncludePath)}\"";
+			return $"-isystem\"{NormalizeCommandLinePath(IncludePath, CompileEnvironment)}\"";
 		}
 
 		protected virtual void GetCompileArguments_IncludePaths(CppCompileEnvironment CompileEnvironment, List<string> Arguments)
 		{
-			Arguments.AddRange(CompileEnvironment.UserIncludePaths.Select(IncludePath => GetUserIncludePathArgument(IncludePath)));
-			Arguments.AddRange(CompileEnvironment.SystemIncludePaths.Select(IncludePath => GetSystemIncludePathArgument(IncludePath)));
+			Arguments.AddRange(CompileEnvironment.UserIncludePaths.Select(IncludePath => GetUserIncludePathArgument(IncludePath, CompileEnvironment)));
+			Arguments.AddRange(CompileEnvironment.SystemIncludePaths.Select(IncludePath => GetSystemIncludePathArgument(IncludePath, CompileEnvironment)));
 
 			if (ShowIncludes)
 			{
@@ -538,24 +538,24 @@ namespace UnrealBuildTool
 			}
 		}
 
-		protected virtual string GetForceIncludeFileArgument(FileReference ForceIncludeFile)
+		protected virtual string GetForceIncludeFileArgument(FileReference ForceIncludeFile, CppCompileEnvironment CompileEnvironment)
 		{
-			return $"-include \"{NormalizeCommandLinePath(ForceIncludeFile)}\"";
+			return $"-include \"{NormalizeCommandLinePath(ForceIncludeFile, CompileEnvironment)}\"";
 		}
 
-		protected virtual string GetForceIncludeFileArgument(FileItem ForceIncludeFile)
+		protected virtual string GetForceIncludeFileArgument(FileItem ForceIncludeFile, CppCompileEnvironment CompileEnvironment)
 		{
-			return GetForceIncludeFileArgument(ForceIncludeFile.Location);
+			return GetForceIncludeFileArgument(ForceIncludeFile.Location, CompileEnvironment);
 		}
 
-		protected virtual string GetIncludePCHFileArgument(FileReference IncludePCHFile)
+		protected virtual string GetIncludePCHFileArgument(FileReference IncludePCHFile, CppCompileEnvironment CompileEnvironment)
 		{
-			return $"-include-pch \"{NormalizeCommandLinePath(IncludePCHFile)}\"";
+			return $"-include-pch \"{NormalizeCommandLinePath(IncludePCHFile, CompileEnvironment)}\"";
 		}
 
-		protected virtual string GetIncludePCHFileArgument(FileItem IncludePCHFile)
+		protected virtual string GetIncludePCHFileArgument(FileItem IncludePCHFile, CppCompileEnvironment CompileEnvironment)
 		{
-			return GetIncludePCHFileArgument(IncludePCHFile.Location);
+			return GetIncludePCHFileArgument(IncludePCHFile.Location, CompileEnvironment);
 		}
 
 		protected virtual void GetCompileArguments_ForceInclude(CppCompileEnvironment CompileEnvironment, List<string> Arguments)
@@ -564,14 +564,14 @@ namespace UnrealBuildTool
 			// This needs to be before the other force include paths to ensure clang uses it instead of the source header file.
 			if (CompileEnvironment.PrecompiledHeaderAction == PrecompiledHeaderAction.Include)
 			{
-				Arguments.Add(GetIncludePCHFileArgument(CompileEnvironment.PrecompiledHeaderFile!));
+				Arguments.Add(GetIncludePCHFileArgument(CompileEnvironment.PrecompiledHeaderFile!, CompileEnvironment));
 			}
 			else if (CompileEnvironment.PrecompiledHeaderAction == PrecompiledHeaderAction.Create && CompileEnvironment.ParentPCHInstance != null)
 			{
-				Arguments.Add(GetIncludePCHFileArgument(CompileEnvironment.ParentPrecompiledHeaderFile!));
+				Arguments.Add(GetIncludePCHFileArgument(CompileEnvironment.ParentPrecompiledHeaderFile!, CompileEnvironment));
 			}
 
-			Arguments.AddRange(CompileEnvironment.ForceIncludeFiles.Select(ForceIncludeFile => GetForceIncludeFileArgument(ForceIncludeFile)));
+			Arguments.AddRange(CompileEnvironment.ForceIncludeFiles.Select(ForceIncludeFile => GetForceIncludeFileArgument(ForceIncludeFile, CompileEnvironment)));
 		}
 
 		protected virtual string GetSourceFileArgument(FileItem SourceFile)
@@ -682,7 +682,7 @@ namespace UnrealBuildTool
 			{
 				Log.TraceInformationOnce("Enabling Profile Guided Optimization (PGO). Linking will take a while.");
 				DirectoryReference? PGODir = DirectoryReference.FromString(CompileEnvironment.PGODirectory!);
-				Arguments.Add($"-fprofile-instr-use=\"{NormalizeCommandLinePath(DirectoryReference.Combine(PGODir!, CompileEnvironment.PGOFilenamePrefix!))}\"");
+				Arguments.Add($"-fprofile-instr-use=\"{NormalizeCommandLinePath(DirectoryReference.Combine(PGODir!, CompileEnvironment.PGOFilenamePrefix!), CompileEnvironment)}\"");
 			}
 			else if (CompileEnvironment.bPGOProfile)
 			{
@@ -884,6 +884,14 @@ namespace UnrealBuildTool
 			}
 		}
 
+		protected virtual void GetCompileArguments_VfsOverlay(CppCompileEnvironment CompileEnvironment, List<string> Arguments)
+		{
+			if (CompileEnvironment.VfsOverlayFile != null)
+			{
+				Arguments.Add($"-Xclang -ivfsoverlay\"{NormalizeCommandLinePath(CompileEnvironment.VfsOverlayFile)}\"");
+			}
+		}
+
 		/// <summary>
 		/// Common compile arguments for all files in a module.
 		/// Override and call base.GetCompileArguments_Global() in derived classes.
@@ -948,6 +956,9 @@ namespace UnrealBuildTool
 		/// <returns>Path to the target file (such as .o)</returns>
 		protected virtual FileItem GetCompileArguments_FileType(CppCompileEnvironment CompileEnvironment, FileItem SourceFile, DirectoryReference OutputDir, List<string> Arguments, Action CompileAction, CPPOutput CompileResult)
 		{
+			// Add virtual file system overlay arguments if enabled. Will be inserted as the first entry in the list
+			GetCompileArguments_VfsOverlay(CompileEnvironment, Arguments);
+
 			// Add the additional response files
 			foreach (FileItem AdditionalResponseFile in CompileEnvironment.AdditionalResponseFiles)
 			{
diff --git a/Engine/Source/Programs/UnrealBuildTool/ToolChain/ISPCToolChain.cs b/Engine/Source/Programs/UnrealBuildTool/ToolChain/ISPCToolChain.cs
--- a/Engine/Source/Programs/UnrealBuildTool/ToolChain/ISPCToolChain.cs
+++ b/Engine/Source/Programs/UnrealBuildTool/ToolChain/ISPCToolChain.cs
@@ -330,7 +330,7 @@ namespace UnrealBuildTool
 		/// </summary>
 		/// <param name="Reference">The FileSystemReference to normalize</param>
 		/// <returns>Normalized path as a string</returns>
-		protected virtual string NormalizeCommandLinePath(FileSystemReference Reference)
+		protected static string NormalizeCommandLinePath(FileSystemReference Reference)
 		{
 			string path = Reference.FullName;
 			// Try to use a relative path to shorten command line length.
@@ -354,11 +354,37 @@ namespace UnrealBuildTool
 		/// </summary>
 		/// <param name="Item">The FileItem to normalize</param>
 		/// <returns>Normalized path as a string</returns>
-		protected virtual string NormalizeCommandLinePath(FileItem Item)
+		protected static string NormalizeCommandLinePath(FileItem Item)
 		{
 			return NormalizeCommandLinePath(Item.Location);
 		}
 
+		/// <summary>
+		/// Normalize a path for use in a command line, making it relative to Engine/Source if under the root directory
+		/// </summary>
+		/// <param name="Reference">The FileSystemReference to normalize</param>
+		/// <param name="CompileEnvironment">The environment to use for normalization</param>
+		/// <returns>Normalized path as a string</returns>
+		protected virtual string NormalizeCommandLinePath(FileSystemReference Reference, CppCompileEnvironment? CompileEnvironment = null)
+		{
+			if (CompileEnvironment?.GetVfsOverlayPath(Reference, out string? vfsPath) == true)
+			{
+				return vfsPath;
+			}
+			return NormalizeCommandLinePath(Reference);
+		}
+
+		/// <summary>
+		/// Normalize a path for use in a command line, making it relative if under the Root Directory
+		/// </summary>
+		/// <param name="Item">The FileItem to normalize</param>
+		/// <param name="CompileEnvironment">The environment to use for normalization</param>
+		/// <returns>Normalized path as a string</returns>
+		protected virtual string NormalizeCommandLinePath(FileItem Item, CppCompileEnvironment? CompileEnvironment = null)
+		{
+			return NormalizeCommandLinePath(Item.Location, CompileEnvironment);
+		}
+
 		protected override CPPOutput GenerateISPCHeaders(CppCompileEnvironment CompileEnvironment, IEnumerable<FileItem> InputFiles, DirectoryReference OutputDir, IActionGraphBuilder Graph)
 		{
 			CPPOutput Result = new CPPOutput();
@@ -447,7 +473,7 @@ namespace UnrealBuildTool
 				List<string> Arguments = new List<string>();
 
 				// Add the ISPC obj file as a prerequisite of the action.
-				Arguments.Add($"\"{NormalizeCommandLinePath(ISPCFile)}\"");
+				Arguments.Add($"\"{NormalizeCommandLinePath(ISPCFile, CompileEnvironment)}\"");
 
 				// Add the ISPC h file to the produced item list.
 				FileItem ISPCIncludeHeaderFile = FileItem.GetItemByFileReference(
@@ -485,11 +511,11 @@ namespace UnrealBuildTool
 				}
 
 				// Add the ISPC file to be compiled.
-				Arguments.Add($"-h \"{NormalizeCommandLinePath(ISPCIncludeHeaderFile)}\"");
+				Arguments.Add($"-h \"{NormalizeCommandLinePath(ISPCIncludeHeaderFile, CompileEnvironment)}\"");
 
 				// Generate the included header dependency list
 				FileItem DependencyListFile = FileItem.GetItemByFileReference(FileReference.Combine(OutputDir, Path.GetFileName(ISPCFile.AbsolutePath) + ".txt"));
-				Arguments.Add($"-MMM \"{NormalizeCommandLinePath(DependencyListFile)}\"");
+				Arguments.Add($"-MMM \"{NormalizeCommandLinePath(DependencyListFile, CompileEnvironment)}\"");
 				CompileAction.DependencyListFile = DependencyListFile;
 				CompileAction.ProducedItems.Add(DependencyListFile);
 
@@ -653,7 +679,7 @@ namespace UnrealBuildTool
 				List<string> Arguments = new List<string>();
 
 				// Add the ISPC file to be compiled.
-				Arguments.Add($"\"{NormalizeCommandLinePath(ISPCFile)}\"");
+				Arguments.Add($"\"{NormalizeCommandLinePath(ISPCFile, CompileEnvironment)}\"");
 
 				List<FileItem> CompiledISPCObjFiles = new List<FileItem>();
 
@@ -709,7 +735,7 @@ namespace UnrealBuildTool
 				}
 
 				// Add the output ISPC obj file
-				Arguments.Add($"-o \"{NormalizeCommandLinePath(CompiledISPCObjFileNoISA)}\"");
+				Arguments.Add($"-o \"{NormalizeCommandLinePath(CompiledISPCObjFileNoISA, CompileEnvironment)}\"");
 
 				// Generate the timing info
 				if (CompileEnvironment.bPrintTimingInfo)
@@ -768,10 +794,10 @@ namespace UnrealBuildTool
 							Action PostCompileAction = Graph.CreateAction(ActionType.Compile);
 
 							List<string> PostCompileArgs = new List<string>();
-							PostCompileArgs.Add($"\"{NormalizeCommandLinePath(CompiledBytecodeObjFile)}\"");
+							PostCompileArgs.Add($"\"{NormalizeCommandLinePath(CompiledBytecodeObjFile, CompileEnvironment)}\"");
 							PostCompileArgs.Add("-c");
 							PostCompileArgs.AddRange(CommonArgs);
-							PostCompileArgs.Add($"-o \"{NormalizeCommandLinePath(FinalCompiledISPCObjFile)}\"");
+							PostCompileArgs.Add($"-o \"{NormalizeCommandLinePath(FinalCompiledISPCObjFile, CompileEnvironment)}\"");
 
 							if (bMergeModules)
 							{
diff --git a/Engine/Source/Programs/UnrealBuildTool/ToolChain/IWYUToolChain.cs b/Engine/Source/Programs/UnrealBuildTool/ToolChain/IWYUToolChain.cs
--- a/Engine/Source/Programs/UnrealBuildTool/ToolChain/IWYUToolChain.cs
+++ b/Engine/Source/Programs/UnrealBuildTool/ToolChain/IWYUToolChain.cs
@@ -51,7 +51,7 @@ namespace UnrealBuildTool
 			// starting with clang 16.x the directory naming changed to include major version only
 			string ClangVersionString = (CompilerToolChainInfo.ClangVersion.Major >= 16) ? CompilerToolChainInfo.ClangVersion.Major.ToString() : CompilerToolChainInfo.ClangVersion.ToString();
 			DirectoryReference SystemPath = DirectoryReference.Combine(BaseLinuxPath!, "lib", "clang", ClangVersionString, "include");
-			CrossCompilingArguments.Add(GetSystemIncludePathArgument(SystemPath));
+			CrossCompilingArguments.Add(GetSystemIncludePathArgument(SystemPath, null));
 
 			string DevPath = ""; //@"include-what-you-use-0.19\vs_projects\bin\RelWithDebInfo";
 
@@ -129,18 +129,20 @@ namespace UnrealBuildTool
 			return new CPPOutput() { ObjectFiles = IwyuFiles };
 		}
 
-		// Include paths must be absolute because we have code gen that can produce includes that looks like this: "../../../Source/Foo/Bar.h"
-		// IWYU combines include path and include and then try to normalize the result. This could result in a path that can't be normalized because there are more
-		// ".." than there are directories to reduce. This silently fails in odd ways
+        // Include paths must be absolute because we have code gen that can produce includes that looks like this: "../../../Source/Foo/Bar.h"
+        // IWYU combines include path and include and then try to normalize the result. This could result in a path that can't be normalized because there are more
+        // ".." than there are directories to reduce. This silently fails in odd ways
 
-		protected override string NormalizeCommandLinePath(FileSystemReference Reference)
+        /// <inheritdoc/>
+        protected override string NormalizeCommandLinePath(FileSystemReference Reference, CppCompileEnvironment? CompileEnvironment)
 		{
 			return Reference.FullName.Replace('\\', '/');
 		}
 
-		protected override string NormalizeCommandLinePath(FileItem Item)
+        /// <inheritdoc/>
+		protected override string NormalizeCommandLinePath(FileItem Item, CppCompileEnvironment? CompileEnvironment)
 		{
-			return Item.FullName.Replace('\\', '/');
+			return NormalizeCommandLinePath(Item.Location, CompileEnvironment);
 		}
 
 		// Code below here is mostly copy-pasted from LinuxToolChain.
@@ -167,8 +169,8 @@ namespace UnrealBuildTool
 			if (ShouldUseLibcxx(CompileEnvironment.Architecture))
 			{
 				Arguments.Add("-nostdinc++");
-				Arguments.Add(GetSystemIncludePathArgument(DirectoryReference.Combine(Unreal.EngineSourceDirectory, "ThirdParty", "Unix", "LibCxx", "include")));
-				Arguments.Add(GetSystemIncludePathArgument(DirectoryReference.Combine(Unreal.EngineSourceDirectory, "ThirdParty", "Unix", "LibCxx", "include", "c++", "v1")));
+				Arguments.Add(GetSystemIncludePathArgument(DirectoryReference.Combine(Unreal.EngineSourceDirectory, "ThirdParty", "Unix", "LibCxx", "include"), CompileEnvironment));
+				Arguments.Add(GetSystemIncludePathArgument(DirectoryReference.Combine(Unreal.EngineSourceDirectory, "ThirdParty", "Unix", "LibCxx", "include", "c++", "v1"), CompileEnvironment));
 			}
 
 			Arguments.Add("-fno-math-errno");               // do not assume that math ops have side effects
