- Fixed the missing interop header line that got deleted

diff --git a/Engine/Source/Programs/UnrealBuildTool/Configuration/UEBuildModuleCPP.cs b/Engine/Source/Programs/UnrealBuildTool/Configuration/UEBuildModuleCPP.cs
--- a/Engine/Source/Programs/UnrealBuildTool/Configuration/UEBuildModuleCPP.cs
+++ b/Engine/Source/Programs/UnrealBuildTool/Configuration/UEBuildModuleCPP.cs
@@ -802,6 +802,7 @@ namespace UnrealBuildTool
 				// currently nothing is using this header, but it could be helpful to have in swift
 				CreateHeaderForDefinitions(SwiftCompileEnvironment, IntermediateDirectory, "swift", Graph);
 				SwiftCompileEnvironment.ForceIncludeFiles.Clear();
+				SwiftCompileEnvironment.ForceIncludeFiles.Add(FileItem.GetItemByPath(Rules.SwiftInteropHeader));
 
 				// build the swift files and bridging headers to expose swift to cpp
 				CPPOutput Output = ToolChain.CompileAllCPPFiles(SwiftCompileEnvironment, InputFiles.SwiftFiles, IntermediateDirectory, Name, Graph);
