- Swift lib paths

diff --git a/Engine/Source/Programs/UnrealBuildTool/Platform/IOS/IOSToolChain.cs b/Engine/Source/Programs/UnrealBuildTool/Platform/IOS/IOSToolChain.cs
--- a/Engine/Source/Programs/UnrealBuildTool/Platform/IOS/IOSToolChain.cs
+++ b/Engine/Source/Programs/UnrealBuildTool/Platform/IOS/IOSToolChain.cs
@@ -450,10 +450,6 @@ namespace UnrealBuildTool
 					// and add to the commandline
 					LinkCommandArguments += String.Format(" \\\"{0}\\\"", Library.FullName);
 				}
-
-				// allow for auto-linked swift libs to be found
-				// @todo Mac will need this too - a lot of this code should probably be moved up to AppleToolChain and shared
-				LinkCommandArguments += $" -L{Settings.GetSDKPath(LinkEnvironment.Architecture)}/usr/lib/swift";
 			}
 
 			// Handle additional framework assets that might need to be shadowed
