UnrealBuildTool: Fixed linux portable outputs with vfs

#rnx

diff --git a/Engine/Source/Programs/UnrealBuildTool/Configuration/UEBuildModuleCPP.cs b/Engine/Source/Programs/UnrealBuildTool/Configuration/UEBuildModuleCPP.cs
--- a/Engine/Source/Programs/UnrealBuildTool/Configuration/UEBuildModuleCPP.cs
+++ b/Engine/Source/Programs/UnrealBuildTool/Configuration/UEBuildModuleCPP.cs
@@ -8,6 +8,7 @@ using System.Linq;
 using System.Text;
 using System.Text.Json;
 using EpicGames.Core;
+using EpicGames.Horde.Jobs.Templates;
 using Microsoft.Extensions.Logging;
 using UnrealBuildBase;
 
@@ -358,9 +359,9 @@ namespace UnrealBuildTool
 			base.AddModuleToCompileEnvironment(SourceModule, SourceBinary, IncludePaths, SystemIncludePaths, ModuleInterfacePaths, Definitions, AdditionalFrameworks, AdditionalPrerequisites, bLegacyPublicIncludePaths, bLegacyParentIncludePaths);
 		}
 
-		private FileItem CreateVfsOverlayFile(CppCompileEnvironment compileEnvironment, DirectoryReference outputDir, IActionGraphBuilder graph)
+		private FileItem CreateVfsOverlayFile(CppCompileEnvironment compileEnvironment, DirectoryReference outputDir, string fileName, IActionGraphBuilder graph, bool allowAsync = true)
 		{
-			FileItem OverlayFile = FileItem.GetItemByFileReference(FileReference.Combine(outputDir, $"vfs.json"));
+			FileItem OverlayFile = FileItem.GetItemByFileReference(FileReference.Combine(outputDir, fileName));
 
 			List<Dictionary<string, object>> jsonRoots = new();
 			foreach (KeyValuePair<CppRootPathFolder, DirectoryReference> item in compileEnvironment.RootPaths)
@@ -382,7 +383,7 @@ namespace UnrealBuildTool
 
 			JsonSerializerOptions options = new() { WriteIndented = true };
 			string jsonStr = JsonSerializer.Serialize(json, options);
-			graph.CreateIntermediateTextFile(OverlayFile, jsonStr);
+			graph.CreateIntermediateTextFile(OverlayFile, jsonStr, allowAsync);
 			return OverlayFile;
 		}
 
@@ -396,10 +397,9 @@ namespace UnrealBuildTool
 			CppCompileEnvironment ModuleCompileEnvironment = CreateModuleCompileEnvironment(Target, BinaryCompileEnvironment, Logger);
 
 			// Setup clang virtual file system overlay file if enabled
-			// TODO: Support clang-cl
-			if (Target.bCompilerVFS && !Target.Platform.IsInGroup(UnrealPlatformGroup.Microsoft))
+			if (Target.bCompilerVFS && (!Target.Platform.IsInGroup(UnrealPlatformGroup.Microsoft) || Target.WindowsPlatform.Compiler.IsClang()))
 			{
-				ModuleCompileEnvironment.VfsOverlayFile = CreateVfsOverlayFile(ModuleCompileEnvironment, IntermediateDirectory, Graph);
+				ModuleCompileEnvironment.VfsOverlayFile = CreateVfsOverlayFile(ModuleCompileEnvironment, IntermediateDirectory, "vfs.json", Graph);
 				ModuleCompileEnvironment.AdditionalPrerequisites.Add(ModuleCompileEnvironment.VfsOverlayFile);
 			}
 
@@ -1364,14 +1364,14 @@ namespace UnrealBuildTool
 				// Create a suffix to distinguish this shared PCH variant from any others.
 				string Variant = GetSuffixForSharedPCH(ModuleCompileEnvironment, Template.BaseCompileEnvironment);
 
-				FileReference SharedDefinitionsLocation = FileReference.Combine(Template.OutputDir, String.Format("SharedDefinitions.{0}{1}.h", Template.Module.Name, Variant));
+				FileReference SharedDefinitionsLocation = FileReference.Combine(Template.OutputDir, $"SharedDefinitions.{Template.Module.Name}{Variant}.h");
 				List<string> NewDefinitions = new();
 				StringBuilder Writer = new StringBuilder();
-				WriteDefinitions($"Shared Definitions for {Template.Module.Name}{Variant}", Array.Empty<string>(), Definitions, Array.Empty<string>(), Writer);
+				WriteDefinitions($"Shared Definitions for {Template.Module.Name}{Variant}", [], Definitions, [], Writer);
 				FileItem SharedDefinitionsFileItem = Graph.CreateIntermediateTextFile(SharedDefinitionsLocation, Writer.ToString(), AllowAsync: false);
 
 				// Create the wrapper file, which sets all the definitions needed to compile it
-				FileReference WrapperLocation = FileReference.Combine(Template.OutputDir, String.Format("SharedPCH.{0}{1}.h", Template.Module.Name, Variant));
+				FileReference WrapperLocation = FileReference.Combine(Template.OutputDir, $"SharedPCH.{Template.Module.Name}{Variant}.h");
 				FileItem WrapperFile = CreatePCHWrapperFile(WrapperLocation, SharedDefinitionsFileItem, Template.HeaderFile, Graph);
 
 				// Create the compile environment for this PCH
@@ -1381,6 +1381,13 @@ namespace UnrealBuildTool
 				CompileEnvironment.PrecompiledHeaderIncludeFilename = WrapperFile.Location;
 				CopySettingsForSharedPCH(ModuleCompileEnvironment, CompileEnvironment);
 
+				// Setup clang virtual file system overlay file if enabled
+				if (Rules.Target.bCompilerVFS && (!Rules.Target.Platform.IsInGroup(UnrealPlatformGroup.Microsoft) || Rules.Target.WindowsPlatform.Compiler.IsClang()))
+				{
+					CompileEnvironment.VfsOverlayFile = CreateVfsOverlayFile(Template.BaseCompileEnvironment, IntermediateDirectory, $"SharedVFS.{Template.Module.Name}{Variant}.json", Graph, allowAsync: false);
+					CompileEnvironment.AdditionalPrerequisites.Add(CompileEnvironment.VfsOverlayFile);
+				}
+
 				// Setup PCH chaining
 				PrecompiledHeaderInstance? ParentPCHInstance = null;
 				Dictionary<string, string>? DefinitionsDictionary = null;
diff --git a/Engine/Source/Programs/UnrealBuildTool/Platform/Windows/VCToolChain.cs b/Engine/Source/Programs/UnrealBuildTool/Platform/Windows/VCToolChain.cs
--- a/Engine/Source/Programs/UnrealBuildTool/Platform/Windows/VCToolChain.cs
+++ b/Engine/Source/Programs/UnrealBuildTool/Platform/Windows/VCToolChain.cs
@@ -415,6 +415,20 @@ namespace UnrealBuildTool
 				{
 					AddDefinition(Arguments, "__clang_analyzer__");
 				}
+
+				if (CompileEnvironment.VfsOverlayFile != null)
+				{
+					Arguments.Add($"-Xclang -ivfsoverlay\"{NormalizeCommandLinePath(CompileEnvironment.VfsOverlayFile)}\"");
+					Arguments.Add("-fcoverage-compilation-dir=.");
+					Arguments.Add("-fdebug-compilation-dir=.");
+					Arguments.Add("-vctoolsdir=.");
+					Arguments.Add("-winsdkdir=.");
+					Arguments.Add("-winsysroot=.");
+					Arguments.Add("-Xclang -gno-codeview-command-line");
+					Arguments.Add("-fintegrated-cc1");
+					DirectoryReference ResourceDir = DirectoryReference.Combine(EnvVars.CompilerDir!, "lib", "clang", EnvVars.CompilerVersion.Components[0].ToString());
+					Arguments.Add($"-resource-dir=\"{NormalizeCommandLinePath(ResourceDir, CompileEnvironment)}\"");
+				}
 			}
 			else if (Target.StaticAnalyzer == StaticAnalyzer.Default && !CompileEnvironment.bDisableStaticAnalysis)
 			{
diff --git a/Engine/Source/Programs/UnrealBuildTool/ToolChain/ClangToolChain.cs b/Engine/Source/Programs/UnrealBuildTool/ToolChain/ClangToolChain.cs
--- a/Engine/Source/Programs/UnrealBuildTool/ToolChain/ClangToolChain.cs
+++ b/Engine/Source/Programs/UnrealBuildTool/ToolChain/ClangToolChain.cs
@@ -512,6 +512,11 @@ namespace UnrealBuildTool
 			return $"-isystem\"{NormalizeCommandLinePath(IncludePath, CompileEnvironment)}\"";
 		}
 
+		protected virtual DirectoryReference GetResourceDir(CppCompileEnvironment? CompileEnvironment)
+		{
+			return DirectoryReference.Combine(Info.BasePath!, "lib", "clang", Info.ClangVersion.Major.ToString());
+		}
+
 		protected virtual void GetCompileArguments_IncludePaths(CppCompileEnvironment CompileEnvironment, List<string> Arguments)
 		{
 			Arguments.AddRange(CompileEnvironment.UserIncludePaths.Select(IncludePath => GetUserIncludePathArgument(IncludePath, CompileEnvironment)));
@@ -574,9 +579,11 @@ namespace UnrealBuildTool
 			Arguments.AddRange(CompileEnvironment.ForceIncludeFiles.Select(ForceIncludeFile => GetForceIncludeFileArgument(ForceIncludeFile, CompileEnvironment)));
 		}
 
-		protected virtual string GetSourceFileArgument(FileItem SourceFile)
+		protected virtual string GetSourceFileArgument(FileItem SourceFile, CppCompileEnvironment CompileEnvironment)
 		{
-			return $"\"{NormalizeCommandLinePath(SourceFile)}\"";
+			// Only pch seems to be ok with getting a vfs path.. so don't provide an environment if not creating pch
+			CppCompileEnvironment? Env = CompileEnvironment.PrecompiledHeaderAction == PrecompiledHeaderAction.Create ? CompileEnvironment : null;
+			return $"\"{NormalizeCommandLinePath(SourceFile, Env)}\"";
 		}
 
 		protected virtual string GetOutputFileArgument(FileItem OutputFile)
@@ -888,6 +895,7 @@ namespace UnrealBuildTool
 		{
 			if (CompileEnvironment.VfsOverlayFile != null)
 			{
+				Arguments.Add($"-resource-dir=\"{NormalizeCommandLinePath(GetResourceDir(CompileEnvironment), CompileEnvironment)}\"");
 				Arguments.Add($"-Xclang -ivfsoverlay\"{NormalizeCommandLinePath(CompileEnvironment.VfsOverlayFile)}\"");
 			}
 		}
@@ -1070,7 +1078,7 @@ namespace UnrealBuildTool
 			CompileAction.ProducedItems.Add(OutputFile);
 
 			// Add the source file path to the command-line.
-			Arguments.Add(GetSourceFileArgument(SourceFile));
+			Arguments.Add(GetSourceFileArgument(SourceFile, CompileEnvironment));
 
 			// Generate the included header dependency list
 			if (!PreprocessDepends)
