UnrealBuildTool: Remove unnecessary using

#rnx

diff --git a/Engine/Source/Programs/UnrealBuildTool/Configuration/UEBuildModuleCPP.cs b/Engine/Source/Programs/UnrealBuildTool/Configuration/UEBuildModuleCPP.cs
--- a/Engine/Source/Programs/UnrealBuildTool/Configuration/UEBuildModuleCPP.cs
+++ b/Engine/Source/Programs/UnrealBuildTool/Configuration/UEBuildModuleCPP.cs
@@ -8,7 +8,6 @@ using System.Linq;
 using System.Text;
 using System.Text.Json;
 using EpicGames.Core;
-using EpicGames.Horde.Jobs.Templates;
 using Microsoft.Extensions.Logging;
 using UnrealBuildBase;
 
