Add a -AutoRTFMClosedStaticLinkage option to UBT and pipe that into an underlying verse-clang-cl flag we are experimenting with.

#rnx

diff --git a/Engine/Source/Programs/UnrealBuildTool/Configuration/ReadOnlyTargetRules.cs b/Engine/Source/Programs/UnrealBuildTool/Configuration/ReadOnlyTargetRules.cs
--- a/Engine/Source/Programs/UnrealBuildTool/Configuration/ReadOnlyTargetRules.cs
+++ b/Engine/Source/Programs/UnrealBuildTool/Configuration/ReadOnlyTargetRules.cs
@@ -140,6 +140,7 @@ namespace UnrealBuildTool
 		public bool bUseAutoRTFMCompiler => Inner.bUseAutoRTFMCompiler;
 		public bool bUseAutoRTFMVerifier => Inner.bUseAutoRTFMVerifier;
 		public bool bAutoRTFMVerify => Inner.bAutoRTFMVerify;
+		public bool bAutoRTFMClosedStaticLinkage => Inner.bAutoRTFMClosedStaticLinkage;
 
 		[Obsolete("Deprecated in UE5.1 - No longer used as Chaos is always enabled.")]
 		public bool bCompileChaos => Inner.bCompileChaos;
diff --git a/Engine/Source/Programs/UnrealBuildTool/Configuration/TargetRules.cs b/Engine/Source/Programs/UnrealBuildTool/Configuration/TargetRules.cs
--- a/Engine/Source/Programs/UnrealBuildTool/Configuration/TargetRules.cs
+++ b/Engine/Source/Programs/UnrealBuildTool/Configuration/TargetRules.cs
@@ -922,6 +922,14 @@ namespace UnrealBuildTool
 		[XmlConfigFile(Category = "BuildConfiguration")]
 		public bool bAutoRTFMVerify { get; set; } = false;
 
+		/// <summary>
+		/// Whether to link closed function declarations statically.
+		/// </summary>
+		[RequiresUniqueBuildEnvironment]
+		[CommandLine("-AutoRTFMClosedStaticLinkage", Value = "true")]
+		[XmlConfigFile(Category = "BuildConfiguration")]
+		public bool bAutoRTFMClosedStaticLinkage { get; set; } = false;
+
 		/// <summary>
 		/// Whether to compile the Chaos physics plugin.
 		/// </summary>
diff --git a/Engine/Source/Programs/UnrealBuildTool/Configuration/UEBuildModuleCPP.cs b/Engine/Source/Programs/UnrealBuildTool/Configuration/UEBuildModuleCPP.cs
--- a/Engine/Source/Programs/UnrealBuildTool/Configuration/UEBuildModuleCPP.cs
+++ b/Engine/Source/Programs/UnrealBuildTool/Configuration/UEBuildModuleCPP.cs
@@ -1667,6 +1667,7 @@ namespace UnrealBuildTool
 			CompileEnvironment.bDisableAutoRTFMInstrumentation = ModuleCompileEnvironment.bDisableAutoRTFMInstrumentation;
 			CompileEnvironment.bEnableAutoRTFMVerification = ModuleCompileEnvironment.bEnableAutoRTFMVerification;
 			CompileEnvironment.bAutoRTFMVerify = ModuleCompileEnvironment.bAutoRTFMVerify;
+			CompileEnvironment.bAutoRTFMClosedStaticLinkage = ModuleCompileEnvironment.bAutoRTFMClosedStaticLinkage;
 			CompileEnvironment.bValidateFormatStrings = ModuleCompileEnvironment.bValidateFormatStrings;
 			CompileEnvironment.bValidateInternalApi = ModuleCompileEnvironment.bValidateInternalApi;
 			CompileEnvironment.bValidateExperimentalApi = ModuleCompileEnvironment.bValidateExperimentalApi;
@@ -2271,6 +2272,7 @@ namespace UnrealBuildTool
 				Result.bDisableAutoRTFMInstrumentation = Rules.bDisableAutoRTFMInstrumentation;
 				Result.bEnableAutoRTFMVerification = Target.bUseAutoRTFMVerifier;
 				Result.bAutoRTFMVerify = Target.bAutoRTFMVerify;
+				Result.bAutoRTFMClosedStaticLinkage = Target.bAutoRTFMClosedStaticLinkage;
 			}
 
 			if (Result.OptimizationLevel != Rules.OptimizationLevel)
diff --git a/Engine/Source/Programs/UnrealBuildTool/Platform/Windows/VCToolChain.cs b/Engine/Source/Programs/UnrealBuildTool/Platform/Windows/VCToolChain.cs
--- a/Engine/Source/Programs/UnrealBuildTool/Platform/Windows/VCToolChain.cs
+++ b/Engine/Source/Programs/UnrealBuildTool/Platform/Windows/VCToolChain.cs
@@ -1344,6 +1344,11 @@ namespace UnrealBuildTool
 				{
 					Arguments.Add("-Xclang -mllvm -Xclang -autortfm-verify");
 				}
+
+				if (CompileEnvironment.bAutoRTFMClosedStaticLinkage)
+				{
+					Arguments.Add("-Xclang -mllvm -Xclang -autortfm-closed-static-linkage");
+				}
 			}
 
 			if (Target.WindowsPlatform.Compiler == WindowsCompiler.Intel)
@@ -1395,6 +1400,11 @@ namespace UnrealBuildTool
 				{
 					Arguments.Add("-Xclang -mllvm -Xclang -autortfm-verify");
 				}
+
+				if (CompileEnvironment.bAutoRTFMClosedStaticLinkage)
+				{
+					Arguments.Add("-Xclang -mllvm -Xclang -autortfm-closed-static-linkage");
+				}
 			}
 		}
 
diff --git a/Engine/Source/Programs/UnrealBuildTool/System/CppCompileEnvironment.cs b/Engine/Source/Programs/UnrealBuildTool/System/CppCompileEnvironment.cs
--- a/Engine/Source/Programs/UnrealBuildTool/System/CppCompileEnvironment.cs
+++ b/Engine/Source/Programs/UnrealBuildTool/System/CppCompileEnvironment.cs
@@ -327,6 +327,11 @@ namespace UnrealBuildTool
 		/// </summary>
 		public bool bAutoRTFMVerify = false;
 
+		/// <summary>
+		/// Whether we should link closed function declarations statically.
+		/// </summary>
+		public bool bAutoRTFMClosedStaticLinkage = false;
+
 		/// <summary>
 		/// If unity builds are enabled this can be used to override if this specific module will build using Unity.
 		/// This is set using the per module configurations in BuildConfiguration.
@@ -825,6 +830,7 @@ namespace UnrealBuildTool
 			bDisableAutoRTFMInstrumentation = Other.bDisableAutoRTFMInstrumentation;
 			bEnableAutoRTFMVerification = Other.bEnableAutoRTFMVerification;
 			bAutoRTFMVerify = Other.bAutoRTFMVerify;
+			bAutoRTFMClosedStaticLinkage = Other.bAutoRTFMClosedStaticLinkage;
 			bCodeCoverage = Other.bCodeCoverage;
 			OptimizationLevel = Other.OptimizationLevel;
 			FPSemantics = Other.FPSemantics;
diff --git a/Engine/Source/Programs/UnrealBuildTool/ToolChain/ClangToolChain.cs b/Engine/Source/Programs/UnrealBuildTool/ToolChain/ClangToolChain.cs
--- a/Engine/Source/Programs/UnrealBuildTool/ToolChain/ClangToolChain.cs
+++ b/Engine/Source/Programs/UnrealBuildTool/ToolChain/ClangToolChain.cs
@@ -416,6 +416,11 @@ namespace UnrealBuildTool
 			{
 				Arguments.Add("-Xclang -mllvm -Xclang -autortfm-verify");
 			}
+
+			if (CompileEnvironment.bAutoRTFMClosedStaticLinkage)
+			{
+				Arguments.Add("-Xclang -mllvm -Xclang -autortfm-closed-static-linkage");
+			}
 		}
 
 		protected virtual void GetCStandardCompileArgument(CppCompileEnvironment CompileEnvironment, List<string> Arguments)
