- Missing library line for swift, and reenabled SwiftMain

diff --git a/Engine/Source/Programs/UnrealBuildTool/ToolChain/AppleToolChain.cs b/Engine/Source/Programs/UnrealBuildTool/ToolChain/AppleToolChain.cs
--- a/Engine/Source/Programs/UnrealBuildTool/ToolChain/AppleToolChain.cs
+++ b/Engine/Source/Programs/UnrealBuildTool/ToolChain/AppleToolChain.cs
@@ -629,6 +629,8 @@ namespace UnrealBuildTool
 
 			// pass the UE definitions to C++ land for the imported obj-c header
 			Arguments.AddRange(CompileEnvironment.Definitions.Select(x => $"-Xcc -D{x}"));
+
+			Console.WriteLine("Arguments: {0}", string.Join(", ", Arguments));
 		}
 
 		private Action CompileSwiftFile(CppCompileEnvironment CompileEnvironment, FileItem SourceFile, DirectoryReference OutputDir, IActionGraphBuilder Graph, CPPOutput CompileResult)
@@ -815,6 +817,9 @@ namespace UnrealBuildTool
 			{
 				Arguments.Add(" -ld_classic");
 			}
+
+			Arguments.Add($"-L{ApplePlatformSDK.GetToolchainDirectory()}/../lib/swift/{Settings.PlatformDirName}");
+			Arguments.Add("-L/usr/lib/swift");
 		}
 
 		#region Stub Xcode Projects
diff --git a/Engine/Source/Runtime/Launch/Launch.Build.cs b/Engine/Source/Runtime/Launch/Launch.Build.cs
--- a/Engine/Source/Runtime/Launch/Launch.Build.cs
+++ b/Engine/Source/Runtime/Launch/Launch.Build.cs
@@ -8,7 +8,7 @@ using EpicGames.Core;
 
 public class Launch : ModuleRules
 {
-	protected virtual bool bAllowSwiftMain => false;
+	protected virtual bool bAllowSwiftMain => Target.IsInPlatformGroup(UnrealPlatformGroup.Apple) && Target.Type != TargetType.Editor;
 	
 	public Launch(ReadOnlyTargetRules Target) : base(Target)
 	{
@@ -242,7 +242,7 @@ public class Launch : ModuleRules
 			}
 		}
 
-		if (bAllowSwiftMain)//Target.IsInPlatformGroup(UnrealPlatformGroup.Apple) && Target.Type != TargetType.Editor)
+		if (bAllowSwiftMain)
 		{
 			SwiftInteropHeader = Path.Combine(ModuleDirectory, "Private/Apple/LaunchCPPToSwift.h");
 		}
