- Reenabled 10.9 min target setting for Mac. UnrealEditorServices has too many deprecation warnings without it

diff --git a/Engine/Source/Programs/UnrealBuildTool/ToolChain/AppleToolChain.cs b/Engine/Source/Programs/UnrealBuildTool/ToolChain/AppleToolChain.cs
--- a/Engine/Source/Programs/UnrealBuildTool/ToolChain/AppleToolChain.cs
+++ b/Engine/Source/Programs/UnrealBuildTool/ToolChain/AppleToolChain.cs
@@ -136,11 +136,11 @@ namespace UnrealBuildTool
 		/// <param name="Platform"></param>
 		/// <param name="TargetType"></param>
 		/// <returns></returns>
-		public virtual string GetTargetTuple(UnrealArch Architecture, UnrealTargetPlatform Platform, TargetType TargetType)
+		public virtual string GetTargetTuple(UnrealArch Architecture, UnrealTargetPlatform Platform, TargetType TargetType, string? ForcedVersion=null)
 		{
 			string Prefix = Architecture.AppleName;
 			string Suffix = (Architecture == UnrealArch.IOSSimulator || Architecture == UnrealArch.TVOSSimulator) ? "-simulator" : "";
-			string TargetVersion = GetTargetVersionForTargetType(TargetType);
+			string TargetVersion = ForcedVersion ?? GetTargetVersionForTargetType(TargetType);
 
 			return $"{Prefix}-apple-{TargetOSName}{TargetVersion}{Suffix}";
 		}
@@ -575,7 +575,8 @@ namespace UnrealBuildTool
 			Arguments.Add("-fmessage-length=0");
 			Arguments.Add("-fpascal-strings");
 
-			Arguments.Add($"-target {ToolChainSettings.Value.GetTargetTuple(CompileEnvironment.Architecture, Target!.Platform, Target!.Type)}");
+			string? Override = (CompileEnvironment.bEnableOSX109Support && CompileEnvironment.Platform == UnrealTargetPlatform.Mac) ? "10.9" : null;
+			Arguments.Add($"-target {ToolChainSettings.Value.GetTargetTuple(CompileEnvironment.Architecture, Target!.Platform, Target!.Type, Override)}");
 		}
 
 		public override CppCompileEnvironment CreateSharedResponseFile(CppCompileEnvironment CompileEnvironment, FileReference OutResponseFile, IActionGraphBuilder Graph)
