- Added missing param comment

diff --git a/Engine/Source/Programs/UnrealBuildTool/ToolChain/AppleToolChain.cs b/Engine/Source/Programs/UnrealBuildTool/ToolChain/AppleToolChain.cs
--- a/Engine/Source/Programs/UnrealBuildTool/ToolChain/AppleToolChain.cs
+++ b/Engine/Source/Programs/UnrealBuildTool/ToolChain/AppleToolChain.cs
@@ -135,6 +135,7 @@ namespace UnrealBuildTool
 		/// <param name="Architecture"></param>
 		/// <param name="Platform"></param>
 		/// <param name="TargetType"></param>
+		/// <param name="ForcedVersion"></param>
 		/// <returns></returns>
 		public virtual string GetTargetTuple(UnrealArch Architecture, UnrealTargetPlatform Platform, TargetType TargetType, string? ForcedVersion=null)
 		{
