UnrealBuildTool: Initial experimental Clang-Cl VFS overlay support

#rnx

diff --git a/Engine/Source/Programs/UnrealBuildTool/Platform/Windows/VCCompileAction.cs b/Engine/Source/Programs/UnrealBuildTool/Platform/Windows/VCCompileAction.cs
--- a/Engine/Source/Programs/UnrealBuildTool/Platform/Windows/VCCompileAction.cs
+++ b/Engine/Source/Programs/UnrealBuildTool/Platform/Windows/VCCompileAction.cs
@@ -2,6 +2,7 @@
 
 using System;
 using System.Collections.Generic;
+using System.Linq;
 using EpicGames.Core;
 using Microsoft.Extensions.Logging;
 using UnrealBuildBase;
@@ -114,6 +115,16 @@ namespace UnrealBuildTool
 		/// </summary>
 		public FileItem? PchThroughHeaderFile { get; set; }
 
+		/// <summary>
+		/// Response file for the compiler
+		/// </summary>
+		public FileItem? VfsOverlayFile { get; set; }
+
+		/// <summary>
+		/// List of additional response paths
+		/// </summary>
+		public List<FileItem> AdditionalResponseFiles { get; } = new List<FileItem>();
+
 		/// <summary>
 		/// List of include paths
 		/// </summary>
@@ -181,7 +192,7 @@ namespace UnrealBuildTool
 		/// <summary>
 		/// Root paths for this action (generally engine root project root, toolchain root, sdk root)
 		/// </summary>
-		public List<DirectoryItem> RootPaths { get; } = new List<DirectoryItem>();
+		public SortedDictionary<CppRootPathFolder, DirectoryReference> RootPaths { get; } = new SortedDictionary<CppRootPathFolder, DirectoryReference>();
 
 		/// <inheritdoc/>
 		public bool bCanExecuteRemotely { get; set; }
@@ -214,7 +225,7 @@ namespace UnrealBuildTool
 		#region Implementation of IAction
 
 		IEnumerable<FileItem> IExternalAction.DeleteItems => DeleteItems;
-		IEnumerable<DirectoryItem> IExternalAction.RootPaths => RootPaths;
+		IEnumerable<DirectoryItem> IExternalAction.RootPaths => RootPaths.Values.Select(x => DirectoryItem.GetItemByDirectoryReference(x));
 		public DirectoryReference WorkingDirectory => Unreal.EngineSourceDirectory;
 		string IExternalAction.CommandDescription
 		{
@@ -255,6 +266,18 @@ namespace UnrealBuildTool
 				{
 					yield return UsingPchFile;
 				}
+				if (VfsOverlayFile != null)
+				{
+					yield return VfsOverlayFile;
+				}
+				foreach (FileItem AdditionalResponseFile in AdditionalResponseFiles)
+				{
+					yield return AdditionalResponseFile;
+				}
+				foreach (FileItem ForceIncludeFile in ForceIncludeFiles)
+				{
+					yield return ForceIncludeFile;
+				}
 				foreach (FileItem AdditionalPrerequisiteItem in AdditionalPrerequisiteItems)
 				{
 					yield return AdditionalPrerequisiteItem;
@@ -383,6 +406,8 @@ namespace UnrealBuildTool
 			CreatePchFile = InAction.CreatePchFile;
 			UsingPchFile = InAction.UsingPchFile;
 			PchThroughHeaderFile = InAction.PchThroughHeaderFile;
+			VfsOverlayFile = InAction.VfsOverlayFile;
+			AdditionalResponseFiles = new List<FileItem>(InAction.AdditionalResponseFiles);
 			IncludePaths = new List<DirectoryReference>(InAction.IncludePaths);
 			SystemIncludePaths = new List<DirectoryReference>(InAction.SystemIncludePaths);
 			Definitions = new List<string>(InAction.Definitions);
@@ -400,7 +425,7 @@ namespace UnrealBuildTool
 			AdditionalPrerequisiteItems = new List<FileItem>(InAction.AdditionalPrerequisiteItems);
 			AdditionalProducedItems = new List<FileItem>(InAction.AdditionalProducedItems);
 			DeleteItems = new List<FileItem>(InAction.DeleteItems);
-			RootPaths = new List<DirectoryItem>(InAction.RootPaths);
+			RootPaths = new SortedDictionary<CppRootPathFolder, DirectoryReference>(InAction.RootPaths, null);
 		}
 
 		/// <summary>
@@ -427,6 +452,8 @@ namespace UnrealBuildTool
 			CreatePchFile = Reader.ReadFileItem();
 			UsingPchFile = Reader.ReadFileItem();
 			PchThroughHeaderFile = Reader.ReadFileItem();
+			VfsOverlayFile = Reader.ReadFileItem();
+			AdditionalResponseFiles = Reader.ReadList(() => Reader.ReadFileItem())!;
 			IncludePaths = Reader.ReadList(() => Reader.ReadCompactDirectoryReference())!;
 			SystemIncludePaths = Reader.ReadList(() => Reader.ReadCompactDirectoryReference())!;
 			Definitions = Reader.ReadList(() => Reader.ReadString())!;
@@ -446,7 +473,7 @@ namespace UnrealBuildTool
 			AdditionalPrerequisiteItems = Reader.ReadList(() => Reader.ReadFileItem())!;
 			AdditionalProducedItems = Reader.ReadList(() => Reader.ReadFileItem())!;
 			DeleteItems = Reader.ReadList(() => Reader.ReadFileItem())!;
-			RootPaths = Reader.ReadList(() => Reader.ReadDirectoryItem())!;
+			RootPaths = new(Reader.ReadDictionary(() => (CppRootPathFolder)Reader.ReadInt(), () => Reader.ReadCompactDirectoryReference())!, null);
 		}
 
 		/// <inheritdoc/>
@@ -470,6 +497,8 @@ namespace UnrealBuildTool
 			Writer.WriteFileItem(CreatePchFile);
 			Writer.WriteFileItem(UsingPchFile);
 			Writer.WriteFileItem(PchThroughHeaderFile);
+			Writer.WriteFileItem(VfsOverlayFile);
+			Writer.WriteList(AdditionalResponseFiles, Item => Writer.WriteFileItem(Item));
 			Writer.WriteList(IncludePaths, Item => Writer.WriteCompactDirectoryReference(Item));
 			Writer.WriteList(SystemIncludePaths, Item => Writer.WriteCompactDirectoryReference(Item));
 			Writer.WriteList(Definitions, Item => Writer.WriteString(Item));
@@ -489,7 +518,7 @@ namespace UnrealBuildTool
 			Writer.WriteList(AdditionalPrerequisiteItems, Item => Writer.WriteFileItem(Item));
 			Writer.WriteList(AdditionalProducedItems, Item => Writer.WriteFileItem(Item));
 			Writer.WriteList(DeleteItems, Item => Writer.WriteFileItem(Item));
-			Writer.WriteList(RootPaths, Item => Writer.WriteDirectoryItem(Item));
+			Writer.WriteDictionary(RootPaths, Key => Writer.WriteInt((int)Key), Value => Writer.WriteCompactDirectoryReference(Value));
 		}
 
 		/// <summary>
@@ -508,20 +537,27 @@ namespace UnrealBuildTool
 
 		public List<string> GetCompilerArguments(ILogger Logger)
 		{
+			IReadOnlyDictionary<CppRootPathFolder, DirectoryReference>? rootPaths = CompilerType.IsClang() && VfsOverlayFile != null ? RootPaths : null;
 			List<string> Arguments = new List<string>();
+
 			if (SourceFile != null)
 			{
 				VCToolChain.AddSourceFile(Arguments, SourceFile);
 			}
 
+			foreach (FileItem AdditionalRsp in AdditionalResponseFiles)
+			{
+				VCToolChain.AddResponseFile(Arguments, AdditionalRsp);
+			}
+
 			foreach (DirectoryReference IncludePath in IncludePaths)
 			{
-				VCToolChain.AddIncludePath(Arguments, IncludePath, CompilerType);
+				VCToolChain.AddIncludePath(Arguments, IncludePath, CompilerType, rootPaths);
 			}
 
 			foreach (DirectoryReference SystemIncludePath in SystemIncludePaths)
 			{
-				VCToolChain.AddSystemIncludePath(Arguments, SystemIncludePath, CompilerType);
+				VCToolChain.AddSystemIncludePath(Arguments, SystemIncludePath, CompilerType, rootPaths);
 			}
 
 			foreach (string Definition in Definitions)
@@ -533,7 +569,7 @@ namespace UnrealBuildTool
 
 			foreach (FileItem ForceIncludeFile in ForceIncludeFiles)
 			{
-				VCToolChain.AddForceIncludeFile(Arguments, ForceIncludeFile);
+				VCToolChain.AddForceIncludeFile(Arguments, ForceIncludeFile, rootPaths);
 			}
 
 			if (CreatePchFile != null)
@@ -543,12 +579,12 @@ namespace UnrealBuildTool
 
 			if (UsingPchFile != null && CompilerType.IsMSVC())
 			{
-				VCToolChain.AddUsingPchFile(Arguments, PchThroughHeaderFile!, UsingPchFile);
+				VCToolChain.AddUsingPchFile(Arguments, PchThroughHeaderFile!, UsingPchFile, rootPaths);
 			}
 
 			if (PreprocessedFile != null)
 			{
-				VCToolChain.AddPreprocessedFile(Arguments, PreprocessedFile, Logger);
+				VCToolChain.AddPreprocessedFile(Arguments, PreprocessedFile);
 			}
 
 			if (ObjectFile != null)
@@ -590,19 +626,17 @@ namespace UnrealBuildTool
 		{
 			if (ResponseFile == null)
 			{
-				return String.Join(" ", Arguments);
+				return String.Join(' ', Arguments);
 			}
-			else
-			{
-				string ResponseFileString = VCToolChain.NormalizeCommandLinePath(ResponseFile);
 
-				// cl.exe can't handle response files with a path longer than 260 characters, and relative paths can push it over the limit
-				if (!System.IO.Path.IsPathRooted(ResponseFileString) && System.IO.Path.Combine(WorkingDirectory.FullName, ResponseFileString).Length > 260)
-				{
-					ResponseFileString = ResponseFile.FullName;
-				}
-				return String.Format("@{0} {1}", Utils.MakePathSafeToUseWithCommandLine(ResponseFileString), String.Join(" ", Arguments));
+			string ResponseFileString = VCToolChain.NormalizeCommandLinePath(ResponseFile);
+
+			// cl.exe can't handle response files with a path longer than 260 characters, and relative paths can push it over the limit
+			if (!System.IO.Path.IsPathRooted(ResponseFileString) && System.IO.Path.Combine(WorkingDirectory.FullName, ResponseFileString).Length > 260)
+			{
+				ResponseFileString = ResponseFile.FullName;
 			}
+			return String.Format("@{0} {1}", Utils.MakePathSafeToUseWithCommandLine(ResponseFileString), String.Join(' ', Arguments));
 		}
 
 		string GetClFilterArguments()
@@ -627,7 +661,7 @@ namespace UnrealBuildTool
 			Arguments.Add(GetClArguments());
 			Arguments.Add("/showIncludes");
 
-			return String.Join(" ", Arguments);
+			return String.Join(' ', Arguments);
 		}
 	}
 
diff --git a/Engine/Source/Programs/UnrealBuildTool/Platform/Windows/VCToolChain.cs b/Engine/Source/Programs/UnrealBuildTool/Platform/Windows/VCToolChain.cs
--- a/Engine/Source/Programs/UnrealBuildTool/Platform/Windows/VCToolChain.cs
+++ b/Engine/Source/Programs/UnrealBuildTool/Platform/Windows/VCToolChain.cs
@@ -173,16 +173,43 @@ namespace UnrealBuildTool
 			return NormalizeCommandLinePath(Item.Location);
 		}
 
+		protected static string NormalizeCommandLinePath(FileSystemReference reference, IReadOnlyDictionary<CppRootPathFolder, DirectoryReference>? RootPaths)
+		{
+			if (RootPaths != null)
+			{
+				foreach (KeyValuePair<CppRootPathFolder, DirectoryReference> item in RootPaths)
+				{
+					if (!reference.IsUnderDirectory(item.Value))
+					{
+						continue;
+					}
+					return $"{CppCompileEnvironment.ClangVfsRootPath}/{item.Key}/{reference.MakeRelativeTo(item.Value).Replace("\\", "/")}";
+				}
+			}
+			return NormalizeCommandLinePath(reference);
+		}
+
+		protected static string NormalizeCommandLinePath(FileItem Item, IReadOnlyDictionary<CppRootPathFolder, DirectoryReference>? RootPaths)
+		{
+			return NormalizeCommandLinePath(Item.Location, RootPaths);
+		}
+
+		public static void AddResponseFile(List<string> Arguments, FileItem ResponseFile)
+		{
+			string ResponseFileString = NormalizeCommandLinePath(ResponseFile);
+			Arguments.Add($"@{Utils.MakePathSafeToUseWithCommandLine(ResponseFileString)}");
+		}
+
 		public static void AddSourceFile(List<string> Arguments, FileItem SourceFile)
 		{
 			string SourceFileString = NormalizeCommandLinePath(SourceFile);
 			Arguments.Add(Utils.MakePathSafeToUseWithCommandLine(SourceFileString));
 		}
 
-		private static void AddIncludePath(List<string> Arguments, DirectoryReference IncludePath, WindowsCompiler Compiler, bool bSystemInclude)
+		private static void AddIncludePath(List<string> Arguments, DirectoryReference IncludePath, WindowsCompiler Compiler, bool bSystemInclude, IReadOnlyDictionary<CppRootPathFolder, DirectoryReference>? RootPaths = null)
 		{
 			// Try to use a relative path to shorten command line length.
-			string IncludePathString = NormalizeCommandLinePath(IncludePath);
+			string IncludePathString = NormalizeCommandLinePath(IncludePath, RootPaths);
 
 			if (Compiler.IsClang() && bSystemInclude)
 			{
@@ -205,19 +232,19 @@ namespace UnrealBuildTool
 			}
 		}
 
-		public static void AddIncludePath(List<string> Arguments, DirectoryReference IncludePath, WindowsCompiler Compiler)
+		public static void AddIncludePath(List<string> Arguments, DirectoryReference IncludePath, WindowsCompiler Compiler, IReadOnlyDictionary<CppRootPathFolder, DirectoryReference>? RootPaths = null)
 		{
-			AddIncludePath(Arguments, IncludePath, Compiler, false);
+			AddIncludePath(Arguments, IncludePath, Compiler, false, RootPaths);
 		}
 
-		public static void AddSystemIncludePath(List<string> Arguments, DirectoryReference IncludePath, WindowsCompiler Compiler)
+		public static void AddSystemIncludePath(List<string> Arguments, DirectoryReference IncludePath, WindowsCompiler Compiler, IReadOnlyDictionary<CppRootPathFolder, DirectoryReference>? RootPaths = null)
 		{
-			AddIncludePath(Arguments, IncludePath, Compiler, true);
+			AddIncludePath(Arguments, IncludePath, Compiler, true, RootPaths);
 		}
 
-		public static void AddForceIncludeFile(List<string> Arguments, FileItem ForceIncludeFile)
+		public static void AddForceIncludeFile(List<string> Arguments, FileItem ForceIncludeFile, IReadOnlyDictionary<CppRootPathFolder, DirectoryReference>? RootPaths)
 		{
-			string ForceIncludeFileString = NormalizeCommandLinePath(ForceIncludeFile);
+			string ForceIncludeFileString = NormalizeCommandLinePath(ForceIncludeFile, RootPaths);
 			Arguments.Add($"/FI\"{ForceIncludeFileString}\"");
 		}
 
@@ -239,10 +266,10 @@ namespace UnrealBuildTool
 			}
 		}
 
-		public static void AddUsingPchFile(List<string> Arguments, FileItem PchThroughHeaderFile, FileItem UsingPchFile)
+		public static void AddUsingPchFile(List<string> Arguments, FileItem PchThroughHeaderFile, FileItem UsingPchFile, IReadOnlyDictionary<CppRootPathFolder, DirectoryReference>? RootPaths)
 		{
-			string PchThroughHeaderFilePath = NormalizeCommandLinePath(PchThroughHeaderFile);
-			string UsingPchFilePath = NormalizeCommandLinePath(UsingPchFile);
+			string PchThroughHeaderFilePath = NormalizeCommandLinePath(PchThroughHeaderFile, RootPaths);
+			string UsingPchFilePath = NormalizeCommandLinePath(UsingPchFile, RootPaths);
 			if (UsingPchFile.Name.EndsWith(".ifc"))
 			{
 				Arguments.Add("/wd4127"); // with header units pragma warning disable is not propagated and this warning is very spammy so explicitly turned off here. Hoping ms will change so warning disable is propgated from HU
@@ -256,7 +283,7 @@ namespace UnrealBuildTool
 			}
 		}
 
-		public static void AddPreprocessedFile(List<string> Arguments, FileItem PreprocessedFile, ILogger Logger)
+		public static void AddPreprocessedFile(List<string> Arguments, FileItem PreprocessedFile)
 		{
 			string PreprocessedFileString = NormalizeCommandLinePath(PreprocessedFile);
 			Arguments.Add("/P"); // Preprocess
@@ -427,7 +454,7 @@ namespace UnrealBuildTool
 					Arguments.Add("-Xclang -gno-codeview-command-line");
 					Arguments.Add("-fintegrated-cc1");
 					DirectoryReference ResourceDir = DirectoryReference.Combine(EnvVars.CompilerDir!, "lib", "clang", EnvVars.CompilerVersion.Components[0].ToString());
-					Arguments.Add($"-resource-dir=\"{NormalizeCommandLinePath(ResourceDir, CompileEnvironment)}\"");
+					Arguments.Add($"-resource-dir=\"{NormalizeCommandLinePath(ResourceDir, CompileEnvironment.RootPaths)}\"");
 				}
 			}
 			else if (Target.StaticAnalyzer == StaticAnalyzer.Default && !CompileEnvironment.bDisableStaticAnalysis)
@@ -699,7 +726,8 @@ namespace UnrealBuildTool
 
 				// Currently the ASan headers are not default around. They can be found at this location so lets use this until this is resolved in the toolchain
 				// Jira with some more info and the MSVC bug at UE-144727
-				AddSystemIncludePath(Arguments, DirectoryReference.Combine(EnvVars.CompilerDir, "crt", "src"), Target.WindowsPlatform.Compiler);
+				AddSystemIncludePath(Arguments, DirectoryReference.Combine(EnvVars.CompilerDir, "crt", "src"), Target.WindowsPlatform.Compiler,
+					Target.WindowsPlatform.Compiler.IsClang() && CompileEnvironment.VfsOverlayFile != null ? CompileEnvironment.RootPaths : null);
 			}
 
 			if (Target.WindowsPlatform.bEnableLibFuzzer)
@@ -1729,8 +1757,6 @@ namespace UnrealBuildTool
 		{
 			VCCompileAction BaseCompileAction = new VCCompileAction(EnvVars);
 
-			IEnumerable<DirectoryItem> rootPaths = CompileEnvironment.RootPaths.Values.Select(x => DirectoryItem.GetItemByDirectoryReference(x));
-
 			// TODO: Revisit this code. We want to use d1trimfile to make outputs machine independent.
 			// but want to make sure we're not causing any frustration for devs (since __FILE__ will show a relative path with lines below)
 			// Also need to find the equivalent for clang
@@ -1747,14 +1773,14 @@ namespace UnrealBuildTool
 				}
 			}
 #endif
-
-			BaseCompileAction.RootPaths.AddRange(rootPaths);
-
-			// Add additional response files
-			foreach (FileItem AdditionalRsp in CompileEnvironment.AdditionalResponseFiles)
+			foreach (KeyValuePair<CppRootPathFolder, DirectoryReference> Item in CompileEnvironment.RootPaths)
 			{
-				BaseCompileAction.Arguments.Add($"@\"{NormalizeCommandLinePath(AdditionalRsp.Location)}\"");
+				BaseCompileAction.RootPaths.Add(Item.Key, Item.Value);
 			}
+			BaseCompileAction.VfsOverlayFile = CompileEnvironment.VfsOverlayFile;
+
+			// Add additional response files
+			BaseCompileAction.AdditionalResponseFiles.AddRange(CompileEnvironment.AdditionalResponseFiles);
 
 			AppendCLArguments_Global(CompileEnvironment, BaseCompileAction.Arguments);
 
@@ -2064,7 +2090,6 @@ namespace UnrealBuildTool
 					CompileAction.AdditionalPrerequisiteItems.AddRange(InlinedFiles);
 				}
 
-				CompileAction.AdditionalPrerequisiteItems.AddRange(CompileEnvironment.ForceIncludeFiles);
 				CompileAction.AdditionalPrerequisiteItems.AddRange(CompileEnvironment.AdditionalPrerequisites);
 
 				if (SourceFile.HasExtension(".ixx"))
@@ -2083,7 +2108,6 @@ namespace UnrealBuildTool
 					CompileDepsAction.DependencyListFile = IfcDepsFile;
 					CompileDepsAction.Arguments.Add($"/sourceDependencies:directives \"{IfcDepsFile.Location}\"");
 					CompileDepsAction.AdditionalPrerequisiteItems.Add(SourceFile);
-					CompileDepsAction.AdditionalPrerequisiteItems.AddRange(CompileEnvironment.ForceIncludeFiles);
 					CompileDepsAction.AdditionalPrerequisiteItems.AddRange(CompileEnvironment.AdditionalPrerequisites);
 					CompileDepsAction.AdditionalProducedItems.Add(IfcDepsFile);
 					Graph.AddAction(CompileDepsAction);
@@ -2624,21 +2648,22 @@ namespace UnrealBuildTool
 		public override CppCompileEnvironment CreateSharedResponseFile(CppCompileEnvironment CompileEnvironment, FileReference OutResponseFile, IActionGraphBuilder Graph)
 		{
 			CppCompileEnvironment NewCompileEnvironment = new CppCompileEnvironment(CompileEnvironment);
+			IReadOnlyDictionary<CppRootPathFolder, DirectoryReference>? rootPaths = Target.WindowsPlatform.Compiler.IsClang() && NewCompileEnvironment.VfsOverlayFile != null ? NewCompileEnvironment.RootPaths : null;
 			List<string> Arguments = new List<string>();
 
 			foreach (DirectoryReference IncludePath in NewCompileEnvironment.UserIncludePaths)
 			{
-				AddIncludePath(Arguments, IncludePath, Target.WindowsPlatform.Compiler);
+				AddIncludePath(Arguments, IncludePath, Target.WindowsPlatform.Compiler, rootPaths);
 			}
 
 			foreach (DirectoryReference IncludePath in NewCompileEnvironment.SystemIncludePaths)
 			{
-				AddSystemIncludePath(Arguments, IncludePath, Target.WindowsPlatform.Compiler);
+				AddSystemIncludePath(Arguments, IncludePath, Target.WindowsPlatform.Compiler, rootPaths);
 			}
 
 			foreach (DirectoryReference IncludePath in EnvVars.IncludePaths)
 			{
-				AddSystemIncludePath(Arguments, IncludePath, Target.WindowsPlatform.Compiler);
+				AddSystemIncludePath(Arguments, IncludePath, Target.WindowsPlatform.Compiler, rootPaths);
 			}
 
 			// Stash shared include paths for validation purposes
@@ -2652,7 +2677,6 @@ namespace UnrealBuildTool
 			FileItem FileItem = FileItem.GetItemByFileReference(OutResponseFile);
 			Graph.CreateIntermediateTextFile(FileItem, Arguments);
 
-			NewCompileEnvironment.AdditionalPrerequisites.Add(FileItem);
 			NewCompileEnvironment.AdditionalResponseFiles.Add(FileItem);
 			NewCompileEnvironment.bHasSharedResponseFile = true;
 
