UnrealBuildTool: Clean up vfs pathing functions to use enum extension

#rnx

diff --git a/Engine/Source/Programs/UnrealBuildTool/Configuration/UEBuildModuleCPP.cs b/Engine/Source/Programs/UnrealBuildTool/Configuration/UEBuildModuleCPP.cs
--- a/Engine/Source/Programs/UnrealBuildTool/Configuration/UEBuildModuleCPP.cs
+++ b/Engine/Source/Programs/UnrealBuildTool/Configuration/UEBuildModuleCPP.cs
@@ -367,7 +367,7 @@ namespace UnrealBuildTool
 			{
 				jsonRoots.Add(new()
 				{
-					["name"] = $"{CppCompileEnvironment.ClangVfsRootPath}/{item.Key}",
+					["name"] = item.Key.AsVfsString(),
 					["type"] = "directory-remap",
 					["external-contents"] = item.Value.FullName.Replace('\\', '/')
 				});
diff --git a/Engine/Source/Programs/UnrealBuildTool/Platform/Windows/VCToolChain.cs b/Engine/Source/Programs/UnrealBuildTool/Platform/Windows/VCToolChain.cs
--- a/Engine/Source/Programs/UnrealBuildTool/Platform/Windows/VCToolChain.cs
+++ b/Engine/Source/Programs/UnrealBuildTool/Platform/Windows/VCToolChain.cs
@@ -183,7 +183,7 @@ namespace UnrealBuildTool
 					{
 						continue;
 					}
-					return $"{CppCompileEnvironment.ClangVfsRootPath}/{item.Key}/{reference.MakeRelativeTo(item.Value).Replace("\\", "/")}";
+					return item.Key.AsVfsString(item.Value, reference);
 				}
 			}
 			return NormalizeCommandLinePath(reference);
diff --git a/Engine/Source/Programs/UnrealBuildTool/System/CppCompileEnvironment.cs b/Engine/Source/Programs/UnrealBuildTool/System/CppCompileEnvironment.cs
--- a/Engine/Source/Programs/UnrealBuildTool/System/CppCompileEnvironment.cs
+++ b/Engine/Source/Programs/UnrealBuildTool/System/CppCompileEnvironment.cs
@@ -152,23 +152,96 @@ namespace UnrealBuildTool
 	}
 
 	/// <summary>
-	/// Named root folders
+	/// Named cpp root folders
 	/// </summary>
-	enum CppRootPathFolder
+	public enum CppRootPathFolder
 	{
+		/// <summary>
+		/// Root/Engine folder
+		/// </summary>
 		Engine,
+
+		/// <summary>
+		/// Project folder
+		/// </summary>
 		Project,
+
+		/// <summary>
+		/// UE Root folder 
+		/// </summary>
 		Root,
 
+		/// <summary>
+		/// Compiler folder
+		/// </summary>
 		Compiler,
+
+		/// <summary>
+		/// Toolchain folder (if needed)
+		/// </summary>
 		Toolchain,
+
+		/// <summary>
+		/// WindowsSDK folder (if needed)
+		/// </summary>
 		WinSDK,
 
+		/// <summary>
+		///  Additional platform SDK folder (if needed)
+		/// </summary>
 		PlatformSDK,
 
+		/// <summary>
+		/// Extra folder (if needed)
+		/// </summary>
 		Extra,
 	}
 
+	/// <summary>
+	/// Extension methods for CppRootPathFolder enum
+	/// </summary>
+	public static class CppRootPathFolderExtensions
+	{
+		/// <summary>
+		/// The root of the vfs path
+		/// </summary>
+		static readonly DirectoryReference _vfsRootPath = new(OperatingSystem.IsWindows() ? "Z:/UE" : "/UE");
+
+		/// <summary>
+		/// Convert a root folder to a vfs path
+		/// </summary>
+		/// <param name="folder">The enum value</param>
+		/// <returns>DirectoryReference for the path</returns>
+		public static DirectoryReference AsVfsPath(this CppRootPathFolder folder) => DirectoryReference.Combine(_vfsRootPath, folder.ToString());
+
+		/// <summary>
+		/// Convert a local path to a vfs path
+		/// </summary>
+		/// <param name="folder">The enum value</param>
+		/// <param name="localRoot">The local root directory</param>
+		/// <param name="localPath">The local path to convert</param>
+		/// <returns>DirectoryReference for the path</returns>
+		/// <remarks>localRoot must be an ancestor of localPath</remarks>
+		public static DirectoryReference AsVfsPath(this CppRootPathFolder folder, DirectoryReference localRoot, FileSystemReference localPath) => DirectoryReference.Combine(AsVfsPath(folder), localPath.MakeRelativeTo(localRoot));
+
+		/// <summary>
+		/// Convert a root folder to a vfs path string
+		/// </summary>
+		/// <param name="folder">The enum value</param>
+		/// <returns>string for the path</returns>
+		public static string AsVfsString(this CppRootPathFolder folder) => AsVfsPath(folder).FullName.Replace('\\', '/');
+
+		/// <summary>
+		/// Convert a local path to a vfs path string
+		/// </summary>
+		/// <param name="folder">The enum value</param>
+		/// <param name="localRoot">The local root directory</param>
+		/// <param name="localPath">The local path to convert</param>
+		/// <returns>string for the path</returns>
+		/// <remarks>localRoot must be an ancestor of localPath</remarks>
+		public static string AsVfsString(this CppRootPathFolder folder, DirectoryReference localRoot, FileSystemReference localPath) => AsVfsPath(folder, localRoot, localPath).FullName.Replace('\\', '/');
+	}
+
 	/// <summary>
 	/// Encapsulates the compilation output of compiling a set of C++ files.
 	/// </summary>
@@ -603,11 +676,6 @@ namespace UnrealBuildTool
 		/// </summary>
 		public FileItem? VfsOverlayFile = null;
 
-		/// <summary>
-		/// The root of the clang vfs path
-		/// </summary>
-		public static readonly string ClangVfsRootPath = OperatingSystem.IsWindows() ? "Z:/UE" : "/UE";
-
 		/// <summary>
 		/// List of files that need to be up to date before compile can proceed
 		/// </summary>
@@ -930,7 +998,7 @@ namespace UnrealBuildTool
 					{
 						continue;
 					}
-					vfsPath = $"{ClangVfsRootPath}/{item.Key}/{reference.MakeRelativeTo(item.Value).Replace("\\", "/")}";
+					vfsPath = item.Key.AsVfsString(item.Value, reference);
 					return true;
 				}
 			}
