UBA: Increase inactivity timeout for compute tasks

Some cases have been seen where CPU or I/O starvation causes timeouts with the previous default value of 20 seconds.

diff --git a/Engine/Source/Programs/UnrealBuildTool/Executors/UnrealBuildAccelerator/UBAAgentCoordinatorHorde.cs b/Engine/Source/Programs/UnrealBuildTool/Executors/UnrealBuildAccelerator/UBAAgentCoordinatorHorde.cs
--- a/Engine/Source/Programs/UnrealBuildTool/Executors/UnrealBuildAccelerator/UBAAgentCoordinatorHorde.cs
+++ b/Engine/Source/Programs/UnrealBuildTool/Executors/UnrealBuildAccelerator/UBAAgentCoordinatorHorde.cs
@@ -274,7 +274,8 @@ namespace UnrealBuildTool
 				{
 					ModePreference = _connectionMode,
 					Encryption = _encryption,
-					Ports = { { UbaPortName, UbaPort }, { UbaProxyPortName, UbaProxyPort } }
+					Ports = { { UbaPortName, UbaPort }, { UbaProxyPortName, UbaProxyPort } },
+					InactivityTimeoutMs = 90000
 				};
 				
 				await ResolveClusterIdAsync(hordeConfig, requirements, requestId, cmr, workerLogger, cancellationToken);
