Refactored how we define MB2\3 and Ansi allocator define in the project.
Moved them into UBT and made them global instead of delegating it to platform implementations which resulted in all sorts of issues

diff --git a/Engine/Plugins/Experimental/ChaosFlesh/Source/ChaosFleshEditor/ChaosFleshEditor.Build.cs b/Engine/Plugins/Experimental/ChaosFlesh/Source/ChaosFleshEditor/ChaosFleshEditor.Build.cs
--- a/Engine/Plugins/Experimental/ChaosFlesh/Source/ChaosFleshEditor/ChaosFleshEditor.Build.cs
+++ b/Engine/Plugins/Experimental/ChaosFlesh/Source/ChaosFleshEditor/ChaosFleshEditor.Build.cs
@@ -137,7 +137,7 @@ namespace UnrealBuildTool.Rules
 			// Note that by that point everything will be using the USD-compatible ANSI allocators anyway, so our overrides in USDMemory.h are also disabled, as they're unnecessary.
 			// Also note that we're forced to use dynamic linking for monolithic targets mainly because static linking the USD libraries disables support for user USD plugins, and secondly
 			// because those static libraries would need to be linked with the --whole-archive argument, and there is currently no standard way of doing that in UE.
-			if (bEnableUsdSdk && Target.LinkType == TargetLinkType.Monolithic && !Target.GlobalDefinitions.Contains("FORCE_ANSI_ALLOCATOR=1") && !Target.GlobalDefinitions.Contains("UE_USE_MALLOC_FILL_BYTES=0"))
+			if (bEnableUsdSdk && Target.LinkType == TargetLinkType.Monolithic && (Target.StaticAllocator != StaticAllocatorType.Ansi) && !Target.GlobalDefinitions.Contains("UE_USE_MALLOC_FILL_BYTES=0"))
 			{
 				PublicDefinitions.Add("USD_FORCE_DISABLED=1");
 				bEnableUsdSdk = false;
diff --git a/Engine/Plugins/Experimental/ChaosFlesh/Source/ChaosFleshEngine/ChaosFleshEngine.Build.cs b/Engine/Plugins/Experimental/ChaosFlesh/Source/ChaosFleshEngine/ChaosFleshEngine.Build.cs
--- a/Engine/Plugins/Experimental/ChaosFlesh/Source/ChaosFleshEngine/ChaosFleshEngine.Build.cs
+++ b/Engine/Plugins/Experimental/ChaosFlesh/Source/ChaosFleshEngine/ChaosFleshEngine.Build.cs
@@ -134,7 +134,7 @@ namespace UnrealBuildTool.Rules
 			// Note that by that point everything will be using the USD-compatible ANSI allocators anyway, so our overrides in USDMemory.h are also disabled, as they're unnecessary.
 			// Also note that we're forced to use dynamic linking for monolithic targets mainly because static linking the USD libraries disables support for user USD plugins, and secondly
 			// because those static libraries would need to be linked with the --whole-archive argument, and there is currently no standard way of doing that in UE.
-			if (bEnableUsdSdk && Target.LinkType == TargetLinkType.Monolithic && !Target.GlobalDefinitions.Contains("FORCE_ANSI_ALLOCATOR=1") && !Target.GlobalDefinitions.Contains("UE_USE_MALLOC_FILL_BYTES=0"))
+			if (bEnableUsdSdk && Target.LinkType == TargetLinkType.Monolithic && (Target.StaticAllocator != StaticAllocatorType.Ansi) && !Target.GlobalDefinitions.Contains("UE_USE_MALLOC_FILL_BYTES=0"))
 			{
 				PublicDefinitions.Add("USD_FORCE_DISABLED=1");
 				bEnableUsdSdk = false;
diff --git a/Engine/Plugins/Runtime/USDCore/Source/UnrealUSDWrapper/UnrealUSDWrapper.Build.cs b/Engine/Plugins/Runtime/USDCore/Source/UnrealUSDWrapper/UnrealUSDWrapper.Build.cs
--- a/Engine/Plugins/Runtime/USDCore/Source/UnrealUSDWrapper/UnrealUSDWrapper.Build.cs
+++ b/Engine/Plugins/Runtime/USDCore/Source/UnrealUSDWrapper/UnrealUSDWrapper.Build.cs
@@ -225,7 +225,7 @@ namespace UnrealBuildTool.Rules
 			// Also note that we're forced to use dynamic linking for monolithic targets mainly because static linking the USD libraries disables support for user USD plugins, and secondly
 			// because those static libraries would need to be linked with the --whole-archive argument, and there is currently no standard way of doing that in UE.
 			bool bIsGameOrMonolithic = Target.LinkType == TargetLinkType.Monolithic || Target.Type == TargetType.Client || Target.Type == TargetType.Server || Target.Type == TargetType.Game;
-			if (bEnableUsdSdk && bIsGameOrMonolithic && !Target.GlobalDefinitions.Contains("FORCE_ANSI_ALLOCATOR=1") && !Target.GlobalDefinitions.Contains("UE_USE_MALLOC_FILL_BYTES=0"))
+			if (bEnableUsdSdk && bIsGameOrMonolithic && (Target.StaticAllocator != StaticAllocatorType.Ansi) && !Target.GlobalDefinitions.Contains("UE_USE_MALLOC_FILL_BYTES=0"))
 			{
 				PublicDefinitions.Add("USD_FORCE_DISABLED=1");
 				bEnableUsdSdk = false;
diff --git a/Engine/Source/Programs/UnrealBuildTool/Configuration/ReadOnlyTargetRules.cs b/Engine/Source/Programs/UnrealBuildTool/Configuration/ReadOnlyTargetRules.cs
--- a/Engine/Source/Programs/UnrealBuildTool/Configuration/ReadOnlyTargetRules.cs
+++ b/Engine/Source/Programs/UnrealBuildTool/Configuration/ReadOnlyTargetRules.cs
@@ -733,5 +733,7 @@ namespace UnrealBuildTool
 				return _bWithAutomationTestsPrivate.Value;
 			}
 		}
+
+		public StaticAllocatorType StaticAllocator => Inner.StaticAllocator;
 	}
 }
diff --git a/Engine/Source/Programs/UnrealBuildTool/Configuration/TargetRules.cs b/Engine/Source/Programs/UnrealBuildTool/Configuration/TargetRules.cs
--- a/Engine/Source/Programs/UnrealBuildTool/Configuration/TargetRules.cs
+++ b/Engine/Source/Programs/UnrealBuildTool/Configuration/TargetRules.cs
@@ -402,6 +402,33 @@ namespace UnrealBuildTool
 		Full
 	}
 
+	/// <summary>
+	/// List of supported allocators to select in compile time for target.
+	/// Can be none for Desktop\Editor targets where allocator can be changed with command line.
+	/// </summary>
+	public enum StaticAllocatorType
+	{
+		/// <summary>
+		/// Default C allocator
+		/// </summary>
+		Ansi,
+
+		/// <summary>
+		/// Binned malloc
+		/// </summary>
+		Binned2,
+
+		/// <summary>
+		/// Newer VM-based binned malloc, 64 bit only
+		/// </summary>
+		Binned3,
+
+		/// <summary>
+		/// Default for Desktop\Editor targets that allows allocator selection to be overriden with a command line
+		/// </summary>
+		None,
+	}
+
 	/// <summary>
 	/// Utility class for EngineIncludeOrderVersion defines
 	/// </summary>
@@ -2838,6 +2865,12 @@ namespace UnrealBuildTool
 		[ConfigSubObject]
 		public WindowsTargetRules WindowsPlatform { get; init; } // Requires 'this' parameter; initialized in constructor
 
+		/// <summary>
+		/// Statically selected allocator to use
+		/// </summary>
+		[RequiresUniqueBuildEnvironment]
+		public StaticAllocatorType StaticAllocator { get; set; } = StaticAllocatorType.None;
+
 		/// <summary>
 		/// Checks if the target needs to care about sdk version - if not, we can share the target between projects that are on differnet sdk versions.
 		/// For instance, ShaderCompileWorker is very tied to, say, console sdk versions, but UnrealPak doesn't have sdk things compiled in
diff --git a/Engine/Source/Programs/UnrealBuildTool/Configuration/UEBuildTarget.cs b/Engine/Source/Programs/UnrealBuildTool/Configuration/UEBuildTarget.cs
--- a/Engine/Source/Programs/UnrealBuildTool/Configuration/UEBuildTarget.cs
+++ b/Engine/Source/Programs/UnrealBuildTool/Configuration/UEBuildTarget.cs
@@ -5724,6 +5724,19 @@ namespace UnrealBuildTool
 				}
 			}
 
+			// Setup defines for statically selected allocator type
+			string[] Mapping = {
+				"FORCE_ANSI_ALLOCATOR",
+				"USE_MALLOC_BINNED2",
+				"USE_MALLOC_BINNED3",
+			};
+
+			for (int i = 0; i < (int)StaticAllocatorType.None; ++i)
+			{
+				string EnabledValue = Rules.StaticAllocator == (StaticAllocatorType)i ? "=1" : "=0";
+				GlobalCompileEnvironment.Definitions.Add(Mapping[i] + EnabledValue);
+			}
+
 			// Initialize the compile and link environments for the platform, configuration, and project.
 			BuildPlatform.SetUpEnvironment(Rules, GlobalCompileEnvironment, GlobalLinkEnvironment);
 			BuildPlatform.SetUpConfigurationEnvironment(Rules, GlobalCompileEnvironment, GlobalLinkEnvironment);
diff --git a/Engine/Source/Programs/UnrealBuildTool/Platform/Android/AndroidToolChain.cs b/Engine/Source/Programs/UnrealBuildTool/Platform/Android/AndroidToolChain.cs
--- a/Engine/Source/Programs/UnrealBuildTool/Platform/Android/AndroidToolChain.cs
+++ b/Engine/Source/Programs/UnrealBuildTool/Platform/Android/AndroidToolChain.cs
@@ -600,14 +600,14 @@ namespace UnrealBuildTool
 
 				if (Sanitizer == ClangSanitizer.Address || Sanitizer == ClangSanitizer.HwAddress)
 				{
-					Arguments.Add("-fno-omit-frame-pointer -DFORCE_ANSI_ALLOCATOR=1");
+					Arguments.Add("-fno-omit-frame-pointer");
 				}
 			}
 
 			//string? SanitizerMode = Environment.GetEnvironmentVariable("ENABLE_ADDRESS_SANITIZER");
 			//if ((SanitizerMode != null && SanitizerMode == "YES") || (Options.HasFlag(ClangToolChainOptions.EnableAddressSanitizer)))
 			//{
-			//	Arguments.Add("-fsanitize=address -fno-omit-frame-pointer -DFORCE_ANSI_ALLOCATOR=1");
+			//	Arguments.Add("-fsanitize=address -fno-omit-frame-pointer");
 			//}
 
 			//string? UndefSanitizerMode = Environment.GetEnvironmentVariable("ENABLE_UNDEFINED_BEHAVIOR_SANITIZER");
diff --git a/Engine/Source/Programs/UnrealBuildTool/Platform/Android/UEBuildAndroid.cs b/Engine/Source/Programs/UnrealBuildTool/Platform/Android/UEBuildAndroid.cs
--- a/Engine/Source/Programs/UnrealBuildTool/Platform/Android/UEBuildAndroid.cs
+++ b/Engine/Source/Programs/UnrealBuildTool/Platform/Android/UEBuildAndroid.cs
@@ -220,6 +220,15 @@ namespace UnrealBuildTool
 			{
 				Target.DisablePlugins.AddRange(Array.Empty<string>());
 			}
+
+			if (Target.StaticAllocator == StaticAllocatorType.None)
+			{
+				Target.StaticAllocator = StaticAllocatorType.Binned3;
+			}
+			if (Target.AndroidPlatform.bEnableAddressSanitizer || Target.AndroidPlatform.bEnableHWAddressSanitizer)
+			{
+				Target.StaticAllocator = StaticAllocatorType.Ansi;
+			}
 		}
 
 		public override bool CanUseXGE()
diff --git a/Engine/Source/Programs/UnrealBuildTool/Platform/IOS/IOSToolChain.cs b/Engine/Source/Programs/UnrealBuildTool/Platform/IOS/IOSToolChain.cs
--- a/Engine/Source/Programs/UnrealBuildTool/Platform/IOS/IOSToolChain.cs
+++ b/Engine/Source/Programs/UnrealBuildTool/Platform/IOS/IOSToolChain.cs
@@ -174,14 +174,12 @@ namespace UnrealBuildTool
 			// TODO: Reconcile with base
 			//base.GetCompilerArguments_Sanitizers(CompileEnvironment, Arguments);
 
-			string? SanitizerMode = Environment.GetEnvironmentVariable("ENABLE_ADDRESS_SANITIZER");
-			if ((SanitizerMode != null && SanitizerMode == "YES") || (Options.HasFlag(ClangToolChainOptions.EnableAddressSanitizer)))
+			if (Options.HasFlag(ClangToolChainOptions.EnableAddressSanitizer))
 			{
-				Arguments.Add("-fsanitize=address -fno-omit-frame-pointer -DFORCE_ANSI_ALLOCATOR=1");
+				Arguments.Add("-fsanitize=address -fno-omit-frame-pointer");
 			}
 
-			string? UndefSanitizerMode = Environment.GetEnvironmentVariable("ENABLE_UNDEFINED_BEHAVIOR_SANITIZER");
-			if ((UndefSanitizerMode != null && UndefSanitizerMode == "YES") || (Options.HasFlag(ClangToolChainOptions.EnableUndefinedBehaviorSanitizer)))
+			if (Options.HasFlag(ClangToolChainOptions.EnableUndefinedBehaviorSanitizer))
 			{
 				Arguments.Add("-fsanitize=undefined -fno-sanitize=bounds,enum,return,float-divide-by-zero");
 			}
@@ -328,15 +326,13 @@ namespace UnrealBuildTool
 				Arguments.Add("-flto");
 			}
 
-			string? SanitizerMode = Environment.GetEnvironmentVariable("ENABLE_ADDRESS_SANITIZER");
-			if ((SanitizerMode != null && SanitizerMode == "YES") || (Options.HasFlag(ClangToolChainOptions.EnableAddressSanitizer)))
+			if (Options.HasFlag(ClangToolChainOptions.EnableAddressSanitizer))
 			{
 				Arguments.Add("-rpath \"@executable_path/Frameworks\"");
 				Arguments.Add("-fsanitize=address");
 			}
 
-			string? UndefSanitizerMode = Environment.GetEnvironmentVariable("ENABLE_UNDEFINED_BEHAVIOR_SANITIZER");
-			if ((UndefSanitizerMode != null && UndefSanitizerMode == "YES") || (Options.HasFlag(ClangToolChainOptions.EnableUndefinedBehaviorSanitizer)))
+			if (Options.HasFlag(ClangToolChainOptions.EnableUndefinedBehaviorSanitizer))
 			{
 				Arguments.Add("-rpath \"@executable_path/libclang_rt.ubsan_ios_dynamic.dylib\"");
 				Arguments.Add("-fsanitize=undefined");
diff --git a/Engine/Source/Programs/UnrealBuildTool/Platform/IOS/UEBuildIOS.cs b/Engine/Source/Programs/UnrealBuildTool/Platform/IOS/UEBuildIOS.cs
--- a/Engine/Source/Programs/UnrealBuildTool/Platform/IOS/UEBuildIOS.cs
+++ b/Engine/Source/Programs/UnrealBuildTool/Platform/IOS/UEBuildIOS.cs
@@ -145,6 +145,30 @@ namespace UnrealBuildTool
 		/// </summary>
 		[CommandLine("-EnableUBSan")]
 		public bool bEnableUndefinedBehaviorSanitizer = false;
+
+		/// <summary>
+		/// Constructor
+		/// </summary>
+		public IOSTargetRules()
+		{
+			string? AddressSanitizer = Environment.GetEnvironmentVariable("ENABLE_ADDRESS_SANITIZER");
+			if (AddressSanitizer != null && AddressSanitizer == "YES")
+			{
+				bEnableAddressSanitizer = true;
+			}
+
+			string? ThreadSanitizer = Environment.GetEnvironmentVariable("ENABLE_THREAD_SANITIZER");
+			if (ThreadSanitizer != null && ThreadSanitizer == "YES")
+			{
+				bEnableThreadSanitizer = true;
+			}
+
+			string? UndefSanitizerMode = Environment.GetEnvironmentVariable("ENABLE_UNDEFINED_BEHAVIOR_SANITIZER");
+			if (UndefSanitizerMode != null && UndefSanitizerMode == "YES")
+			{
+				bEnableUndefinedBehaviorSanitizer = true;
+			}
+		}
 	}
 
 	/// <summary>
@@ -769,6 +793,16 @@ namespace UnrealBuildTool
 			}
 
 			Target.bCheckSystemHeadersForModification = false;
+
+			if (Target.StaticAllocator == StaticAllocatorType.None)
+			{
+				// Prefer MB2 as MB3 requires extended virtual address space entitlement
+				Target.StaticAllocator = StaticAllocatorType.Binned2;
+			}
+			if (Target.IOSPlatform.bEnableAddressSanitizer)
+			{
+				Target.StaticAllocator = StaticAllocatorType.Ansi;
+			}
 		}
 
 		public override void ValidateModule(UEBuildModule Module, ReadOnlyTargetRules Target)
diff --git a/Engine/Source/Programs/UnrealBuildTool/Platform/Linux/LinuxToolChain.cs b/Engine/Source/Programs/UnrealBuildTool/Platform/Linux/LinuxToolChain.cs
--- a/Engine/Source/Programs/UnrealBuildTool/Platform/Linux/LinuxToolChain.cs
+++ b/Engine/Source/Programs/UnrealBuildTool/Platform/Linux/LinuxToolChain.cs
@@ -460,20 +460,6 @@ namespace UnrealBuildTool
 		{
 			base.GetCompilerArguments_Sanitizers(CompileEnvironment, Arguments);
 
-			// ASan
-			if (Options.HasFlag(ClangToolChainOptions.EnableAddressSanitizer))
-			{
-				// Force using the ANSI allocator if ASan is enabled
-				Arguments.Add("-DFORCE_ANSI_ALLOCATOR=1");
-			}
-
-			// TSan
-			if (Options.HasFlag(ClangToolChainOptions.EnableThreadSanitizer))
-			{
-				// Force using the ANSI allocator if TSan is enabled
-				Arguments.Add("-DFORCE_ANSI_ALLOCATOR=1");
-			}
-
 			// UBSan
 			if (Options.HasFlag(ClangToolChainOptions.EnableUndefinedBehaviorSanitizer))
 			{
@@ -487,7 +473,6 @@ namespace UnrealBuildTool
 				// -fsanitize-memory-track-origins adds a 1.5x-2.5x slow down ontop of MSan normal amount of overhead
 				// -fsanitize-memory-track-origins=1 is faster but collects only allocation points but not intermediate stores
 				Arguments.Add("-fsanitize-memory-track-origins");
-				Arguments.Add("-DFORCE_ANSI_ALLOCATOR=1");
 			}
 
 			// LibFuzzer
diff --git a/Engine/Source/Programs/UnrealBuildTool/Platform/Linux/UEBuildLinux.cs b/Engine/Source/Programs/UnrealBuildTool/Platform/Linux/UEBuildLinux.cs
--- a/Engine/Source/Programs/UnrealBuildTool/Platform/Linux/UEBuildLinux.cs
+++ b/Engine/Source/Programs/UnrealBuildTool/Platform/Linux/UEBuildLinux.cs
@@ -282,6 +282,16 @@ namespace UnrealBuildTool
 
 			// Disable chaining PCHs for the moment because it is crashing clang
 			Target.bChainPCHs = false;
+
+			if (Target.StaticAllocator == StaticAllocatorType.None)
+			{
+				// default to FMallocBinned2 as it uses less Virtual Memory Areas
+				Target.StaticAllocator = StaticAllocatorType.Binned2;
+			}
+			if (Target.LinuxPlatform.bEnableAddressSanitizer || Target.LinuxPlatform.bEnableThreadSanitizer || Target.LinuxPlatform.bEnableMemorySanitizer)
+			{
+				Target.StaticAllocator = StaticAllocatorType.Ansi;
+			}
 		}
 
 		public override bool CanUseXGE()
diff --git a/Engine/Source/Programs/UnrealBuildTool/Platform/Mac/UEBuildMac.cs b/Engine/Source/Programs/UnrealBuildTool/Platform/Mac/UEBuildMac.cs
--- a/Engine/Source/Programs/UnrealBuildTool/Platform/Mac/UEBuildMac.cs
+++ b/Engine/Source/Programs/UnrealBuildTool/Platform/Mac/UEBuildMac.cs
@@ -61,6 +61,30 @@ namespace UnrealBuildTool
 		/// </summary>
 		[ConfigFile(ConfigHierarchyType.Engine, "/Script/MacTargetPlatform.MacTargetSettings", "bEnableRayTracing")]
 		public bool bEnableRayTracing = false;
+
+		/// <summary>
+		/// Constructor
+		/// </summary>
+		public MacTargetRules()
+		{
+			string? AddressSanitizer = Environment.GetEnvironmentVariable("ENABLE_ADDRESS_SANITIZER");
+			if (AddressSanitizer != null && AddressSanitizer == "YES")
+			{
+				bEnableAddressSanitizer = true;
+			}
+
+			string? ThreadSanitizer = Environment.GetEnvironmentVariable("ENABLE_THREAD_SANITIZER");
+			if (ThreadSanitizer != null && ThreadSanitizer == "YES")
+			{
+				bEnableThreadSanitizer = true;
+			}
+
+			string? UndefSanitizerMode = Environment.GetEnvironmentVariable("ENABLE_UNDEFINED_BEHAVIOR_SANITIZER");
+			if (UndefSanitizerMode != null && UndefSanitizerMode == "YES")
+			{
+				bEnableUndefinedBehaviorSanitizer = true;
+			}
+		}
 	}
 
 	/// <summary>
@@ -360,11 +384,14 @@ namespace UnrealBuildTool
 				Target.GlobalDefinitions.Add("HAS_METAL=0");
 			}
 
+			if (Target.StaticAllocator == StaticAllocatorType.None)
+			{
+				Target.StaticAllocator = StaticAllocatorType.Binned3;
+			}
 			// Force using the ANSI allocator if ASan is enabled
-			string? AddressSanitizer = Environment.GetEnvironmentVariable("ENABLE_ADDRESS_SANITIZER");
-			if (Target.MacPlatform.bEnableAddressSanitizer || (AddressSanitizer != null && AddressSanitizer == "YES"))
+			if (Target.MacPlatform.bEnableAddressSanitizer)
 			{
-				Target.GlobalDefinitions.Add("FORCE_ANSI_ALLOCATOR=1");
+				Target.StaticAllocator = StaticAllocatorType.Ansi;
 			}
 
 			Target.GlobalDefinitions.Add("GL_SILENCE_DEPRECATION=1");
@@ -614,19 +641,15 @@ namespace UnrealBuildTool
 		{
 			ClangToolChainOptions Options = ClangToolChainOptions.None;
 
-			string? AddressSanitizer = Environment.GetEnvironmentVariable("ENABLE_ADDRESS_SANITIZER");
-			string? ThreadSanitizer = Environment.GetEnvironmentVariable("ENABLE_THREAD_SANITIZER");
-			string? UndefSanitizerMode = Environment.GetEnvironmentVariable("ENABLE_UNDEFINED_BEHAVIOR_SANITIZER");
-
-			if (Target.MacPlatform.bEnableAddressSanitizer || (AddressSanitizer != null && AddressSanitizer == "YES"))
+			if (Target.MacPlatform.bEnableAddressSanitizer)
 			{
 				Options |= ClangToolChainOptions.EnableAddressSanitizer;
 			}
-			if (Target.MacPlatform.bEnableThreadSanitizer || (ThreadSanitizer != null && ThreadSanitizer == "YES"))
+			if (Target.MacPlatform.bEnableThreadSanitizer)
 			{
 				Options |= ClangToolChainOptions.EnableThreadSanitizer;
 			}
-			if (Target.MacPlatform.bEnableUndefinedBehaviorSanitizer || (UndefSanitizerMode != null && UndefSanitizerMode == "YES"))
+			if (Target.MacPlatform.bEnableUndefinedBehaviorSanitizer)
 			{
 				Options |= ClangToolChainOptions.EnableUndefinedBehaviorSanitizer;
 			}
diff --git a/Engine/Source/Programs/UnrealBuildTool/Platform/Windows/UEBuildWindows.cs b/Engine/Source/Programs/UnrealBuildTool/Platform/Windows/UEBuildWindows.cs
--- a/Engine/Source/Programs/UnrealBuildTool/Platform/Windows/UEBuildWindows.cs
+++ b/Engine/Source/Programs/UnrealBuildTool/Platform/Windows/UEBuildWindows.cs
@@ -1335,6 +1335,15 @@ namespace UnrealBuildTool
 
 				ValidateToolchainVersion(Target);
 			}
+
+			if (Target.StaticAllocator == StaticAllocatorType.None)
+			{
+				Target.StaticAllocator = StaticAllocatorType.Binned2;
+			}
+			if (Target.WindowsPlatform.bEnableAddressSanitizer)
+			{
+				Target.StaticAllocator = StaticAllocatorType.Ansi;
+			}
 		}
 
 		static bool _toolchainWarningLogged = false;
diff --git a/Engine/Source/Programs/UnrealBuildTool/Platform/Windows/VCToolChain.cs b/Engine/Source/Programs/UnrealBuildTool/Platform/Windows/VCToolChain.cs
--- a/Engine/Source/Programs/UnrealBuildTool/Platform/Windows/VCToolChain.cs
+++ b/Engine/Source/Programs/UnrealBuildTool/Platform/Windows/VCToolChain.cs
@@ -709,10 +709,6 @@ namespace UnrealBuildTool
 					Arguments.Add("-mllvm -asan-globals=0");
 				}
 
-				// Use the CRT allocator so that ASan is able to hook into it for better error
-				// detection.
-				AddDefinition(Arguments, "FORCE_ANSI_ALLOCATOR=1");
-
 				if (Target.WindowsPlatform.Compiler.IsMSVC())
 				{
 					// MSVC has no support for __has_feature(address_sanitizer)
diff --git a/Engine/Source/Runtime/Core/Private/Android/AndroidPlatformMemory.cpp b/Engine/Source/Runtime/Core/Private/Android/AndroidPlatformMemory.cpp
--- a/Engine/Source/Runtime/Core/Private/Android/AndroidPlatformMemory.cpp
+++ b/Engine/Source/Runtime/Core/Private/Android/AndroidPlatformMemory.cpp
@@ -430,12 +430,6 @@ EPlatformMemorySizeBucket FAndroidPlatformMemory::GetMemorySizeBucket()
 	return FGenericPlatformMemory::GetMemorySizeBucket();
 }
 
-// Set rather to use BinnedMalloc2 for binned malloc, can be overridden below
-#define USE_MALLOC_BINNED2 PLATFORM_ANDROID_ARM64
-#if !defined(USE_MALLOC_BINNED3)
-	#define USE_MALLOC_BINNED3 (0)
-#endif
-
 FMalloc* FAndroidPlatformMemory::BaseAllocator()
 {
 	static FMalloc* Instance = nullptr;
@@ -457,7 +451,7 @@ FMalloc* FAndroidPlatformMemory::BaseAllocator()
 
 	const bool bHeapProfilingSupported = AndroidHeapProfiling::Init();
 
-#if USE_MALLOC_BINNED3 && PLATFORM_ANDROID_ARM64
+#if USE_MALLOC_BINNED3
 	if (bHeapProfilingSupported)
 	{
 		return (Instance = new FMallocProfilingProxy<FMallocBinned3>());
diff --git a/Engine/Source/Runtime/Core/Private/Apple/ApplePlatformMemory.cpp b/Engine/Source/Runtime/Core/Private/Apple/ApplePlatformMemory.cpp
--- a/Engine/Source/Runtime/Core/Private/Apple/ApplePlatformMemory.cpp
+++ b/Engine/Source/Runtime/Core/Private/Apple/ApplePlatformMemory.cpp
@@ -337,13 +337,6 @@ void FApplePlatformMemory::SetAllocatorToUse()
 	}
     else if (USE_MALLOC_BINNED2)
     {
-        if(!CanOverallocateVirtualMemory())
-        {
-			NSLog(@"MallocBinned2 requested but com.apple.developer.kernel.extended-virtual-addressing entitlement not found. Check your entitlements. Falling back to Ansi.\n");
-            AllocatorToUse = EMemoryAllocatorToUse::Ansi;
-            return;
-        }
-
 		NSLog(@"Using MallocBinned2 allocator.\n");
         AllocatorToUse = EMemoryAllocatorToUse::Binned2;
         return;
diff --git a/Engine/Source/Runtime/Core/Private/Unix/UnixPlatformMemory.cpp b/Engine/Source/Runtime/Core/Private/Unix/UnixPlatformMemory.cpp
--- a/Engine/Source/Runtime/Core/Private/Unix/UnixPlatformMemory.cpp
+++ b/Engine/Source/Runtime/Core/Private/Unix/UnixPlatformMemory.cpp
@@ -45,9 +45,6 @@
   #endif // defined(_M_X64) || defined(__x86_64__) || defined (__amd64__) 
 #endif // ifndef UE4_DO_ROOT_PRIVILEGE_CHECK
 
-// Set rather to use BinnedMalloc2 for binned malloc, can be overridden below
-#define USE_MALLOC_BINNED2 (1)
-
 // Set to 1 if we should try to use /proc/self/smaps_rollup in FUnixPlatformMemory::GetExtendedStats().
 // There is a potential tradeoff in that smaps_rollup appears to be quite a bit faster in
 //   straight testing, but it also looks like it blocks mmap() calls on other threads until it finishes.
diff --git a/Engine/Source/Runtime/Core/Private/Windows/WindowsPlatformMemory.cpp b/Engine/Source/Runtime/Core/Private/Windows/WindowsPlatformMemory.cpp
--- a/Engine/Source/Runtime/Core/Private/Windows/WindowsPlatformMemory.cpp
+++ b/Engine/Source/Runtime/Core/Private/Windows/WindowsPlatformMemory.cpp
@@ -133,12 +133,6 @@ void FWindowsPlatformMemory::Init()
 	DumpStats( *GLog );
 }
 
-// Set rather to use BinnedMalloc2 for binned malloc, can be overridden below
-#define USE_MALLOC_BINNED2 (1)
-#if !defined(USE_MALLOC_BINNED3)
-	#define USE_MALLOC_BINNED3 (0)
-#endif
-
 FMalloc* FWindowsPlatformMemory::BaseAllocator()
 {
 	static FMalloc* Instance = nullptr;
diff --git a/Engine/Source/Runtime/Core/Public/Apple/ApplePlatformMemory.h b/Engine/Source/Runtime/Core/Public/Apple/ApplePlatformMemory.h
--- a/Engine/Source/Runtime/Core/Public/Apple/ApplePlatformMemory.h
+++ b/Engine/Source/Runtime/Core/Public/Apple/ApplePlatformMemory.h
@@ -14,16 +14,6 @@
 
 NS_ASSUME_NONNULL_BEGIN
 
-// if USE_MALLOC_BINNED3 isn't configured disable it by default
-#ifndef USE_MALLOC_BINNED3
-	#define USE_MALLOC_BINNED3 0
-#endif // USE_MALLOC_BINNED3
-
-
-// if USE_MALLOC_BINNED2 isn't configured enable it if USE_MALLOC_BINNED3 isn't true
-#ifndef USE_MALLOC_BINNED2
-#define USE_MALLOC_BINNED2 !USE_MALLOC_BINNED3
-#endif // USE_MALLOC_BINNED2
 
 // Support for passing LLM alloc data to XCode Instruments.
 // Use USE_APPLE_SUPPORT_INSTRUMENTED_ALLOCS to enable it.
