UnrealBuildTool: Setting to warn if a referenced plugin module violates the settings from the .uplugin

#rnx

diff --git a/Engine/Source/Programs/UnrealBuildTool/Configuration/CppCompileWarnings.cs b/Engine/Source/Programs/UnrealBuildTool/Configuration/CppCompileWarnings.cs
--- a/Engine/Source/Programs/UnrealBuildTool/Configuration/CppCompileWarnings.cs
+++ b/Engine/Source/Programs/UnrealBuildTool/Configuration/CppCompileWarnings.cs
@@ -100,6 +100,30 @@ namespace UnrealBuildTool
 
 		#region -- Rules Validation Warning Settings --
 
+		/// <summary>
+		/// How to treat module unsupported validation messages
+		/// </summary>
+		[XmlConfigFile(Category = "BuildConfiguration")]
+		[CommandLine("-ModuleUnsupportedWarningLevel=")]
+		public WarningLevel ModuleUnsupportedWarningLevel
+		{
+			get => ResolveWarning(null, ParentCppCompileWarnings?.ModuleUnsupportedWarningLevel, _moduleUnsupportedWarningLevel);
+			set => _moduleUnsupportedWarningLevel = value;
+		}
+		private WarningLevel _moduleUnsupportedWarningLevel;
+
+		/// <summary>
+		/// How to treat plugin specific module unsupported validation messages
+		/// </summary>
+		[XmlConfigFile(Category = "BuildConfiguration")]
+		[CommandLine("-PluginModuleUnsupportedWarningLevel=")]
+		public WarningLevel PluginModuleUnsupportedWarningLevel
+		{
+			get => ResolveWarning(null, ParentCppCompileWarnings?.PluginModuleUnsupportedWarningLevel, _pluginModuleUnsupportedWarningLevel);
+			set => _pluginModuleUnsupportedWarningLevel = value;
+		}
+		private WarningLevel _pluginModuleUnsupportedWarningLevel;
+
 		/// <summary>
 		/// How to treat general module include path validation messages
 		/// </summary>
@@ -283,6 +307,9 @@ namespace UnrealBuildTool
 			_pchPerformanceIssueWarningLevel = WarningLevel.Off;
 
 			// Rules Validation Warning Settings
+			_moduleUnsupportedWarningLevel = WarningLevel.Warning;
+			_pluginModuleUnsupportedWarningLevel = WarningLevel.Off;
+
 			_moduleIncludePathWarningLevel = WarningLevel.Off;
 			_moduleIncludePrivateWarningLevel = WarningLevel.Off;
 			_moduleIncludeSubdirectoryWarningLevel = WarningLevel.Off;
@@ -306,6 +333,9 @@ namespace UnrealBuildTool
 			_pchPerformanceIssueWarningLevel = WarningLevel.Off;
 
 			// Rules Validation Warning Settings
+			_moduleUnsupportedWarningLevel = WarningLevel.Warning;
+			_pluginModuleUnsupportedWarningLevel = WarningLevel.Off;
+
 			_moduleIncludePathWarningLevel = WarningLevel.Off;
 			_moduleIncludePrivateWarningLevel = WarningLevel.Off;
 			_moduleIncludeSubdirectoryWarningLevel = WarningLevel.Off;
diff --git a/Engine/Source/Programs/UnrealBuildTool/Configuration/ReadOnlyCppCompileWarnings.cs b/Engine/Source/Programs/UnrealBuildTool/Configuration/ReadOnlyCppCompileWarnings.cs
--- a/Engine/Source/Programs/UnrealBuildTool/Configuration/ReadOnlyCppCompileWarnings.cs
+++ b/Engine/Source/Programs/UnrealBuildTool/Configuration/ReadOnlyCppCompileWarnings.cs
@@ -53,7 +53,8 @@ namespace UnrealBuildTool
 		#endregion -- Performance Warning Settings --
 
 		#region -- Rules Validation Warning Settings --
-
+		public WarningLevel ModuleUnsupportedWarningLevel => _inner.ModuleUnsupportedWarningLevel;
+		public WarningLevel PluginModuleUnsupportedWarningLevel => _inner.PluginModuleUnsupportedWarningLevel;
 		public WarningLevel ModuleIncludePathWarningLevel => _inner.ModuleIncludePathWarningLevel;
 		public WarningLevel ModuleIncludePrivateWarningLevel => _inner.ModuleIncludePrivateWarningLevel;
 		public WarningLevel ModuleIncludeSubdirectoryWarningLevel => _inner.ModuleIncludeSubdirectoryWarningLevel;
diff --git a/Engine/Source/Programs/UnrealBuildTool/Configuration/UEBuildModule.cs b/Engine/Source/Programs/UnrealBuildTool/Configuration/UEBuildModule.cs
--- a/Engine/Source/Programs/UnrealBuildTool/Configuration/UEBuildModule.cs
+++ b/Engine/Source/Programs/UnrealBuildTool/Configuration/UEBuildModule.cs
@@ -2,11 +2,13 @@
 
 using System;
 using System.Collections.Generic;
+using System.Configuration;
 using System.IO;
 using System.Linq;
 using System.Text;
 using System.Threading.Tasks;
 using EpicGames.Core;
+using Microsoft.CodeAnalysis;
 using Microsoft.Extensions.Logging;
 using UnrealBuildBase;
 
@@ -1134,9 +1136,34 @@ namespace UnrealBuildTool
 			ReferenceStack.Add(this);
 
 			// Check if this module is invalid for the current target, based on Module attributes
-			if (!ModuleRules.IsValidForTarget(Rules.GetType(), Rules.Target, out string? InvalidReason))
+			if (Rules.CppCompileWarningSettings.ModuleUnsupportedWarningLevel != WarningLevel.Off)
 			{
-				Logger.LogWarning("Warning: Referenced module '{Module}' does not support {InvalidReason} via {ReferenceChain}", ThisRefName, InvalidReason, NextReferenceChain);
+				if (!ModuleRules.IsValidForTarget(Rules.GetType(), Rules.Target, out string? InvalidReason))
+				{
+					if (Rules.CppCompileWarningSettings.ModuleUnsupportedWarningLevel == WarningLevel.Warning)
+					{
+						Logger.LogWarning("Warning: Referenced module '{Module}' not supported '{InvalidReason}' via {ReferenceChain}", ThisRefName, InvalidReason, NextReferenceChain);
+					}
+					else if (Rules.CppCompileWarningSettings.ModuleUnsupportedWarningLevel == WarningLevel.Error)
+					{
+						Logger.LogError("Error: Referenced module '{Module}' not supported '{InvalidReason}' via {ReferenceChain}", ThisRefName, InvalidReason, NextReferenceChain);
+					}
+				}
+			}
+			if (Rules.CppCompileWarningSettings.PluginModuleUnsupportedWarningLevel != WarningLevel.Off)
+			{
+				ModuleDescriptor? ModuleInfo = Rules.Plugin?.Descriptor.Modules?.FirstOrDefault(x => x.Name == Rules.Name);
+				if (ModuleInfo?.IsCompiledInConfiguration(Rules.Target.Platform, Rules.Target.Configuration, Rules.Target.Name, Rules.Target.Type, Rules.Target.bBuildDeveloperTools, Rules.Target.bBuildRequiresCookedData, out string? InvalidReason) == false)
+				{
+					if (Rules.CppCompileWarningSettings.PluginModuleUnsupportedWarningLevel == WarningLevel.Warning)
+					{
+						Logger.LogWarning("Warning: Referenced plugin '{Plugin}' module '{Module}' not supported '{InvalidReason}' via {ReferenceChain}", Rules.Plugin?.File.GetFileName(), ThisRefName, InvalidReason, NextReferenceChain);
+					}
+					else if (Rules.CppCompileWarningSettings.PluginModuleUnsupportedWarningLevel == WarningLevel.Error)
+					{
+						Logger.LogError("Error: Referenced plugin '{Plugin}' module '{Module}' not supported '{InvalidReason}' via {ReferenceChain}", Rules.Plugin?.File.GetFileName(), ThisRefName, InvalidReason, NextReferenceChain);
+					}
+				}
 			}
 
 			// Create all the referenced modules. This path can be recursive, so we check against PublicDependencyModules to ensure we don't recurse through the 
@@ -1190,6 +1217,20 @@ namespace UnrealBuildTool
 			List<UEBuildModule> ReferenceStack = new List<UEBuildModule>();
 			RecursivelyCheckForCircularReferences(ReferenceChain, ReferenceStack, Logger);
 
+			// Check if this module is invalid for the current target, based on Module attributes
+			if (Rules.CppCompileWarningSettings.ModuleUnsupportedWarningLevel == WarningLevel.Error && !ModuleRules.IsValidForTarget(Rules.GetType(), Rules.Target, out string? _))
+			{
+				return true;
+			}
+			if (Rules.CppCompileWarningSettings.PluginModuleUnsupportedWarningLevel == WarningLevel.Error)
+			{
+				ModuleDescriptor? ModuleInfo = Rules.Plugin?.Descriptor.Modules?.FirstOrDefault(x => x.Name == Rules.Name);
+				if (ModuleInfo?.IsCompiledInConfiguration(Rules.Target.Platform, Rules.Target.Configuration, Rules.Target.Name, Rules.Target.Type, Rules.Target.bBuildDeveloperTools, Rules.Target.bBuildRequiresCookedData, out string? InvalidReason) == false)
+				{
+					return true;
+				}
+			}
+
 			return false; // No fatal errors
 		}
 
diff --git a/Engine/Source/Programs/UnrealBuildTool/System/ModuleDescriptor.cs b/Engine/Source/Programs/UnrealBuildTool/System/ModuleDescriptor.cs
--- a/Engine/Source/Programs/UnrealBuildTool/System/ModuleDescriptor.cs
+++ b/Engine/Source/Programs/UnrealBuildTool/System/ModuleDescriptor.cs
@@ -3,6 +3,7 @@
 using System;
 using System.Collections.Generic;
 using System.Diagnostics;
+using System.Diagnostics.CodeAnalysis;
 using System.Linq;
 using EpicGames.Core;
 using JsonExtensions;
@@ -583,6 +584,21 @@ namespace UnrealBuildTool
 		/// <param name="bBuildDeveloperTools">Whether the configuration includes developer tools (typically UEBuildConfiguration.bBuildDeveloperTools for UBT callers)</param>
 		/// <param name="bBuildRequiresCookedData">Whether the configuration requires cooked content (typically UEBuildConfiguration.bBuildRequiresCookedData for UBT callers)</param>
 		public bool IsCompiledInConfiguration(UnrealTargetPlatform Platform, UnrealTargetConfiguration Configuration, string TargetName, TargetType TargetType, bool bBuildDeveloperTools, bool bBuildRequiresCookedData)
+		{
+			return IsCompiledInConfiguration(Platform, Configuration, TargetName, TargetType, bBuildDeveloperTools, bBuildRequiresCookedData, out string? _);
+		}
+
+		/// <summary>
+		/// Determines whether the given plugin module is part of the current build.
+		/// </summary>
+		/// <param name="Platform">The platform being compiled for</param>
+		/// <param name="Configuration">The target configuration being compiled for</param>
+		/// <param name="TargetName">Name of the target being built</param>
+		/// <param name="TargetType">The type of the target being compiled</param>
+		/// <param name="bBuildDeveloperTools">Whether the configuration includes developer tools (typically UEBuildConfiguration.bBuildDeveloperTools for UBT callers)</param>
+		/// <param name="bBuildRequiresCookedData">Whether the configuration requires cooked content (typically UEBuildConfiguration.bBuildRequiresCookedData for UBT callers)</param>
+		/// <param name="invalidReason">Out parameter, reason why this plugin module is invalid</param>
+		public bool IsCompiledInConfiguration(UnrealTargetPlatform Platform, UnrealTargetConfiguration Configuration, string TargetName, TargetType TargetType, bool bBuildDeveloperTools, bool bBuildRequiresCookedData, [NotNullWhen(false)] out string? invalidReason)
 		{
 			// Check the platform is allowed
 			// important note: we don't check the length of the platform allow list, because if an unknown platform was read in, but was not valid, the 
@@ -591,36 +607,42 @@ namespace UnrealBuildTool
 
 			if (PlatformAllowList != null && !PlatformAllowList.Contains(Platform))
 			{
+				invalidReason = $"PlatformAllowList does not include {Platform}";
 				return false;
 			}
 
 			// Check the platform is not denied
 			if (PlatformDenyList != null && PlatformDenyList.Contains(Platform))
 			{
+				invalidReason = $"PlatformDenyList includes {Platform}";
 				return false;
 			}
 
 			// Check the target is allowed
 			if (TargetAllowList != null && TargetAllowList.Length > 0 && !TargetAllowList.Contains(TargetType))
 			{
+				invalidReason = $"TargetAllowList does not include {TargetType}";
 				return false;
 			}
 
 			// Check the target is not denied
 			if (TargetDenyList != null && TargetDenyList.Contains(TargetType))
 			{
+				invalidReason = $"PlatformDenyList includes {TargetType}";
 				return false;
 			}
 
 			// Check the target configuration is allowed
 			if (TargetConfigurationAllowList != null && TargetConfigurationAllowList.Length > 0 && !TargetConfigurationAllowList.Contains(Configuration))
 			{
+				invalidReason = $"TargetConfigurationAllowList does not include {Configuration}";
 				return false;
 			}
 
 			// Check the target configuration is not denied
 			if (TargetConfigurationDenyList != null && TargetConfigurationDenyList.Contains(Configuration))
 			{
+				invalidReason = $"TargetConfigurationDenyList includes {Configuration}";
 				return false;
 			}
 
@@ -629,12 +651,19 @@ namespace UnrealBuildTool
 				// Check the program name is on the allow list. Note that this behavior is slightly different to other allow/deny checks; we will allow a module of any type if it's explicitly allowed for this program.
 				if (ProgramAllowList != null && ProgramAllowList.Length > 0)
 				{
-					return ProgramAllowList.Contains(TargetName);
+					if (!ProgramAllowList.Contains(TargetName))
+					{
+						invalidReason = $"ProgramAllowList does not include {TargetName}";
+						return false;
+					}
+					invalidReason = null;
+					return true;
 				}
 
 				// Check the program name is not denied
 				if (ProgramDenyList != null && ProgramDenyList.Contains(TargetName))
 				{
+					invalidReason = $"ProgramDenyList includes {TargetName}";
 					return false;
 				}
 			}
@@ -643,17 +672,20 @@ namespace UnrealBuildTool
 				// Check that the TargetName is allowed
 				if (GameTargetAllowList != null && GameTargetAllowList.Length > 0 && !GameTargetAllowList.Contains(TargetName))
 				{
+					invalidReason = $"GameTargetAllowList does not include {TargetName}";
 					return false;
 				}
 				
 				// Check that the TargetName is not denied
 				if (GameTargetDenyList != null && GameTargetDenyList.Contains(TargetName))
 				{
+					invalidReason = $"GameTargetDenyList includes {TargetName}";
 					return false;
 				}
 			}
-			
+
 			// Check the module is compatible with this target.
+			invalidReason = $"Module Type {Type} not supported for Target Type {TargetType}";
 			switch (Type)
 			{
 				case ModuleHostType.Runtime:
