[UBT]
* Fixed paths for vfs files used by shared pch

diff --git a/Engine/Source/Programs/UnrealBuildTool/Configuration/UEBuildModuleCPP.cs b/Engine/Source/Programs/UnrealBuildTool/Configuration/UEBuildModuleCPP.cs
--- a/Engine/Source/Programs/UnrealBuildTool/Configuration/UEBuildModuleCPP.cs
+++ b/Engine/Source/Programs/UnrealBuildTool/Configuration/UEBuildModuleCPP.cs
@@ -1383,7 +1383,7 @@ namespace UnrealBuildTool
 				// Setup clang virtual file system overlay file if enabled
 				if (Rules.Target.bCompilerVFS && (!Rules.Target.Platform.IsInGroup(UnrealPlatformGroup.Microsoft) || Rules.Target.WindowsPlatform.Compiler.IsClang()))
 				{
-					CompileEnvironment.VfsOverlayFile = CreateVfsOverlayFile(Template.BaseCompileEnvironment, IntermediateDirectory, $"SharedVFS.{Template.Module.Name}{Variant}.json", Graph, allowAsync: false);
+					CompileEnvironment.VfsOverlayFile = CreateVfsOverlayFile(Template.BaseCompileEnvironment, Template.OutputDir, $"SharedVFS.{Template.Module.Name}{Variant}.json", Graph, allowAsync: false);
 					CompileEnvironment.AdditionalPrerequisites.Add(CompileEnvironment.VfsOverlayFile);
 				}
 
