[UBA]
* Removed pch code from UbaExecutor.cs because we will most likely never use it (we can readd it if we will)

diff --git a/Engine/Source/Programs/UnrealBuildTool/Executors/UnrealBuildAccelerator/UBAExecutor.cs b/Engine/Source/Programs/UnrealBuildTool/Executors/UnrealBuildAccelerator/UBAExecutor.cs
--- a/Engine/Source/Programs/UnrealBuildTool/Executors/UnrealBuildAccelerator/UBAExecutor.cs
+++ b/Engine/Source/Programs/UnrealBuildTool/Executors/UnrealBuildAccelerator/UBAExecutor.cs
@@ -257,7 +257,7 @@ namespace UnrealBuildTool
 			{
 				return Task.Factory.StartNew(() =>
 				{
-					ProcessStartInfo startInfo = _executor.GetActionStartInfo(action, out FileItem? pchItem);
+					ProcessStartInfo startInfo = _executor.GetActionStartInfo(action);
 					uint bucket = UBAExecutor.GetActionCacheBucket(action);
 					using (IRootPaths rootPaths = _executor.GetActionRootPaths(action))
 					{
@@ -650,7 +650,7 @@ namespace UnrealBuildTool
 			!_forcedRetryActions.ContainsKey(action) &&
 			!ForceLocalNoDetour(action);
 
-		ProcessStartInfo GetActionStartInfo(LinkedAction action, out FileItem? pchItem)
+		ProcessStartInfo GetActionStartInfo(LinkedAction action)
 		{
 			string description = action.StatusDescription;
 			if (!String.IsNullOrEmpty(action.CommandDescription))
@@ -669,36 +669,6 @@ namespace UnrealBuildTool
 				LogFile = UBAConfig.bLogEnabled ? (action.Inner.ProducedItems.First().Location.GetFileName() + ".log") : null,
 			};
 
-			pchItem = null;
-
-			// This is not used. Should probably be deleted.
-			/*
-			if (_cacheClient != null && UBAConfig.bWriteCache && action.ArtifactMode.HasFlag(ArtifactMode.Enabled))
-			{
-				startInfo.TrackInputs = true;
-			}
-			else
-			{
-				// TODO: Revisit this code. This was added to make non-deterministic pch deterministic from a caskey perspective
-				// It is not used atm.
-
-				bool usingLtcg = true; // ltcg linking checks what pch was used and it seems like it needs to be identical in other ways than timestamp
-				pchItem = action.Inner.ProducedItems.FirstOrDefault(item => item.Name.EndsWith(".pch", StringComparison.OrdinalIgnoreCase));
-				if (pchItem != null)
-				{
-					startInfo.Priority = System.Diagnostics.ProcessPriorityClass.AboveNormal;
-					if (!usingLtcg && action.ArtifactMode.HasFlag(ArtifactMode.PropagateInputs))
-					{
-						startInfo.TrackInputs = true;
-					}
-					else
-					{
-						pchItem = null;
-					}
-				}
-			}
-			*/
-
 			return startInfo;
 		}
 
@@ -833,7 +803,7 @@ namespace UnrealBuildTool
 				}
 				bool enableDetour = !ForceLocalNoDetour(action) && action.bCanExecuteInUBA && !_forcedRetryActions.ContainsKey(action);
 
-				ProcessStartInfo startInfo = GetActionStartInfo(action, out FileItem? pchItem);
+				ProcessStartInfo startInfo = GetActionStartInfo(action);
 				using (IProcess process = _session!.RunProcess(startInfo, false, null, enableDetour))
 				{
 					Interlocked.Add(ref _localProcessedActions, 1);
@@ -845,6 +815,7 @@ namespace UnrealBuildTool
 						return Task.CompletedTask;
 					}
 
+					// If not detoured we manually have to report created files
 					if (!enableDetour && process.ExitCode == 0)
 					{
 						_session!.RegisterNewFiles(action.ProducedItems.Where(x => FileReference.Exists(x.Location)).Select(x => x.FullName).ToArray());
@@ -860,7 +831,7 @@ namespace UnrealBuildTool
 					List<string> logLines = process.LogLines;
 					RemoveLogLineSpam(logLines);
 					string? additionalDescription = !enableDetour ? "(UBA disabled)" : null;
-					ActionFinished(queue, new ExecuteResults(logLines, process.ExitCode, executionTime, processorTime, additionalDescription), action, pchItem, process);
+					ActionFinished(queue, new ExecuteResults(logLines, process.ExitCode, executionTime, processorTime, additionalDescription), action, process);
 				}
 				return Task.CompletedTask;
 			};
@@ -923,7 +894,7 @@ namespace UnrealBuildTool
 					}
 				}
 
-				ProcessStartInfo startInfo = GetActionStartInfo(action, out FileItem? pchItem);
+				ProcessStartInfo startInfo = GetActionStartInfo(action);
 				_session!.RunProcessRemote(startInfo, (s, e) =>
 				{
 					if (e.ExitCode == 99999) // Process was cancelled by executor
@@ -977,7 +948,7 @@ namespace UnrealBuildTool
 					TimeSpan executionTime = e.TotalWallTime;
 					List<string> logLines = e.LogLines;
 					logLines.RemoveAll((line) => line.StartsWith("   Creating library ", StringComparison.OrdinalIgnoreCase) && line.EndsWith(".exp", StringComparison.OrdinalIgnoreCase));
-					ActionFinished(queue, new ExecuteResults(logLines, e.ExitCode, executionTime, processorTime, additionalDescription), action, pchItem, process);
+					ActionFinished(queue, new ExecuteResults(logLines, e.ExitCode, executionTime, processorTime, additionalDescription), action, process);
 				}, action.Weight, knownInputs, knownInputsCount, action.bCanExecuteInUBACrossArchitecture);
 				return Task.CompletedTask;
 			};
@@ -1002,7 +973,7 @@ namespace UnrealBuildTool
 			queue.OnActionCompleted(action, false, cancelResults);
 		}
 
-		protected void ActionFinished(ImmediateActionQueue queue, ExecuteResults results, LinkedAction action, FileItem? pchItem = null, IProcess? process = null)
+		protected void ActionFinished(ImmediateActionQueue queue, ExecuteResults results, LinkedAction action, IProcess? process = null)
 		{
 			if (_bIsCancelled)
 			{
@@ -1011,10 +982,6 @@ namespace UnrealBuildTool
 			}
 
 			bool success = results.ExitCode == 0;
-			if (pchItem != null && process != null && success)
-			{
-				_session!.SetCustomCasKeyFromTrackedInputs(pchItem.FullName, action.WorkingDirectory.FullName, process);
-			}
 
 			if (results.ExitCode == 0 && action.bForceWarningsAsError && results.LogLines.Any(x => x.Contains("): warning: ")))
 			{
