Update NSAppTransportSecurity support for iOS17+.
Also, when ZenServer is enabled, allow http over 169.245.0.0/16 via USB for non-shipping builds

diff --git a/Engine/Source/Programs/UnrealBuildTool/Platform/IOS/UEDeployIOS.cs b/Engine/Source/Programs/UnrealBuildTool/Platform/IOS/UEDeployIOS.cs
--- a/Engine/Source/Programs/UnrealBuildTool/Platform/IOS/UEDeployIOS.cs
+++ b/Engine/Source/Programs/UnrealBuildTool/Platform/IOS/UEDeployIOS.cs
@@ -216,7 +216,7 @@ namespace UnrealBuildTool
 			return MinVersionToReturn;
 		}
 
-		public static void WritePlistFile(FileReference PlistFile, DirectoryReference? ProjectLocation, UnrealPluginLanguage? UPL, string GameName, bool bIsUnrealGame, string ProjectName, ILogger Logger)
+		public static void WritePlistFile(FileReference PlistFile, DirectoryReference? ProjectLocation, UnrealPluginLanguage? UPL, UnrealTargetConfiguration Config, string GameName, bool bIsUnrealGame, string ProjectName, ILogger Logger)
 		{
 			ConfigHierarchy Ini = ConfigCache.ReadHierarchy(ConfigHierarchyType.Engine, ProjectLocation, UnrealTargetPlatform.IOS);
 			// required capabilities
@@ -276,6 +276,13 @@ namespace UnrealBuildTool
 			bool bDisableHTTPS;
 			Ini.GetBool("/Script/IOSRuntimeSettings.IOSRuntimeSettings", "bDisableHTTPS", out bDisableHTTPS);
 
+			bool bUseZenStore = false;
+			if (Config != UnrealTargetConfiguration.Shipping)
+			{
+				ConfigHierarchy PlatformGameConfig = ConfigCache.ReadHierarchy(ConfigHierarchyType.Game, ProjectLocation, UnrealTargetPlatform.IOS);
+				PlatformGameConfig.GetBool("/Script/UnrealEd.ProjectPackagingSettings", "bUseZenStore", out bUseZenStore);
+			}
+
 			// bundle display name
 			string BundleDisplayName;
 			Ini.GetString("/Script/IOSRuntimeSettings.IOSRuntimeSettings", "BundleDisplayName", out BundleDisplayName);
@@ -407,11 +414,25 @@ namespace UnrealBuildTool
 				Text.AppendLine(String.Format("\t<string>{0}</string>", LocationWhenInUseDescription));
 			}
 			// disable HTTPS requirement
-			if (bDisableHTTPS)
+			if (bDisableHTTPS || bUseZenStore)
 			{
 				Text.AppendLine("\t<key>NSAppTransportSecurity</key>");
 				Text.AppendLine("\t\t<dict>");
 				Text.AppendLine("\t\t\t<key>NSAllowsArbitraryLoads</key><true/>");
+				Text.AppendLine("\t\t\t<key>NSAllowsArbitraryLoadsInWebContent</key><true/>");
+				Text.AppendLine("\t\t\t<key>NSAllowsLocalNetworking</key><true/>");
+				
+				if (bUseZenStore)
+				{
+					Text.AppendLine("\t\t\t<key>NSExceptionDomains</key>");
+					Text.AppendLine("\t\t\t<dict>");
+					Text.AppendLine("\t\t\t\t<key>169.245.0.0/16</key>");
+					Text.AppendLine("\t\t\t\t<dict>");
+					Text.AppendLine("\t\t\t\t\t<key>NSIncludesSubdomains</key><true/>");
+					Text.AppendLine("\t\t\t\t\t<key>NSExceptionAllowsInsecureHTTPLoads</key><true/>");
+					Text.AppendLine("\t\t\t\t</dict>");
+					Text.AppendLine("\t\t\t</dict>");
+				}
 				Text.AppendLine("\t\t</dict>");
 			}
 
@@ -504,7 +525,7 @@ namespace UnrealBuildTool
 			VersionUtilities.BuildDirectory = BuildDirectory;
 			VersionUtilities.GameName = GameName;
 
-			WritePlistFile(new FileReference(PListFile), DirRef, UPL, GameName, bIsUnrealGame, ProjectName, Logger);
+			WritePlistFile(new FileReference(PListFile), DirRef, UPL, Config, GameName, bIsUnrealGame, ProjectName, Logger);
 
 			if (BuildHostPlatform.Current.Platform == UnrealTargetPlatform.Mac && !bBuildAsFramework)
 			{
@@ -612,7 +633,7 @@ namespace UnrealBuildTool
 			// disable https requirement
 			bool bDisableHTTPS;
 			Ini.GetBool("/Script/IOSRuntimeSettings.IOSRuntimeSettings", "bDisableHTTPS", out bDisableHTTPS);
-
+			
 			// short version string
 			string BundleShortVersion;
 			Ini.GetString("/Script/IOSRuntimeSettings.IOSRuntimeSettings", "VersionInfo", out BundleShortVersion);
