EpicGames.Core: Obsolete RuntimePlatform.IsWindows in favor of OperatingSystem.IsWindows()

#rnx

diff --git a/Engine/Source/Programs/AutomationTool/Android/AndroidFileClient.cs b/Engine/Source/Programs/AutomationTool/Android/AndroidFileClient.cs
--- a/Engine/Source/Programs/AutomationTool/Android/AndroidFileClient.cs
+++ b/Engine/Source/Programs/AutomationTool/Android/AndroidFileClient.cs
@@ -111,7 +111,7 @@ namespace AutomationTool
 
 			public OptimalADB(string ADBExecutable = "", int inPort = DefaultPort)
 			{
-				Executable = (ADBExecutable != "") ? ADBExecutable : Environment.ExpandEnvironmentVariables("%ANDROID_HOME%/platform-tools/adb" + (RuntimePlatform.IsWindows ? ".exe" : ""));
+				Executable = (ADBExecutable != "") ? ADBExecutable : Environment.ExpandEnvironmentVariables("%ANDROID_HOME%/platform-tools/adb" + (OperatingSystem.IsWindows() ? ".exe" : ""));
 				Port = inPort;
 			}
 
@@ -561,7 +561,7 @@ namespace AutomationTool
 
 		public static IProcessResult RunAdbCommand(string SerialNumber, string Args, string Input = null, CommandUtils.ERunOptions Options = CommandUtils.ERunOptions.NoLoggingOfRunCommand, bool bShouldLogCommand = false)
 		{
-			string AdbCommand = Environment.ExpandEnvironmentVariables("%ANDROID_HOME%/platform-tools/adb" + (RuntimePlatform.IsWindows ? ".exe" : ""));
+			string AdbCommand = Environment.ExpandEnvironmentVariables("%ANDROID_HOME%/platform-tools/adb" + (OperatingSystem.IsWindows() ? ".exe" : ""));
 			if (Options.HasFlag(CommandUtils.ERunOptions.AllowSpew) || Options.HasFlag(CommandUtils.ERunOptions.SpewIsVerbose))
 			{
 				LastSpewFilename = "";
diff --git a/Engine/Source/Programs/AutomationTool/Android/AndroidPlatform.Automation.cs b/Engine/Source/Programs/AutomationTool/Android/AndroidPlatform.Automation.cs
--- a/Engine/Source/Programs/AutomationTool/Android/AndroidPlatform.Automation.cs
+++ b/Engine/Source/Programs/AutomationTool/Android/AndroidPlatform.Automation.cs
@@ -773,7 +773,7 @@ public class AndroidPlatform : Platform
 
 	public static string GetStorageQueryCommand(bool bForcePC = false)
     {
-		if (!bForcePC && !RuntimePlatform.IsWindows)
+		if (!bForcePC && !OperatingSystem.IsWindows())
 		{
 			return "shell 'echo $EXTERNAL_STORAGE'";
 		}
@@ -1493,7 +1493,7 @@ public class AndroidPlatform : Platform
 					File.WriteAllText(SymbolizeBatchName, string.Join(LineEnding, UniversalSymbolizeBatchLines) + LineEnding);
 				}
 
-				if (!RuntimePlatform.IsWindows)
+				if (!OperatingSystem.IsWindows())
 				{
 					if (bHaveAPK)
 					{
@@ -2243,7 +2243,7 @@ public class AndroidPlatform : Platform
 
 	private static IProcessResult RunAdbCommand(string SerialNumber, string Args, string Input = null, ERunOptions Options = ERunOptions.Default, bool bShouldLogCommand = false)
 	{
-		string AdbCommand = Environment.ExpandEnvironmentVariables("%ANDROID_HOME%/platform-tools/adb" + (RuntimePlatform.IsWindows ? ".exe" : ""));
+		string AdbCommand = Environment.ExpandEnvironmentVariables("%ANDROID_HOME%/platform-tools/adb" + (OperatingSystem.IsWindows() ? ".exe" : ""));
 		if (Options.HasFlag(ERunOptions.AllowSpew) || Options.HasFlag(ERunOptions.SpewIsVerbose))
 		{
 			LastSpewFilename = "";
@@ -2254,7 +2254,7 @@ public class AndroidPlatform : Platform
 
 	private string RunAndLogAdbCommand(ProjectParams Params, string SerialNumber, string Args, out int SuccessCode)
 	{
-		string AdbCommand = Environment.ExpandEnvironmentVariables("%ANDROID_HOME%/platform-tools/adb" + (RuntimePlatform.IsWindows ? ".exe" : ""));
+		string AdbCommand = Environment.ExpandEnvironmentVariables("%ANDROID_HOME%/platform-tools/adb" + (OperatingSystem.IsWindows() ? ".exe" : ""));
 		LastSpewFilename = "";
 		return RunAndLog(CmdEnv, AdbCommand, GetAdbCommandLine(SerialNumber, Args), out SuccessCode, SpewFilterCallback: new ProcessResult.SpewFilterCallbackType(ADBSpewFilter));
 	}
@@ -3939,7 +3939,7 @@ public class AndroidPlatform : Platform
 		uint BestVersion = 0;
 		foreach (string CandidateDir in Subdirs)
 		{
-			string AaptFilename = Path.Combine(CandidateDir, RuntimePlatform.IsWindows ? "aapt.exe" : "aapt");
+			string AaptFilename = Path.Combine(CandidateDir, OperatingSystem.IsWindows() ? "aapt.exe" : "aapt");
 			uint RevisionValue = 0;
 
 			if (File.Exists(AaptFilename))
diff --git a/Engine/Source/Programs/AutomationTool/AutomationUtils/Automation.cs b/Engine/Source/Programs/AutomationTool/AutomationUtils/Automation.cs
--- a/Engine/Source/Programs/AutomationTool/AutomationUtils/Automation.cs
+++ b/Engine/Source/Programs/AutomationTool/AutomationUtils/Automation.cs
@@ -43,7 +43,7 @@ namespace AutomationTool
 		static void Domain_ProcessExit(object sender, EventArgs e)
 		{
 			// Kill all spawned processes (Console instead of Log because logging is closed at this time anyway)
-			if (ShouldKillProcesses && RuntimePlatform.IsWindows)
+			if (ShouldKillProcesses && OperatingSystem.IsWindows())
 			{			
 				ProcessManager.KillAll();
 			}
@@ -226,7 +226,7 @@ namespace AutomationTool
 				// Try to kill process before app domain exits to leave the other KillAll call to extreme edge cases
 				NoThrow(() =>
 				{
-					if (ShouldKillProcesses && RuntimePlatform.IsWindows) ProcessManager.KillAll();
+					if (ShouldKillProcesses && OperatingSystem.IsWindows()) ProcessManager.KillAll();
 				}, "Kill All Processes");
 			}
 		}
diff --git a/Engine/Source/Programs/AutomationTool/AutomationUtils/CommandEnvironment.cs b/Engine/Source/Programs/AutomationTool/AutomationUtils/CommandEnvironment.cs
--- a/Engine/Source/Programs/AutomationTool/AutomationUtils/CommandEnvironment.cs
+++ b/Engine/Source/Programs/AutomationTool/AutomationUtils/CommandEnvironment.cs
@@ -130,7 +130,7 @@ namespace AutomationTool
 
 			RobocopyExe = GetSystemExePath("robocopy.exe");
 			MountExe = GetSystemExePath("mount.exe");
-			CmdExe = RuntimePlatform.IsWindows ? GetSystemExePath("cmd.exe") : "/bin/sh";
+			CmdExe = OperatingSystem.IsWindows() ? GetSystemExePath("cmd.exe") : "/bin/sh";
 			MallocNanoZone = "0";
 			CommandUtils.SetEnvVar(EnvVarNames.MacMallocNanoZone, MallocNanoZone);
 
diff --git a/Engine/Source/Programs/AutomationTool/AutomationUtils/CommandUtils.cs b/Engine/Source/Programs/AutomationTool/AutomationUtils/CommandUtils.cs
--- a/Engine/Source/Programs/AutomationTool/AutomationUtils/CommandUtils.cs
+++ b/Engine/Source/Programs/AutomationTool/AutomationUtils/CommandUtils.cs
@@ -2250,11 +2250,11 @@ namespace AutomationTool
 				}
 				else
 				{
-					if (RuntimePlatform.IsMac)
+					if (OperatingSystem.IsMacOS())
 					{
 						CachedRootBuildStorageDirectory = CombinePaths("/Volumes", "Builds");
 					}
-					else if (RuntimePlatform.IsLinux)
+					else if (OperatingSystem.IsLinux())
 					{
 						CachedRootBuildStorageDirectory = CombinePaths("/mnt", "Builds");
 					}
@@ -2477,7 +2477,7 @@ namespace AutomationTool
 		public static IEnumerable<string> LegacyUnzipFiles(string ZipFileName, string BaseDirectory, bool OverwriteFiles = true)
 		{
 			List<string> OutputFileNames = new List<string>();
-			if (!RuntimePlatform.IsWindows)
+			if (!OperatingSystem.IsWindows())
 			{
 				CommandUtils.CreateDirectory(BaseDirectory);
 
diff --git a/Engine/Source/Programs/AutomationTool/AutomationUtils/UnrealBuild.cs b/Engine/Source/Programs/AutomationTool/AutomationUtils/UnrealBuild.cs
--- a/Engine/Source/Programs/AutomationTool/AutomationUtils/UnrealBuild.cs
+++ b/Engine/Source/Programs/AutomationTool/AutomationUtils/UnrealBuild.cs
@@ -737,7 +737,7 @@ namespace AutomationTool
 		[Obsolete("Deprecated in UE5.1; use UnrealBuildToolDll")]
 		public static string GetUBTExecutable()
 		{
-			return CommandUtils.CombinePaths(CommandUtils.CmdEnv.LocalRoot, @"Engine/Binaries/DotNET/UnrealBuildTool/UnrealBuildTool" + (RuntimePlatform.IsWindows ? ".exe" : ""));
+			return CommandUtils.CombinePaths(CommandUtils.CmdEnv.LocalRoot, @"Engine/Binaries/DotNET/UnrealBuildTool/UnrealBuildTool" + (OperatingSystem.IsWindows() ? ".exe" : ""));
 		}
 
 		[Obsolete("Deprecated in UE5.1; use UnrealBuildToolDll")]
diff --git a/Engine/Source/Programs/AutomationTool/BuildGraph/BgScriptReader.cs b/Engine/Source/Programs/AutomationTool/BuildGraph/BgScriptReader.cs
--- a/Engine/Source/Programs/AutomationTool/BuildGraph/BgScriptReader.cs
+++ b/Engine/Source/Programs/AutomationTool/BuildGraph/BgScriptReader.cs
@@ -1052,7 +1052,7 @@ namespace AutomationTool
 				if (ValidateName(element, name))
 				{
 					string envVarName = name;
-					if (!RuntimePlatform.IsWindows)
+					if (!OperatingSystem.IsWindows())
 					{
 						// Non-windows platforms don't allow dashes in variable names. The engine platform layer substitutes underscores for them.
 						envVarName = envVarName.Replace("-", "_", StringComparison.Ordinal);
diff --git a/Engine/Source/Programs/AutomationTool/Gauntlet/Platform/Android/Gauntlet.TargetDeviceAndroid.cs b/Engine/Source/Programs/AutomationTool/Gauntlet/Platform/Android/Gauntlet.TargetDeviceAndroid.cs
--- a/Engine/Source/Programs/AutomationTool/Gauntlet/Platform/Android/Gauntlet.TargetDeviceAndroid.cs
+++ b/Engine/Source/Programs/AutomationTool/Gauntlet/Platform/Android/Gauntlet.TargetDeviceAndroid.cs
@@ -637,7 +637,7 @@ namespace Gauntlet
 			CommandUtils.ERunOptions RunOptions = CommandUtils.ERunOptions.AppMustExist | CommandUtils.ERunOptions.NoWaitForExit | CommandUtils.ERunOptions.SpewIsVerbose;
 			RunOptions |= Log.IsVeryVerbose ? ERunOptions.AllowSpew : ERunOptions.NoLoggingOfRunCommand;
 
-			string AdbCommand = Environment.ExpandEnvironmentVariables("%ANDROID_HOME%/platform-tools/adb" + (RuntimePlatform.IsWindows ? ".exe" : ""));
+			string AdbCommand = Environment.ExpandEnvironmentVariables("%ANDROID_HOME%/platform-tools/adb" + (OperatingSystem.IsWindows() ? ".exe" : ""));
 			return new LongProcessResult(AdbCommand, Args, RunOptions, LocalCache: LocalCache);
 		}
 
@@ -1681,7 +1681,7 @@ namespace Gauntlet
 				throw new InvalidOperationException("A screen recording session is already in progress.");
 			}
 
-			string AdbCommand = Environment.ExpandEnvironmentVariables("%ANDROID_HOME%/platform-tools/adb" + (RuntimePlatform.IsWindows ? ".exe" : ""));
+			string AdbCommand = Environment.ExpandEnvironmentVariables("%ANDROID_HOME%/platform-tools/adb" + (OperatingSystem.IsWindows() ? ".exe" : ""));
 
 			ADBProcess = new Process
 			{
diff --git a/Engine/Source/Programs/AutomationTool/IOS/IOSPlatform.Automation.cs b/Engine/Source/Programs/AutomationTool/IOS/IOSPlatform.Automation.cs
--- a/Engine/Source/Programs/AutomationTool/IOS/IOSPlatform.Automation.cs
+++ b/Engine/Source/Programs/AutomationTool/IOS/IOSPlatform.Automation.cs
@@ -1514,7 +1514,7 @@ Console.WriteLine("target receipt executable path: {0}", InExecutablePath);
 			PlatformGameConfig.GetBool("/Script/IOSRuntimeSettings.IOSRuntimeSettings", "bGenerateXCArchive", out bXCArchive);
 		}
 
-		if (bXCArchive && !RuntimePlatform.IsWindows)
+		if (bXCArchive && !OperatingSystem.IsWindows())
 		{
 			// Always put the archive in the current user's Library/Developer/Xcode/Archives path if not on the build machine
 			string ArchivePath = "[REDACTED_USER_HOME] + Environment.UserName + "/Library/Developer/Xcode/Archives";
@@ -1667,7 +1667,7 @@ Console.WriteLine("target receipt executable path: {0}", InExecutablePath);
 			Text.AppendLine("</plist>");
 			File.WriteAllText(Path.Combine(ArchiveName, "Info.plist"), Text.ToString());
 		}
-		else if (bXCArchive && RuntimePlatform.IsWindows)
+		else if (bXCArchive && OperatingSystem.IsWindows())
 		{
 			Logger.LogWarning("Can not produce an XCArchive on windows");
 		}
diff --git a/Engine/Source/Programs/AutomationTool/Linux/LinuxPlatform.Automation.cs b/Engine/Source/Programs/AutomationTool/Linux/LinuxPlatform.Automation.cs
--- a/Engine/Source/Programs/AutomationTool/Linux/LinuxPlatform.Automation.cs
+++ b/Engine/Source/Programs/AutomationTool/Linux/LinuxPlatform.Automation.cs
@@ -177,7 +177,7 @@ public abstract class BaseLinuxPlatform : Platform
 		// write out the 
 		FileReference.WriteAllText(IntermediateFile, Script.ToString());
 
-		if (!RuntimePlatform.IsWindows)
+		if (!OperatingSystem.IsWindows())
 		{
 			var Result = CommandUtils.Run("env", string.Format("-- \"chmod\" \"+x\" \"{0}\"", IntermediateFile.ToString().Replace("'", "'\"'\"'")));
 			if (Result.ExitCode != 0)
diff --git a/Engine/Source/Programs/AutomationTool/LowLevelTests/Tests/LowLevelTests.WebTestsLowLevelTestsExtension.cs b/Engine/Source/Programs/AutomationTool/LowLevelTests/Tests/LowLevelTests.WebTestsLowLevelTestsExtension.cs
--- a/Engine/Source/Programs/AutomationTool/LowLevelTests/Tests/LowLevelTests.WebTestsLowLevelTestsExtension.cs
+++ b/Engine/Source/Programs/AutomationTool/LowLevelTests/Tests/LowLevelTests.WebTestsLowLevelTestsExtension.cs
@@ -42,8 +42,8 @@ namespace LowLevelTests
 		{
 			ProcessStartInfo StartInfo = new ProcessStartInfo();
 			StartInfo.WorkingDirectory = WebTestsServerDir();
-			StartInfo.FileName = RuntimePlatform.IsWindows ? "cmd.exe" : "/bin/sh";
-			StartInfo.Arguments = RuntimePlatform.IsWindows ? "/c createenv.bat" : "-c './createenv.sh'";
+			StartInfo.FileName = OperatingSystem.IsWindows() ? "cmd.exe" : "/bin/sh";
+			StartInfo.Arguments = OperatingSystem.IsWindows() ? "/c createenv.bat" : "-c './createenv.sh'";
 			StartInfo.UseShellExecute = false;
 			StartInfo.CreateNoWindow = true;
 
@@ -58,7 +58,7 @@ namespace LowLevelTests
 		private void AsyncLaunchWebServerProcess()
 		{
 			string WorkingDir = WebTestsServerDir();
-			string PythonFile = Path.Combine(WorkingDir, "env", RuntimePlatform.IsWindows ? "Scripts" : "bin", RuntimePlatform.IsWindows ? "python.exe" : "python");
+			string PythonFile = Path.Combine(WorkingDir, "env", OperatingSystem.IsWindows() ? "Scripts" : "bin", OperatingSystem.IsWindows() ? "python.exe" : "python");
 
 			ProcessStartInfo StartInfo = new ProcessStartInfo();
 			StartInfo.WorkingDirectory = WorkingDir;
diff --git a/Engine/Source/Programs/AutomationTool/Scripts/BuildHlslcc.Automation.cs b/Engine/Source/Programs/AutomationTool/Scripts/BuildHlslcc.Automation.cs
--- a/Engine/Source/Programs/AutomationTool/Scripts/BuildHlslcc.Automation.cs
+++ b/Engine/Source/Programs/AutomationTool/Scripts/BuildHlslcc.Automation.cs
@@ -517,7 +517,7 @@ class BuildHlslcc : BuildCommand
 	}
 	private static void SetupStaticBuildEnvironment()
 	{
-		if (RuntimePlatform.IsWindows)
+		if (OperatingSystem.IsWindows())
 		{
 			string VS2022Path = GetMsDevExe(WindowsCompiler.VisualStudio2022);
 			if (VS2022Path != null)
diff --git a/Engine/Source/Programs/AutomationTool/Scripts/RunProjectCommand.Automation.cs b/Engine/Source/Programs/AutomationTool/Scripts/RunProjectCommand.Automation.cs
--- a/Engine/Source/Programs/AutomationTool/Scripts/RunProjectCommand.Automation.cs
+++ b/Engine/Source/Programs/AutomationTool/Scripts/RunProjectCommand.Automation.cs
@@ -1001,7 +1001,7 @@ namespace AutomationScripts
 		private static bool StartUnrealTrace()
 		{
 			// [TEMPORARY] - UnrealTrace server is currently only available on Windows
-			if (!RuntimePlatform.IsWindows)
+			if (!OperatingSystem.IsWindows())
 			{
 				return true;
 			}
diff --git a/Engine/Source/Programs/AutomationTool/Scripts/ZipUtils.Automation.cs b/Engine/Source/Programs/AutomationTool/Scripts/ZipUtils.Automation.cs
--- a/Engine/Source/Programs/AutomationTool/Scripts/ZipUtils.Automation.cs
+++ b/Engine/Source/Programs/AutomationTool/Scripts/ZipUtils.Automation.cs
@@ -159,7 +159,7 @@ public class ZipUtils : BuildCommand
 					}
 				}
 
-				if (!RuntimePlatform.IsWindows && CommandUtils.IsProbablyAMacOrIOSExe(OutputFileName))
+				if (!OperatingSystem.IsWindows() && CommandUtils.IsProbablyAMacOrIOSExe(OutputFileName))
 				{
 					FixUnixFilePermissions(OutputFileName);
 				}
diff --git a/Engine/Source/Programs/AutomationTool/Turnkey/Commands/CreateBuild.cs b/Engine/Source/Programs/AutomationTool/Turnkey/Commands/CreateBuild.cs
--- a/Engine/Source/Programs/AutomationTool/Turnkey/Commands/CreateBuild.cs
+++ b/Engine/Source/Programs/AutomationTool/Turnkey/Commands/CreateBuild.cs
@@ -67,7 +67,7 @@ namespace Turnkey.Commands
 			string OutputDir = TurnkeyUtils.ParseParamValue("OutputDir", null, CommandOptions);
 			if (OutputDir == null)
 			{
-				//if (RuntimePlatform.IsWindows)
+				//if (OperatingSystem.IsWindows())
 				//{
 				//	string ChosenFile = null;
 				//	System.Threading.Thread t = new System.Threading.Thread(x =>
diff --git a/Engine/Source/Programs/AutomationTool/Turnkey/Commands/ExecuteBuild.cs b/Engine/Source/Programs/AutomationTool/Turnkey/Commands/ExecuteBuild.cs
--- a/Engine/Source/Programs/AutomationTool/Turnkey/Commands/ExecuteBuild.cs
+++ b/Engine/Source/Programs/AutomationTool/Turnkey/Commands/ExecuteBuild.cs
@@ -145,7 +145,7 @@ namespace Turnkey.Commands
 					// if the -outputdir option was specified, use that, otherwise ask user for directory
 					if (OutputDir == null)
 					{
-						//if (RuntimePlatform.IsWindows)
+						//if (OperatingSystem.IsWindows())
 						//{
 						//	string ChosenFile = null;
 						//	System.Threading.Thread t = new System.Threading.Thread(x =>
diff --git a/Engine/Source/Programs/AutomationTool/Turnkey/TurnkeyUtils.cs b/Engine/Source/Programs/AutomationTool/Turnkey/TurnkeyUtils.cs
--- a/Engine/Source/Programs/AutomationTool/Turnkey/TurnkeyUtils.cs
+++ b/Engine/Source/Programs/AutomationTool/Turnkey/TurnkeyUtils.cs
@@ -293,7 +293,7 @@ namespace Turnkey
 				ProjectNames.AddRange(ProjectsByType[CurrentMode].Select(x => (CurrentMode == ProjectMode.Recent) ? x.FullName : x.GetFileNameWithoutAnyExtensions()));
 
 				// and finally a manual entry option
-				if (RuntimePlatform.IsWindows)
+				if (OperatingSystem.IsWindows())
 				{
 					ProjectNames.Add("[Browse...]");
 				}
diff --git a/Engine/Source/Programs/Horde/HordeAgent/Commands/Service/ServiceInstall.cs b/Engine/Source/Programs/Horde/HordeAgent/Commands/Service/ServiceInstall.cs
--- a/Engine/Source/Programs/Horde/HordeAgent/Commands/Service/ServiceInstall.cs
+++ b/Engine/Source/Programs/Horde/HordeAgent/Commands/Service/ServiceInstall.cs
@@ -55,7 +55,7 @@ namespace HordeAgent.Commands.Service
 		/// <returns>Exit code</returns>
 		public override Task<int> ExecuteAsync(ILogger logger)
 		{
-			if (!RuntimePlatform.IsWindows)
+			if (!OperatingSystem.IsWindows())
 			{
 				logger.LogError("This command requires Windows");
 				return Task.FromResult(1);
diff --git a/Engine/Source/Programs/Horde/HordeAgent/Leases/Handlers/UpgradeHandler.cs b/Engine/Source/Programs/Horde/HordeAgent/Leases/Handlers/UpgradeHandler.cs
--- a/Engine/Source/Programs/Horde/HordeAgent/Leases/Handlers/UpgradeHandler.cs
+++ b/Engine/Source/Programs/Horde/HordeAgent/Leases/Handlers/UpgradeHandler.cs
@@ -87,7 +87,7 @@ namespace HordeAgent.Leases.Handlers
 					}
 
 					executable = "HordeAgent";
-					if (RuntimePlatform.IsWindows)
+					if (OperatingSystem.IsWindows())
 					{
 						executable += ".exe";
 					}
diff --git a/Engine/Source/Programs/Horde/Plugins/Compute/HordeServer.Compute/Agents/Relay/Nftables.cs b/Engine/Source/Programs/Horde/Plugins/Compute/HordeServer.Compute/Agents/Relay/Nftables.cs
--- a/Engine/Source/Programs/Horde/Plugins/Compute/HordeServer.Compute/Agents/Relay/Nftables.cs
+++ b/Engine/Source/Programs/Horde/Plugins/Compute/HordeServer.Compute/Agents/Relay/Nftables.cs
@@ -139,7 +139,7 @@ public class Nftables
 
 	private async Task<string> GetNftablesVersionAsync(CancellationToken cancellationToken)
 	{
-		if (!RuntimePlatform.IsLinux)
+		if (!OperatingSystem.IsLinux())
 		{
 			throw new NftablesException("nftables can only run under Linux");
 		}
diff --git a/Engine/Source/Programs/Shared/EpicGames.Build/System/DirectoryItem.cs b/Engine/Source/Programs/Shared/EpicGames.Build/System/DirectoryItem.cs
--- a/Engine/Source/Programs/Shared/EpicGames.Build/System/DirectoryItem.cs
+++ b/Engine/Source/Programs/Shared/EpicGames.Build/System/DirectoryItem.cs
@@ -49,7 +49,7 @@ namespace UnrealBuildBase
 		private DirectoryItem(DirectoryReference Location, DirectoryInfo Info)
 		{
 			this.Location = Location;
-			if (RuntimePlatform.IsWindows)
+			if (OperatingSystem.IsWindows())
 			{
 				this.Info = new Lazy<DirectoryInfo>(Info);
 			}
diff --git a/Engine/Source/Programs/Shared/EpicGames.Build/System/FileItem.cs b/Engine/Source/Programs/Shared/EpicGames.Build/System/FileItem.cs
--- a/Engine/Source/Programs/Shared/EpicGames.Build/System/FileItem.cs
+++ b/Engine/Source/Programs/Shared/EpicGames.Build/System/FileItem.cs
@@ -44,7 +44,7 @@ namespace UnrealBuildBase
 		private FileItem(FileReference Location, FileInfo Info)
 		{
 			this.Location = Location;
-			if (RuntimePlatform.IsWindows)
+			if (OperatingSystem.IsWindows())
 			{
 				this.Info = new Lazy<FileInfo>(Info);
 			}
diff --git a/Engine/Source/Programs/Shared/EpicGames.Build/UEBuildPlatformSDK.cs b/Engine/Source/Programs/Shared/EpicGames.Build/UEBuildPlatformSDK.cs
--- a/Engine/Source/Programs/Shared/EpicGames.Build/UEBuildPlatformSDK.cs
+++ b/Engine/Source/Programs/Shared/EpicGames.Build/UEBuildPlatformSDK.cs
@@ -1336,15 +1336,15 @@ namespace EpicGames.Core
 
 		private static string GetAutoSDKHostPlatform()
 		{
-			if (RuntimePlatform.IsWindows)
+			if (OperatingSystem.IsWindows())
 			{
 				return "Win64";
 			}
-			else if (RuntimePlatform.IsMac)
+			else if (OperatingSystem.IsMacOS())
 			{
 				return "Mac";
 			}
-			else if (RuntimePlatform.IsLinux)
+			else if (OperatingSystem.IsLinux())
 			{
 				return "Linux";
 			}
@@ -1710,7 +1710,7 @@ namespace EpicGames.Core
 		/// <param name="Hook">Hook type</param>
 		protected virtual string GetHookExecutableName(SDKHookType Hook)
 		{
-			if (RuntimePlatform.IsWindows)
+			if (OperatingSystem.IsWindows())
 			{
 				if (Hook == SDKHookType.Uninstall)
 				{
@@ -1965,7 +1965,7 @@ namespace EpicGames.Core
 
 					// actually perform the PATH stripping / adding.
 					string? OrigPathVar = Environment.GetEnvironmentVariable("PATH");
-					string PathDelimiter = RuntimePlatform.IsWindows ? ";" : ":";
+					string PathDelimiter = OperatingSystem.IsWindows() ? ";" : ":";
 					string[] PathVars = [];
 					if (!String.IsNullOrEmpty(OrigPathVar))
 					{
diff --git a/Engine/Source/Programs/Shared/EpicGames.Build/Unreal.cs b/Engine/Source/Programs/Shared/EpicGames.Build/Unreal.cs
--- a/Engine/Source/Programs/Shared/EpicGames.Build/Unreal.cs
+++ b/Engine/Source/Programs/Shared/EpicGames.Build/Unreal.cs
@@ -448,11 +448,11 @@ namespace UnrealBuildBase
 		/// </summary>
 		private static DirectoryReference GetUserSettingDirectory()
 		{
-			if (RuntimePlatform.IsMac)
+			if (OperatingSystem.IsMacOS())
 			{
 				return new DirectoryReference(Path.Combine(Environment.GetFolderPath(Environment.SpecialFolder.Personal), "Library", "Application Support", "Epic"));
 			}
-			else if (RuntimePlatform.IsLinux)
+			else if (OperatingSystem.IsLinux())
 			{
 				return new DirectoryReference(Path.Combine(Environment.GetFolderPath(Environment.SpecialFolder.ApplicationData), "Epic"));
 			}
@@ -485,7 +485,7 @@ namespace UnrealBuildBase
 			string PersonalFolder = Environment.GetFolderPath(Environment.SpecialFolder.Personal);
 			if (!String.IsNullOrEmpty(PersonalFolder))
 			{
-				if (RuntimePlatform.IsMac || RuntimePlatform.IsLinux)
+				if (OperatingSystem.IsMacOS() || OperatingSystem.IsLinux())
 				{
 					return new DirectoryReference(System.IO.Path.Combine(PersonalFolder, "Documents"));
 				}
@@ -508,7 +508,7 @@ namespace UnrealBuildBase
 				// this likely can't fail, but just in case, fallback to preview implementation
 				string machineName = System.Net.Dns.GetHostName();
 
-				if (RuntimePlatform.IsMac && machineName.EndsWith(".local"))
+				if (OperatingSystem.IsMacOS() && machineName.EndsWith(".local"))
 				{
 					machineName = machineName.Replace(".local", "");
 				}
diff --git a/Engine/Source/Programs/Shared/EpicGames.Core/FileUtils.cs b/Engine/Source/Programs/Shared/EpicGames.Core/FileUtils.cs
--- a/Engine/Source/Programs/Shared/EpicGames.Core/FileUtils.cs
+++ b/Engine/Source/Programs/Shared/EpicGames.Core/FileUtils.cs
@@ -93,12 +93,12 @@ namespace EpicGames.Core
 		/// <summary>
 		/// Comparer that should be used for native path comparisons
 		/// </summary>
-		public static IEqualityComparer<string> PlatformPathComparer { get; } = RuntimePlatform.IsLinux ? StringComparer.Ordinal : StringComparer.OrdinalIgnoreCase;
+		public static IEqualityComparer<string> PlatformPathComparer { get; } = OperatingSystem.IsLinux() ? StringComparer.Ordinal : StringComparer.OrdinalIgnoreCase;
 
 		/// <summary>
 		/// Utf8 string comparer that should be used for native path comparisons
 		/// </summary>
-		public static IEqualityComparer<Utf8String> PlatformPathComparerUtf8 { get; } = RuntimePlatform.IsLinux ? Utf8StringComparer.Ordinal : Utf8StringComparer.OrdinalIgnoreCase;
+		public static IEqualityComparer<Utf8String> PlatformPathComparerUtf8 { get; } = OperatingSystem.IsLinux() ? Utf8StringComparer.Ordinal : Utf8StringComparer.OrdinalIgnoreCase;
 
 		/// <summary>
 		/// Read all text for a file
@@ -170,7 +170,7 @@ namespace EpicGames.Core
 		/// <returns>String containing locking information</returns>
 		public static string? GetFileLockInfo(string fileName)
 		{
-			if (RuntimePlatform.IsWindows)
+			if (OperatingSystem.IsWindows())
 			{
 				List<FileLockInfoWin32>? lockInfoList;
 				try
@@ -201,7 +201,7 @@ namespace EpicGames.Core
 		/// <param name="fileName">Name of the file to delete</param>
 		public static void ForceDeleteFile(string fileName)
 		{
-			if (RuntimePlatform.IsWindows)
+			if (OperatingSystem.IsWindows())
 			{
 				ForceDeleteFileWin32(fileName);
 			}
@@ -217,7 +217,7 @@ namespace EpicGames.Core
 		/// <param name="file">The file to delete</param>
 		public static void ForceDeleteFile(FileInfo file)
 		{
-			if (RuntimePlatform.IsWindows)
+			if (OperatingSystem.IsWindows())
 			{
 				ForceDeleteFileWin32(file.FullName);
 				file.Refresh();
@@ -245,7 +245,7 @@ namespace EpicGames.Core
 		/// <param name="location">The file to delete</param>
 		public static void ForceDeleteFile(FileReference location)
 		{
-			if (RuntimePlatform.IsWindows)
+			if (OperatingSystem.IsWindows())
 			{
 				ForceDeleteFileWin32(location.FullName);
 			}
@@ -261,7 +261,7 @@ namespace EpicGames.Core
 		/// <param name="directoryName">Directory to delete</param>
 		public static void ForceDeleteDirectory(string directoryName)
 		{
-			if (RuntimePlatform.IsWindows)
+			if (OperatingSystem.IsWindows())
 			{
 				ForceDeleteLongDirectoryWin32("\\\\?\\" + directoryName);
 			}
@@ -279,7 +279,7 @@ namespace EpicGames.Core
 		{
 			if (directory.Exists)
 			{
-				if (RuntimePlatform.IsWindows)
+				if (OperatingSystem.IsWindows())
 				{
 					ForceDeleteLongDirectoryWin32("\\\\?\\" + directory.FullName);
 				}
@@ -335,7 +335,7 @@ namespace EpicGames.Core
 		/// <param name="directory">Directory to delete</param>
 		public static void ForceDeleteDirectoryContents(DirectoryInfo directory)
 		{
-			if (RuntimePlatform.IsWindows)
+			if (OperatingSystem.IsWindows())
 			{
 				ForceDeleteLongDirectoryContentsWin32(directory.FullName);
 			}
diff --git a/Engine/Source/Programs/Shared/EpicGames.Core/ManagedProcess.cs b/Engine/Source/Programs/Shared/EpicGames.Core/ManagedProcess.cs
--- a/Engine/Source/Programs/Shared/EpicGames.Core/ManagedProcess.cs
+++ b/Engine/Source/Programs/Shared/EpicGames.Core/ManagedProcess.cs
@@ -128,7 +128,7 @@ namespace EpicGames.Core
 		/// <summary>
 		/// Determines support for using job objects
 		/// </summary>
-		internal static bool SupportsJobObjects => RuntimePlatform.IsWindows;
+		internal static bool SupportsJobObjects => OperatingSystem.IsWindows();
 
 		/// <summary>
 		/// Constructor
diff --git a/Engine/Source/Programs/Shared/EpicGames.Core/RuntimePlatform.cs b/Engine/Source/Programs/Shared/EpicGames.Core/RuntimePlatform.cs
--- a/Engine/Source/Programs/Shared/EpicGames.Core/RuntimePlatform.cs
+++ b/Engine/Source/Programs/Shared/EpicGames.Core/RuntimePlatform.cs
@@ -1,28 +1,31 @@
 // Copyright Epic Games, Inc. All Rights Reserved.
 
-using System.Runtime.InteropServices;
+using System;
 
 namespace EpicGames.Core
 {
 	/// <summary>
-	/// Allows testing different platforsm
+	/// Allows testing different platforms
 	/// </summary>
 	public static class RuntimePlatform
 	{
 		/// <summary>
 		/// Whether we are currently running on Linux.
 		/// </summary>
-		public static readonly bool IsLinux = RuntimeInformation.IsOSPlatform(OSPlatform.Linux);
+		[Obsolete("Replace with OperatingSystem.IsLinux()")]
+		public static bool IsLinux => OperatingSystem.IsLinux();
 
 		/// <summary>
 		/// Whether we are currently running on a MacOS platform.
 		/// </summary>
-		public static readonly bool IsMac = RuntimeInformation.IsOSPlatform(OSPlatform.OSX);
+		[Obsolete("Replace with OperatingSystem.IsMacOS()")]
+		public static bool IsMac => OperatingSystem.IsMacOS();
 
 		/// <summary>
 		/// Whether we are currently running a Windows platform.
 		/// </summary>
-		public static readonly bool IsWindows = RuntimeInformation.IsOSPlatform(OSPlatform.Windows);
+		[Obsolete("Replace with OperatingSystem.IsWindows()")]
+		public static bool IsWindows => OperatingSystem.IsWindows();
 
 		/// <summary>
 		/// The platform type
@@ -48,11 +51,11 @@ namespace EpicGames.Core
 		/// <summary>
 		/// The current runtime platform
 		/// </summary>
-		public static readonly Type Current = IsWindows ? Type.Windows : IsMac ? Type.Mac : Type.Linux;
+		public static readonly Type Current = OperatingSystem.IsWindows() ? Type.Windows : OperatingSystem.IsMacOS() ? Type.Mac : Type.Linux;
 
 		/// <summary>
 		/// The extension executables have on the current platform
 		/// </summary>
-		public static readonly string ExeExtension = IsWindows ? ".exe" : "";
+		public static readonly string ExeExtension = OperatingSystem.IsWindows() ? ".exe" : "";
 	}
 }
diff --git a/Engine/Source/Programs/Shared/EpicGames.Perforce/PerforceChildProcess.cs b/Engine/Source/Programs/Shared/EpicGames.Perforce/PerforceChildProcess.cs
--- a/Engine/Source/Programs/Shared/EpicGames.Perforce/PerforceChildProcess.cs
+++ b/Engine/Source/Programs/Shared/EpicGames.Perforce/PerforceChildProcess.cs
@@ -94,7 +94,7 @@ namespace EpicGames.Perforce
 		public static string GetExecutable()
 		{
 			string perforceFileName = "p4.exe";
-			if (!RuntimePlatform.IsWindows)
+			if (!OperatingSystem.IsWindows())
 			{
 				string[] p4Paths = {
 					"/usr/bin/p4", // Default path
diff --git a/Engine/Source/Programs/UnrealBuildTool/Configuration/UEBuildTarget.cs b/Engine/Source/Programs/UnrealBuildTool/Configuration/UEBuildTarget.cs
--- a/Engine/Source/Programs/UnrealBuildTool/Configuration/UEBuildTarget.cs
+++ b/Engine/Source/Programs/UnrealBuildTool/Configuration/UEBuildTarget.cs
@@ -1530,7 +1530,8 @@ namespace UnrealBuildTool
 			}
 
 			// Check that each project module with a dependency on a plugin module, that the plugin is either enabled by default or it's enabled by the uproject file.
-			if (ProjectDescriptor?.Modules?.Any(x => x.Name == module.Name) == true)
+			ModuleDescriptor? moduleDescriptor = ProjectDescriptor?.Modules?.FirstOrDefault(x => x.Name == module.Name);
+			if (moduleDescriptor != null)
 			{
 				string projectName = ProjectFile?.FullName ?? TargetRulesFile.GetFileName();
 
@@ -1548,6 +1549,7 @@ namespace UnrealBuildTool
 					// Is the modules plugin enabled by default?
 					if (dependencyPluginInfo != null && dependencyPluginInfo.IsEnabledByDefault(bAllowEnginePluginsEnabledByDefault) == false)
 					{
+						ModuleDescriptor? dependencyModuleInfo = dependencyPluginInfo.Descriptor.Modules?.FirstOrDefault(x => x.Name == dependencyModule.Name);
 						// Try and find the project plugin reference in the plugins list.
 						PluginReferenceDescriptor? projectPluginReference = ProjectDescriptor?.Plugins?.FirstOrDefault(x => x.Name == dependencyPluginInfo.Name);
 						if (projectPluginReference != null)
@@ -5697,7 +5699,7 @@ namespace UnrealBuildTool
 				UEBuildBinary ExecutableBinary = Binaries[0];
 
 				bool IsCurrentPlatform;
-				if (!RuntimePlatform.IsWindows)
+				if (!OperatingSystem.IsWindows())
 				{
 					IsCurrentPlatform = Platform == UnrealTargetPlatform.Mac || (UEBuildPlatform.IsPlatformInGroup(Platform, UnrealPlatformGroup.Unix) && Platform == BuildHostPlatform.Current.Platform);
 				}
diff --git a/Engine/Source/Programs/UnrealBuildTool/Modes/BuildMode.cs b/Engine/Source/Programs/UnrealBuildTool/Modes/BuildMode.cs
--- a/Engine/Source/Programs/UnrealBuildTool/Modes/BuildMode.cs
+++ b/Engine/Source/Programs/UnrealBuildTool/Modes/BuildMode.cs
@@ -322,12 +322,12 @@ namespace UnrealBuildTool
 				Dictionary<string, string> Variables = GetTargetVariables(TargetDescriptor);
 
 				// Get the output path to the script
-				string ScriptExtension = RuntimePlatform.IsWindows ? ".bat" : ".sh";
+				string ScriptExtension = OperatingSystem.IsWindows() ? ".bat" : ".sh";
 				FileReference ScriptFile = FileReference.Combine(Directory, String.Format("{0}-{1}{2}", FilePrefix, ScriptFiles.Count + 1, ScriptExtension));
 
 				// Write it to disk
 				List<string> Contents = new List<string>();
-				if (RuntimePlatform.IsWindows)
+				if (OperatingSystem.IsWindows())
 				{
 					Contents.Insert(0, "@echo off");
 				}
@@ -1296,7 +1296,7 @@ namespace UnrealBuildTool
 				return;
 			}
 
-			if (!RuntimePlatform.IsWindows)
+			if (!OperatingSystem.IsWindows())
 			{
 				return;
 			}
diff --git a/Engine/Source/Programs/UnrealBuildTool/Platform/Android/AndroidPlatformSDK.cs b/Engine/Source/Programs/UnrealBuildTool/Platform/Android/AndroidPlatformSDK.cs
--- a/Engine/Source/Programs/UnrealBuildTool/Platform/Android/AndroidPlatformSDK.cs
+++ b/Engine/Source/Programs/UnrealBuildTool/Platform/Android/AndroidPlatformSDK.cs
@@ -59,7 +59,7 @@ namespace UnrealBuildTool
 			}
 				
 			// If we are not running on Windows and we are still missing a key then go and find it from the .bash_profile
-			if (!RuntimePlatform.IsWindows && !EnvVarNames.All(s => AndroidEnv.ContainsKey(s.Key)))
+			if (!OperatingSystem.IsWindows() && !EnvVarNames.All(s => AndroidEnv.ContainsKey(s.Key)))
 			{
 				string BashProfilePath = Path.Combine(Environment.GetFolderPath(Environment.SpecialFolder.UserProfile), ".bash_profile");
 				if (!File.Exists(BashProfilePath))
diff --git a/Engine/Source/Programs/UnrealBuildTool/Platform/Android/AndroidToolChain.cs b/Engine/Source/Programs/UnrealBuildTool/Platform/Android/AndroidToolChain.cs
--- a/Engine/Source/Programs/UnrealBuildTool/Platform/Android/AndroidToolChain.cs
+++ b/Engine/Source/Programs/UnrealBuildTool/Platform/Android/AndroidToolChain.cs
@@ -1305,7 +1305,7 @@ namespace UnrealBuildTool
 			// To work around this we call clang to just get the command it would execute and generate a
 			// second response file to directly call ld with the right arguments instead of calling through clang.
 			/* disable while tracking down some linker errors this introduces
-						if (RuntimePlatform.IsWindows)
+						if (OperatingSystem.IsWindows())
 						{
 							// capture the actual link command without running it
 							ProcessStartInfo StartInfo = new ProcessStartInfo();
diff --git a/Engine/Source/Programs/UnrealBuildTool/Platform/Android/UEDeployAndroid.cs b/Engine/Source/Programs/UnrealBuildTool/Platform/Android/UEDeployAndroid.cs
--- a/Engine/Source/Programs/UnrealBuildTool/Platform/Android/UEDeployAndroid.cs
+++ b/Engine/Source/Programs/UnrealBuildTool/Platform/Android/UEDeployAndroid.cs
@@ -354,7 +354,7 @@ namespace UnrealBuildTool
 				// valid directories will have a source.properties with the Pkg.Revision (there is no guarantee we can use the directory name as revision)
 				foreach (string CandidateDir in Subdirs)
 				{
-					string AaptFilename = Path.Combine(CandidateDir, RuntimePlatform.IsWindows ? "aapt.exe" : "aapt");
+					string AaptFilename = Path.Combine(CandidateDir, OperatingSystem.IsWindows() ? "aapt.exe" : "aapt");
 					string RevisionString = "";
 					uint RevisionValue = 0;
 
@@ -1330,11 +1330,11 @@ namespace UnrealBuildTool
 
 		private static string GetPlatformNDKHostName()
 		{
-			if (RuntimePlatform.IsLinux)
+			if (OperatingSystem.IsLinux())
 			{
 				return "linux-x86_64";
 			}
-			else if (RuntimePlatform.IsMac)
+			else if (OperatingSystem.IsMacOS())
 			{
 				return "darwin-x86_64";
 			}
@@ -4287,7 +4287,7 @@ namespace UnrealBuildTool
 			string BuildToolsVersion = GetBuildToolsVersion();
 
 			// cache some tools paths
-			//string NDKBuildPath = Environment.ExpandEnvironmentVariables("%NDKROOT%/ndk-build" + (RuntimePlatform.IsWindows ? ".cmd" : ""));
+			//string NDKBuildPath = Environment.ExpandEnvironmentVariables("%NDKROOT%/ndk-build" + (OperatingSystem.IsWindows() ? ".cmd" : ""));
 			//bool HasNDKPath = File.Exists(NDKBuildPath);
 
 			// set up some directory info
@@ -4805,11 +4805,11 @@ namespace UnrealBuildTool
 
 					// sign the APK
 					string BuildToolsPath = Path.Combine(Environment.ExpandEnvironmentVariables("%ANDROID_HOME%"), "Build-Tools", BuildToolsVersion);
-					string APKSignerExecutable = Path.Combine(BuildToolsPath, "apksigner" + (RuntimePlatform.IsWindows ? ".bat" : ""));
+					string APKSignerExecutable = Path.Combine(BuildToolsPath, "apksigner" + (OperatingSystem.IsWindows() ? ".bat" : ""));
 
 					ProcessStartInfo StartInfo = new ProcessStartInfo();
 					StartInfo.WorkingDirectory = IntermediateAndroidPath;
-					if (RuntimePlatform.IsWindows)
+					if (OperatingSystem.IsWindows())
 					{
 						StartInfo.FileName = "cmd.exe";
 						StartInfo.ArgumentList.Add("/c");
@@ -5292,7 +5292,7 @@ popd
 				if (!bSkipGradleBuild)
 				{
 					string GradleScriptPath = Path.Combine(UnrealBuildGradlePath, "gradlew");
-					if (!RuntimePlatform.IsWindows)
+					if (!OperatingSystem.IsWindows())
 					{
 						// fix permissions for Mac/Linux
 						RunCommandLineProgramWithException(UnrealBuildGradlePath, "/bin/sh", String.Format("-c 'chmod 0755 \"{0}\"'", GradleScriptPath.Replace("'", "'\"'\"'")), Logger, "Fix gradlew permissions");
@@ -5325,7 +5325,7 @@ popd
 
 							// For build machine, get a list of current java.exe tasks already running to ignore for cleanup
 							List<int> PreviousJavaProcesses = new List<int>();
-							if (bIsBuildMachine && RuntimePlatform.IsWindows)
+							if (bIsBuildMachine && OperatingSystem.IsWindows())
 							{
 								string Output;
 								int ExitCode = RunCommandLineProgramAndReturnResultOutput(UnrealBuildGradlePath, "cmd.exe", "/c tasklist.exe /NH /FI \"imagename eq java.exe\"", Logger, out Output, "Collecting existing java.exe processes");
@@ -5348,9 +5348,9 @@ popd
 							}
 
 							// Use gradle to build the .apk file
-							string ShellExecutable = RuntimePlatform.IsWindows ? "cmd.exe" : "/bin/sh";
-							string ShellParametersBegin = RuntimePlatform.IsWindows ? "/c " : "-c '";
-							string ShellParametersEnd = RuntimePlatform.IsWindows ? "" : "'";
+							string ShellExecutable = OperatingSystem.IsWindows() ? "cmd.exe" : "/bin/sh";
+							string ShellParametersBegin = OperatingSystem.IsWindows() ? "/c " : "-c '";
+							string ShellParametersEnd = OperatingSystem.IsWindows() ? "" : "'";
 							RunCommandLineProgramWithExceptionAndFiltering(UnrealBuildGradlePath, ShellExecutable, ShellParametersBegin + "\"" + GradleScriptPath + "\" " + GradleOptions + ShellParametersEnd, Logger, "Making .apk with Gradle...");
 
 							if (!String.IsNullOrEmpty(GradleSecondCallOptions))
@@ -5365,7 +5365,7 @@ popd
 								//RunCommandLineProgramWithException(UnrealBuildGradlePath, ShellExecutable, ShellParametersBegin + "\"" + GradleScriptPath + "\" " + GradleOptions + ShellParametersEnd, "Listing all tasks...");
 
 								// on Windows sometimes minifyReleaseWithR8 is keeping a lock on classes.dex so kill java.exe not previously running
-								if (RuntimePlatform.IsWindows)
+								if (OperatingSystem.IsWindows())
 								{
 									string Output;
 									int ExitCode = RunCommandLineProgramAndReturnResultOutput(UnrealBuildGradlePath, "cmd.exe", "/c tasklist.exe /NH /FI \"imagename eq java.exe\"", Logger, out Output, "Collecting existing java.exe processes");
@@ -5726,7 +5726,7 @@ popd
 					File.WriteAllText(GradlePropertiesFilename, GradlePropertiesContent);
 
 					string GradleScriptPath = Path.Combine(UnrealBuildGradlePath, "gradlew");
-					if (!RuntimePlatform.IsWindows)
+					if (!OperatingSystem.IsWindows())
 					{
 						// fix permissions for Mac/Linux
 						RunCommandLineProgramWithException(UnrealBuildGradlePath, "/bin/sh", String.Format("-c 'chmod 0755 \"{0}\"'", GradleScriptPath.Replace("'", "'\"'\"'")), Logger, "Fix gradlew permissions");
@@ -5753,9 +5753,9 @@ popd
 					//Directory.CreateDirectory(Path.GetDirectoryName(DestApkName));
 
 					// Use gradle to build the .apk file
-					string ShellExecutable = RuntimePlatform.IsWindows ? "cmd.exe" : "/bin/sh";
-					string ShellParametersBegin = RuntimePlatform.IsWindows ? "/c " : "-c '";
-					string ShellParametersEnd = RuntimePlatform.IsWindows ? "" : "'";
+					string ShellExecutable = OperatingSystem.IsWindows() ? "cmd.exe" : "/bin/sh";
+					string ShellParametersBegin = OperatingSystem.IsWindows() ? "/c " : "-c '";
+					string ShellParametersEnd = OperatingSystem.IsWindows() ? "" : "'";
 					RunCommandLineProgramWithExceptionAndFiltering(UnrealBuildGradlePath, ShellExecutable, ShellParametersBegin + "\"" + GradleScriptPath + "\" " + GradleOptions + ShellParametersEnd, Logger, "Making .aab with Gradle...");
 
 					if (!String.IsNullOrEmpty(GradleSecondCallOptions))
@@ -5788,7 +5788,7 @@ popd
 						string UnrealBuildGradlePath = Path.Combine(UnrealBuildPath, "gradle");
 
 						string GradleScriptPath = Path.Combine(UnrealBuildGradlePath, "gradlew");
-						if (!RuntimePlatform.IsWindows)
+						if (!OperatingSystem.IsWindows())
 						{
 							// fix permissions for Mac/Linux
 							RunCommandLineProgramWithException(UnrealBuildGradlePath, "/bin/sh", String.Format("-c 'chmod 0755 \"{0}\"'", GradleScriptPath.Replace("'", "'\"'\"'")), Logger, "Fix gradlew permissions");
@@ -5815,9 +5815,9 @@ popd
 						//Directory.CreateDirectory(Path.GetDirectoryName(DestApkName));
 
 						// Use gradle to build the .apk file
-						string ShellExecutable = RuntimePlatform.IsWindows ? "cmd.exe" : "/bin/sh";
-						string ShellParametersBegin = RuntimePlatform.IsWindows ? "/c " : "-c '";
-						string ShellParametersEnd = RuntimePlatform.IsWindows ? "" : "'";
+						string ShellExecutable = OperatingSystem.IsWindows() ? "cmd.exe" : "/bin/sh";
+						string ShellParametersBegin = OperatingSystem.IsWindows() ? "/c " : "-c '";
+						string ShellParametersEnd = OperatingSystem.IsWindows() ? "" : "'";
 						RunCommandLineProgramWithExceptionAndFiltering(UnrealBuildGradlePath, ShellExecutable, ShellParametersBegin + "\"" + GradleScriptPath + "\" " + GradleOptions + ShellParametersEnd, Logger, "Making .aab with Gradle...");
 
 						if (!String.IsNullOrEmpty(GradleSecondCallOptions))
diff --git a/Engine/Source/Programs/UnrealBuildTool/Platform/Mac/ApplePlatformSDK.cs b/Engine/Source/Programs/UnrealBuildTool/Platform/Mac/ApplePlatformSDK.cs
--- a/Engine/Source/Programs/UnrealBuildTool/Platform/Mac/ApplePlatformSDK.cs
+++ b/Engine/Source/Programs/UnrealBuildTool/Platform/Mac/ApplePlatformSDK.cs
@@ -270,7 +270,7 @@ namespace UnrealBuildTool
 
 			// iTunes is technically only need to deploy to and run on connected devices.
 			// This code removes requirement for Windows builders to have Xcode installed.
-			if (Status == SDKStatus.Invalid && !RuntimePlatform.IsMac && Unreal.IsBuildMachine())
+			if (Status == SDKStatus.Invalid && !OperatingSystem.IsMacOS() && Unreal.IsBuildMachine())
 			{
 				Status = SDKStatus.Valid;
 			}
diff --git a/Engine/Source/Programs/UnrealBuildTool/Platform/Windows/MicrosoftPlatformSDK.cs b/Engine/Source/Programs/UnrealBuildTool/Platform/Windows/MicrosoftPlatformSDK.cs
--- a/Engine/Source/Programs/UnrealBuildTool/Platform/Windows/MicrosoftPlatformSDK.cs
+++ b/Engine/Source/Programs/UnrealBuildTool/Platform/Windows/MicrosoftPlatformSDK.cs
@@ -34,7 +34,7 @@ namespace UnrealBuildTool
 		[SupportedOSPlatform("windows")]
 		protected override string? GetInstalledSDKVersion()
 		{
-			if (!RuntimePlatform.IsWindows)
+			if (!OperatingSystem.IsWindows())
 			{
 				return null;
 			}
diff --git a/Engine/Source/Programs/UnrealBuildTool/Platform/Windows/UEBuildWindows.cs b/Engine/Source/Programs/UnrealBuildTool/Platform/Windows/UEBuildWindows.cs
--- a/Engine/Source/Programs/UnrealBuildTool/Platform/Windows/UEBuildWindows.cs
+++ b/Engine/Source/Programs/UnrealBuildTool/Platform/Windows/UEBuildWindows.cs
@@ -740,7 +740,7 @@ namespace UnrealBuildTool
 			get
 			{
 				DirectoryReference? NetFxSdkDir;
-				if (RuntimePlatform.IsWindows && MicrosoftPlatformSDK.TryGetNetFxSdkInstallDir(out NetFxSdkDir))
+				if (OperatingSystem.IsWindows() && MicrosoftPlatformSDK.TryGetNetFxSdkInstallDir(out NetFxSdkDir))
 				{
 					return NetFxSdkDir.FullName;
 				}
diff --git a/Engine/Source/Programs/UnrealBuildTool/Platform/Windows/VCToolChain.cs b/Engine/Source/Programs/UnrealBuildTool/Platform/Windows/VCToolChain.cs
--- a/Engine/Source/Programs/UnrealBuildTool/Platform/Windows/VCToolChain.cs
+++ b/Engine/Source/Programs/UnrealBuildTool/Platform/Windows/VCToolChain.cs
@@ -2307,7 +2307,7 @@ namespace UnrealBuildTool
 
 						// Extract CompileScore data from traces
 						FileReference ScoreDataExtractor = FileReference.Combine(Unreal.RootDirectory, "Engine", "Extras", "ThirdPartyNotUE", "CompileScore", "ScoreDataExtractor.exe");
-						if (RuntimePlatform.IsWindows && FileReference.Exists(ScoreDataExtractor))
+						if (OperatingSystem.IsWindows() && FileReference.Exists(ScoreDataExtractor))
 						{
 							FileItem CompileScoreOutput = FileItem.GetItemByFileReference(FileReference.Combine(Makefile.ProjectIntermediateDirectory, $"{Target.Name}.scor"));
 
diff --git a/Engine/Source/Programs/UnrealBuildTool/ProjectFiles/ProjectFileGenerator.cs b/Engine/Source/Programs/UnrealBuildTool/ProjectFiles/ProjectFileGenerator.cs
--- a/Engine/Source/Programs/UnrealBuildTool/ProjectFiles/ProjectFileGenerator.cs
+++ b/Engine/Source/Programs/UnrealBuildTool/ProjectFiles/ProjectFileGenerator.cs
@@ -2139,7 +2139,7 @@ namespace UnrealBuildTool
 		{
 			if (ShouldGenerateIntelliSenseData() && Targets.Any())
 			{
-				string ProgressInfoText = RuntimePlatform.IsWindows ? "Binding IntelliSense data..." : "Generating data for project indexing...";
+				string ProgressInfoText = OperatingSystem.IsWindows() ? "Binding IntelliSense data..." : "Generating data for project indexing...";
 				using (ProgressWriter Progress = new ProgressWriter(ProgressInfoText, true, Logger))
 				{
 					// Ignore projects for platforms we can't build for
diff --git a/Engine/Source/Programs/UnrealBuildTool/System/BuildHostPlatform.cs b/Engine/Source/Programs/UnrealBuildTool/System/BuildHostPlatform.cs
--- a/Engine/Source/Programs/UnrealBuildTool/System/BuildHostPlatform.cs
+++ b/Engine/Source/Programs/UnrealBuildTool/System/BuildHostPlatform.cs
@@ -42,15 +42,15 @@ namespace UnrealBuildTool
 			{
 				if (CurrentPlatform == null)
 				{
-					if (RuntimePlatform.IsWindows)
+					if (OperatingSystem.IsWindows())
 					{
 						CurrentPlatform = new WindowsBuildHostPlatform();
 					}
-					else if (RuntimePlatform.IsMac)
+					else if (OperatingSystem.IsMacOS())
 					{
 						CurrentPlatform = new MacBuildHostPlatform();
 					}
-					else if (RuntimePlatform.IsLinux)
+					else if (OperatingSystem.IsLinux())
 					{
 						CurrentPlatform = new LinuxBuildHostPlatform();
 					}
diff --git a/Engine/Source/Programs/UnrealBuildTool/System/Utils.cs b/Engine/Source/Programs/UnrealBuildTool/System/Utils.cs
--- a/Engine/Source/Programs/UnrealBuildTool/System/Utils.cs
+++ b/Engine/Source/Programs/UnrealBuildTool/System/Utils.cs
@@ -1023,7 +1023,7 @@ namespace UnrealBuildTool
 		public static int GetLogicalProcessorCount()
 		{
 			// This function uses Windows P/Invoke calls; if we're not running on Windows, just return the default.
-			if (RuntimePlatform.IsWindows)
+			if (OperatingSystem.IsWindows())
 			{
 				const int ERROR_INSUFFICIENT_BUFFER = 122;
 
@@ -1066,7 +1066,7 @@ namespace UnrealBuildTool
 					}
 				}
 			}
-			else if (RuntimePlatform.IsLinux)
+			else if (OperatingSystem.IsLinux())
 			{
 				// query socket/logical core pairings.  There should not be duplicates in this list since each hyperthread
 				// will show up as it's own logical "cpu".  Including the socket number allows us to count multi-processor
@@ -1089,7 +1089,7 @@ namespace UnrealBuildTool
 		/// <returns>The number of physical cores, or -1 if it could not be obtained</returns>
 		public static int GetPhysicalProcessorCount()
 		{
-			if (RuntimePlatform.IsWindows)
+			if (OperatingSystem.IsWindows())
 			{
 				const int ERROR_INSUFFICIENT_BUFFER = 122;
 
@@ -1129,7 +1129,7 @@ namespace UnrealBuildTool
 					}
 				}
 			}
-			else if (RuntimePlatform.IsMac)
+			else if (OperatingSystem.IsMacOS())
 			{
 				ulong Size = 4;
 				if (0 == sysctlbyname("hw.physicalcpu", out int Value, ref Size, IntPtr.Zero, 0))
@@ -1137,7 +1137,7 @@ namespace UnrealBuildTool
 					return Value;
 				}
 			}
-			else if (RuntimePlatform.IsLinux)
+			else if (OperatingSystem.IsLinux())
 			{
 				// query socket/physical core pairings.  There will be duplicates in this if there are hyperthreads
 				// using the HashSet ensures that those duplicates are removed and we only count the first "cpu" found
@@ -1280,7 +1280,7 @@ namespace UnrealBuildTool
 			// On Mac, MemoryInfo.MemoryLoadBytes includes memory used to cache disk-backed files ("Cached Files" in
 			// Activity Monitor), which can result in a significant over-estimate of memory pressure.
 			// We treat memory used for caching of disk-backed files as free for use in compilation tasks.
-			if (RuntimePlatform.IsMac)
+			if (OperatingSystem.IsMacOS())
 			{
 				// host_statistics64() flavor, from <mach/host_info.h>
 				int HOST_VM_INFO64 = 4;
@@ -1334,7 +1334,7 @@ namespace UnrealBuildTool
 		public static bool GetTotalCpuUtilization(out float Utilization)
 		{
 			Utilization = 0.0f;
-			if (RuntimePlatform.IsWindows)
+			if (OperatingSystem.IsWindows())
 			{
 				const uint ERROR_SUCCESS = 0;
 				if (!CpuInitialized)
@@ -1425,7 +1425,7 @@ namespace UnrealBuildTool
 				// limiting based on total system memory and relying on the OS to swap out other processes to make room is fine.
 				// That said only Mac has been extensively tested here in this scenario, and does tend to have faster disk access
 				// than other platforms. So for now we'll change Mac and leave Windows/Linux using the old behavior
-				long AvailableMemoryBytes = RuntimePlatform.IsMac ? GetTotalSystemMemoryBytes() : GetFreeMemoryBytes();
+				long AvailableMemoryBytes = OperatingSystem.IsMacOS() ? GetTotalSystemMemoryBytes() : GetFreeMemoryBytes();
 
 				if (AvailableMemoryBytes != -1)
 				{
@@ -1866,7 +1866,7 @@ namespace UnrealBuildTool
 				// ignored
 			}
 
-			if (RuntimePlatform.IsWindows)
+			if (OperatingSystem.IsWindows())
 			{
 				try
 				{
diff --git a/Engine/Source/Programs/UnrealBuildTool/ToolChain/ClangToolChain.cs b/Engine/Source/Programs/UnrealBuildTool/ToolChain/ClangToolChain.cs
--- a/Engine/Source/Programs/UnrealBuildTool/ToolChain/ClangToolChain.cs
+++ b/Engine/Source/Programs/UnrealBuildTool/ToolChain/ClangToolChain.cs
@@ -303,7 +303,7 @@ namespace UnrealBuildTool
 
 					// Extract CompileScore data from traces
 					FileReference ScoreDataExtractor = FileReference.Combine(Unreal.RootDirectory, "Engine", "Extras", "ThirdPartyNotUE", "CompileScore", "ScoreDataExtractor.exe");
-					if (RuntimePlatform.IsWindows && FileReference.Exists(ScoreDataExtractor))
+					if (OperatingSystem.IsWindows() && FileReference.Exists(ScoreDataExtractor))
 					{
 						FileItem CompileScoreOutput = FileItem.GetItemByFileReference(FileReference.Combine(Makefile.ProjectIntermediateDirectory, $"{Target.Name}.scor"));
 
@@ -639,7 +639,7 @@ namespace UnrealBuildTool
 			}
 
 			// https://clang.llvm.org/docs/ClangCommandLineReference.html#cmdoption-clang-fdiagnostics-format
-			if (RuntimePlatform.IsWindows)
+			if (OperatingSystem.IsWindows())
 			{
 				Arguments.Add("-fdiagnostics-format=msvc");
 			}
@@ -1238,7 +1238,7 @@ namespace UnrealBuildTool
 			FileReference ResponseFileName = GetResponseFileName(CompileEnvironment, TargetFile);
 			List<string> ResponseFileContents = ExpandResponseFileContents(FileArguments);
 
-			if (RuntimePlatform.IsWindows)
+			if (OperatingSystem.IsWindows())
 			{
 				// Enables Clang Cmd Line Work Around
 				// Clang reads the response file and computes the total cmd line length.
@@ -1349,7 +1349,7 @@ namespace UnrealBuildTool
 				PreprocessResponseFileContents.AddRange(PreprocessGlobalArguments);
 				PreprocessResponseFileContents.AddRange(PreprocessFileArguments);
 
-				if (RuntimePlatform.IsWindows)
+				if (OperatingSystem.IsWindows())
 				{
 					int CmdLineLength = Info.Clang.ToString().Length + String.Join(' ', ResponseFileContents).Length;
 					bool bIsInDangerZone = CmdLineLength >= ClangCmdlineDangerZone && CmdLineLength <= ClangCmdLineMaxSize;
