[UBT]
* Added CppRootPath type and changed all rootpaths code to use it

diff --git a/Engine/Source/Programs/UnrealBuildTool/Configuration/UEBuildTarget.cs b/Engine/Source/Programs/UnrealBuildTool/Configuration/UEBuildTarget.cs
--- a/Engine/Source/Programs/UnrealBuildTool/Configuration/UEBuildTarget.cs
+++ b/Engine/Source/Programs/UnrealBuildTool/Configuration/UEBuildTarget.cs
@@ -3071,15 +3071,6 @@ namespace UnrealBuildTool
 				stripAction.bCanExecuteRemotely = false; // Can run remotely but is so fast so it would be a waste
 				stripAction.ArtifactMode = ArtifactMode.Enabled;
 
-				List<DirectoryItem> rootPaths = new();
-				rootPaths.Add(DirectoryItem.GetItemByDirectoryReference(Unreal.EngineDirectory));
-				if (ProjectFile != null && (!CompileEnvironment.bUseSharedBuildEnvironment || CompileEnvironment.AllIncludePath.Any(x => x.IsUnderDirectory(ProjectFile.Directory))))
-				{
-					rootPaths.Add(DirectoryItem.GetItemByDirectoryReference(ProjectFile.Directory));
-				}
-				rootPaths.Add(DirectoryItem.GetItemByDirectoryReference(Unreal.RootDirectory));
-				stripAction.RootPaths = rootPaths;
-
 				Makefile.OutputItems.AddRange(stripAction.ProducedItems);
 			}
 
diff --git a/Engine/Source/Programs/UnrealBuildTool/Executors/UnrealBuildAccelerator/UBAExecutor.cs b/Engine/Source/Programs/UnrealBuildTool/Executors/UnrealBuildAccelerator/UBAExecutor.cs
--- a/Engine/Source/Programs/UnrealBuildTool/Executors/UnrealBuildAccelerator/UBAExecutor.cs
+++ b/Engine/Source/Programs/UnrealBuildTool/Executors/UnrealBuildAccelerator/UBAExecutor.cs
@@ -689,9 +689,9 @@ namespace UnrealBuildTool
 				// And since pch contains absolute paths we unfortunately can't share cache data between machines that have different paths
 				if (action.ArtifactMode.HasFlag(ArtifactMode.AbsolutePath))
 				{
-					foreach (DirectoryItem root in action.RootPaths)
+					foreach (KeyValuePair<CppRootPathFolder, DirectoryReference> kv in action.RootPaths)
 					{
-						hasher.Update(System.Text.Encoding.UTF8.GetBytes(root.FullName));
+						hasher.Update(System.Text.Encoding.UTF8.GetBytes(kv.Value.FullName));
 					}
 				}
 
@@ -702,9 +702,9 @@ namespace UnrealBuildTool
 		IRootPaths GetActionRootPaths(LinkedAction action)
 		{
 			IRootPaths rootPaths = IRootPaths.Create(_ubaLogger!);
-			foreach (DirectoryItem root in action.RootPaths)
+			foreach (KeyValuePair<CppRootPathFolder, DirectoryReference> kv in action.RootPaths)
 			{
-				rootPaths.RegisterRoot(root.FullName, true);
+				rootPaths.RegisterRoot(kv.Value.FullName, true);
 			}
 
 			DirectoryReference? autoSdkDir;
diff --git a/Engine/Source/Programs/UnrealBuildTool/Platform/Android/AndroidToolChain.cs b/Engine/Source/Programs/UnrealBuildTool/Platform/Android/AndroidToolChain.cs
--- a/Engine/Source/Programs/UnrealBuildTool/Platform/Android/AndroidToolChain.cs
+++ b/Engine/Source/Programs/UnrealBuildTool/Platform/Android/AndroidToolChain.cs
@@ -1082,7 +1082,7 @@ namespace UnrealBuildTool
 
 			// Create an action that invokes the linker.
 			Action LinkAction = Graph.CreateAction(ActionType.Link);
-			LinkAction.RootPaths.AddRange(LinkEnvironment.RootPaths.Values.Select(x => DirectoryItem.GetItemByDirectoryReference(x)));
+			LinkAction.RootPaths = LinkEnvironment.RootPaths;
 			LinkAction.WorkingDirectory = Unreal.EngineSourceDirectory;
 
 			if (LinkEnvironment.bIsBuildingLibrary)
diff --git a/Engine/Source/Programs/UnrealBuildTool/Platform/IOS/IOSToolChain.cs b/Engine/Source/Programs/UnrealBuildTool/Platform/IOS/IOSToolChain.cs
--- a/Engine/Source/Programs/UnrealBuildTool/Platform/IOS/IOSToolChain.cs
+++ b/Engine/Source/Programs/UnrealBuildTool/Platform/IOS/IOSToolChain.cs
@@ -393,7 +393,7 @@ namespace UnrealBuildTool
 
 			// Create an action that invokes the linker.
 			Action LinkAction = Graph.CreateAction(ActionType.Link);
-			LinkAction.RootPaths.AddRange(LinkEnvironment.RootPaths.Values.Select(x => DirectoryItem.GetItemByDirectoryReference(x)));
+			LinkAction.RootPaths = LinkEnvironment.RootPaths;
 
 			// RPC utility parameters are in terms of the Mac side
 			LinkAction.WorkingDirectory = GetMacDevSrcRoot();
diff --git a/Engine/Source/Programs/UnrealBuildTool/Platform/Linux/LinuxToolChain.cs b/Engine/Source/Programs/UnrealBuildTool/Platform/Linux/LinuxToolChain.cs
--- a/Engine/Source/Programs/UnrealBuildTool/Platform/Linux/LinuxToolChain.cs
+++ b/Engine/Source/Programs/UnrealBuildTool/Platform/Linux/LinuxToolChain.cs
@@ -1116,7 +1116,7 @@ namespace UnrealBuildTool
 
 			// Create an action that invokes the linker.
 			Action LinkAction = Graph.CreateAction(ActionType.Link);
-			LinkAction.RootPaths.AddRange(LinkEnvironment.RootPaths.Values.Select(x => DirectoryItem.GetItemByDirectoryReference(x)));
+			LinkAction.RootPaths = LinkEnvironment.RootPaths;
 			LinkAction.WorkingDirectory = Unreal.EngineSourceDirectory;
 
 			string LinkCommandString;
@@ -1574,7 +1574,7 @@ namespace UnrealBuildTool
 				{
 					// Create the action to relink the library. This actions does not overwrite the source file so it can be executed in parallel
 					Action RelinkAction = Graph.CreateAction(ActionType.Link);
-					RelinkAction.RootPaths.AddRange(LinkEnvironment.RootPaths.Values.Select(x => DirectoryItem.GetItemByDirectoryReference(x)));
+					RelinkAction.RootPaths = LinkEnvironment.RootPaths;
 					RelinkAction.WorkingDirectory = LinkAction.WorkingDirectory;
 					RelinkAction.StatusDescription = LinkAction.StatusDescription;
 					RelinkAction.CommandDescription = "Relink";
diff --git a/Engine/Source/Programs/UnrealBuildTool/Platform/Mac/MacToolChain.cs b/Engine/Source/Programs/UnrealBuildTool/Platform/Mac/MacToolChain.cs
--- a/Engine/Source/Programs/UnrealBuildTool/Platform/Mac/MacToolChain.cs
+++ b/Engine/Source/Programs/UnrealBuildTool/Platform/Mac/MacToolChain.cs
@@ -680,7 +680,7 @@ namespace UnrealBuildTool
 
 			// Create an action that invokes the linker.
 			Action LinkAction = Graph.CreateAction(ActionType.Link);
-			LinkAction.RootPaths.AddRange(LinkEnvironment.RootPaths.Values.Select(x => DirectoryItem.GetItemByDirectoryReference(x)));
+			LinkAction.RootPaths = LinkEnvironment.RootPaths;
 
 			FileReference LinkerPath = LinkEnvironment.bIsBuildingLibrary ? Info.Archiver : Info.Clang;
 
diff --git a/Engine/Source/Programs/UnrealBuildTool/Platform/Windows/VCCompileAction.cs b/Engine/Source/Programs/UnrealBuildTool/Platform/Windows/VCCompileAction.cs
--- a/Engine/Source/Programs/UnrealBuildTool/Platform/Windows/VCCompileAction.cs
+++ b/Engine/Source/Programs/UnrealBuildTool/Platform/Windows/VCCompileAction.cs
@@ -192,7 +192,7 @@ namespace UnrealBuildTool
 		/// <summary>
 		/// Root paths for this action (generally engine root project root, toolchain root, sdk root)
 		/// </summary>
-		public SortedDictionary<CppRootPathFolder, DirectoryReference> RootPaths { get; } = new SortedDictionary<CppRootPathFolder, DirectoryReference>();
+		public CppRootPaths RootPaths { get; set; }
 
 		/// <inheritdoc/>
 		public bool bCanExecuteRemotely { get; set; }
@@ -225,7 +225,6 @@ namespace UnrealBuildTool
 		#region Implementation of IAction
 
 		IEnumerable<FileItem> IExternalAction.DeleteItems => DeleteItems;
-		IEnumerable<DirectoryItem> IExternalAction.RootPaths => RootPaths.Values.Select(x => DirectoryItem.GetItemByDirectoryReference(x));
 		public DirectoryReference WorkingDirectory => Unreal.EngineSourceDirectory;
 		string IExternalAction.CommandDescription
 		{
@@ -380,6 +379,8 @@ namespace UnrealBuildTool
 			CompilerExe = FileItem.GetItemByFileReference(Environment.CompilerPath);
 			CompilerType = Environment.Compiler;
 			ToolChainVersion = Environment.ToolChainVersion.ToString();
+
+			RootPaths = new();
 		}
 
 		/// <summary>
@@ -425,7 +426,7 @@ namespace UnrealBuildTool
 			AdditionalPrerequisiteItems = new List<FileItem>(InAction.AdditionalPrerequisiteItems);
 			AdditionalProducedItems = new List<FileItem>(InAction.AdditionalProducedItems);
 			DeleteItems = new List<FileItem>(InAction.DeleteItems);
-			RootPaths = new SortedDictionary<CppRootPathFolder, DirectoryReference>(InAction.RootPaths, null);
+			RootPaths = new(InAction.RootPaths);
 		}
 
 		/// <summary>
@@ -473,7 +474,7 @@ namespace UnrealBuildTool
 			AdditionalPrerequisiteItems = Reader.ReadList(() => Reader.ReadFileItem())!;
 			AdditionalProducedItems = Reader.ReadList(() => Reader.ReadFileItem())!;
 			DeleteItems = Reader.ReadList(() => Reader.ReadFileItem())!;
-			RootPaths = new(Reader.ReadDictionary(() => (CppRootPathFolder)Reader.ReadInt(), () => Reader.ReadCompactDirectoryReference())!, null);
+			RootPaths = new(Reader);
 		}
 
 		/// <inheritdoc/>
@@ -518,7 +519,7 @@ namespace UnrealBuildTool
 			Writer.WriteList(AdditionalPrerequisiteItems, Item => Writer.WriteFileItem(Item));
 			Writer.WriteList(AdditionalProducedItems, Item => Writer.WriteFileItem(Item));
 			Writer.WriteList(DeleteItems, Item => Writer.WriteFileItem(Item));
-			Writer.WriteDictionary(RootPaths, Key => Writer.WriteInt((int)Key), Value => Writer.WriteCompactDirectoryReference(Value));
+			RootPaths.Write(Writer);
 		}
 
 		/// <summary>
@@ -537,7 +538,7 @@ namespace UnrealBuildTool
 
 		public List<string> GetCompilerArguments(ILogger Logger)
 		{
-			IReadOnlyDictionary<CppRootPathFolder, DirectoryReference>? rootPaths = CompilerType.IsClang() && VfsOverlayFile != null ? RootPaths : null;
+			CppRootPaths? rootPaths = CompilerType.IsClang() && VfsOverlayFile != null ? RootPaths : null;
 			List<string> Arguments = new List<string>();
 
 			if (SourceFile != null)
diff --git a/Engine/Source/Programs/UnrealBuildTool/Platform/Windows/VCToolChain.cs b/Engine/Source/Programs/UnrealBuildTool/Platform/Windows/VCToolChain.cs
--- a/Engine/Source/Programs/UnrealBuildTool/Platform/Windows/VCToolChain.cs
+++ b/Engine/Source/Programs/UnrealBuildTool/Platform/Windows/VCToolChain.cs
@@ -173,7 +173,7 @@ namespace UnrealBuildTool
 			return NormalizeCommandLinePath(Item.Location);
 		}
 
-		protected static string NormalizeCommandLinePath(FileSystemReference reference, IReadOnlyDictionary<CppRootPathFolder, DirectoryReference>? RootPaths)
+		protected static string NormalizeCommandLinePath(FileSystemReference reference, CppRootPaths? RootPaths)
 		{
 			if (RootPaths != null)
 			{
@@ -189,7 +189,7 @@ namespace UnrealBuildTool
 			return NormalizeCommandLinePath(reference);
 		}
 
-		protected static string NormalizeCommandLinePath(FileItem Item, IReadOnlyDictionary<CppRootPathFolder, DirectoryReference>? RootPaths)
+		protected static string NormalizeCommandLinePath(FileItem Item, CppRootPaths? RootPaths)
 		{
 			return NormalizeCommandLinePath(Item.Location, RootPaths);
 		}
@@ -206,7 +206,7 @@ namespace UnrealBuildTool
 			Arguments.Add(Utils.MakePathSafeToUseWithCommandLine(SourceFileString));
 		}
 
-		private static void AddIncludePath(List<string> Arguments, DirectoryReference IncludePath, WindowsCompiler Compiler, bool bSystemInclude, IReadOnlyDictionary<CppRootPathFolder, DirectoryReference>? RootPaths = null)
+		private static void AddIncludePath(List<string> Arguments, DirectoryReference IncludePath, WindowsCompiler Compiler, bool bSystemInclude, CppRootPaths? RootPaths = null)
 		{
 			// Try to use a relative path to shorten command line length.
 			string IncludePathString = NormalizeCommandLinePath(IncludePath, RootPaths);
@@ -232,17 +232,17 @@ namespace UnrealBuildTool
 			}
 		}
 
-		public static void AddIncludePath(List<string> Arguments, DirectoryReference IncludePath, WindowsCompiler Compiler, IReadOnlyDictionary<CppRootPathFolder, DirectoryReference>? RootPaths = null)
+		public static void AddIncludePath(List<string> Arguments, DirectoryReference IncludePath, WindowsCompiler Compiler, CppRootPaths? RootPaths = null)
 		{
 			AddIncludePath(Arguments, IncludePath, Compiler, false, RootPaths);
 		}
 
-		public static void AddSystemIncludePath(List<string> Arguments, DirectoryReference IncludePath, WindowsCompiler Compiler, IReadOnlyDictionary<CppRootPathFolder, DirectoryReference>? RootPaths = null)
+		public static void AddSystemIncludePath(List<string> Arguments, DirectoryReference IncludePath, WindowsCompiler Compiler, CppRootPaths? RootPaths = null)
 		{
 			AddIncludePath(Arguments, IncludePath, Compiler, true, RootPaths);
 		}
 
-		public static void AddForceIncludeFile(List<string> Arguments, FileItem ForceIncludeFile, IReadOnlyDictionary<CppRootPathFolder, DirectoryReference>? RootPaths)
+		public static void AddForceIncludeFile(List<string> Arguments, FileItem ForceIncludeFile, CppRootPaths? RootPaths)
 		{
 			string ForceIncludeFileString = NormalizeCommandLinePath(ForceIncludeFile, RootPaths);
 			Arguments.Add($"/FI\"{ForceIncludeFileString}\"");
@@ -266,7 +266,7 @@ namespace UnrealBuildTool
 			}
 		}
 
-		public static void AddUsingPchFile(List<string> Arguments, FileItem PchThroughHeaderFile, FileItem UsingPchFile, IReadOnlyDictionary<CppRootPathFolder, DirectoryReference>? RootPaths)
+		public static void AddUsingPchFile(List<string> Arguments, FileItem PchThroughHeaderFile, FileItem UsingPchFile, CppRootPaths? RootPaths)
 		{
 			string PchThroughHeaderFilePath = NormalizeCommandLinePath(PchThroughHeaderFile, RootPaths);
 			string UsingPchFilePath = NormalizeCommandLinePath(UsingPchFile, RootPaths);
@@ -1769,10 +1769,7 @@ namespace UnrealBuildTool
 				}
 			}
 #endif
-			foreach (KeyValuePair<CppRootPathFolder, DirectoryReference> Item in CompileEnvironment.RootPaths)
-			{
-				BaseCompileAction.RootPaths.Add(Item.Key, Item.Value);
-			}
+			BaseCompileAction.RootPaths = CompileEnvironment.RootPaths;
 			BaseCompileAction.VfsOverlayFile = CompileEnvironment.VfsOverlayFile;
 
 			// Add additional response files
@@ -2404,7 +2401,7 @@ namespace UnrealBuildTool
 			foreach (FileItem RCFile in InputFiles)
 			{
 				Action CompileAction = Graph.CreateAction(ActionType.Compile);
-				CompileAction.RootPaths.AddRange(CompileEnvironment.RootPaths.Values.Select(x => DirectoryItem.GetItemByDirectoryReference(x)));
+				CompileAction.RootPaths = CompileEnvironment.RootPaths;
 				CompileAction.CommandDescription = "Resource";
 				CompileAction.WorkingDirectory = Unreal.EngineSourceDirectory;
 				CompileAction.CommandPath = EnvVars.ResourceCompilerPath;
@@ -2616,7 +2613,7 @@ namespace UnrealBuildTool
 			CompileAction.CommandVersion = EnvVars.ToolChainVersion.ToString();
 			CompileAction.bShouldOutputStatusDescription = false;
 			CompileAction.bCanExecuteRemotely = false; // Incompatible with remote distribution
-			CompileAction.RootPaths.AddRange(CompileEnvironment.RootPaths.Values.Select(x => DirectoryItem.GetItemByDirectoryReference(x)));
+			CompileAction.RootPaths = CompileEnvironment.RootPaths;
 			CompileAction.ArtifactMode = ArtifactMode.Enabled;
 		}
 
@@ -2644,7 +2641,7 @@ namespace UnrealBuildTool
 		public override CppCompileEnvironment CreateSharedResponseFile(CppCompileEnvironment CompileEnvironment, FileReference OutResponseFile, IActionGraphBuilder Graph)
 		{
 			CppCompileEnvironment NewCompileEnvironment = new CppCompileEnvironment(CompileEnvironment);
-			IReadOnlyDictionary<CppRootPathFolder, DirectoryReference>? rootPaths = Target.WindowsPlatform.Compiler.IsClang() && NewCompileEnvironment.VfsOverlayFile != null ? NewCompileEnvironment.RootPaths : null;
+			CppRootPaths? rootPaths = Target.WindowsPlatform.Compiler.IsClang() && NewCompileEnvironment.VfsOverlayFile != null ? NewCompileEnvironment.RootPaths : null;
 			List<string> Arguments = new List<string>();
 
 			foreach (DirectoryReference IncludePath in NewCompileEnvironment.UserIncludePaths)
@@ -2985,7 +2982,7 @@ namespace UnrealBuildTool
 
 			// Create an action that invokes the linker.
 			Action LinkAction = Graph.CreateAction(ActionType.Link);
-			LinkAction.RootPaths.AddRange(LinkEnvironment.RootPaths.Values.Select(x => DirectoryItem.GetItemByDirectoryReference(x)));
+			LinkAction.RootPaths = LinkEnvironment.RootPaths;
 			string ReadableArch = UnrealArchitectureConfig.ForPlatform(LinkEnvironment.Platform).ConvertToReadableArchitecture(LinkEnvironment.Architecture);
 			LinkAction.CommandDescription += $"Link [{ReadableArch}]";
 			LinkAction.WorkingDirectory = Unreal.EngineSourceDirectory;
diff --git a/Engine/Source/Programs/UnrealBuildTool/System/Action.cs b/Engine/Source/Programs/UnrealBuildTool/System/Action.cs
--- a/Engine/Source/Programs/UnrealBuildTool/System/Action.cs
+++ b/Engine/Source/Programs/UnrealBuildTool/System/Action.cs
@@ -98,7 +98,7 @@ namespace UnrealBuildTool
 		/// <summary>
 		/// Root paths for this action (generally engine root project root, toolchain root, sdk root)
 		/// </summary>
-		IEnumerable<DirectoryItem> RootPaths { get; }
+		CppRootPaths RootPaths { get; }
 
 		/// <summary>
 		/// For C++ source files, specifies a dependency list file used to check changes to header files
@@ -250,7 +250,7 @@ namespace UnrealBuildTool
 		/// Root paths for this action (generally engine root project root, toolchain root, sdk root)
         /// Order needs to be mainted, so a list is used
 		/// </summary>
-		public List<DirectoryItem> RootPaths { get; set; } = new();
+		public CppRootPaths RootPaths { get; set; } = new();
 
 		/// <summary>
 		/// For C++ source files, specifies a dependency list file used to check changes to header files
@@ -356,8 +356,6 @@ namespace UnrealBuildTool
 		IEnumerable<FileItem> IExternalAction.PrerequisiteItems => PrerequisiteItems;
 		IEnumerable<FileItem> IExternalAction.ProducedItems => ProducedItems;
 		IEnumerable<FileItem> IExternalAction.DeleteItems => DeleteItems;
-        IEnumerable<DirectoryItem> IExternalAction.RootPaths => RootPaths;
-
         public Action(ActionType InActionType)
 		{
 			ActionType = InActionType;
@@ -376,7 +374,7 @@ namespace UnrealBuildTool
 			PrerequisiteItems = new(InOther.PrerequisiteItems);
 			ProducedItems = new(InOther.ProducedItems);
 			DeleteItems = new(InOther.DeleteItems);
-            RootPaths = new(InOther.RootPaths);
+			RootPaths = new(InOther.RootPaths);
 			DependencyListFile = InOther.DependencyListFile;
 			WorkingDirectory = InOther.WorkingDirectory;
 			CommandPath = InOther.CommandPath;
@@ -424,8 +422,8 @@ namespace UnrealBuildTool
 			PrerequisiteItems = Reader.ReadSortedSet(() => Reader.ReadFileItem())!;
 			ProducedItems = Reader.ReadSortedSet(() => Reader.ReadFileItem())!;
 			DeleteItems = Reader.ReadSortedSet(() => Reader.ReadFileItem())!;
-            RootPaths = Reader.ReadList(() => Reader.ReadDirectoryItem())!;
-            DependencyListFile = Reader.ReadFileItem();
+			RootPaths = new(Reader);
+			DependencyListFile = Reader.ReadFileItem();
 			bUseActionHistory = Reader.ReadBool();
 			bIsHighPriority = Reader.ReadBool();
 			Weight = Reader.ReadDouble();
@@ -459,7 +457,7 @@ namespace UnrealBuildTool
 			Writer.WriteSortedSet(PrerequisiteItems, Item => Writer.WriteFileItem(Item));
 			Writer.WriteSortedSet(ProducedItems, Item => Writer.WriteFileItem(Item));
 			Writer.WriteSortedSet(DeleteItems, Item => Writer.WriteFileItem(Item));
-            Writer.WriteList(RootPaths, Item => Writer.WriteDirectoryItem(Item));
+			RootPaths.Write(Writer);
 			Writer.WriteFileItem(DependencyListFile);
 			Writer.WriteBool(bUseActionHistory);
 			Writer.WriteBool(bIsHighPriority);
@@ -579,9 +577,12 @@ namespace UnrealBuildTool
 				Action.DeleteItems.UnionWith(DeleteItems.Select(x => FileItem.GetItemByPath(x)));
 			}
 
-			if (Object.TryGetStringArrayField("RootPaths", out string[]? RootPaths))
+			if (Object.TryGetObjectArrayField("RootPaths", out JsonObject[]? RootPaths))
 			{
-                Action.RootPaths.AddRange(RootPaths.Select(x => DirectoryItem.GetItemByPath(x)));
+				foreach (JsonObject Root in RootPaths)
+				{
+					Action.RootPaths[Root.GetEnumField<CppRootPathFolder>("Root")] = DirectoryReference.FromString(Root.GetStringField("Path"))!;
+				}
 			}
 
 			if (Object.TryGetStringField("DependencyListFile", out string? DependencyListFile))
@@ -686,9 +687,12 @@ namespace UnrealBuildTool
 			Writer.WriteArrayEnd();
 
 			Writer.WriteArrayStart("RootPaths");
-			foreach (DirectoryItem RootPath in Action.RootPaths)
+			foreach (KeyValuePair<CppRootPathFolder, DirectoryReference> RootPath in Action.RootPaths)
 			{
-				Writer.WriteValue(RootPath.FullName);
+				Writer.WriteObjectStart();
+				Writer.WriteEnumValue("Root", RootPath.Key);
+				Writer.WriteValue("Path", RootPath.Value.FullName);
+				Writer.WriteObjectEnd();
 			}
 			Writer.WriteArrayEnd();
 
@@ -777,7 +781,7 @@ namespace UnrealBuildTool
 		public IEnumerable<FileItem> PrerequisiteItems => Inner.PrerequisiteItems;
 		public IEnumerable<FileItem> ProducedItems => Inner.ProducedItems;
 		public IEnumerable<FileItem> DeleteItems => Inner.DeleteItems;
-		public IEnumerable<DirectoryItem> RootPaths => Inner.RootPaths;
+		public CppRootPaths RootPaths => Inner.RootPaths;
 		public FileItem? DependencyListFile => Inner.DependencyListFile;
 		public DirectoryReference WorkingDirectory => Inner.WorkingDirectory;
 		public FileReference CommandPath => Inner.CommandPath;
diff --git a/Engine/Source/Programs/UnrealBuildTool/System/CppCompileEnvironment.cs b/Engine/Source/Programs/UnrealBuildTool/System/CppCompileEnvironment.cs
--- a/Engine/Source/Programs/UnrealBuildTool/System/CppCompileEnvironment.cs
+++ b/Engine/Source/Programs/UnrealBuildTool/System/CppCompileEnvironment.cs
@@ -242,6 +242,50 @@ namespace UnrealBuildTool
 		public static string AsVfsString(this CppRootPathFolder folder, DirectoryReference localRoot, FileSystemReference localPath) => AsVfsPath(folder, localRoot, localPath).FullName.Replace('\\', '/');
 	}
 
+	/// <summary>
+	/// Encapsulates the roots used for cache and vfs
+	/// </summary>
+	class CppRootPaths : IEnumerable<KeyValuePair<CppRootPathFolder, DirectoryReference>>
+	{
+		public CppRootPaths()
+		{
+			Lookup = [];
+		}
+
+		public CppRootPaths(BinaryArchiveReader Reader)
+		{
+			Lookup = new(Reader.ReadDictionary(() => (CppRootPathFolder)Reader.ReadInt(), () => Reader.ReadCompactDirectoryReference())!, null);
+		}
+
+		public CppRootPaths(CppRootPaths Other)
+		{
+			Lookup = new(Other.Lookup);
+		}
+
+		public void Write(BinaryArchiveWriter Writer)
+		{
+			Writer.WriteDictionary(Lookup, Key => Writer.WriteInt((int)Key), Value => Writer.WriteCompactDirectoryReference(Value));
+		}
+
+		public DirectoryReference this[CppRootPathFolder key]
+		{
+			set => Lookup[key] = value;
+		}
+
+		public IEnumerator<KeyValuePair<CppRootPathFolder, DirectoryReference>> GetEnumerator()
+		{
+			return Lookup.GetEnumerator();
+		}
+
+		System.Collections.IEnumerator System.Collections.IEnumerable.GetEnumerator()
+		{
+			return Lookup.GetEnumerator();
+		}
+
+		private readonly SortedDictionary<CppRootPathFolder, DirectoryReference> Lookup;
+	}
+
+
 	/// <summary>
 	/// Encapsulates the compilation output of compiling a set of C++ files.
 	/// </summary>
@@ -312,7 +356,7 @@ namespace UnrealBuildTool
 		/// <summary>
 		/// Root paths for this environment. Used for Clang VFS and UBA
 		/// </summary>
-		public readonly SortedDictionary<CppRootPathFolder, DirectoryReference> RootPaths;
+		public readonly CppRootPaths RootPaths;
 
 		/// <summary>
 		/// Templates for shared precompiled headers
@@ -863,7 +907,7 @@ namespace UnrealBuildTool
 			Configuration = Other.Configuration;
 			Architectures = new(Other.Architectures);
 			MetadataCache = Other.MetadataCache;
-			RootPaths = new(Other.RootPaths, null);
+			RootPaths = new(Other.RootPaths);
 			SharedPCHs = Other.SharedPCHs;
 			PrecompiledHeaderIncludeFilename = Other.PrecompiledHeaderIncludeFilename;
 			PrecompiledHeaderAction = Other.PrecompiledHeaderAction;
diff --git a/Engine/Source/Programs/UnrealBuildTool/System/LinkEnvironment.cs b/Engine/Source/Programs/UnrealBuildTool/System/LinkEnvironment.cs
--- a/Engine/Source/Programs/UnrealBuildTool/System/LinkEnvironment.cs
+++ b/Engine/Source/Programs/UnrealBuildTool/System/LinkEnvironment.cs
@@ -36,7 +36,7 @@ namespace UnrealBuildTool
 		/// <summary>
 		/// Root paths for this environment. Used for Clang VFS and UBA
 		/// </summary>
-		public readonly SortedDictionary<CppRootPathFolder, DirectoryReference> RootPaths;
+		public readonly CppRootPaths RootPaths;
 
 		/// <summary>
 		/// On Mac, indicates the path to the target's application bundle
@@ -391,7 +391,7 @@ namespace UnrealBuildTool
 			Platform = Other.Platform;
 			Configuration = Other.Configuration;
 			Architectures = Other.Architectures;
-			RootPaths = new(Other.RootPaths, null);
+			RootPaths = new(Other.RootPaths);
 			BundleDirectory = Other.BundleDirectory;
 			OutputDirectory = Other.OutputDirectory;
 			IntermediateDirectory = Other.IntermediateDirectory;
diff --git a/Engine/Source/Programs/UnrealBuildTool/ToolChain/ClangToolChain.cs b/Engine/Source/Programs/UnrealBuildTool/ToolChain/ClangToolChain.cs
--- a/Engine/Source/Programs/UnrealBuildTool/ToolChain/ClangToolChain.cs
+++ b/Engine/Source/Programs/UnrealBuildTool/ToolChain/ClangToolChain.cs
@@ -1216,7 +1216,7 @@ namespace UnrealBuildTool
 		protected virtual Action CompileCPPFile(CppCompileEnvironment CompileEnvironment, FileItem SourceFile, DirectoryReference OutputDir, string ModuleName, IActionGraphBuilder Graph, IReadOnlyCollection<string> GlobalArguments, CPPOutput Result)
 		{
 			Action CompileAction = Graph.CreateAction(ActionType.Compile);
-			CompileAction.RootPaths.AddRange(CompileEnvironment.RootPaths.Values.Select(x => DirectoryItem.GetItemByDirectoryReference(x)));
+			CompileAction.RootPaths = CompileEnvironment.RootPaths;
 
 			// If we are using the AutoRTFM compiler, we make the compile action depend on the version of the compiler itself.
 			// This lets us update the compiler (which might not cause a version update of the compiler, which instead tracks
@@ -1314,7 +1314,7 @@ namespace UnrealBuildTool
 			if (PreprocessDepends)
 			{
 				Action PrepassAction = Graph.CreateAction(ActionType.Compile);
-				PrepassAction.RootPaths.AddRange(CompileAction.RootPaths);
+				PrepassAction.RootPaths = CompileAction.RootPaths;
 				PrepassAction.PrerequisiteItems.UnionWith(CompileAction.PrerequisiteItems);
 				PrepassAction.PrerequisiteItems.Remove(CompilerResponseFileItem);
 				PrepassAction.CommandDescription = "Preprocess Depends";
diff --git a/Engine/Source/Programs/UnrealBuildTool/ToolChain/ISPCToolChain.cs b/Engine/Source/Programs/UnrealBuildTool/ToolChain/ISPCToolChain.cs
--- a/Engine/Source/Programs/UnrealBuildTool/ToolChain/ISPCToolChain.cs
+++ b/Engine/Source/Programs/UnrealBuildTool/ToolChain/ISPCToolChain.cs
@@ -447,12 +447,10 @@ namespace UnrealBuildTool
 				}
 			}
 
-			IEnumerable<DirectoryItem> RootPaths = CompileEnvironment.RootPaths.Values.Select(x => DirectoryItem.GetItemByDirectoryReference(x));
-
 			foreach (FileItem ISPCFile in InputFiles)
 			{
 				Action CompileAction = Graph.CreateAction(ActionType.Compile);
-				CompileAction.RootPaths.AddRange(RootPaths);
+				CompileAction.RootPaths = CompileEnvironment.RootPaths;
 
 				CompileAction.CommandDescription = $"Generate Header [{ISPCArch}]";
 				CompileAction.WorkingDirectory = Unreal.EngineSourceDirectory;
@@ -649,12 +647,10 @@ namespace UnrealBuildTool
 				}
 			}
 
-			IEnumerable<DirectoryItem> RootPaths = CompileEnvironment.RootPaths.Values.Select(x => DirectoryItem.GetItemByDirectoryReference(x));
-
 			foreach (FileItem ISPCFile in InputFiles)
 			{
 				Action CompileAction = Graph.CreateAction(ActionType.Compile);
-				CompileAction.RootPaths.AddRange(RootPaths);
+				CompileAction.RootPaths = CompileEnvironment.RootPaths;
 				CompileAction.CommandDescription = $"Compile [{ISPCArch}]";
 				CompileAction.WorkingDirectory = Unreal.EngineSourceDirectory;
 				CompileAction.CommandPath = new FileReference(GetISPCHostCompilerPath(BuildHostPlatform.Current.Platform));
@@ -824,7 +820,7 @@ namespace UnrealBuildTool
 								CompileAction.CommandVersion = $"{CompileAction.CommandVersion} Compressed";
 							}
 
-							PostCompileAction.RootPaths.AddRange(RootPaths);
+							PostCompileAction.RootPaths = CompileEnvironment.RootPaths;
 							PostCompileAction.ArtifactMode = ArtifactMode.Enabled;
 
 							// Disable remote execution to workaround mismatched case on XGE
