[UBT]
* Changed pch cpp files to use includes without path and add additional include paths instead

diff --git a/Engine/Source/Programs/UnrealBuildTool/Configuration/UEBuildModuleCPP.cs b/Engine/Source/Programs/UnrealBuildTool/Configuration/UEBuildModuleCPP.cs
--- a/Engine/Source/Programs/UnrealBuildTool/Configuration/UEBuildModuleCPP.cs
+++ b/Engine/Source/Programs/UnrealBuildTool/Configuration/UEBuildModuleCPP.cs
@@ -1910,7 +1910,7 @@ namespace UnrealBuildTool
 							}
 
 							StringBuilder Writer = new();
-							WriteDefinitions($"Shared PCH Definitions for {Name}", new string[] { GetIncludeString(Instance.DefinitionsFile) }, NewDefinitions, Undefintions, Writer);
+							WriteDefinitions($"Shared PCH Definitions for {Name}", new string[] { Instance.DefinitionsFile.Name }, NewDefinitions, Undefintions, Writer);
 							PrivateDefinitionsFileItem = Graph.CreateIntermediateTextFile(PrivateDefinitionsFile, Writer.ToString());
 						}
 
@@ -1919,6 +1919,7 @@ namespace UnrealBuildTool
 						CompileEnvironment.ForceIncludeFiles.Insert(0, PrivateDefinitionsFileItem);
 						CompileEnvironment.PrecompiledHeaderAction = PrecompiledHeaderAction.Include;
 						CompileEnvironment.PrecompiledHeaderIncludeFilename = Instance.HeaderFile.Location;
+						CompileEnvironment.UserIncludePaths.Add(Instance.DefinitionsFile.Directory.Location);
 						CompileEnvironment.PCHInstance = Instance;
 
 						LinkInputFiles.AddRange(Instance.Output.ObjectFiles);
@@ -1978,16 +1979,6 @@ namespace UnrealBuildTool
 			CompileEnvironment.UserIncludePaths.Add(IntermediateDirectory);
 		}
 
-		static string GetIncludeString(FileItem FileItem)
-		{
-			string IncludeFileString = FileItem.AbsolutePath;
-			if (FileItem.Location.IsUnderDirectory(Unreal.RootDirectory))
-			{
-				IncludeFileString = FileItem.Location.MakeRelativeTo(Unreal.EngineSourceDirectory);
-			}
-			return IncludeFileString.Replace('\\', '/');
-		}
-
 		/// <summary>
 		/// Create a header file containing the module definitions, which also includes the PCH itself. Including through another file is necessary on 
 		/// Clang, since we get warnings about #pragma once otherwise, but it also allows us to consistently define the preprocessor state on all 
@@ -2003,11 +1994,9 @@ namespace UnrealBuildTool
 			// Build the contents of the wrapper file
 			StringBuilder WrapperContents = new StringBuilder();
 			{
-				string IncludeFileString = GetIncludeString(IncludedFile);
-				string DefinitionsFileString = GetIncludeString(DefinitionsFile);
-				WrapperContents.AppendLine("// PCH for {0}", IncludeFileString);
-				WrapperContents.AppendLine("#include \"{0}\"", DefinitionsFileString);
-				WrapperContents.AppendLine("#include \"{0}\"", IncludeFileString);
+				WrapperContents.AppendLine("// PCH for {0}", IncludedFile.Name);
+				WrapperContents.AppendLine("#include \"{0}\"", DefinitionsFile.Name);
+				WrapperContents.AppendLine("#include \"{0}\"", IncludedFile.Name);
 			}
 
 			// Create the item
