- Added support for targeting a device with xcodebuild, so BuildCookRun can add a device to the provision if needed

diff --git a/Engine/Source/Programs/AutomationTool/Apple/ApplePlatform.Automation.cs b/Engine/Source/Programs/AutomationTool/Apple/ApplePlatform.Automation.cs
--- a/Engine/Source/Programs/AutomationTool/Apple/ApplePlatform.Automation.cs
+++ b/Engine/Source/Programs/AutomationTool/Apple/ApplePlatform.Automation.cs
@@ -151,7 +151,7 @@ public abstract class ApplePlatform : Platform
 					}
 				}
 
-				AppleExports.BuildWithStubXcodeProject(SC.RawProjectPath, Target.Platform, Target.Architectures, Target.Configuration, TargetName, AppleExports.XcodeBuildMode.Stage, Logger, ExtraOptions);
+				AppleExports.BuildWithStubXcodeProject(SC.RawProjectPath, Target.Platform, Target.Architectures, Target.Configuration, TargetName, AppleExports.XcodeBuildMode.Stage, Logger, ExtraOptions, DestinationIds:Params.DeviceNames);
 			}
 		}
 	}
@@ -260,7 +260,7 @@ public abstract class ApplePlatform : Platform
 				// the archive will be created in the standard Archives location accessible via Xcode. Using -archive will copy it out into
 				// the specified location for use as needed
 				AppleExports.XcodeBuildMode BuildMode = Params.Distribution ? AppleExports.XcodeBuildMode.Distribute : AppleExports.XcodeBuildMode.Package;
-				if (AppleExports.BuildWithStubXcodeProject(Params.RawProjectPath, Receipt.Platform, Receipt.Architectures, Receipt.Configuration, TargetName, BuildMode, Logger, ExtraOptions) == 0)
+				if (AppleExports.BuildWithStubXcodeProject(Params.RawProjectPath, Receipt.Platform, Receipt.Architectures, Receipt.Configuration, TargetName, BuildMode, Logger, ExtraOptions, DestinationIds:Params.DeviceNames) == 0)
 				{
 					Logger.LogInformation("==================================================================================================================================================");
 					if (Params.Distribution)
diff --git a/Engine/Source/Programs/UnrealBuildTool/Platform/Apple/AppleExports.cs b/Engine/Source/Programs/UnrealBuildTool/Platform/Apple/AppleExports.cs
--- a/Engine/Source/Programs/UnrealBuildTool/Platform/Apple/AppleExports.cs
+++ b/Engine/Source/Programs/UnrealBuildTool/Platform/Apple/AppleExports.cs
@@ -210,9 +210,10 @@ namespace UnrealBuildTool
 		/// <param name="Logger"></param>
 		/// <param name="ExtraOptions">Any extra options to pass to xcodebuild</param>
 		/// <param name="bForceDummySigning">If true, force signing with the - identity</param>
+		/// <param name="DestinationIds">Optional list of devices to target, mainly for adding to provision if needed</param>
 		/// <returns>xcode's exit code</returns>
 		public static int BuildWithStubXcodeProject(FileReference? ProjectFile, UnrealTargetPlatform Platform, UnrealArchitectures Architectures, UnrealTargetConfiguration Configuration,
-			string TargetName, XcodeBuildMode BuildMode, ILogger Logger, string ExtraOptions = "", bool bForceDummySigning = false)
+			string TargetName, XcodeBuildMode BuildMode, ILogger Logger, string ExtraOptions = "", bool bForceDummySigning = false, List<string>? DestinationIds = null)
 		{
 			DirectoryReference? GeneratedProjectFile;
 			// we don't use distro flag when making a modern project
@@ -223,6 +224,7 @@ namespace UnrealBuildTool
 				return 1;
 			}
 
+			bool bUseAutomaticCodeSigning = false;
 			if (bForceDummySigning)
 			{
 				ExtraOptions += " CODE_SIGN_STYLE=Manual CODE_SIGN_IDENTITY= PROVISIONING_PROFILE_SPECIFIER=";
@@ -231,7 +233,6 @@ namespace UnrealBuildTool
 			{
 				// look for the special app store connect key information
 				ConfigHierarchy SharedPlatformIni = ConfigCache.ReadHierarchy(ConfigHierarchyType.Engine, ProjectFile?.Directory, UnrealTargetPlatform.Mac);
-				bool bUseAutomaticCodeSigning;
 				SharedPlatformIni.TryGetValue("/Script/MacTargetPlatform.XcodeProjectSettings", "bUseAutomaticCodeSigning", out bUseAutomaticCodeSigning);
 
 				// disable automatic signing, if some extra options imply manual
@@ -251,7 +252,7 @@ namespace UnrealBuildTool
 
 			// run xcodebuild on the generated project to make the .app
 			string XcodeBuildAction = BuildMode == XcodeBuildMode.Distribute ? "archive" : "build";
-			return AppleExports.FinalizeAppWithXcode(GeneratedProjectFile!, Platform, Architectures, TargetName, Configuration.ToString(), XcodeBuildAction, ExtraOptions + $" UE_XCODE_BUILD_MODE={BuildMode}", Logger);
+			return AppleExports.FinalizeAppWithXcode(GeneratedProjectFile!, Platform, Architectures, bUseAutomaticCodeSigning, TargetName, Configuration.ToString(), XcodeBuildAction, ExtraOptions + $" UE_XCODE_BUILD_MODE={BuildMode}", Logger, DestinationIds);
 		}
 
 		/// <summary>
@@ -274,15 +275,17 @@ namespace UnrealBuildTool
 		/// <param name="XcodeProject">The .xcworkspace file to build</param>
 		/// <param name="Platform">THe platform to make the .app for</param>
 		/// <param name="Architectures">The architecture we are targeting (really, this is for Simulator to pass in the -destination field)</param>
+		/// <param name="bUseAutomaticCodeSigning">True when using automatic codesigning via xcodebuild</param>
 		/// <param name="SchemeName">The name of the scheme (basically the target on the .xcworkspace)</param>
 		/// <param name="Configuration">Which configuration to make (Debug, etc)</param>
 		/// <param name="Action">Action (build, archive, etc)</param>
 		/// <param name="ExtraOptions">Extra options to pass to xcodebuild</param>
 		/// <param name="Logger">Logging object</param>
+		/// <param name="DestinationIds">Optional list of devices to target, mainly for adding to provision if needed</param>
 		/// <returns>xcode's exit code</returns>
-		public static int FinalizeAppWithXcode(DirectoryReference XcodeProject, UnrealTargetPlatform Platform, UnrealArchitectures Architectures, string SchemeName, string Configuration, string Action, string ExtraOptions, ILogger Logger)
+		public static int FinalizeAppWithXcode(DirectoryReference XcodeProject, UnrealTargetPlatform Platform, UnrealArchitectures Architectures, bool bUseAutomaticCodeSigning, string SchemeName, string Configuration, string Action, string ExtraOptions, ILogger Logger, List<string>? DestinationIds=null)
 		{
-			return AppleToolChain.FinalizeAppWithXcode(XcodeProject, Platform, Architectures, SchemeName, Configuration, Action, ExtraOptions, Logger);
+			return AppleToolChain.FinalizeAppWithXcode(XcodeProject, Platform, Architectures, bUseAutomaticCodeSigning, SchemeName, Configuration, Action, ExtraOptions, Logger, DestinationIds);
 		}
 
 		/// <summary>
diff --git a/Engine/Source/Programs/UnrealBuildTool/ToolChain/AppleToolChain.cs b/Engine/Source/Programs/UnrealBuildTool/ToolChain/AppleToolChain.cs
--- a/Engine/Source/Programs/UnrealBuildTool/ToolChain/AppleToolChain.cs
+++ b/Engine/Source/Programs/UnrealBuildTool/ToolChain/AppleToolChain.cs
@@ -909,13 +909,38 @@ namespace UnrealBuildTool
 			Environment.CurrentDirectory = CurrentCWD;
 		}
 
-		internal static int FinalizeAppWithXcode(DirectoryReference XcodeProject, UnrealTargetPlatform Platform, UnrealArchitectures Architectures, string SchemeName, string Configuration, string Action, string ExtraOptions, ILogger Logger)
+		static string IdOrNameParam(string Id)
+		{
+			string OldFormat = "^[0-9,a-f]{40}$";
+			string NewFormat = "^[0-9,A-F]{8}-[0-9,A-F]{16}$";
+			if (Regex.Match(Id, OldFormat).Success || Regex.Match(Id, NewFormat).Success)
+			{
+				return $"id={Id}";
+			}
+			return $"name={Id}";
+		}
+
+		internal static int FinalizeAppWithXcode(DirectoryReference XcodeProject, UnrealTargetPlatform Platform, UnrealArchitectures Architectures, bool bUseAutomaticCodeSigning, string SchemeName, string Configuration, string Action, string ExtraOptions, ILogger Logger, List<string>? DestinationIds=null)
 		{
 			// Acquire a different mutex to the regular UBT instance, since this mode will be called as part of a build. We need the mutex to ensure that building two modular configurations 
 			// in parallel don't clash over writing shared *.modules files (eg. DebugGame and Development editors).
 			string MutexName = GlobalSingleInstanceMutex.GetUniqueMutexForPath("UnrealBuildTool_XcodeBuild", Unreal.RootDirectory);
 			using (new GlobalSingleInstanceMutex(MutexName, true))
 			{
+				string DestinationParam;
+				// use generic:
+				//   - for Mac, the destination seems to want name="My Mac", not the real name of the Mac
+				//   - when no destinations were supplied
+				//   - when using manual codesigning because the main use of this is to update provisioning information (reducing potential for issues)
+				if (Platform == UnrealTargetPlatform.Mac || DestinationIds == null || DestinationIds.Count == 0 || !bUseAutomaticCodeSigning)
+				{
+					DestinationParam = $"-destination generic/platform=\"{AppleExports.GetDestinationPlatform(Platform, Architectures)}\"";
+				}
+				else
+				{
+					string DestPlatform = AppleExports.GetDestinationPlatform(Platform, Architectures);
+					DestinationParam = string.Join(" ", DestinationIds.Select(x => $"-destination \"platform={DestPlatform},{IdOrNameParam(x)}\""));
+				}
 				List<string> Arguments = new()
 				{
 					"UBT_NO_POST_DEPLOY=true",
@@ -924,7 +949,7 @@ namespace UnrealBuildTool
 					$"-workspace \"{XcodeProject.FullName}\"",
 					$"-scheme \"{SchemeName}\"",
 					$"-configuration \"{Configuration}\"",
-					$"-destination generic/platform=\"{AppleExports.GetDestinationPlatform(Platform, Architectures)}\"",
+					DestinationParam,
 					"-hideShellScriptEnvironment",
 					// xcode gets confused it we _just_ wrote out entitlements while generating the temp project, and it thinks it was modified _during_ building
 					// but it wasn't, it was written before the build started
@@ -933,6 +958,8 @@ namespace UnrealBuildTool
 					//$"-sdk {SDKName}",
 				};
 
+				Console.WriteLine("params: {0}", String.Join(" ", Arguments));
+
 				Process LocalProcess = new Process();
 				LocalProcess.StartInfo = new ProcessStartInfo("/usr/bin/env", String.Join(" ", Arguments));
 				LocalProcess.OutputDataReceived += (Sender, Args) => { LocalProcessOutput(Args, false, Logger); };
