[UBT/UHT]
* Changed so module's include paths are provided in manifest.
* Fixed so module's include paths are used to figure out which relative path to use

diff --git a/Engine/Source/Programs/Shared/EpicGames.Core/UHTTypes.cs b/Engine/Source/Programs/Shared/EpicGames.Core/UHTTypes.cs
--- a/Engine/Source/Programs/Shared/EpicGames.Core/UHTTypes.cs
+++ b/Engine/Source/Programs/Shared/EpicGames.Core/UHTTypes.cs
@@ -109,7 +109,7 @@ namespace EpicGames.Core
 			[JsonConverter(typeof(JsonStringEnumConverter))]
 			public EPackageOverrideType OverrideModuleType { get; set; } = EPackageOverrideType.None;
 			public string BaseDirectory { get; set; } = String.Empty;
-			public string IncludeBase { get; set; } = String.Empty; // The include path which all UHT-generated includes should be relative to
+			public List<string> IncludePaths { get; set; } = []; // The include paths which all UHT-generated includes should be relative to
 			public string OutputDirectory { get; set; } = String.Empty;
 			public List<string> ClassesHeaders { get; set; } = [];
 			public List<string> PublicHeaders { get; set; } = [];
diff --git a/Engine/Source/Programs/Shared/EpicGames.UHT/Exporters/CodeGen/UhtCodeGenerator.cs b/Engine/Source/Programs/Shared/EpicGames.UHT/Exporters/CodeGen/UhtCodeGenerator.cs
--- a/Engine/Source/Programs/Shared/EpicGames.UHT/Exporters/CodeGen/UhtCodeGenerator.cs
+++ b/Engine/Source/Programs/Shared/EpicGames.UHT/Exporters/CodeGen/UhtCodeGenerator.cs
@@ -10,6 +10,7 @@ using EpicGames.Core;
 using EpicGames.UHT.Tables;
 using EpicGames.UHT.Types;
 using EpicGames.UHT.Utils;
+using Microsoft.Extensions.Logging;
 
 namespace EpicGames.UHT.Exporters.CodeGen
 {
@@ -240,7 +241,27 @@ namespace EpicGames.UHT.Exporters.CodeGen
 		{
 			ref HeaderInfo headerInfo = ref HeaderInfos[headerFile.HeaderFileTypeIndex];
 
-			headerInfo.IncludePath = Path.GetRelativePath(headerFile.Module.Module.IncludeBase, headerFile.FilePath).Replace('\\', '/');
+			// Find the shortest matching relative path
+			string? relativePath = null;
+			foreach (string includePath in headerFile.Module.Module.IncludePaths)
+			{
+				if (headerFile.FilePath.StartsWith(includePath, StringComparison.Ordinal))
+				{
+					if (relativePath == null || (includePath.Length - headerFile.FilePath.Length - 1) < relativePath.Length)
+					{
+						relativePath = headerFile.FilePath.Substring(includePath.Length + 1);
+					}
+				}
+			}
+
+			// This will create a "../" path which is not great
+			if (relativePath == null)
+			{
+				Session.Logger.LogWarning("There are no include paths in module {ModuleName} for {FileName}. Will create \"../\" path", headerFile.Module.ShortName, headerFile.FilePath);
+				relativePath = Path.GetRelativePath(headerFile.Module.Module.IncludePaths[0], headerFile.FilePath);
+			}
+
+			headerInfo.IncludePath = relativePath.Replace('\\', '/');
 
 			// Convert the file path to a C identifier
 			string filePath = headerFile.FilePath;
diff --git a/Engine/Source/Programs/UnrealBuildTool/System/ExternalExecution.cs b/Engine/Source/Programs/UnrealBuildTool/System/ExternalExecution.cs
--- a/Engine/Source/Programs/UnrealBuildTool/System/ExternalExecution.cs
+++ b/Engine/Source/Programs/UnrealBuildTool/System/ExternalExecution.cs
@@ -101,7 +101,7 @@ namespace UnrealBuildTool
 		/// <summary>
 		/// The include search path for generated headers to include other headers
 		/// </summary>
-		public DirectoryReference ModuleIncludeBase;
+		public DirectoryReference[] ModuleIncludePaths;
 
 		/// <summary>
 		/// Module type
@@ -182,7 +182,7 @@ namespace UnrealBuildTool
 			string ModuleName, 
 			FileReference ModuleRulesFile, 
 			DirectoryReference[] ModuleDirectories, 
-			DirectoryReference ModuleIncludeBase, 
+			DirectoryReference[] ModuleIncludePaths, 
 			UHTModuleType ModuleType, 
 			DirectoryItem GeneratedCodeDirectory, 
 			EGeneratedCodeVersion GeneratedCodeVersion, 
@@ -196,7 +196,7 @@ namespace UnrealBuildTool
 			this.ModuleName = ModuleName;
 			this.ModuleRulesFile = ModuleRulesFile;
 			this.ModuleDirectories = ModuleDirectories;
-			this.ModuleIncludeBase = ModuleIncludeBase;
+			this.ModuleIncludePaths = ModuleIncludePaths;
 			this.ModuleType = ModuleType.ToString();
 			OverrideModuleType = OverrideType.ToString();
 			PublicUObjectClassesHeaders = new List<FileItem>();
@@ -218,7 +218,7 @@ namespace UnrealBuildTool
 			ModuleName = Reader.ReadString()!;
 			ModuleRulesFile = Reader.ReadFileReference();
 			ModuleDirectories = Reader.ReadArray<DirectoryReference>(Reader.ReadDirectoryReferenceNotNull)!;
-			ModuleIncludeBase = Reader.ReadDirectoryReference()!;
+			ModuleIncludePaths = Reader.ReadArray<DirectoryReference>(Reader.ReadDirectoryReferenceNotNull)!;
 			ModuleType = Reader.ReadString()!;
 			OverrideModuleType = Reader.ReadString()!;
 			PublicUObjectClassesHeaders = Reader.ReadList(() => Reader.ReadFileItem())!;
@@ -241,7 +241,7 @@ namespace UnrealBuildTool
 			Writer.WriteString(ModuleName);
 			Writer.WriteFileReference(ModuleRulesFile);
 			Writer.WriteArray<DirectoryReference>(ModuleDirectories, Writer.WriteDirectoryReference);
-			Writer.WriteDirectoryReference(ModuleIncludeBase);
+			Writer.WriteArray<DirectoryReference>(ModuleIncludePaths, Writer.WriteDirectoryReference);
 			Writer.WriteString(ModuleType);
 			Writer.WriteString(OverrideModuleType);
 			Writer.WriteList(PublicUObjectClassesHeaders, Item => Writer.WriteFileItem(Item));
@@ -430,32 +430,6 @@ namespace UnrealBuildTool
 			}
 		}
 
-		static DirectoryReference FindIncludeBase(UEBuildModuleCPP Module, ILogger Logger)
-		{
-			// Plugin Source directory may be an available include path for modules under a UPlugin (see UEBuildTarget.FindOrCreateModuleByName)
-			if (Module.Rules.Plugin != null)
-			{
-				DirectoryReference PluginSourceDirectoryName = DirectoryReference.Combine(Module.Rules.Plugin.Directory, "Source");
-				if (Module.Rules.File.IsUnderDirectory(PluginSourceDirectoryName))
-				{
-					return PluginSourceDirectoryName;
-				}
-			}
-
-			// Project Source directory is always an available include path for modules under a UProject (see UEBuildTarget.FindOrCreateModuleByName)
-			if (Module.Rules.Target.ProjectFile != null)
-			{
-				DirectoryReference ProjectSourceDirectoryName = DirectoryReference.Combine(Module.Rules.Target.ProjectFile.Directory, "Source");
-				if (Module.Rules.File.IsUnderDirectory(ProjectSourceDirectoryName))
-				{
-					return ProjectSourceDirectoryName;
-				}
-			}
-
-			// Otherwise use Engine/Source
-			return Unreal.EngineSourceDirectory;
-		}
-
 		public static void SetupUObjectModules(
 			IEnumerable<UEBuildModuleCPP> ModulesToGenerateHeadersFor,
 			UnrealTargetPlatform Platform, 
@@ -487,13 +461,14 @@ namespace UnrealBuildTool
 					UEBuildModuleCPP Module = ModulesSortedByType[Idx];
 
 					DirectoryItem GeneratedCodeDirectory = DirectoryItem.GetItemByDirectoryReference(Module.GeneratedCodeDirectoryUHT!);
-					DirectoryReference ModuleIncludeBase = FindIncludeBase(Module, Logger);
+
+					DirectoryReference[] ModuleIncludePaths = Module.PublicIncludePaths.Union(Module.InternalIncludePaths).Union(Module.PrivateIncludePaths).ToArray();
 
 					UHTModuleInfo Info = new UHTModuleInfo(
 						Module.Name,
 						Module.RulesFile,
 						Module.ModuleDirectories,
-						ModuleIncludeBase,
+						ModuleIncludePaths,
 						ModuleToType[Module],
 						GeneratedCodeDirectory,
 						GeneratedCodeVersion,
@@ -1063,7 +1038,7 @@ namespace UnrealBuildTool
 						ModuleType = (UHTModuleType)Enum.Parse(typeof(UHTModuleType), UObjectModule.ModuleType),
 						OverrideModuleType = (EPackageOverrideType)Enum.Parse(typeof(EPackageOverrideType), UObjectModule.OverrideModuleType),
 						BaseDirectory = UObjectModule.ModuleDirectories[0].FullName,
-						IncludeBase = UObjectModule.ModuleIncludeBase.FullName,
+						IncludePaths = UObjectModule.ModuleIncludePaths.Select(IncludePath => IncludePath.FullName).ToList(),
 						OutputDirectory = Path.GetDirectoryName(UObjectModule.GeneratedCPPFilenameBase)!,
 						ClassesHeaders = UObjectModule.PublicUObjectClassesHeaders.Select((Header) => Header.AbsolutePath).ToList(),
 						PublicHeaders = UObjectModule.PublicUObjectHeaders.Select((Header) => Header.AbsolutePath).ToList(),
