[UBT]
* Added ModuleRules.ExtraRootPath that can be used per module to set path to files outside normal environment

diff --git a/Engine/Source/Programs/UnrealBuildTool/Configuration/ModuleRules.cs b/Engine/Source/Programs/UnrealBuildTool/Configuration/ModuleRules.cs
--- a/Engine/Source/Programs/UnrealBuildTool/Configuration/ModuleRules.cs
+++ b/Engine/Source/Programs/UnrealBuildTool/Configuration/ModuleRules.cs
@@ -1402,6 +1402,11 @@ namespace UnrealBuildTool
 		/// </summary>
 		public bool? bRequiresImplementModule { get; set; }
 
+		/// <summary>
+		/// Module can provide ExtraRootPath. This will be part of the rootpaths table and is needed if module is using paths outside normal environment
+		/// </summary>
+		public string? ExtraRootPath { get; set; }
+
 		/// <summary>
 		/// If this module has associated Verse code, this is the Verse root path of it
 		/// </summary>
diff --git a/Engine/Source/Programs/UnrealBuildTool/Configuration/UEBuildModuleCPP.cs b/Engine/Source/Programs/UnrealBuildTool/Configuration/UEBuildModuleCPP.cs
--- a/Engine/Source/Programs/UnrealBuildTool/Configuration/UEBuildModuleCPP.cs
+++ b/Engine/Source/Programs/UnrealBuildTool/Configuration/UEBuildModuleCPP.cs
@@ -2419,6 +2419,15 @@ namespace UnrealBuildTool
 				}
 			}
 
+			if (Rules.ExtraRootPath != null)
+			{
+				DirectoryReference? ExtraRootPath = DirectoryReference.FromString(Rules.ExtraRootPath);
+				if (ExtraRootPath != null)
+				{
+					Result.RootPaths[CppRootPathFolder.Extra] = ExtraRootPath;
+				}
+			}
+
 			// Project root path should only be added if the module needs it, to prevent violating shared environment actions
 			if (Target.ProjectFile != null && (!Result.bUseSharedBuildEnvironment || Result.AllIncludePath.Any(x => x.IsUnderDirectory(Target.ProjectFile.Directory))))
 			{
