[UBT]
* Removed provided CompileEnvironment to Normalize calls in ispc toolchain since ispc does not support vfs

diff --git a/Engine/Source/Programs/UnrealBuildTool/ToolChain/ISPCToolChain.cs b/Engine/Source/Programs/UnrealBuildTool/ToolChain/ISPCToolChain.cs
--- a/Engine/Source/Programs/UnrealBuildTool/ToolChain/ISPCToolChain.cs
+++ b/Engine/Source/Programs/UnrealBuildTool/ToolChain/ISPCToolChain.cs
@@ -471,7 +471,7 @@ namespace UnrealBuildTool
 				List<string> Arguments = new List<string>();
 
 				// Add the ISPC obj file as a prerequisite of the action.
-				Arguments.Add($"\"{NormalizeCommandLinePath(ISPCFile, CompileEnvironment)}\"");
+				Arguments.Add($"\"{NormalizeCommandLinePath(ISPCFile)}\"");
 
 				// Add the ISPC h file to the produced item list.
 				FileItem ISPCIncludeHeaderFile = FileItem.GetItemByFileReference(
@@ -509,11 +509,11 @@ namespace UnrealBuildTool
 				}
 
 				// Add the ISPC file to be compiled.
-				Arguments.Add($"-h \"{NormalizeCommandLinePath(ISPCIncludeHeaderFile, CompileEnvironment)}\"");
+				Arguments.Add($"-h \"{NormalizeCommandLinePath(ISPCIncludeHeaderFile)}\"");
 
 				// Generate the included header dependency list
 				FileItem DependencyListFile = FileItem.GetItemByFileReference(FileReference.Combine(OutputDir, Path.GetFileName(ISPCFile.AbsolutePath) + ".txt"));
-				Arguments.Add($"-MMM \"{NormalizeCommandLinePath(DependencyListFile, CompileEnvironment)}\"");
+				Arguments.Add($"-MMM \"{NormalizeCommandLinePath(DependencyListFile)}\"");
 				CompileAction.DependencyListFile = DependencyListFile;
 				CompileAction.ProducedItems.Add(DependencyListFile);
 
@@ -675,7 +675,7 @@ namespace UnrealBuildTool
 				List<string> Arguments = new List<string>();
 
 				// Add the ISPC file to be compiled.
-				Arguments.Add($"\"{NormalizeCommandLinePath(ISPCFile, CompileEnvironment)}\"");
+				Arguments.Add($"\"{NormalizeCommandLinePath(ISPCFile)}\"");
 
 				List<FileItem> CompiledISPCObjFiles = new List<FileItem>();
 
@@ -731,7 +731,7 @@ namespace UnrealBuildTool
 				}
 
 				// Add the output ISPC obj file
-				Arguments.Add($"-o \"{NormalizeCommandLinePath(CompiledISPCObjFileNoISA, CompileEnvironment)}\"");
+				Arguments.Add($"-o \"{NormalizeCommandLinePath(CompiledISPCObjFileNoISA)}\"");
 
 				// Generate the timing info
 				if (CompileEnvironment.bPrintTimingInfo)
@@ -790,10 +790,10 @@ namespace UnrealBuildTool
 							Action PostCompileAction = Graph.CreateAction(ActionType.Compile);
 
 							List<string> PostCompileArgs = new List<string>();
-							PostCompileArgs.Add($"\"{NormalizeCommandLinePath(CompiledBytecodeObjFile, CompileEnvironment)}\"");
+							PostCompileArgs.Add($"\"{NormalizeCommandLinePath(CompiledBytecodeObjFile)}\"");
 							PostCompileArgs.Add("-c");
 							PostCompileArgs.AddRange(CommonArgs);
-							PostCompileArgs.Add($"-o \"{NormalizeCommandLinePath(FinalCompiledISPCObjFile, CompileEnvironment)}\"");
+							PostCompileArgs.Add($"-o \"{NormalizeCommandLinePath(FinalCompiledISPCObjFile)}\"");
 
 							if (bMergeModules)
 							{
