[UBT]
* Fixed so source file in vctoolchain uses vfs path (this only works when using uba, but is needed to get portable pch)

diff --git a/Engine/Source/Programs/UnrealBuildTool/Platform/Windows/VCCompileAction.cs b/Engine/Source/Programs/UnrealBuildTool/Platform/Windows/VCCompileAction.cs
--- a/Engine/Source/Programs/UnrealBuildTool/Platform/Windows/VCCompileAction.cs
+++ b/Engine/Source/Programs/UnrealBuildTool/Platform/Windows/VCCompileAction.cs
@@ -543,7 +543,7 @@ namespace UnrealBuildTool
 
 			if (SourceFile != null)
 			{
-				VCToolChain.AddSourceFile(Arguments, SourceFile);
+				VCToolChain.AddSourceFile(Arguments, SourceFile, rootPaths);
 			}
 
 			foreach (FileItem AdditionalRsp in AdditionalResponseFiles)
diff --git a/Engine/Source/Programs/UnrealBuildTool/Platform/Windows/VCToolChain.cs b/Engine/Source/Programs/UnrealBuildTool/Platform/Windows/VCToolChain.cs
--- a/Engine/Source/Programs/UnrealBuildTool/Platform/Windows/VCToolChain.cs
+++ b/Engine/Source/Programs/UnrealBuildTool/Platform/Windows/VCToolChain.cs
@@ -200,9 +200,9 @@ namespace UnrealBuildTool
 			Arguments.Add($"@{Utils.MakePathSafeToUseWithCommandLine(ResponseFileString)}");
 		}
 
-		public static void AddSourceFile(List<string> Arguments, FileItem SourceFile)
+		public static void AddSourceFile(List<string> Arguments, FileItem SourceFile, CppRootPaths? RootPaths)
 		{
-			string SourceFileString = NormalizeCommandLinePath(SourceFile);
+			string SourceFileString = NormalizeCommandLinePath(SourceFile, RootPaths);
 			Arguments.Add(Utils.MakePathSafeToUseWithCommandLine(SourceFileString));
 		}
 
