[UBA]
* Added so root paths are registered into uba session

diff --git a/Engine/Source/Programs/UnrealBuildTool/Executors/UnrealBuildAccelerator/UBAExecutor.cs b/Engine/Source/Programs/UnrealBuildTool/Executors/UnrealBuildAccelerator/UBAExecutor.cs
--- a/Engine/Source/Programs/UnrealBuildTool/Executors/UnrealBuildAccelerator/UBAExecutor.cs
+++ b/Engine/Source/Programs/UnrealBuildTool/Executors/UnrealBuildAccelerator/UBAExecutor.cs
@@ -717,6 +717,21 @@ namespace UnrealBuildTool
 			return rootPaths;
 		}
 
+		ulong GetActionRootsHandle(LinkedAction action)
+		{
+			using MemoryStream outputsMemory = new(1024);
+			using (BinaryWriter writer = new(outputsMemory, System.Text.Encoding.UTF8, true))
+			{
+				foreach (var kv in action.RootPaths)
+				{
+					writer.Write((byte)0);// kv.Key);
+					writer.Write(kv.Key.AsVfsPath().FullName);
+					writer.Write(kv.Value.FullName);
+				}
+			}
+			return _session!.RegisterRoots(outputsMemory.GetBuffer(), (uint)outputsMemory.Position);
+		}
+
 		public class DepsFile
 		{
 			public class DepsData
@@ -804,6 +819,9 @@ namespace UnrealBuildTool
 				bool enableDetour = !ForceLocalNoDetour(action) && action.bCanExecuteInUBA && !_forcedRetryActions.ContainsKey(action);
 
 				ProcessStartInfo startInfo = GetActionStartInfo(action);
+
+				startInfo.RootsHandle = GetActionRootsHandle(action);
+
 				using (IProcess process = _session!.RunProcess(startInfo, false, null, enableDetour))
 				{
 					Interlocked.Add(ref _localProcessedActions, 1);
@@ -895,6 +913,7 @@ namespace UnrealBuildTool
 				}
 
 				ProcessStartInfo startInfo = GetActionStartInfo(action);
+				startInfo.RootsHandle = GetActionRootsHandle(action);
 				_session!.RunProcessRemote(startInfo, (s, e) =>
 				{
 					if (e.ExitCode == 99999) // Process was cancelled by executor
