UnrealBuildTool: Correctly handle Horde UbaEnabled .ini setting

#rnx

diff --git a/Engine/Source/Programs/UnrealBuildTool/Executors/UnrealBuildAccelerator/UBAAgentCoordinatorHorde.cs b/Engine/Source/Programs/UnrealBuildTool/Executors/UnrealBuildAccelerator/UBAAgentCoordinatorHorde.cs
--- a/Engine/Source/Programs/UnrealBuildTool/Executors/UnrealBuildAccelerator/UBAAgentCoordinatorHorde.cs
+++ b/Engine/Source/Programs/UnrealBuildTool/Executors/UnrealBuildAccelerator/UBAAgentCoordinatorHorde.cs
@@ -690,10 +690,9 @@ namespace UnrealBuildTool
 			_ubaConfig = ubaConfig;
 
 			ConfigCache.ReadSettings(projectDir, BuildHostPlatform.Current.Platform, HordeConfig);
-			if (HordeConfig.HordeEnabled?.Equals("BuildMachineOnly", StringComparison.OrdinalIgnoreCase) == true && !Unreal.IsBuildMachine())
+			if (HordeConfig.HordeEnabled?.Equals("False", StringComparison.OrdinalIgnoreCase) == true || (HordeConfig.HordeEnabled?.Equals("BuildMachineOnly", StringComparison.OrdinalIgnoreCase) == true && !Unreal.IsBuildMachine()))
 			{
-				// reset
-				HordeConfig = new();
+				HordeConfig.bDisableHorde = true;
 			}
 			if (Unreal.IsBuildMachine())
 			{
@@ -703,10 +702,6 @@ namespace UnrealBuildTool
 					HordeConfig.HordeServer = hordeUrl;
 				}
 			}	
-			if (HordeConfig.HordeEnabled?.Equals("False", StringComparison.OrdinalIgnoreCase) == true)
-			{
-				HordeConfig.bDisableHorde = true;
-			}
 			XmlConfig.ApplyTo(HordeConfig);
 			additionalArguments?.ApplyTo(HordeConfig);
 
