[UBT]
* Changed so vfs file is using platform path separator to work well with cache (this will be revisited one day)

diff --git a/Engine/Source/Programs/UnrealBuildTool/Configuration/UEBuildModuleCPP.cs b/Engine/Source/Programs/UnrealBuildTool/Configuration/UEBuildModuleCPP.cs
--- a/Engine/Source/Programs/UnrealBuildTool/Configuration/UEBuildModuleCPP.cs
+++ b/Engine/Source/Programs/UnrealBuildTool/Configuration/UEBuildModuleCPP.cs
@@ -369,7 +369,7 @@ namespace UnrealBuildTool
 				{
 					["name"] = item.Key.AsVfsString(),
 					["type"] = "directory-remap",
-					["external-contents"] = item.Value.FullName.Replace('\\', '/')
+					["external-contents"] = item.Value.FullName
 				});
 			}
 
