[UBT]
* Changed so precompiled headers don't set artifact mode absolutepath if vfs is enabled

diff --git a/Engine/Source/Programs/UnrealBuildTool/Platform/Windows/VCToolChain.cs b/Engine/Source/Programs/UnrealBuildTool/Platform/Windows/VCToolChain.cs
--- a/Engine/Source/Programs/UnrealBuildTool/Platform/Windows/VCToolChain.cs
+++ b/Engine/Source/Programs/UnrealBuildTool/Platform/Windows/VCToolChain.cs
@@ -1993,7 +1993,11 @@ namespace UnrealBuildTool
 
 					if (CompileEnvironment.PrecompiledHeaderAction != PrecompiledHeaderAction.None)
 					{
-						CompileAction.ArtifactMode |= ArtifactMode.AbsolutePath; // Unfortunately we require matching absolute paths for pch to be cached
+						if (!Target.bCompilerVFS)
+						{
+							// Unfortunately we require matching absolute paths for pch to be cached if not using vfs
+							CompileAction.ArtifactMode |= ArtifactMode.AbsolutePath;
+						}
 
 						if (Target.WindowsPlatform.Compiler.IsClang())
 						{
