Fit UnrealEditor.modules not being updated after HotReload
This was causing the FModuleManager to detect the module as missing because it's name in UnrealEditor.modules was still the one with the HotReload suffix.
By adding Metadata files to the list of TemporaryFiles in HotReloadState we ensure their will be regenerated after a HotReload.
We also make sure that ProducedItems that have PrerequisiteItems will be outdated when the PrerequisiteItems doesn't exist since we can assume that they are currently being generated.

#tests Local testing with HotReload & Regular build on Mac & Windows

diff --git a/Engine/Source/Programs/UnrealBuildTool/System/ActionGraph.cs b/Engine/Source/Programs/UnrealBuildTool/System/ActionGraph.cs
--- a/Engine/Source/Programs/UnrealBuildTool/System/ActionGraph.cs
+++ b/Engine/Source/Programs/UnrealBuildTool/System/ActionGraph.cs
@@ -794,6 +794,14 @@ namespace UnrealBuildTool
 							}
 						}
 					}
+					else
+					{
+						// If the PrerequisiteItem doesn't exist, it means it's being produced right now and we need to update the ProducedItem that
+						// depends on it
+						Logger.LogDebug("{StatusDescription}: Prerequisite {PrerequisiteItem} doesn't exist.", RootAction.StatusDescription, PrerequisiteItem.Location);
+						bIsOutdated = true;
+						break;
+					}
 				}
 			}
 
diff --git a/Engine/Source/Programs/UnrealBuildTool/System/HotReload.cs b/Engine/Source/Programs/UnrealBuildTool/System/HotReload.cs
--- a/Engine/Source/Programs/UnrealBuildTool/System/HotReload.cs
+++ b/Engine/Source/Programs/UnrealBuildTool/System/HotReload.cs
@@ -75,13 +75,27 @@ namespace UnrealBuildTool
 			// Now filter out all the hot reload files and update the state
 			foreach (LinkedAction Action in ActionsToExecute)
 			{
-				foreach (FileItem ProducedItem in Action.ProducedItems)
+				// Metadata doesn't change name but their content might change
+				// Example:
+				// UnrealEditor.modules will be patched by HotReload to include the updated module name with suffix
+				// When going back to regular build we need to restore the previous module name.
+				if (Action.ActionType == ActionType.WriteMetadata)
+				{
+					foreach (FileItem FileItem in Action.ProducedItems)
+					{
+						TemporaryFiles.Add(FileItem.Location);
+					}
+				}
+				else
 				{
-					FileReference? OriginalLocation;
-					if (HotReloadFileToOriginalFile.TryGetValue(ProducedItem.Location, out OriginalLocation))
+					foreach (FileItem ProducedItem in Action.ProducedItems)
 					{
-						OriginalFileToHotReloadFile[OriginalLocation] = ProducedItem.Location;
-						TemporaryFiles.Add(ProducedItem.Location);
+						FileReference? OriginalLocation;
+						if (HotReloadFileToOriginalFile.TryGetValue(ProducedItem.Location, out OriginalLocation))
+						{
+							OriginalFileToHotReloadFile[OriginalLocation] = ProducedItem.Location;
+							TemporaryFiles.Add(ProducedItem.Location);
+						}
 					}
 				}
 			}
