- FIxed a comment with correct path

diff --git a/Engine/Source/Programs/UnrealBuildTool/Platform/Windows/MicrosoftPlatformSDK.Versions.cs b/Engine/Source/Programs/UnrealBuildTool/Platform/Windows/MicrosoftPlatformSDK.Versions.cs
--- a/Engine/Source/Programs/UnrealBuildTool/Platform/Windows/MicrosoftPlatformSDK.Versions.cs
+++ b/Engine/Source/Programs/UnrealBuildTool/Platform/Windows/MicrosoftPlatformSDK.Versions.cs
@@ -8,7 +8,7 @@ using EpicGames.Core;
 namespace UnrealBuildTool
 {
 	////////////////////////////////////////////////////////////////////////////////
-	// If you are looking for version numbers, see Engine/Config/Windows_SDK.json
+	// If you are looking for version numbers, see Engine/Config/Windows/Windows_SDK.json
 	////////////////////////////////////////////////////////////////////////////////
 
 	partial class MicrosoftPlatformSDK : UEBuildPlatformSDK
