UnrealBuildTool: Support opt-in unity files for source code that isn't .cpp

diff --git a/Engine/Source/Programs/UnrealBuildTool/Configuration/ModuleRules.cs b/Engine/Source/Programs/UnrealBuildTool/Configuration/ModuleRules.cs
--- a/Engine/Source/Programs/UnrealBuildTool/Configuration/ModuleRules.cs
+++ b/Engine/Source/Programs/UnrealBuildTool/Configuration/ModuleRules.cs
@@ -265,6 +265,43 @@ namespace UnrealBuildTool
 			VisibileForDll,
 		}
 
+		/// <summary>
+		/// File types that can use unity files for this module, if enabled.
+		/// </summary>
+		[Flags]
+		public enum UnityFileType
+		{
+			/// <summary>
+			/// No files
+			/// </summary>
+			None,
+
+			/// <summary>
+			/// C++
+			/// </summary>
+			Cpp,
+
+			/// <summary>
+			/// C
+			/// </summary>
+			C,
+
+			/// <summary>
+			/// C++ (alternate extension)
+			/// </summary>
+			CC,
+
+			/// <summary>
+			/// Objective-C
+			/// </summary>
+			M,
+
+			/// <summary>
+			/// Objective-C++
+			/// </summary>
+			MM,
+		}
+
 		/// <summary>
 		/// Information about a file which is required by the target at runtime, and must be moved around with it.
 		/// </summary>
@@ -1100,6 +1137,11 @@ namespace UnrealBuildTool
 		/// </summary>
 		public bool bMergeUnityFiles { get; set; } = true;
 
+		/// <summary>
+		/// Source file types that will be placed into unity files when enabled.
+		/// </summary>
+		public UnityFileType UnityFileTypes { get; set;  } = UnityFileType.Cpp;
+
 		/// <summary>
 		/// The number of source files in this module before unity build will be activated for that module.  If set to
 		/// anything besides -1, will override the default setting which is controlled by MinGameModuleSourceFilesForUnityBuild
diff --git a/Engine/Source/Programs/UnrealBuildTool/Configuration/UEBuildModuleCPP.cs b/Engine/Source/Programs/UnrealBuildTool/Configuration/UEBuildModuleCPP.cs
--- a/Engine/Source/Programs/UnrealBuildTool/Configuration/UEBuildModuleCPP.cs
+++ b/Engine/Source/Programs/UnrealBuildTool/Configuration/UEBuildModuleCPP.cs
@@ -506,7 +506,7 @@ namespace UnrealBuildTool
 			}
 
 			// Set up the environment with which to compile the CPP files
-			CppCompileEnvironment CompileEnvironment = new(ModuleCompileEnvironment);
+			CppCompileEnvironment CppCompileEnvironment = new(ModuleCompileEnvironment);
 
 			// Compile any module interfaces
 			if (InputFiles.IXXFiles.Count > 0 && Target.bEnableCppModules)
@@ -519,40 +519,40 @@ namespace UnrealBuildTool
 				CPPOutput ModuleOutput = ToolChain.CompileAllCPPFiles(IxxCompileEnvironment, InputFiles.IXXFiles, IntermediateDirectory, Name, Graph);
 
 				LinkInputFiles.AddRange(ModuleOutput.ObjectFiles);
-				CompileEnvironment.AdditionalPrerequisites.AddRange(ModuleOutput.CompiledModuleInterfaces);
+				CppCompileEnvironment.AdditionalPrerequisites.AddRange(ModuleOutput.CompiledModuleInterfaces);
 			}
 
 			// Configure the precompiled headers for this module
-			CompileEnvironment = SetupPrecompiledHeaders(Target, ToolChain, CompileEnvironment, LinkInputFiles, Graph, Logger);
-			if (CompileEnvironment.PerArchPrecompiledHeaderFiles != null)
+			CppCompileEnvironment = SetupPrecompiledHeaders(Target, ToolChain, CppCompileEnvironment, LinkInputFiles, Graph, Logger);
+			if (CppCompileEnvironment.PerArchPrecompiledHeaderFiles != null)
 			{
-				foreach (UnrealArch Arch in CompileEnvironment.PerArchPrecompiledHeaderFiles.Keys)
+				foreach (UnrealArch Arch in CppCompileEnvironment.PerArchPrecompiledHeaderFiles.Keys)
 				{
-					Logger.LogDebug("Module '{ModuleName}' uses PCH '{PCHIncludeFilename}' for Architecture '{Arch}'", Name, Arch, CompileEnvironment.PerArchPrecompiledHeaderFiles[Arch]);
+					Logger.LogDebug("Module '{ModuleName}' uses PCH '{PCHIncludeFilename}' for Architecture '{Arch}'", Name, Arch, CppCompileEnvironment.PerArchPrecompiledHeaderFiles[Arch]);
 				}
 			}
-			else if (CompileEnvironment.PrecompiledHeaderFile != null)
+			else if (CppCompileEnvironment.PrecompiledHeaderFile != null)
 			{
-				Logger.LogDebug("Module '{ModuleName}' uses PCH '{PCHIncludeFilename}'", Name, CompileEnvironment.PrecompiledHeaderFile);
+				Logger.LogDebug("Module '{ModuleName}' uses PCH '{PCHIncludeFilename}'", Name, CppCompileEnvironment.PrecompiledHeaderFile);
 			}
 
 			// Write all the definitions to a separate file
-			CreateHeaderForDefinitions(CompileEnvironment, IntermediateDirectory, null, Graph);
+			CreateHeaderForDefinitions(CppCompileEnvironment, IntermediateDirectory, null, Graph);
 
 			// Create shared rsp for the normal cpp files
 			FileReference SharedResponseFile = FileReference.Combine(IntermediateDirectory, $"{Rules.ShortName ?? Name}.Shared{UEToolChain.ResponseExt}");
-			CompileEnvironment = ToolChain.CreateSharedResponseFile(CompileEnvironment, SharedResponseFile, Graph);
+			CppCompileEnvironment = ToolChain.CreateSharedResponseFile(CppCompileEnvironment, SharedResponseFile, Graph);
 
 			// Mapping of source file to unity file. We output this to intermediate directories for other tools (eg. live coding) to use.
 			Dictionary<FileItem, FileItem> SourceFileToUnityFile = new Dictionary<FileItem, FileItem>();
 
 			List<FileItem> CPPFiles = new List<FileItem>(InputFiles.CPPFiles);
 			List<FileItem> GeneratedFileItems = new List<FileItem>();
-			CppCompileEnvironment GeneratedCPPCompileEnvironment = CompileEnvironment;
+			CppCompileEnvironment GeneratedCPPCompileEnvironment = CppCompileEnvironment;
 			bool bMergeUnityFiles = !(Target.DisableMergingModuleAndGeneratedFilesInUnityFiles?.Contains(Name) ?? false) && Rules.bMergeUnityFiles;
 
 			// Compile all the generated CPP files
-			if (GeneratedCppDirectories != null && !CompileEnvironment.bHackHeaderGenerator)
+			if (GeneratedCppDirectories != null && !CppCompileEnvironment.bHackHeaderGenerator)
 			{
 				Dictionary<string, FileItem> GeneratedFiles = new Dictionary<string, FileItem>();
 				foreach (string GeneratedDir in GeneratedCppDirectories)
@@ -592,11 +592,11 @@ namespace UnrealBuildTool
 							string Key = Prefix + ListOfInlinedGenCppsItem;
 							if (GeneratedFiles.Remove(Key, out FileItem? FoundGenCppFile))
 							{
-								if (!CompileEnvironment.FileInlineGenCPPMap.ContainsKey(CPPFileItem))
+								if (!CppCompileEnvironment.FileInlineGenCPPMap.ContainsKey(CPPFileItem))
 								{
-									CompileEnvironment.FileInlineGenCPPMap[CPPFileItem] = new List<FileItem>();
+									CppCompileEnvironment.FileInlineGenCPPMap[CPPFileItem] = new List<FileItem>();
 								}
-								CompileEnvironment.FileInlineGenCPPMap[CPPFileItem].Add(FoundGenCppFile);
+								CppCompileEnvironment.FileInlineGenCPPMap[CPPFileItem].Add(FoundGenCppFile);
 							}
 							else
 							{
@@ -604,7 +604,7 @@ namespace UnrealBuildTool
 							}
 						}
 					}
-					ModuleCompileEnvironment.FileInlineGenCPPMap = CompileEnvironment.FileInlineGenCPPMap;
+					ModuleCompileEnvironment.FileInlineGenCPPMap = CppCompileEnvironment.FileInlineGenCPPMap;
 
 					if (Rules.bEnableNonInlinedGenCppWarnings)
 					{
@@ -683,10 +683,10 @@ namespace UnrealBuildTool
 
 				if (Target.WindowsPlatform.bEnableInstrumentation)
 				{
-					CompileEnvironment.UserIncludePaths.Add(DirectoryReference.Combine(Unreal.EngineSourceDirectory, "Runtime", "Instrumentation", "Public"));
-					CompileEnvironment.UserIncludePaths.Add(DirectoryReference.Combine(Unreal.EngineSourceDirectory, "Runtime", "Instrumentation", "Internal"));
-					CompileEnvironment.UserIncludePaths.Add(DirectoryReference.Combine(Unreal.EngineSourceDirectory, "Runtime", "Core", "Public"));
-					CompileEnvironment.UserIncludePaths.Add(DirectoryReference.Combine(Unreal.EngineSourceDirectory, "Runtime", "Core", "Internal"));
+					CppCompileEnvironment.UserIncludePaths.Add(DirectoryReference.Combine(Unreal.EngineSourceDirectory, "Runtime", "Instrumentation", "Public"));
+					CppCompileEnvironment.UserIncludePaths.Add(DirectoryReference.Combine(Unreal.EngineSourceDirectory, "Runtime", "Instrumentation", "Internal"));
+					CppCompileEnvironment.UserIncludePaths.Add(DirectoryReference.Combine(Unreal.EngineSourceDirectory, "Runtime", "Core", "Public"));
+					CppCompileEnvironment.UserIncludePaths.Add(DirectoryReference.Combine(Unreal.EngineSourceDirectory, "Runtime", "Core", "Internal"));
 
 					// Give the ability for core to hotpatch itself for the per module functions to save another set of JMPs.
 					if (Name == "Core")
@@ -700,7 +700,7 @@ namespace UnrealBuildTool
 				FileItem PerModuleFile = Graph.CreateIntermediateTextFile(NewDeleteOverrides, Content, bMergeModules);
 				if (bMergeModules)
 				{
-					LinkInputFiles.AddRange(ToolChain.CompileAllCPPFiles(CompileEnvironment, new[] { PerModuleFile }, IntermediateDirectory, Name, Graph).ObjectFiles);
+					LinkInputFiles.AddRange(ToolChain.CompileAllCPPFiles(CppCompileEnvironment, new[] { PerModuleFile }, IntermediateDirectory, Name, Graph).ObjectFiles);
 				}
 				else
 				{
@@ -765,8 +765,8 @@ namespace UnrealBuildTool
 				Logger.LogTrace("Module '{ModuleName}' not using unity build mode", Name);
 			}
 
-			CompileEnvironment.bUseUnity = bModuleUsesUnityBuild;
-			GeneratedCPPCompileEnvironment.bUseUnity = bModuleUsesUnityBuild;
+			CppCompileEnvironment.bUseUnity = bModuleUsesUnityBuild && Rules.UnityFileTypes.HasFlag(ModuleRules.UnityFileType.Cpp);
+			GeneratedCPPCompileEnvironment.bUseUnity = bModuleUsesUnityBuild && Rules.UnityFileTypes.HasFlag(ModuleRules.UnityFileType.Cpp);
 
 			// Create and register a special action that can be used to compile single files (even when unity is enabled) for generated files
 			if (GeneratedFileItems.Any())
@@ -782,7 +782,7 @@ namespace UnrealBuildTool
 			// Create and register a special action that can be used to compile single files (even when unity is enabled) for normal files
 			if (ModuleDirectories.Any())
 			{
-				CppCompileEnvironment SingleCompileEnvironment = new CppCompileEnvironment(CompileEnvironment);
+				CppCompileEnvironment SingleCompileEnvironment = new CppCompileEnvironment(CppCompileEnvironment);
 				SingleCompileEnvironment.PrecompiledHeaderAction = PrecompiledHeaderAction.None;
 				foreach (DirectoryReference Directory in ModuleDirectories)
 				{
@@ -791,7 +791,7 @@ namespace UnrealBuildTool
 			}
 
 			// Compile Generated CPP Files
-			if (bModuleUsesUnityBuild)
+			if (bModuleUsesUnityBuild && Rules.UnityFileTypes.HasFlag(ModuleRules.UnityFileType.Cpp))
 			{
 				if (bMergeUnityFiles)
 				{
@@ -799,7 +799,7 @@ namespace UnrealBuildTool
 				}
 				else
 				{
-					Unity.GenerateUnityCPPs(Target, GeneratedFileItems, new List<FileItem>(), CompileEnvironment, WorkingSet, (Rules.ShortName ?? Name) + ".gen", IntermediateDirectory, Graph, SourceFileToUnityFile,
+					Unity.GenerateUnitySource(Target, GeneratedFileItems, new List<FileItem>(), CppCompileEnvironment, WorkingSet, (Rules.ShortName ?? Name) + ".gen", IntermediateDirectory, Graph, SourceFileToUnityFile,
 						out List<FileItem> NormalGeneratedFiles, out List<FileItem> AdaptiveGeneratedFiles, NumIncludedBytesPerUnityCPP);
 					LinkInputFiles.AddRange(CompileFilesWithToolChain(Target, ToolChain, GeneratedCPPCompileEnvironment, ModuleCompileEnvironment, NormalGeneratedFiles, AdaptiveGeneratedFiles, Graph, Logger).ObjectFiles);
 				}
@@ -833,32 +833,35 @@ namespace UnrealBuildTool
 				// modify both ModuleCompileEnvironment and CompileEnvironment because they are branched already, and are
 				// both used when compiling files that may need the headers
 				ModuleCompileEnvironment.AdditionalPrerequisites.AddRange(HeaderFiles);
-				CompileEnvironment.AdditionalPrerequisites.AddRange(HeaderFiles);
+				CppCompileEnvironment.AdditionalPrerequisites.AddRange(HeaderFiles);
 				if (HeaderFiles.Count() > 0)
 				{
 					// @todo: this could have different per-arch headers- should be rare but possible
 					ModuleCompileEnvironment.UserIncludePaths.Add(HeaderFiles.First().Location.Directory);
-					CompileEnvironment.UserIncludePaths.Add(HeaderFiles.First().Location.Directory);
+					CppCompileEnvironment.UserIncludePaths.Add(HeaderFiles.First().Location.Directory);
 				}
 			}
 
 			// Compile CPP files
-			if (bModuleUsesUnityBuild)
+			if (CPPFiles.Count > 0)
 			{
-				Unity.GenerateUnityCPPs(Target, CPPFiles, InputFiles.HeaderFiles, CompileEnvironment, WorkingSet, Rules.ShortName ?? Name, IntermediateDirectory, Graph, SourceFileToUnityFile,
-					out List<FileItem> NormalFiles, out List<FileItem> AdaptiveFiles, NumIncludedBytesPerUnityCPP);
-				LinkInputFiles.AddRange(CompileFilesWithToolChain(Target, ToolChain, CompileEnvironment, ModuleCompileEnvironment, NormalFiles, AdaptiveFiles, Graph, Logger).ObjectFiles);
-			}
-			else
-			{
-				Unity.GetAdaptiveFiles(Target, CPPFiles, InputFiles.HeaderFiles, CompileEnvironment, WorkingSet, Rules.ShortName ?? Name, IntermediateDirectory, Graph,
-					out List<FileItem> NormalFiles, out List<FileItem> AdaptiveFiles);
-				if (!NormalFiles.Where(file => !file.HasExtension(".gen.cpp") && !GeneratedCPPCompileEnvironment.FileMatchesExtraGeneratedCPPTypes(file.FullName)).Any())
+				if (bModuleUsesUnityBuild && Rules.UnityFileTypes.HasFlag(ModuleRules.UnityFileType.Cpp))
 				{
-					NormalFiles = CPPFiles;
-					AdaptiveFiles.RemoveAll(new HashSet<FileItem>(NormalFiles).Contains);
+					Unity.GenerateUnitySource(Target, CPPFiles, InputFiles.HeaderFiles, CppCompileEnvironment, WorkingSet, Rules.ShortName ?? Name, IntermediateDirectory, Graph, SourceFileToUnityFile,
+						out List<FileItem> NormalFiles, out List<FileItem> AdaptiveFiles, NumIncludedBytesPerUnityCPP);
+					LinkInputFiles.AddRange(CompileFilesWithToolChain(Target, ToolChain, CppCompileEnvironment, ModuleCompileEnvironment, NormalFiles, AdaptiveFiles, Graph, Logger).ObjectFiles);
+				}
+				else
+				{
+					Unity.GetAdaptiveFiles(Target, CPPFiles, InputFiles.HeaderFiles, CppCompileEnvironment, WorkingSet, Rules.ShortName ?? Name, IntermediateDirectory, Graph,
+						out List<FileItem> NormalFiles, out List<FileItem> AdaptiveFiles);
+					if (!NormalFiles.Where(file => !file.HasExtension(".gen.cpp") && !GeneratedCPPCompileEnvironment.FileMatchesExtraGeneratedCPPTypes(file.FullName)).Any())
+					{
+						NormalFiles = CPPFiles;
+						AdaptiveFiles.RemoveAll(new HashSet<FileItem>(NormalFiles).Contains);
+					}
+					LinkInputFiles.AddRange(CompileFilesWithToolChain(Target, ToolChain, CppCompileEnvironment, ModuleCompileEnvironment, NormalFiles, AdaptiveFiles, Graph, Logger).ObjectFiles);
 				}
-				LinkInputFiles.AddRange(CompileFilesWithToolChain(Target, ToolChain, CompileEnvironment, ModuleCompileEnvironment, NormalFiles, AdaptiveFiles, Graph, Logger).ObjectFiles);
 			}
 
 			// Compile ISPC files directly
@@ -867,36 +870,80 @@ namespace UnrealBuildTool
 				LinkInputFiles.AddRange(ToolChain.CompileAllISPCFiles(ModuleCompileEnvironment, InputFiles.ISPCFiles, IntermediateDirectory, Graph).ObjectFiles);
 			}
 
-			// Compile C files directly. Do not use a PCH here, because a C++ PCH is not compatible with C source files.
+			// Compile C files. Do not use a PCH here, because a C++ PCH is not compatible with C source files.
 			if (InputFiles.CFiles.Count > 0)
 			{
 				CppCompileEnvironment CCompileEnvironment = new CppCompileEnvironment(ModuleCompileEnvironment);
+				CCompileEnvironment.bUseUnity = bModuleUsesUnityBuild && Rules.UnityFileTypes.HasFlag(ModuleRules.UnityFileType.C);
 				CreateHeaderForDefinitions(CCompileEnvironment, IntermediateDirectory, "c", Graph);
-				LinkInputFiles.AddRange(ToolChain.CompileAllCPPFiles(CCompileEnvironment, InputFiles.CFiles, IntermediateDirectory, Name, Graph).ObjectFiles);
+
+				if (CCompileEnvironment.bUseUnity)
+				{
+					Unity.GenerateUnitySource(Target, InputFiles.CFiles, InputFiles.HeaderFiles, CCompileEnvironment, WorkingSet, Rules.ShortName ?? Name, IntermediateDirectory, Graph, SourceFileToUnityFile,
+						out List<FileItem> NormalFiles, out List<FileItem> AdaptiveFiles, NumIncludedBytesPerUnityCPP);
+					LinkInputFiles.AddRange(CompileFilesWithToolChain(Target, ToolChain, CCompileEnvironment, ModuleCompileEnvironment, NormalFiles, AdaptiveFiles, Graph, Logger).ObjectFiles);
+				}
+				else
+				{
+					LinkInputFiles.AddRange(ToolChain.CompileAllCPPFiles(CCompileEnvironment, InputFiles.CFiles, IntermediateDirectory, Name, Graph).ObjectFiles);
+				}
 			}
 
-			// Compile CC files directly.
+			// Compile CC files
 			if (InputFiles.CCFiles.Count > 0)
 			{
-				CppCompileEnvironment CCCompileEnvironment = new CppCompileEnvironment(CompileEnvironment);
+				CppCompileEnvironment CCCompileEnvironment = new CppCompileEnvironment(CppCompileEnvironment);
+				CCCompileEnvironment.bUseUnity = bModuleUsesUnityBuild && Rules.UnityFileTypes.HasFlag(ModuleRules.UnityFileType.CC);
 				CreateHeaderForDefinitions(CCCompileEnvironment, IntermediateDirectory, "cc", Graph);
-				LinkInputFiles.AddRange(ToolChain.CompileAllCPPFiles(CCCompileEnvironment, InputFiles.CCFiles, IntermediateDirectory, Name, Graph).ObjectFiles);
+
+				if (CCCompileEnvironment.bUseUnity)
+				{
+					Unity.GenerateUnitySource(Target, InputFiles.CCFiles, InputFiles.HeaderFiles, CCCompileEnvironment, WorkingSet, Rules.ShortName ?? Name, IntermediateDirectory, Graph, SourceFileToUnityFile,
+						out List<FileItem> NormalFiles, out List<FileItem> AdaptiveFiles, NumIncludedBytesPerUnityCPP);
+					LinkInputFiles.AddRange(CompileFilesWithToolChain(Target, ToolChain, CCCompileEnvironment, CppCompileEnvironment, NormalFiles, AdaptiveFiles, Graph, Logger).ObjectFiles);
+				}
+				else
+				{
+					LinkInputFiles.AddRange(ToolChain.CompileAllCPPFiles(CCCompileEnvironment, InputFiles.CCFiles, IntermediateDirectory, Name, Graph).ObjectFiles);
+				}
 			}
 
-			// Compile M files directly. Do not use a PCH here, because a C++ PCH is not compatible with C source files.
+			// Compile M files. Do not use a PCH here, because a C++ PCH is not compatible with Objective-C source files.
 			if (InputFiles.MFiles.Count > 0)
 			{
 				CppCompileEnvironment MCompileEnvironment = new CppCompileEnvironment(ModuleCompileEnvironment);
+				MCompileEnvironment.bUseUnity = bModuleUsesUnityBuild && Rules.UnityFileTypes.HasFlag(ModuleRules.UnityFileType.M);
 				CreateHeaderForDefinitions(MCompileEnvironment, IntermediateDirectory, "m", Graph);
-				LinkInputFiles.AddRange(ToolChain.CompileAllCPPFiles(MCompileEnvironment, InputFiles.MFiles, IntermediateDirectory, Name, Graph).ObjectFiles);
+
+				if (MCompileEnvironment.bUseUnity)
+				{
+					Unity.GenerateUnitySource(Target, InputFiles.MFiles, InputFiles.HeaderFiles, MCompileEnvironment, WorkingSet, Rules.ShortName ?? Name, IntermediateDirectory, Graph, SourceFileToUnityFile,
+						out List<FileItem> NormalFiles, out List<FileItem> AdaptiveFiles, NumIncludedBytesPerUnityCPP);
+					LinkInputFiles.AddRange(CompileFilesWithToolChain(Target, ToolChain, MCompileEnvironment, ModuleCompileEnvironment, NormalFiles, AdaptiveFiles, Graph, Logger).ObjectFiles);
+				}
+				else
+				{
+					LinkInputFiles.AddRange(ToolChain.CompileAllCPPFiles(MCompileEnvironment, InputFiles.MFiles, IntermediateDirectory, Name, Graph).ObjectFiles);
+				}
 			}
 
-			// Compile MM files directly.
+			// Compile MM files
 			if (InputFiles.MMFiles.Count > 0)
 			{
-				CppCompileEnvironment MMCompileEnvironment = new CppCompileEnvironment(CompileEnvironment);
+				CppCompileEnvironment MMCompileEnvironment = new CppCompileEnvironment(CppCompileEnvironment);
+				MMCompileEnvironment.bUseUnity = bModuleUsesUnityBuild && Rules.UnityFileTypes.HasFlag(ModuleRules.UnityFileType.MM);
 				CreateHeaderForDefinitions(MMCompileEnvironment, IntermediateDirectory, "mm", Graph);
-				LinkInputFiles.AddRange(ToolChain.CompileAllCPPFiles(MMCompileEnvironment, InputFiles.MMFiles, IntermediateDirectory, Name, Graph).ObjectFiles);
+
+				if (MMCompileEnvironment.bUseUnity)
+				{
+					Unity.GenerateUnitySource(Target, InputFiles.MMFiles, InputFiles.HeaderFiles, MMCompileEnvironment, WorkingSet, Rules.ShortName ?? Name, IntermediateDirectory, Graph, SourceFileToUnityFile,
+						out List<FileItem> NormalFiles, out List<FileItem> AdaptiveFiles, NumIncludedBytesPerUnityCPP);
+					LinkInputFiles.AddRange(CompileFilesWithToolChain(Target, ToolChain, MMCompileEnvironment, CppCompileEnvironment, NormalFiles, AdaptiveFiles, Graph, Logger).ObjectFiles);
+				}
+				else
+				{
+					LinkInputFiles.AddRange(ToolChain.CompileAllCPPFiles(MMCompileEnvironment, InputFiles.MMFiles, IntermediateDirectory, Name, Graph).ObjectFiles);
+				};
 			}
 
 			// Compile RC files. The resource compiler does not work with response files, and using the regular compile environment can easily result in the 
@@ -957,7 +1004,7 @@ namespace UnrealBuildTool
 				}
 
 				// Collect the headers that should be built
-				List<FileItem> HeaderFileItems = GetCompilableHeaders(InputFiles, CompileEnvironment);
+				List<FileItem> HeaderFileItems = GetCompilableHeaders(InputFiles, CppCompileEnvironment);
 				if (HeaderFileItems.Count > 0)
 				{
 					CppCompileEnvironment HeaderCompileEnvironment = new CppCompileEnvironment(ModuleCompileEnvironment);
@@ -1690,7 +1737,7 @@ namespace UnrealBuildTool
 		}
 
 		/// <summary>
-		/// Compiles the provided CPP unity files. Will
+		/// Compiles the provided source files. Will possibly update environment settings for adaptive files if any are provided.
 		/// </summary>
 		private CPPOutput CompileFilesWithToolChain(
 			ReadOnlyTargetRules Target,
@@ -1703,45 +1750,48 @@ namespace UnrealBuildTool
 			ILogger Logger)
 		{
 			bool bAdaptiveUnityDisablesPCH = false;
-			if (Rules.PCHUsage == ModuleRules.PCHUsageMode.UseExplicitOrSharedPCHs)
+			if (CompileEnvironment.PrecompiledHeaderAction == PrecompiledHeaderAction.Include)
 			{
-				if (Rules.bTreatAsEngineModule || Rules.PrivatePCHHeaderFile == null)
-				{
-					bAdaptiveUnityDisablesPCH = Target.bAdaptiveUnityDisablesPCH;
-				}
-				else
+				if (Rules.PCHUsage == ModuleRules.PCHUsageMode.UseExplicitOrSharedPCHs)
 				{
-					bAdaptiveUnityDisablesPCH = Target.bAdaptiveUnityDisablesPCHForProject;
+					if (Rules.bTreatAsEngineModule || Rules.PrivatePCHHeaderFile == null)
+					{
+						bAdaptiveUnityDisablesPCH = Target.bAdaptiveUnityDisablesPCH;
+					}
+					else
+					{
+						bAdaptiveUnityDisablesPCH = Target.bAdaptiveUnityDisablesPCHForProject;
+					}
 				}
-			}
 
-			if (AdaptiveFiles.Count > 0)
-			{
-				if (Target.bAdaptiveUnityCreatesDedicatedPCH)
+				if (AdaptiveFiles.Count > 0)
 				{
-					Graph.AddDiagnostic("[Adaptive Build] Creating dedicated PCH for each excluded file. Set bAdaptiveUnityCreatesDedicatedPCH to false in BuildConfiguration.xml to change this behavior.");
-				}
-				else if (Target.bAdaptiveUnityDisablesPCH)
-				{
-					Graph.AddDiagnostic("[Adaptive Build] Disabling PCH for excluded files. Set bAdaptiveUnityDisablesPCH to false in BuildConfiguration.xml to change this behavior.");
-				}
+					if (Target.bAdaptiveUnityCreatesDedicatedPCH)
+					{
+						Graph.AddDiagnostic("[Adaptive Build] Creating dedicated PCH for each excluded file. Set bAdaptiveUnityCreatesDedicatedPCH to false in BuildConfiguration.xml to change this behavior.");
+					}
+					else if (Target.bAdaptiveUnityDisablesPCH)
+					{
+						Graph.AddDiagnostic("[Adaptive Build] Disabling PCH for excluded files. Set bAdaptiveUnityDisablesPCH to false in BuildConfiguration.xml to change this behavior.");
+					}
 
-				if (Target.bAdaptiveUnityDisablesOptimizations)
-				{
-					Graph.AddDiagnostic("[Adaptive Build] Disabling optimizations for excluded files. Set bAdaptiveUnityDisablesOptimizations to false in BuildConfiguration.xml to change this behavior.");
+					if (Target.bAdaptiveUnityDisablesOptimizations)
+					{
+						Graph.AddDiagnostic("[Adaptive Build] Disabling optimizations for excluded files. Set bAdaptiveUnityDisablesOptimizations to false in BuildConfiguration.xml to change this behavior.");
+					}
+					if (Target.bAdaptiveUnityEnablesEditAndContinue)
+					{
+						Graph.AddDiagnostic("[Adaptive Build] Enabling Edit & Continue for excluded files. Set bAdaptiveUnityEnablesEditAndContinue to false in BuildConfiguration.xml to change this behavior.");
+					}
+
+					Graph.AddDiagnostic($"[Adaptive Build] Excluded from {Name} unity file: " + String.Join(", ", AdaptiveFiles.Select(File => Path.GetFileName(File.AbsolutePath))));
 				}
-				if (Target.bAdaptiveUnityEnablesEditAndContinue)
+
+				if ((!Target.bAdaptiveUnityDisablesOptimizations && !bAdaptiveUnityDisablesPCH && !Target.bAdaptiveUnityCreatesDedicatedPCH) || Target.bStressTestUnity)
 				{
-					Graph.AddDiagnostic("[Adaptive Build] Enabling Edit & Continue for excluded files. Set bAdaptiveUnityEnablesEditAndContinue to false in BuildConfiguration.xml to change this behavior.");
+					NormalFiles = NormalFiles.Concat(AdaptiveFiles).ToList();
+					AdaptiveFiles = new List<FileItem>();
 				}
-
-				Graph.AddDiagnostic($"[Adaptive Build] Excluded from {Name} unity file: " + String.Join(", ", AdaptiveFiles.Select(File => Path.GetFileName(File.AbsolutePath))));
-			}
-
-			if ((!Target.bAdaptiveUnityDisablesOptimizations && !bAdaptiveUnityDisablesPCH && !Target.bAdaptiveUnityCreatesDedicatedPCH) || Target.bStressTestUnity)
-			{
-				NormalFiles = NormalFiles.Concat(AdaptiveFiles).ToList();
-				AdaptiveFiles = new List<FileItem>();
 			}
 
 			CPPOutput OutputFiles = new CPPOutput();
@@ -1766,11 +1816,11 @@ namespace UnrealBuildTool
 
 				// Create a per-file PCH
 				CPPOutput AdaptiveOutput;
-				if (Target.bAdaptiveUnityCreatesDedicatedPCH)
+				if (CompileEnvironment.PrecompiledHeaderAction == PrecompiledHeaderAction.Include && Target.bAdaptiveUnityCreatesDedicatedPCH)
 				{
 					AdaptiveOutput = CompileAdaptiveNonUnityFilesWithDedicatedPCH(ToolChain, AdaptiveUnityEnvironment, AdaptiveFiles, IntermediateDirectory, Name, Graph);
 				}
-				else if (bAdaptiveUnityDisablesPCH)
+				else if (CompileEnvironment.PrecompiledHeaderAction == PrecompiledHeaderAction.Include && bAdaptiveUnityDisablesPCH)
 				{
 					AdaptiveOutput = CompileAdaptiveNonUnityFilesWithoutPCH(ToolChain, AdaptiveUnityEnvironment, AdaptiveFiles, IntermediateDirectory, Name, Graph);
 				}
diff --git a/Engine/Source/Programs/UnrealBuildTool/System/Unity.cs b/Engine/Source/Programs/UnrealBuildTool/System/Unity.cs
--- a/Engine/Source/Programs/UnrealBuildTool/System/Unity.cs
+++ b/Engine/Source/Programs/UnrealBuildTool/System/Unity.cs
@@ -142,25 +142,25 @@ namespace UnrealBuildTool
 		}
 
 		/// <summary>
-		/// Given a set of C++ files, generates another set of C++ files that #include all the original
+		/// Given a set of source files, generates another set of source files that #include all the original
 		/// files, the goal being to compile the same code in fewer translation units.
 		/// The "unity" files are written to the IntermediateDirectory.
 		/// </summary>
 		/// <param name="Target">The target we're building</param>
-		/// <param name="CPPFiles">The C++ files to #include.</param>
-		/// <param name="HeaderFiles">The header files that might correspond to the C++ files.</param>
-		/// <param name="CompileEnvironment">The environment that is used to compile the C++ files.</param>
+		/// <param name="SourceFiles">The source files to #include.</param>
+		/// <param name="HeaderFiles">The header files that might correspond to the source files.</param>
+		/// <param name="CompileEnvironment">The environment that is used to compile the source files.</param>
 		/// <param name="WorkingSet">Interface to query files which belong to the working set</param>
 		/// <param name="BaseName">Base name to use for the Unity files</param>
-		/// <param name="IntermediateDirectory">Intermediate directory for unity cpp files</param>
+		/// <param name="IntermediateDirectory">Intermediate directory for unity source files</param>
 		/// <param name="Graph">The makefile being built</param>
 		/// <param name="SourceFileToUnityFile">Receives a mapping of source file to unity file</param>
 		/// <param name="NormalFiles">Receives the files to compile using the normal configuration.</param>
 		/// <param name="AdaptiveFiles">Receives the files to compile using the adaptive unity configuration.</param>
-		/// <param name="NumIncludedBytesPerUnityCPP">An approximate number of bytes of C++ code to target for inclusion in a single unified C++ file.</param>
-		public static void GenerateUnityCPPs(
+		/// <param name="NumIncludedBytesPerUnitySource">An approximate number of bytes of source code to target for inclusion in a single unified source file.</param>
+		public static void GenerateUnitySource(
 			ReadOnlyTargetRules Target,
-			List<FileItem> CPPFiles,
+			List<FileItem> SourceFiles,
 			List<FileItem> HeaderFiles,
 			CppCompileEnvironment CompileEnvironment,
 			ISourceFileWorkingSet WorkingSet,
@@ -170,21 +170,21 @@ namespace UnrealBuildTool
 			Dictionary<FileItem, FileItem> SourceFileToUnityFile,
 			out List<FileItem> NormalFiles,
 			out List<FileItem> AdaptiveFiles,
-			int NumIncludedBytesPerUnityCPP)
+			int NumIncludedBytesPerUnitySource)
 		{
-			List<FileItem> NewCPPFiles = new List<FileItem>();
+			List<FileItem> NewSourceFiles = new List<FileItem>();
 
 			// Figure out size of all input files combined. We use this to determine whether to use larger unity threshold or not.
-			long TotalBytesInCPPFiles = CPPFiles.Sum(F => F.Length);
+			long TotalBytesInSourceFiles = SourceFiles.Sum(F => F.Length);
 
 			// We have an increased threshold for unity file size if, and only if, all files fit into the same unity file. This
 			// is beneficial when dealing with PCH files. The default PCH creation limit is X unity files so if we generate < X 
 			// this could be fairly slow and we'd rather bump the limit a bit to group them all into the same unity file.
 			// Optimization only makes sense if PCH files are enabled.
-			bool bForceIntoSingleUnityFile = Target.bStressTestUnity || (TotalBytesInCPPFiles < NumIncludedBytesPerUnityCPP * 2 && Target.bUsePCHFiles);
+			bool bForceIntoSingleUnityFile = Target.bStressTestUnity || (TotalBytesInSourceFiles < NumIncludedBytesPerUnitySource * 2 && Target.bUsePCHFiles);
 
 			// Even if every single file in the module appears in the working set adaptive unity should still be used even if it's slower.
-			GetAdaptiveFiles(Target, CPPFiles, HeaderFiles, CompileEnvironment, WorkingSet, BaseName, IntermediateDirectory, Graph, out NormalFiles, out AdaptiveFiles);
+			GetAdaptiveFiles(Target, SourceFiles, HeaderFiles, CompileEnvironment, WorkingSet, BaseName, IntermediateDirectory, Graph, out NormalFiles, out AdaptiveFiles);
 
 			// Build the list of unity files.
 			List<FileCollection> AllUnityFiles;
@@ -193,8 +193,8 @@ namespace UnrealBuildTool
 				// Note that we're relying on this not only sorting files within each directory, but also the directories
 				// themselves, so the whole list of file paths is the same across computers.
 				// Case-insensitive file path compare, because you never know what is going on with local file systems.
-				List<FileItem> SortedCPPFiles = new List<FileItem>(CPPFiles);
-				SortedCPPFiles.Sort((A, B) =>
+				List<FileItem> SortedSourceFiles = new List<FileItem>(SourceFiles);
+				SortedSourceFiles.Sort((A, B) =>
 				{
 					// Generated files from UHT need to be first in the list because they implement templated functions that aren't
 					// declared in the header but are required to link. If they are placed later in the list, you will see
@@ -216,31 +216,31 @@ namespace UnrealBuildTool
 				});
 
 				HashSet<FileItem> AdaptiveFileSet = new HashSet<FileItem>(AdaptiveFiles);
-				UnityFileBuilder CPPUnityFileBuilder = new UnityFileBuilder(bForceIntoSingleUnityFile ? -1 : NumIncludedBytesPerUnityCPP);
-				foreach (FileItem CPPFile in SortedCPPFiles)
+				UnityFileBuilder SourceUnityFileBuilder = new UnityFileBuilder(bForceIntoSingleUnityFile ? -1 : NumIncludedBytesPerUnitySource);
+				foreach (FileItem SourceFile in SortedSourceFiles)
 				{
-					if (!bForceIntoSingleUnityFile && CPPFile.AbsolutePath.Contains(".GeneratedWrapper.", StringComparison.InvariantCultureIgnoreCase))
+					if (!bForceIntoSingleUnityFile && SourceFile.AbsolutePath.Contains(".GeneratedWrapper.", StringComparison.InvariantCultureIgnoreCase))
 					{
-						NewCPPFiles.Add(CPPFile);
+						NewSourceFiles.Add(SourceFile);
 					}
 
 					// When adaptive unity is enabled, go ahead and exclude any source files that we're actively working with
-					if (AdaptiveFileSet.Contains(CPPFile))
+					if (AdaptiveFileSet.Contains(SourceFile))
 					{
 						// Let the unity file builder know about the file, so that we can retain the existing size of the unity blobs.
 						// This won't actually make the source file part of the unity blob, but it will keep track of how big the
 						// file is so that other existing unity blobs from the same module won't be invalidated. This prevents much
 						// longer compile times the first time you build after your working file set changes.
-						CPPUnityFileBuilder.AddVirtualFile(CPPFile);
+						SourceUnityFileBuilder.AddVirtualFile(SourceFile);
 					}
 					else
 					{
 						// Compile this file as part of the unity blob
-						CPPUnityFileBuilder.AddFile(CPPFile);
+						SourceUnityFileBuilder.AddFile(SourceFile);
 					}
 				}
 
-				AllUnityFiles = CPPUnityFileBuilder.GetUnityFiles();
+				AllUnityFiles = SourceUnityFileBuilder.GetUnityFiles();
 			}
 
 			// Create a set of CPP files that combine smaller CPP files into larger compilation units, along with the corresponding 
@@ -250,71 +250,72 @@ namespace UnrealBuildTool
 			{
 				++CurrentUnityFileCount;
 
-				StringWriter OutputUnityCPPWriter = new StringWriter();
+				string UnityExt = UnityFile.Files.First().Location.GetExtension().ToLowerInvariant();
+				StringWriter OutputUnitySourceWriter = new();
 
-				OutputUnityCPPWriter.WriteLine("// This file is automatically generated at compile-time to include some subset of the user-created cpp files.");
+				OutputUnitySourceWriter.WriteLine($"// This file is automatically generated at compile-time to include some subset of the user-created {UnityExt.Trim('.')} files.");
 
 				// Determine unity file path name
-				string UnityCPPFileName;
+				string UnitySourceFileName;
 				if (AllUnityFiles.Count > 1)
 				{
 					if (Target.bDetailedUnityFiles)
 					{
-						UnityCPPFileName = String.Format("{0}{1}.{2}_of_{3}.cpp", ModulePrefix, BaseName, CurrentUnityFileCount, AllUnityFiles.Count);
+						UnitySourceFileName = $"{ModulePrefix}{BaseName}.{CurrentUnityFileCount}_of_{AllUnityFiles.Count}{UnityExt}";
 					}
 					else
 					{
-						UnityCPPFileName = String.Format("{0}{1}.{2}.cpp", ModulePrefix, BaseName, CurrentUnityFileCount);
+						UnitySourceFileName = $"{ModulePrefix}{BaseName}.{CurrentUnityFileCount}{UnityExt}";
 					}
 				}
 				else
 				{
-					UnityCPPFileName = String.Format("{0}{1}.cpp", ModulePrefix, BaseName);
+					UnitySourceFileName = $"{ModulePrefix}{BaseName}{UnityExt}";
 				}
-				FileReference UnityCPPFilePath = FileReference.Combine(IntermediateDirectory, UnityCPPFileName);
+				FileReference UnitySourceFilePath = FileReference.Combine(IntermediateDirectory, UnitySourceFileName);
 
-				List<FileItem> InlinedGenCPPFilesInUnity = new();
+				List<FileItem> InlinedGenCPPFilesInUnity = [];
 
 				// Add source files to the unity file
-				foreach (FileItem CPPFile in UnityFile.Files)
+				foreach (FileItem SourceFile in UnityFile.Files)
 				{
-					string CPPFileString = CPPFile.AbsolutePath;
-					if (CompileEnvironment.GetVfsOverlayPath(CPPFile.Location, out string? vfsPath))
+					string SourceFileString = SourceFile.AbsolutePath;
+					if (CompileEnvironment.GetVfsOverlayPath(SourceFile.Location, out string? vfsPath))
 					{
-						CPPFileString = vfsPath;
+						SourceFileString = vfsPath;
 					}
-					else if (CPPFile.Location.IsUnderDirectory(Unreal.RootDirectory))
+					else if (SourceFile.Location.IsUnderDirectory(Unreal.RootDirectory))
 					{
-						CPPFileString = CPPFile.Location.MakeRelativeTo(Unreal.EngineSourceDirectory);
+						SourceFileString = SourceFile.Location.MakeRelativeTo(Unreal.EngineSourceDirectory);
 					}
-					OutputUnityCPPWriter.WriteLine("#include \"{0}\"", CPPFileString.Replace('\\', '/'));
+					OutputUnitySourceWriter.WriteLine("#include \"{0}\"", SourceFileString.Replace('\\', '/'));
 
 					List<FileItem>? InlinedGenCPPFiles;
-					if (CompileEnvironment.FileInlineGenCPPMap.TryGetValue(CPPFile, out InlinedGenCPPFiles))
+					if (CompileEnvironment.FileInlineGenCPPMap.TryGetValue(SourceFile, out InlinedGenCPPFiles))
 					{
 						InlinedGenCPPFilesInUnity.AddRange(InlinedGenCPPFiles);
 					}
 				}
 
 				// Write the unity file to the intermediate folder.
-				FileItem UnityCPPFile = Graph.CreateIntermediateTextFile(UnityCPPFilePath, OutputUnityCPPWriter.ToString());
-				NewCPPFiles.Add(UnityCPPFile);
+				FileItem UnitySourceFile = Graph.CreateIntermediateTextFile(UnitySourceFilePath, OutputUnitySourceWriter.ToString());
+				NewSourceFiles.Add(UnitySourceFile);
 
 				// Store all the inlined gen.cpp files
-				CompileEnvironment.FileInlineGenCPPMap[UnityCPPFile] = InlinedGenCPPFilesInUnity;
+				CompileEnvironment.FileInlineGenCPPMap[UnitySourceFile] = InlinedGenCPPFilesInUnity;
 
 				// Store the mapping of source files to unity files in the makefile
 				foreach (FileItem SourceFile in UnityFile.Files)
 				{
-					SourceFileToUnityFile[SourceFile] = UnityCPPFile;
+					SourceFileToUnityFile[SourceFile] = UnitySourceFile;
 				}
 				foreach (FileItem SourceFile in UnityFile.VirtualFiles)
 				{
-					SourceFileToUnityFile[SourceFile] = UnityCPPFile;
+					SourceFileToUnityFile[SourceFile] = UnitySourceFile;
 				}
 			}
 
-			NormalFiles = NewCPPFiles;
+			NormalFiles = NewSourceFiles;
 		}
 
 		public static void GetAdaptiveFiles(
