- Catch a case where there's an empty named device that coming in to codesigning

diff --git a/Engine/Source/Programs/UnrealBuildTool/ToolChain/AppleToolChain.cs b/Engine/Source/Programs/UnrealBuildTool/ToolChain/AppleToolChain.cs
--- a/Engine/Source/Programs/UnrealBuildTool/ToolChain/AppleToolChain.cs
+++ b/Engine/Source/Programs/UnrealBuildTool/ToolChain/AppleToolChain.cs
@@ -932,7 +932,7 @@ namespace UnrealBuildTool
 				//   - for Mac, the destination seems to want name="My Mac", not the real name of the Mac
 				//   - when no destinations were supplied
 				//   - when using manual codesigning because the main use of this is to update provisioning information (reducing potential for issues)
-				if (Platform == UnrealTargetPlatform.Mac || DestinationIds == null || DestinationIds.Count == 0 || !bUseAutomaticCodeSigning)
+				if (Platform == UnrealTargetPlatform.Mac || DestinationIds == null || DestinationIds.Count == 0 || DestinationIds[0] == "" || !bUseAutomaticCodeSigning)
 				{
 					DestinationParam = $"-destination generic/platform=\"{AppleExports.GetDestinationPlatform(Platform, Architectures)}\"";
 				}
