UnrealBuildTool: Continued vfs support, no longer using .json overlay

#rnx

diff --git a/Engine/Source/Programs/UnrealBuildTool/Configuration/UEBuildModule.cs b/Engine/Source/Programs/UnrealBuildTool/Configuration/UEBuildModule.cs
--- a/Engine/Source/Programs/UnrealBuildTool/Configuration/UEBuildModule.cs
+++ b/Engine/Source/Programs/UnrealBuildTool/Configuration/UEBuildModule.cs
@@ -955,12 +955,28 @@ namespace UnrealBuildTool
 		/// </summary>
 		public virtual List<FileItem> Compile(ReadOnlyTargetRules Target, UEToolChain ToolChain, CppCompileEnvironment CompileEnvironment, ISourceFileWorkingSet WorkingSet, IActionGraphBuilder Graph, ILogger Logger)
 		{
+			CppCompileEnvironment ModuleCompileEnvironment = new CppCompileEnvironment(CompileEnvironment);
+			IEnumerable<DirectoryReference> AllIncludePaths = [.. ModuleCompileEnvironment.UserIncludePaths, .. ModuleCompileEnvironment.SystemIncludePaths];
+			if (Target.ProjectFile != null && (Target.BuildEnvironment != TargetBuildEnvironment.Shared || AllIncludePaths.Any(x => x.IsUnderDirectory(Target.ProjectFile.Directory))))
+			{
+				ModuleCompileEnvironment.RootPaths[CppRootPathFolder.Project] = Target.ProjectFile.Directory;
+			}
+
+			if (Rules.ExtraRootPath != null)
+			{
+				DirectoryReference? ExtraRootPath = DirectoryReference.FromString(Rules.ExtraRootPath);
+				if (ExtraRootPath != null)
+				{
+					ModuleCompileEnvironment.RootPaths[CppRootPathFolder.Extra] = ExtraRootPath;
+				}
+			}
+
 			// Generate type libraries for Windows
 			foreach (ModuleRules.TypeLibrary TypeLibrary in Rules.TypeLibraries)
 			{
 				FileReference OutputFile = FileReference.Combine(IntermediateDirectory, TypeLibrary.Header);
 				FileReference? OutputInclude = TypeLibrary.Include != null ? FileReference.Combine(IntermediateDirectory, TypeLibrary.Include) : null;
-				ToolChain.GenerateTypeLibraryHeader(CompileEnvironment, TypeLibrary, OutputFile, OutputInclude, Graph);
+				ToolChain.GenerateTypeLibraryHeader(ModuleCompileEnvironment, TypeLibrary, OutputFile, OutputInclude, Graph);
 			}
 
 			return new List<FileItem>();
diff --git a/Engine/Source/Programs/UnrealBuildTool/Configuration/UEBuildModuleCPP.cs b/Engine/Source/Programs/UnrealBuildTool/Configuration/UEBuildModuleCPP.cs
--- a/Engine/Source/Programs/UnrealBuildTool/Configuration/UEBuildModuleCPP.cs
+++ b/Engine/Source/Programs/UnrealBuildTool/Configuration/UEBuildModuleCPP.cs
@@ -395,13 +395,6 @@ namespace UnrealBuildTool
 
 			CppCompileEnvironment ModuleCompileEnvironment = CreateModuleCompileEnvironment(Target, BinaryCompileEnvironment, Logger);
 
-			// Setup clang virtual file system overlay file if enabled
-			if (Target.bCompilerVFS && (!Target.Platform.IsInGroup(UnrealPlatformGroup.Microsoft) || Target.WindowsPlatform.Compiler.IsClang()))
-			{
-				ModuleCompileEnvironment.VfsOverlayFile = CreateVfsOverlayFile(ModuleCompileEnvironment, IntermediateDirectory, "vfs.json", Graph);
-				ModuleCompileEnvironment.AdditionalPrerequisites.Add(ModuleCompileEnvironment.VfsOverlayFile);
-			}
-
 			// If the module is precompiled, read the object files from the manifest
 			if (Rules.bUsePrecompiled && Target.LinkType == TargetLinkType.Monolithic)
 			{
@@ -1443,13 +1436,6 @@ namespace UnrealBuildTool
 				CompileEnvironment.PrecompiledHeaderIncludeFilename = WrapperFile.Location;
 				CopySettingsForSharedPCH(ModuleCompileEnvironment, CompileEnvironment);
 
-				// Setup clang virtual file system overlay file if enabled
-				if (Rules.Target.bCompilerVFS && (!Rules.Target.Platform.IsInGroup(UnrealPlatformGroup.Microsoft) || Rules.Target.WindowsPlatform.Compiler.IsClang()))
-				{
-					CompileEnvironment.VfsOverlayFile = CreateVfsOverlayFile(Template.BaseCompileEnvironment, Template.OutputDir, $"SharedVFS.{Template.Module.Name}{Variant}.json", Graph, allowAsync: false);
-					CompileEnvironment.AdditionalPrerequisites.Add(CompileEnvironment.VfsOverlayFile);
-				}
-
 				// Setup PCH chaining
 				PrecompiledHeaderInstance? ParentPCHInstance = null;
 				Dictionary<string, string>? DefinitionsDictionary = null;
@@ -2547,6 +2533,22 @@ namespace UnrealBuildTool
 			Dictionary<UEBuildModule, bool> ModuleToIncludePathsOnlyFlag = new Dictionary<UEBuildModule, bool>();
 			FindModulesInPublicCompileEnvironment(ModuleToIncludePathsOnlyFlag);
 
+			// Set up additional root paths
+			IEnumerable<DirectoryReference> AllIncludePaths = [.. CompileEnvironment.UserIncludePaths, .. CompileEnvironment.SystemIncludePaths];
+			if (Target.ProjectFile != null && (!Target.bUseSharedBuildEnvironment || AllIncludePaths.Any(x => x.IsUnderDirectory(Target.ProjectFile.Directory))))
+			{
+				CompileEnvironment.RootPaths[CppRootPathFolder.Project] = Target.ProjectFile.Directory;
+			}
+
+			if (Rules.ExtraRootPath != null)
+			{
+				DirectoryReference? ExtraRootPath = DirectoryReference.FromString(Rules.ExtraRootPath);
+				if (ExtraRootPath != null)
+				{
+					CompileEnvironment.RootPaths[CppRootPathFolder.Extra] = ExtraRootPath;
+				}
+			}
+
 			// Now set up the compile environment for the modules in the original order that we encountered them
 			foreach (UEBuildModule Module in ModuleToIncludePathsOnlyFlag.Keys)
 			{
diff --git a/Engine/Source/Programs/UnrealBuildTool/Executors/UnrealBuildAccelerator/UBAExecutor.cs b/Engine/Source/Programs/UnrealBuildTool/Executors/UnrealBuildAccelerator/UBAExecutor.cs
--- a/Engine/Source/Programs/UnrealBuildTool/Executors/UnrealBuildAccelerator/UBAExecutor.cs
+++ b/Engine/Source/Programs/UnrealBuildTool/Executors/UnrealBuildAccelerator/UBAExecutor.cs
@@ -719,6 +719,11 @@ namespace UnrealBuildTool
 
 		ulong GetActionRootsHandle(LinkedAction action)
 		{
+			if (!action.RootPaths.bUseVfs)
+			{
+				return 0;
+			}
+
 			using MemoryStream outputsMemory = new(1024);
 			using (BinaryWriter writer = new(outputsMemory, System.Text.Encoding.UTF8, true))
 			{
diff --git a/Engine/Source/Programs/UnrealBuildTool/Platform/Linux/LinuxToolChain.cs b/Engine/Source/Programs/UnrealBuildTool/Platform/Linux/LinuxToolChain.cs
--- a/Engine/Source/Programs/UnrealBuildTool/Platform/Linux/LinuxToolChain.cs
+++ b/Engine/Source/Programs/UnrealBuildTool/Platform/Linux/LinuxToolChain.cs
@@ -570,7 +570,7 @@ namespace UnrealBuildTool
 			if (CrossCompiling())
 			{
 				Arguments.Add($"-target {CompileEnvironment.Architecture.LinuxName}");        // Set target triple
-				Arguments.Add($"--sysroot=\"{NormalizeCommandLinePath(LinuxInfo.BaseLinuxPath!, CompileEnvironment)}\"");
+				Arguments.Add($"--sysroot=\"{NormalizeCommandLinePath(LinuxInfo.BaseLinuxPath!, CompileEnvironment.RootPaths)}\"");
 			}
 		}
 
diff --git a/Engine/Source/Programs/UnrealBuildTool/Platform/Mac/MacToolChain.cs b/Engine/Source/Programs/UnrealBuildTool/Platform/Mac/MacToolChain.cs
--- a/Engine/Source/Programs/UnrealBuildTool/Platform/Mac/MacToolChain.cs
+++ b/Engine/Source/Programs/UnrealBuildTool/Platform/Mac/MacToolChain.cs
@@ -207,7 +207,7 @@ namespace UnrealBuildTool
 				FileReference FrameworkPath = new FileReference(Path.GetFullPath(Framework.Name));
 				if (!FrameworksSearchPaths.Contains(FrameworkPath.Directory.FullName))
 				{
-					Arguments.Add($"-F \"{NormalizeCommandLinePath(FrameworkPath.Directory, CompileEnvironment)}\"");
+					Arguments.Add($"-F \"{NormalizeCommandLinePath(FrameworkPath.Directory, CompileEnvironment.RootPaths)}\"");
 					FrameworksSearchPaths.Add(FrameworkPath.Directory.FullName);
 				}
 			}
diff --git a/Engine/Source/Programs/UnrealBuildTool/Platform/Windows/PVSToolChain.cs b/Engine/Source/Programs/UnrealBuildTool/Platform/Windows/PVSToolChain.cs
--- a/Engine/Source/Programs/UnrealBuildTool/Platform/Windows/PVSToolChain.cs
+++ b/Engine/Source/Programs/UnrealBuildTool/Platform/Windows/PVSToolChain.cs
@@ -645,8 +645,8 @@ namespace UnrealBuildTool
 			base.SetUpGlobalEnvironment(Target, GlobalCompileEnvironment, GlobalLinkEnvironment);
 			InnerToolChain.SetUpGlobalEnvironment(Target, GlobalCompileEnvironment, GlobalLinkEnvironment);
 
-			GlobalCompileEnvironment.RootPaths[CppRootPathFolder.Extra] = AnalyzerFile.Directory;
-			GlobalLinkEnvironment.RootPaths[CppRootPathFolder.Extra] = AnalyzerFile.Directory;
+			GlobalCompileEnvironment.RootPaths[CppRootPathFolder.Analyzer] = AnalyzerFile.Directory;
+			GlobalLinkEnvironment.RootPaths[CppRootPathFolder.Analyzer] = AnalyzerFile.Directory;
 		}
   
 		public override void SetEnvironmentVariables()
diff --git a/Engine/Source/Programs/UnrealBuildTool/Platform/Windows/VCCompileAction.cs b/Engine/Source/Programs/UnrealBuildTool/Platform/Windows/VCCompileAction.cs
--- a/Engine/Source/Programs/UnrealBuildTool/Platform/Windows/VCCompileAction.cs
+++ b/Engine/Source/Programs/UnrealBuildTool/Platform/Windows/VCCompileAction.cs
@@ -538,27 +538,26 @@ namespace UnrealBuildTool
 
 		public List<string> GetCompilerArguments(ILogger Logger)
 		{
-			CppRootPaths? rootPaths = CompilerType.IsClang() && VfsOverlayFile != null ? RootPaths : null;
 			List<string> Arguments = new List<string>();
 
 			if (SourceFile != null)
 			{
-				VCToolChain.AddSourceFile(Arguments, SourceFile, rootPaths);
+				VCToolChain.AddSourceFile(Arguments, SourceFile, RootPaths);
 			}
 
 			foreach (FileItem AdditionalRsp in AdditionalResponseFiles)
 			{
-				VCToolChain.AddResponseFile(Arguments, AdditionalRsp);
+				VCToolChain.AddResponseFile(Arguments, AdditionalRsp, RootPaths);
 			}
 
 			foreach (DirectoryReference IncludePath in IncludePaths)
 			{
-				VCToolChain.AddIncludePath(Arguments, IncludePath, CompilerType, rootPaths);
+				VCToolChain.AddIncludePath(Arguments, IncludePath, CompilerType, RootPaths);
 			}
 
 			foreach (DirectoryReference SystemIncludePath in SystemIncludePaths)
 			{
-				VCToolChain.AddSystemIncludePath(Arguments, SystemIncludePath, CompilerType, rootPaths);
+				VCToolChain.AddSystemIncludePath(Arguments, SystemIncludePath, CompilerType, RootPaths);
 			}
 
 			foreach (string Definition in Definitions)
@@ -570,53 +569,53 @@ namespace UnrealBuildTool
 
 			foreach (FileItem ForceIncludeFile in ForceIncludeFiles)
 			{
-				VCToolChain.AddForceIncludeFile(Arguments, ForceIncludeFile, rootPaths);
+				VCToolChain.AddForceIncludeFile(Arguments, ForceIncludeFile, RootPaths);
 			}
 
 			if (CreatePchFile != null)
 			{
-				VCToolChain.AddCreatePchFile(Arguments, PchThroughHeaderFile!, CreatePchFile);
+				VCToolChain.AddCreatePchFile(Arguments, PchThroughHeaderFile!, CreatePchFile, RootPaths);
 			}
 
 			if (UsingPchFile != null && CompilerType.IsMSVC())
 			{
-				VCToolChain.AddUsingPchFile(Arguments, PchThroughHeaderFile!, UsingPchFile, rootPaths);
+				VCToolChain.AddUsingPchFile(Arguments, PchThroughHeaderFile!, UsingPchFile, RootPaths);
 			}
 
 			if (PreprocessedFile != null)
 			{
-				VCToolChain.AddPreprocessedFile(Arguments, PreprocessedFile);
+				VCToolChain.AddPreprocessedFile(Arguments, PreprocessedFile, RootPaths);
 			}
 
 			if (ObjectFile != null)
 			{
-				VCToolChain.AddObjectFile(Arguments, ObjectFile);
+				VCToolChain.AddObjectFile(Arguments, ObjectFile, RootPaths);
 			}
 
 			if (AssemblyFile != null)
 			{
-				VCToolChain.AddAssemblyFile(Arguments, AssemblyFile);
+				VCToolChain.AddAssemblyFile(Arguments, AssemblyFile, RootPaths);
 			}
 
 			if (AnalyzeLogFile != null)
 			{
-				VCToolChain.AddAnalyzeLogFile(Arguments, AnalyzeLogFile);
+				VCToolChain.AddAnalyzeLogFile(Arguments, AnalyzeLogFile, RootPaths);
 			}
 
 			if (ExperimentalLogFile != null)
 			{
-				VCToolChain.AddExperimentalLogFile(Arguments, ExperimentalLogFile);
+				VCToolChain.AddExperimentalLogFile(Arguments, ExperimentalLogFile, RootPaths);
 			}
 
 			// A better way to express this? .json is used as output for /sourceDependencies), but .md.json is used as output for /sourceDependencies:directives)
 			if (DependencyListFile != null && DependencyListFile.HasExtension(".json") && !DependencyListFile.HasExtension(".md.json"))
 			{
-				VCToolChain.AddSourceDependenciesFile(Arguments, DependencyListFile);
+				VCToolChain.AddSourceDependenciesFile(Arguments, DependencyListFile, RootPaths);
 			}
 
 			if (DependencyListFile != null && DependencyListFile.HasExtension(".d"))
 			{
-				VCToolChain.AddSourceDependsFile(Arguments, DependencyListFile);
+				VCToolChain.AddSourceDependsFile(Arguments, DependencyListFile, RootPaths);
 			}
 
 			Arguments.AddRange(this.Arguments);
diff --git a/Engine/Source/Programs/UnrealBuildTool/Platform/Windows/VCToolChain.cs b/Engine/Source/Programs/UnrealBuildTool/Platform/Windows/VCToolChain.cs
--- a/Engine/Source/Programs/UnrealBuildTool/Platform/Windows/VCToolChain.cs
+++ b/Engine/Source/Programs/UnrealBuildTool/Platform/Windows/VCToolChain.cs
@@ -173,40 +173,29 @@ namespace UnrealBuildTool
 			return NormalizeCommandLinePath(Item.Location);
 		}
 
-		protected static string NormalizeCommandLinePath(FileSystemReference reference, CppRootPaths? RootPaths)
+		public static new string NormalizeCommandLinePath(FileSystemReference Reference, CppRootPaths RootPaths)
 		{
-			if (RootPaths != null)
-			{
-				foreach (KeyValuePair<CppRootPathFolder, DirectoryReference> item in RootPaths)
-				{
-					if (!reference.IsUnderDirectory(item.Value))
-					{
-						continue;
-					}
-					return item.Key.AsVfsString(item.Value, reference);
-				}
-			}
-			return NormalizeCommandLinePath(reference);
+			return RootPaths.GetVfsOverlayPath(Reference, out string? vfsPath) ? vfsPath : NormalizeCommandLinePath(Reference);
 		}
 
-		protected static string NormalizeCommandLinePath(FileItem Item, CppRootPaths? RootPaths)
+		public static new string NormalizeCommandLinePath(FileItem Item, CppRootPaths RootPaths)
 		{
 			return NormalizeCommandLinePath(Item.Location, RootPaths);
 		}
 
-		public static void AddResponseFile(List<string> Arguments, FileItem ResponseFile)
+		public static void AddResponseFile(List<string> Arguments, FileItem ResponseFile, CppRootPaths RootPaths)
 		{
-			string ResponseFileString = NormalizeCommandLinePath(ResponseFile);
+			string ResponseFileString = NormalizeCommandLinePath(ResponseFile, RootPaths);
 			Arguments.Add($"@{Utils.MakePathSafeToUseWithCommandLine(ResponseFileString)}");
 		}
 
-		public static void AddSourceFile(List<string> Arguments, FileItem SourceFile, CppRootPaths? RootPaths)
+		public static void AddSourceFile(List<string> Arguments, FileItem SourceFile, CppRootPaths RootPaths)
 		{
 			string SourceFileString = NormalizeCommandLinePath(SourceFile, RootPaths);
 			Arguments.Add(Utils.MakePathSafeToUseWithCommandLine(SourceFileString));
 		}
 
-		private static void AddIncludePath(List<string> Arguments, DirectoryReference IncludePath, WindowsCompiler Compiler, bool bSystemInclude, CppRootPaths? RootPaths = null)
+		private static void AddIncludePath(List<string> Arguments, DirectoryReference IncludePath, WindowsCompiler Compiler, bool bSystemInclude, CppRootPaths RootPaths)
 		{
 			// Try to use a relative path to shorten command line length.
 			string IncludePathString = NormalizeCommandLinePath(IncludePath, RootPaths);
@@ -232,26 +221,26 @@ namespace UnrealBuildTool
 			}
 		}
 
-		public static void AddIncludePath(List<string> Arguments, DirectoryReference IncludePath, WindowsCompiler Compiler, CppRootPaths? RootPaths = null)
+		public static void AddIncludePath(List<string> Arguments, DirectoryReference IncludePath, WindowsCompiler Compiler, CppRootPaths RootPaths)
 		{
 			AddIncludePath(Arguments, IncludePath, Compiler, false, RootPaths);
 		}
 
-		public static void AddSystemIncludePath(List<string> Arguments, DirectoryReference IncludePath, WindowsCompiler Compiler, CppRootPaths? RootPaths = null)
+		public static void AddSystemIncludePath(List<string> Arguments, DirectoryReference IncludePath, WindowsCompiler Compiler, CppRootPaths RootPaths)
 		{
 			AddIncludePath(Arguments, IncludePath, Compiler, true, RootPaths);
 		}
 
-		public static void AddForceIncludeFile(List<string> Arguments, FileItem ForceIncludeFile, CppRootPaths? RootPaths)
+		public static void AddForceIncludeFile(List<string> Arguments, FileItem ForceIncludeFile, CppRootPaths RootPaths)
 		{
 			string ForceIncludeFileString = NormalizeCommandLinePath(ForceIncludeFile, RootPaths);
 			Arguments.Add($"/FI\"{ForceIncludeFileString}\"");
 		}
 
-		public static void AddCreatePchFile(List<string> Arguments, FileItem PchThroughHeaderFile, FileItem CreatePchFile)
+		public static void AddCreatePchFile(List<string> Arguments, FileItem PchThroughHeaderFile, FileItem CreatePchFile, CppRootPaths RootPaths)
 		{
 			string PchThroughHeaderFilePath = PchThroughHeaderFile.Location.GetFileName();
-			string CreatePchFilePath = NormalizeCommandLinePath(CreatePchFile);
+			string CreatePchFilePath = NormalizeCommandLinePath(CreatePchFile, RootPaths);
 			if (CreatePchFile.Name.EndsWith(".ifc"))
 			{
 				Arguments.Add($"/exportHeader");
@@ -266,7 +255,7 @@ namespace UnrealBuildTool
 			}
 		}
 
-		public static void AddUsingPchFile(List<string> Arguments, FileItem PchThroughHeaderFile, FileItem UsingPchFile, CppRootPaths? RootPaths)
+		public static void AddUsingPchFile(List<string> Arguments, FileItem PchThroughHeaderFile, FileItem UsingPchFile, CppRootPaths RootPaths)
 		{
 			string PchThroughHeaderFilePath = NormalizeCommandLinePath(PchThroughHeaderFile, RootPaths);
 			string UsingPchFilePath = NormalizeCommandLinePath(UsingPchFile, RootPaths);
@@ -283,48 +272,48 @@ namespace UnrealBuildTool
 			}
 		}
 
-		public static void AddPreprocessedFile(List<string> Arguments, FileItem PreprocessedFile)
+		public static void AddPreprocessedFile(List<string> Arguments, FileItem PreprocessedFile, CppRootPaths RootPaths)
 		{
-			string PreprocessedFileString = NormalizeCommandLinePath(PreprocessedFile);
+			string PreprocessedFileString = NormalizeCommandLinePath(PreprocessedFile, RootPaths);
 			Arguments.Add("/P"); // Preprocess
 			Arguments.Add("/C"); // Preserve comments when preprocessing
 			Arguments.Add($"/Fi\"{PreprocessedFileString}\""); // Preprocess to a file
 		}
 
-		public static void AddObjectFile(List<string> Arguments, FileItem ObjectFile)
+		public static void AddObjectFile(List<string> Arguments, FileItem ObjectFile, CppRootPaths RootPaths)
 		{
-			string ObjectFileString = NormalizeCommandLinePath(ObjectFile);
+			string ObjectFileString = NormalizeCommandLinePath(ObjectFile, RootPaths);
 			Arguments.Add($"/Fo\"{ObjectFileString}\"");
 		}
 
-		public static void AddAssemblyFile(List<string> Arguments, FileItem AssemblyFile)
+		public static void AddAssemblyFile(List<string> Arguments, FileItem AssemblyFile, CppRootPaths RootPaths)
 		{
-			string AssemblyFileString = NormalizeCommandLinePath(AssemblyFile);
+			string AssemblyFileString = NormalizeCommandLinePath(AssemblyFile, RootPaths);
 			Arguments.Add("/FAs");// Write out an assembly file (.asm) with the c++ code embedded in it via comments
 			Arguments.Add($"/Fa\"{AssemblyFileString}\""); // Set the output patj for the asm file
 		}
 
-		public static void AddAnalyzeLogFile(List<string> Arguments, FileItem AnalyzeLogFile)
+		public static void AddAnalyzeLogFile(List<string> Arguments, FileItem AnalyzeLogFile, CppRootPaths RootPaths)
 		{
-			string AnalyzeLogFileString = NormalizeCommandLinePath(AnalyzeLogFile);
+			string AnalyzeLogFileString = NormalizeCommandLinePath(AnalyzeLogFile, RootPaths);
 			Arguments.Add("/analyze:log " + Utils.MakePathSafeToUseWithCommandLine(AnalyzeLogFileString));
 		}
 
-		public static void AddExperimentalLogFile(List<string> Arguments, FileItem ExperimentalLogFile)
+		public static void AddExperimentalLogFile(List<string> Arguments, FileItem ExperimentalLogFile, CppRootPaths RootPaths)
 		{
-			string ExperimentalLogFileString = NormalizeCommandLinePath(ExperimentalLogFile);
+			string ExperimentalLogFileString = NormalizeCommandLinePath(ExperimentalLogFile, RootPaths);
 			Arguments.Add("/experimental:log " + Utils.MakePathSafeToUseWithCommandLine(ExperimentalLogFileString));
 		}
 
-		public static void AddSourceDependenciesFile(List<string> Arguments, FileItem SourceDependenciesFile)
+		public static void AddSourceDependenciesFile(List<string> Arguments, FileItem SourceDependenciesFile, CppRootPaths RootPaths)
 		{
-			string SourceDependenciesFileString = NormalizeCommandLinePath(SourceDependenciesFile);
+			string SourceDependenciesFileString = NormalizeCommandLinePath(SourceDependenciesFile, RootPaths);
 			Arguments.Add("/sourceDependencies " + Utils.MakePathSafeToUseWithCommandLine(SourceDependenciesFileString));
 		}
 
-		public static void AddSourceDependsFile(List<string> Arguments, FileItem SourceDependsFile)
+		public static void AddSourceDependsFile(List<string> Arguments, FileItem SourceDependsFile, CppRootPaths RootPaths)
 		{
-			string SourceDependsFileString = NormalizeCommandLinePath(SourceDependsFile);
+			string SourceDependsFileString = NormalizeCommandLinePath(SourceDependsFile, RootPaths);
 			Arguments.Add($"/clang:-MD /clang:-MF\"{SourceDependsFileString}\"");
 		}
 
@@ -443,9 +432,8 @@ namespace UnrealBuildTool
 					AddDefinition(Arguments, "__clang_analyzer__");
 				}
 
-				if (CompileEnvironment.VfsOverlayFile != null)
+				if (CompileEnvironment.RootPaths.bUseVfs)
 				{
-					Arguments.Add($"-Xclang -ivfsoverlay\"{NormalizeCommandLinePath(CompileEnvironment.VfsOverlayFile)}\"");
 					Arguments.Add("-fcoverage-compilation-dir=.");
 					Arguments.Add("-fdebug-compilation-dir=.");
 					Arguments.Add("-vctoolsdir=.");
@@ -727,8 +715,7 @@ namespace UnrealBuildTool
 
 				// Currently the ASan headers are not default around. They can be found at this location so lets use this until this is resolved in the toolchain
 				// Jira with some more info and the MSVC bug at UE-144727
-				AddSystemIncludePath(Arguments, DirectoryReference.Combine(EnvVars.CompilerDir, "crt", "src"), Target.WindowsPlatform.Compiler,
-					Target.WindowsPlatform.Compiler.IsClang() && CompileEnvironment.VfsOverlayFile != null ? CompileEnvironment.RootPaths : null);
+				AddSystemIncludePath(Arguments, DirectoryReference.Combine(EnvVars.CompilerDir, "crt", "src"), Target.WindowsPlatform.Compiler, CompileEnvironment.RootPaths);
 			}
 
 			if (Target.WindowsPlatform.bEnableLibFuzzer)
@@ -1775,7 +1762,6 @@ namespace UnrealBuildTool
 			}
 #endif
 			BaseCompileAction.RootPaths = CompileEnvironment.RootPaths;
-			BaseCompileAction.VfsOverlayFile = CompileEnvironment.VfsOverlayFile;
 
 			// Add additional response files
 			BaseCompileAction.AdditionalResponseFiles.AddRange(CompileEnvironment.AdditionalResponseFiles);
@@ -2407,6 +2393,8 @@ namespace UnrealBuildTool
 		{
 			CPPOutput Result = new CPPOutput();
 
+			CppRootPaths RootPaths = CompileEnvironment.RootPaths;
+
 			foreach (FileItem RCFile in InputFiles)
 			{
 				Action CompileAction = Graph.CreateAction(ActionType.Compile);
@@ -2438,17 +2426,17 @@ namespace UnrealBuildTool
 				// Include paths. Don't use AddIncludePath() here, since it uses the full path and exceeds the max command line length.
 				foreach (DirectoryReference IncludePath in CompileEnvironment.UserIncludePaths)
 				{
-					Arguments.Add($"/I \"{NormalizeCommandLinePath(IncludePath)}\"");
+					Arguments.Add($"/I \"{NormalizeCommandLinePath(IncludePath, RootPaths)}\"");
 				}
 
 				// System include paths.
 				foreach (DirectoryReference SystemIncludePath in CompileEnvironment.SystemIncludePaths)
 				{
-					Arguments.Add($"/I \"{NormalizeCommandLinePath(SystemIncludePath)}\"");
+					Arguments.Add($"/I \"{NormalizeCommandLinePath(SystemIncludePath, RootPaths)}\"");
 				}
 				foreach (DirectoryReference SystemIncludePath in EnvVars.IncludePaths)
 				{
-					Arguments.Add($"/I \"{NormalizeCommandLinePath(SystemIncludePath)}\"");
+					Arguments.Add($"/I \"{NormalizeCommandLinePath(SystemIncludePath, RootPaths)}\"");
 				}
 
 				// Preprocessor definitions.
@@ -2473,7 +2461,7 @@ namespace UnrealBuildTool
 				CompileAction.PrerequisiteItems.Add(FileItem.GetItemByFileReference(IconFile));
 
 				// Setup the compile environment, setting the icon to use via a macro. This is used in Default.rc2.
-				AddDefinition(Arguments, $"BUILD_ICON_FILE_NAME=\"\\\"{NormalizeCommandLinePath(IconFile).Replace("\\", "\\\\")}\\\"\"");
+				AddDefinition(Arguments, $"BUILD_ICON_FILE_NAME=\"\\\"{NormalizeCommandLinePath(IconFile, RootPaths).Replace("\\", "\\\\")}\\\"\"");
 
 				// Apply the target settings for the resources
 				if (!CompileEnvironment.bUseSharedBuildEnvironment)
@@ -2507,11 +2495,11 @@ namespace UnrealBuildTool
 						)
 					);
 				CompileAction.ProducedItems.Add(CompiledResourceFile);
-				Arguments.Add($"/fo \"{NormalizeCommandLinePath(CompiledResourceFile)}\"");
+				Arguments.Add($"/fo \"{NormalizeCommandLinePath(CompiledResourceFile, RootPaths)}\"");
 				Result.ObjectFiles.Add(CompiledResourceFile);
 
 				// Add the RC file as a prerequisite of the action.
-				Arguments.Add($"\"{NormalizeCommandLinePath(RCFile)}\"");
+				Arguments.Add($"\"{NormalizeCommandLinePath(RCFile, RootPaths)}\"");
 
 				// Create a response file for the resource compilier
 				FileItem ResponseFile = FileItem.GetItemByFileReference(GetResponseFileName(CompileEnvironment, CompiledResourceFile));
@@ -2540,13 +2528,15 @@ namespace UnrealBuildTool
 
 		public override void GenerateTypeLibraryHeader(CppCompileEnvironment CompileEnvironment, ModuleRules.TypeLibrary TypeLibrary, FileReference OutputFile, FileReference? OutputHeader, IActionGraphBuilder Graph)
 		{
+			CppRootPaths RootPaths = CompileEnvironment.RootPaths;
+
 			// Create the input file
 			StringBuilder Contents = new StringBuilder();
 			Contents.AppendLine("#include <windows.h>");
 			Contents.AppendLine("#include <unknwn.h>");
 			Contents.AppendLine();
 
-			Contents.AppendFormat("#import \"{0}\"", TypeLibrary.FileName.Replace('\\', '/'));
+			Contents.AppendFormat("#import \"{0}\"", NormalizeCommandLinePath(FileReference.FromString(TypeLibrary.FileName), RootPaths).Replace('\\', '/'));
 			if (!String.IsNullOrEmpty(TypeLibrary.Attributes))
 			{
 				Contents.Append(' ');
@@ -2561,44 +2551,45 @@ namespace UnrealBuildTool
 
 			List<string> Arguments = new List<string>
 			{
-				$"\"{InputFile.Location}\"",
+				$"\"{NormalizeCommandLinePath(InputFile, RootPaths)}\"",
 				"/c",
 				"/nologo",
-				$"/Fo\"{ObjectFile.Location}\""
+				$"/Fo\"{NormalizeCommandLinePath(ObjectFile, RootPaths)}\""
 			};
 
 			foreach (DirectoryReference IncludePath in CompileEnvironment.UserIncludePaths)
 			{
-				AddIncludePath(Arguments, IncludePath, EnvVars.ToolChain);
+				AddIncludePath(Arguments, IncludePath, EnvVars.ToolChain, RootPaths);
 			}
 
 			foreach (DirectoryReference IncludePath in CompileEnvironment.SystemIncludePaths)
 			{
-				AddSystemIncludePath(Arguments, IncludePath, EnvVars.ToolChain);
+				AddSystemIncludePath(Arguments, IncludePath, EnvVars.ToolChain, RootPaths);
 			}
 
 			foreach (DirectoryReference IncludePath in EnvVars.IncludePaths)
 			{
-				AddSystemIncludePath(Arguments, IncludePath, EnvVars.ToolChain);
+				AddSystemIncludePath(Arguments, IncludePath, EnvVars.ToolChain, RootPaths);
 			}
 
 			FileItem ResponseFile = Graph.CreateIntermediateTextFile(OutputFile.ChangeExtension(ResponseExt), String.Join(" ", Arguments));
 
 			// Build the command for touching the output file to update the time stamp
 			string TouchActionCommand;
+			string OutfileFullName = NormalizeCommandLinePath(OutputFile);
 			if (BuildHostPlatform.Current.IsRunningOnWine())
 			{
-				TouchActionCommand = $"copy /Y \"{OutputFile.FullName}\" \"%TMP%/{OutputFile.GetFileName()}\" && type \"%TMP%/{OutputFile.GetFileName()}\" > \"{OutputFile.FullName}\"";
+				TouchActionCommand = $"copy /Y \"{OutfileFullName}\" \"%TMP%/{OutputFile.GetFileName()}\" && type \"%TMP%/{OutputFile.GetFileName()}\" > \"{OutfileFullName}\"";
 			}
 			else
 			{
-				TouchActionCommand = $"copy /b \"{OutputFile.FullName}\"+,, \"{OutputFile.FullName}\"";
+				TouchActionCommand = $"copy /b \"{OutfileFullName}\"+,, \"{OutfileFullName}\"";
 			}
 
 			// Build the batch file
 			StringBuilder BatchFileContents = new();
 			BatchFileContents.AppendLine("@echo off");
-			BatchFileContents.AppendLine($"\"{EnvVars.ToolchainCompilerPath}\" @\"{ResponseFile.FullName}\"");
+			BatchFileContents.AppendLine($"\"{NormalizeCommandLinePath(EnvVars.ToolchainCompilerPath)}\" @\"{NormalizeCommandLinePath(ResponseFile)}\"");
 			BatchFileContents.AppendLine(TouchActionCommand);
 
 			FileItem BatchFile = Graph.CreateIntermediateTextFile(OutputFile.ChangeExtension(".bat"), BatchFileContents.ToString());
@@ -2650,22 +2641,21 @@ namespace UnrealBuildTool
 		public override CppCompileEnvironment CreateSharedResponseFile(CppCompileEnvironment CompileEnvironment, FileReference OutResponseFile, IActionGraphBuilder Graph)
 		{
 			CppCompileEnvironment NewCompileEnvironment = new CppCompileEnvironment(CompileEnvironment);
-			CppRootPaths? rootPaths = Target.WindowsPlatform.Compiler.IsClang() && NewCompileEnvironment.VfsOverlayFile != null ? NewCompileEnvironment.RootPaths : null;
 			List<string> Arguments = new List<string>();
 
 			foreach (DirectoryReference IncludePath in NewCompileEnvironment.UserIncludePaths)
 			{
-				AddIncludePath(Arguments, IncludePath, Target.WindowsPlatform.Compiler, rootPaths);
+				AddIncludePath(Arguments, IncludePath, Target.WindowsPlatform.Compiler, NewCompileEnvironment.RootPaths);
 			}
 
 			foreach (DirectoryReference IncludePath in NewCompileEnvironment.SystemIncludePaths)
 			{
-				AddSystemIncludePath(Arguments, IncludePath, Target.WindowsPlatform.Compiler, rootPaths);
+				AddSystemIncludePath(Arguments, IncludePath, Target.WindowsPlatform.Compiler, NewCompileEnvironment.RootPaths);
 			}
 
 			foreach (DirectoryReference IncludePath in EnvVars.IncludePaths)
 			{
-				AddSystemIncludePath(Arguments, IncludePath, Target.WindowsPlatform.Compiler, rootPaths);
+				AddSystemIncludePath(Arguments, IncludePath, Target.WindowsPlatform.Compiler, NewCompileEnvironment.RootPaths);
 			}
 
 			// Stash shared include paths for validation purposes
@@ -2730,6 +2720,8 @@ namespace UnrealBuildTool
 
 		public override FileItem LinkFiles(LinkEnvironment LinkEnvironment, bool bBuildImportLibraryOnly, IActionGraphBuilder Graph)
 		{
+			CppRootPaths RootPaths = LinkEnvironment.RootPaths;
+			
 			if (LinkEnvironment.bIsBuildingDotNetAssembly)
 			{
 				return FileItem.GetItemByFileReference(LinkEnvironment.OutputFilePath);
@@ -2809,11 +2801,11 @@ namespace UnrealBuildTool
 				// Add the library paths to the argument list.
 				foreach (DirectoryReference LibraryPath in LinkEnvironment.SystemLibraryPaths)
 				{
-					Arguments.Add($"/LIBPATH:\"{NormalizeCommandLinePath(LibraryPath)}\"");
+					Arguments.Add($"/LIBPATH:\"{NormalizeCommandLinePath(LibraryPath, RootPaths)}\"");
 				}
 				foreach (DirectoryReference LibraryPath in EnvVars.LibraryPaths)
 				{
-					Arguments.Add($"/LIBPATH:\"{NormalizeCommandLinePath(LibraryPath)}\"");
+					Arguments.Add($"/LIBPATH:\"{NormalizeCommandLinePath(LibraryPath, RootPaths)}\"");
 				}
 
 				// Add the excluded default libraries to the argument list.
@@ -2871,7 +2863,7 @@ namespace UnrealBuildTool
 			List<string> InputFileNames = new List<string>();
 			foreach (FileItem InputFile in LinkEnvironment.InputFiles)
 			{
-				InputFileNames.Add(String.Format("\"{0}\"", NormalizeCommandLinePath(InputFile)));
+				InputFileNames.Add(String.Format("\"{0}\"", NormalizeCommandLinePath(InputFile, RootPaths)));
 				PrerequisiteItems.Add(InputFile);
 			}
 
@@ -2879,7 +2871,7 @@ namespace UnrealBuildTool
 			{
 				foreach (FileReference Library in LinkEnvironment.Libraries)
 				{
-					InputFileNames.Add(String.Format("\"{0}\"", NormalizeCommandLinePath(Library)));
+					InputFileNames.Add(String.Format("\"{0}\"", NormalizeCommandLinePath(Library, RootPaths)));
 					PrerequisiteItems.Add(FileItem.GetItemByFileReference(Library));
 				}
 				foreach (string SystemLibrary in LinkEnvironment.SystemLibraries)
@@ -2891,7 +2883,7 @@ namespace UnrealBuildTool
 				{
 					PrerequisiteItems.Add(NatvisFile);
 					Arguments.Add(String.Format("/NATVIS:\"{0}\"",
-						NormalizeCommandLinePath(NatvisFile)));
+						NormalizeCommandLinePath(NatvisFile, RootPaths)));
 				}
 
 				if (Target.WindowsPlatform.ManifestFile != null)
@@ -2903,7 +2895,7 @@ namespace UnrealBuildTool
 			Arguments.AddRange(InputFileNames);
 
 			// Add the output file to the command-line.
-			Arguments.Add($"/OUT:\"{NormalizeCommandLinePath(OutputFile)}\"");
+			Arguments.Add($"/OUT:\"{NormalizeCommandLinePath(OutputFile, RootPaths)}\"");
 
 			// For import libraries and exports generated by cross-referenced builds, we don't track output files. VS 15.3+ doesn't touch timestamps for libs
 			// and exp files with no modifications, breaking our dependency checking, but incremental linking will fall back to a full link if we delete it.
@@ -2948,7 +2940,7 @@ namespace UnrealBuildTool
 						}
 
 						FileItem PDBFile = FileItem.GetItemByFileReference(PDBFilePath);
-						Arguments.Add($"/PDB:\"{NormalizeCommandLinePath(PDBFilePath)}\"");
+						Arguments.Add($"/PDB:\"{NormalizeCommandLinePath(PDBFilePath, RootPaths)}\"");
 						ProducedItems.Add(PDBFile);
 					}
 
diff --git a/Engine/Source/Programs/UnrealBuildTool/System/CppCompileEnvironment.cs b/Engine/Source/Programs/UnrealBuildTool/System/CppCompileEnvironment.cs
--- a/Engine/Source/Programs/UnrealBuildTool/System/CppCompileEnvironment.cs
+++ b/Engine/Source/Programs/UnrealBuildTool/System/CppCompileEnvironment.cs
@@ -191,6 +191,11 @@ namespace UnrealBuildTool
 		/// </summary>
 		PlatformSDK,
 
+		/// <summary>
+		/// Analyzer folder (if needed)
+		/// </summary>
+		Analyzer,
+
 		/// <summary>
 		/// Extra folder (if needed)
 		/// </summary>
@@ -255,16 +260,19 @@ namespace UnrealBuildTool
 		public CppRootPaths(BinaryArchiveReader Reader)
 		{
 			Lookup = new(Reader.ReadDictionary(() => (CppRootPathFolder)Reader.ReadInt(), () => Reader.ReadCompactDirectoryReference())!, null);
+			bUseVfs = Reader.ReadBool();
 		}
 
 		public CppRootPaths(CppRootPaths Other)
 		{
 			Lookup = new(Other.Lookup);
+			bUseVfs = Other.bUseVfs;
 		}
 
 		public void Write(BinaryArchiveWriter Writer)
 		{
 			Writer.WriteDictionary(Lookup, Key => Writer.WriteInt((int)Key), Value => Writer.WriteCompactDirectoryReference(Value));
+			Writer.WriteBool(bUseVfs);
 		}
 
 		public DirectoryReference this[CppRootPathFolder key]
@@ -282,7 +290,32 @@ namespace UnrealBuildTool
 			return Lookup.GetEnumerator();
 		}
 
-		private readonly SortedDictionary<CppRootPathFolder, DirectoryReference> Lookup;
+		private SortedDictionary<CppRootPathFolder, DirectoryReference> Lookup { get; }
+		public bool bUseVfs { get; set; }
+
+		/// <summary>
+		/// Get a virtual filesystem overlay path if in use for a local path
+		/// </summary>
+		/// <param name="reference">The path to convert</param>
+		/// <param name="vfsPath">Out param, the converted path</param>
+		/// <returns>True if the local path can be mapped to a virtual path when enabled, false otherwise</returns>
+		public bool GetVfsOverlayPath(FileSystemReference reference, [NotNullWhen(true)] out string? vfsPath)
+		{
+			if (bUseVfs)
+			{
+				foreach (KeyValuePair<CppRootPathFolder, DirectoryReference> item in Lookup)
+				{
+					if (!reference.IsUnderDirectory(item.Value))
+					{
+						continue;
+					}
+					vfsPath = item.Key.AsVfsString(item.Value, reference);
+					return true;
+				}
+			}
+			vfsPath = null;
+			return false;
+		}
 	}
 
 
@@ -715,11 +748,6 @@ namespace UnrealBuildTool
 		/// </summary>
 		public List<FileItem> ForceIncludeFiles = new List<FileItem>();
 
-		/// <summary>
-		/// Path to a virtual file system overlay file if in use
-		/// </summary>
-		public FileItem? VfsOverlayFile = null;
-
 		/// <summary>
 		/// List of files that need to be up to date before compile can proceed
 		/// </summary>
@@ -974,7 +1002,6 @@ namespace UnrealBuildTool
 			ModuleInterfacePaths = new HashSet<DirectoryReference>(Other.ModuleInterfacePaths);
 			bCheckSystemHeadersForModification = Other.bCheckSystemHeadersForModification;
 			ForceIncludeFiles.AddRange(Other.ForceIncludeFiles);
-			VfsOverlayFile = Other.VfsOverlayFile;
 			AdditionalPrerequisites.AddRange(Other.AdditionalPrerequisites);
 			ExtraGeneratedCPPFileTypes = Other.ExtraGeneratedCPPFileTypes;
 			CollidingNames = Other.CollidingNames;
@@ -1025,29 +1052,5 @@ namespace UnrealBuildTool
 
 			return false;
 		}
-
-		/// <summary>
-		/// Get a virtual filesystem overlay path if in use for a local path
-		/// </summary>
-		/// <param name="reference">The path to convert</param>
-		/// <param name="vfsPath">Out param, the converted path</param>
-		/// <returns>True if the local path can be mapped to a virtual path when enabled, false otherwise</returns>
-		public bool GetVfsOverlayPath(FileSystemReference reference, [NotNullWhen(true)] out string? vfsPath)
-		{
-			if (VfsOverlayFile != null)
-			{
-				foreach (KeyValuePair<CppRootPathFolder, DirectoryReference> item in RootPaths)
-				{
-					if (!reference.IsUnderDirectory(item.Value))
-					{
-						continue;
-					}
-					vfsPath = item.Key.AsVfsString(item.Value, reference);
-					return true;
-				}
-			}
-			vfsPath = null;
-			return false;
-		}
 	}
 }
diff --git a/Engine/Source/Programs/UnrealBuildTool/System/Unity.cs b/Engine/Source/Programs/UnrealBuildTool/System/Unity.cs
--- a/Engine/Source/Programs/UnrealBuildTool/System/Unity.cs
+++ b/Engine/Source/Programs/UnrealBuildTool/System/Unity.cs
@@ -280,7 +280,7 @@ namespace UnrealBuildTool
 				foreach (FileItem SourceFile in UnityFile.Files)
 				{
 					string SourceFileString = SourceFile.AbsolutePath;
-					if (CompileEnvironment.GetVfsOverlayPath(SourceFile.Location, out string? vfsPath))
+					if (CompileEnvironment.RootPaths.GetVfsOverlayPath(SourceFile.Location, out string? vfsPath))
 					{
 						SourceFileString = vfsPath;
 					}
diff --git a/Engine/Source/Programs/UnrealBuildTool/ToolChain/AppleToolChain.cs b/Engine/Source/Programs/UnrealBuildTool/ToolChain/AppleToolChain.cs
--- a/Engine/Source/Programs/UnrealBuildTool/ToolChain/AppleToolChain.cs
+++ b/Engine/Source/Programs/UnrealBuildTool/ToolChain/AppleToolChain.cs
@@ -613,7 +613,7 @@ namespace UnrealBuildTool
 
 			if (CompileEnvironment.ForceIncludeFiles.Count() > 0)
 			{
-				Arguments.Add($"-import-objc-header \"{NormalizeCommandLinePath(CompileEnvironment.ForceIncludeFiles.First(), CompileEnvironment)}\"");
+				Arguments.Add($"-import-objc-header \"{NormalizeCommandLinePath(CompileEnvironment.ForceIncludeFiles.First(), CompileEnvironment.RootPaths)}\"");
 			}
 
 			// tell swift to define PLATFORM_FOO, and tell C++ to define PLATFORM_FOO=1
diff --git a/Engine/Source/Programs/UnrealBuildTool/ToolChain/ClangToolChain.cs b/Engine/Source/Programs/UnrealBuildTool/ToolChain/ClangToolChain.cs
--- a/Engine/Source/Programs/UnrealBuildTool/ToolChain/ClangToolChain.cs
+++ b/Engine/Source/Programs/UnrealBuildTool/ToolChain/ClangToolChain.cs
@@ -509,22 +509,27 @@ namespace UnrealBuildTool
 
 		protected virtual string GetUserIncludePathArgument(DirectoryReference IncludePath, CppCompileEnvironment CompileEnvironment)
 		{
-			return $"-I\"{NormalizeCommandLinePath(IncludePath, CompileEnvironment)}\"";
+			return $"-I\"{NormalizeCommandLinePath(IncludePath, CompileEnvironment.RootPaths)}\"";
 		}
 
-		protected virtual string GetSystemIncludePathArgument(DirectoryReference IncludePath, CppCompileEnvironment? CompileEnvironment)
+		protected virtual string GetSystemIncludePathArgument(DirectoryReference IncludePath, CppCompileEnvironment CompileEnvironment)
 		{
-			return $"-isystem\"{NormalizeCommandLinePath(IncludePath, CompileEnvironment)}\"";
+			return $"-isystem\"{NormalizeCommandLinePath(IncludePath, CompileEnvironment.RootPaths)}\"";
 		}
 
-		protected virtual DirectoryReference GetResourceDir(CppCompileEnvironment? CompileEnvironment)
+		protected virtual DirectoryReference GetResourceDir(CppCompileEnvironment CompileEnvironment)
 		{
 			return DirectoryReference.Combine(Info.BasePath!, "lib", "clang", Info.ClangVersion.Major.ToString());
 		}
 
 		protected virtual void GetCompileArguments_IncludePaths(CppCompileEnvironment CompileEnvironment, List<string> Arguments)
 		{
-			Arguments.AddRange(CompileEnvironment.UserIncludePaths.Select(IncludePath => GetUserIncludePathArgument(IncludePath, CompileEnvironment)));
+            if (CompileEnvironment.RootPaths.bUseVfs)
+            {
+                Arguments.Add($"-resource-dir=\"{NormalizeCommandLinePath(GetResourceDir(CompileEnvironment), CompileEnvironment.RootPaths)}\"");
+            }
+
+            Arguments.AddRange(CompileEnvironment.UserIncludePaths.Select(IncludePath => GetUserIncludePathArgument(IncludePath, CompileEnvironment)));
 			Arguments.AddRange(CompileEnvironment.SystemIncludePaths.Select(IncludePath => GetSystemIncludePathArgument(IncludePath, CompileEnvironment)));
 
 			if (ShowIncludes)
@@ -550,7 +555,7 @@ namespace UnrealBuildTool
 
 		protected virtual string GetForceIncludeFileArgument(FileReference ForceIncludeFile, CppCompileEnvironment CompileEnvironment)
 		{
-			return $"-include \"{NormalizeCommandLinePath(ForceIncludeFile, CompileEnvironment)}\"";
+			return $"-include \"{NormalizeCommandLinePath(ForceIncludeFile, CompileEnvironment.RootPaths)}\"";
 		}
 
 		protected virtual string GetForceIncludeFileArgument(FileItem ForceIncludeFile, CppCompileEnvironment CompileEnvironment)
@@ -560,7 +565,7 @@ namespace UnrealBuildTool
 
 		protected virtual string GetIncludePCHFileArgument(FileReference IncludePCHFile, CppCompileEnvironment CompileEnvironment)
 		{
-			return $"-include-pch \"{NormalizeCommandLinePath(IncludePCHFile, CompileEnvironment)}\"";
+			return $"-include-pch \"{NormalizeCommandLinePath(IncludePCHFile, CompileEnvironment.RootPaths)}\"";
 		}
 
 		protected virtual string GetIncludePCHFileArgument(FileItem IncludePCHFile, CppCompileEnvironment CompileEnvironment)
@@ -586,24 +591,27 @@ namespace UnrealBuildTool
 
 		protected virtual string GetSourceFileArgument(FileItem SourceFile, CppCompileEnvironment CompileEnvironment)
 		{
-			// Only pch seems to be ok with getting a vfs path.. so don't provide an environment if not creating pch
-			CppCompileEnvironment? Env = CompileEnvironment.PrecompiledHeaderAction == PrecompiledHeaderAction.Create ? CompileEnvironment : null;
-			return $"\"{NormalizeCommandLinePath(SourceFile, Env)}\"";
+			return $"\"{NormalizeCommandLinePath(SourceFile, CompileEnvironment.RootPaths)}\"";
 		}
 
-		protected virtual string GetOutputFileArgument(FileItem OutputFile)
+		protected virtual string GetOutputFileArgument(FileItem OutputFile, CppRootPaths RootPaths)
 		{
-			return $"-o \"{NormalizeCommandLinePath(OutputFile)}\"";
+			return $"-o \"{NormalizeCommandLinePath(OutputFile, RootPaths)}\"";
 		}
 
-		protected virtual string GetDepencenciesListFileArgument(FileItem DependencyListFile)
+        protected virtual string GetOutputFileArgument(FileItem OutputFile, CppCompileEnvironment CompileEnvironment)
+        {
+            return GetOutputFileArgument(OutputFile, CompileEnvironment.RootPaths);
+        }
+
+        protected virtual string GetDepencenciesListFileArgument(FileItem DependencyListFile, CppCompileEnvironment CompileEnvironment)
 		{
-			return $"-MD -MF\"{NormalizeCommandLinePath(DependencyListFile)}\"";
+			return $"-MD -MF\"{NormalizeCommandLinePath(DependencyListFile, CompileEnvironment.RootPaths)}\"";
 		}
 
-		protected virtual string GetPreprocessDepencenciesListFileArgument(FileItem DependencyListFile)
+		protected virtual string GetPreprocessDepencenciesListFileArgument(FileItem DependencyListFile, CppCompileEnvironment CompileEnvironment)
 		{
-			return $"-M -MF\"{NormalizeCommandLinePath(DependencyListFile)}\"";
+			return $"-M -MF\"{NormalizeCommandLinePath(DependencyListFile, CompileEnvironment.RootPaths)}\"";
 		}
 
 		protected virtual string GetResponseFileArgument(FileItem ResponseFile)
@@ -694,7 +702,7 @@ namespace UnrealBuildTool
 			{
 				Log.TraceInformationOnce("Enabling Profile Guided Optimization (PGO). Linking will take a while.");
 				DirectoryReference? PGODir = DirectoryReference.FromString(CompileEnvironment.PGODirectory!);
-				Arguments.Add($"-fprofile-instr-use=\"{NormalizeCommandLinePath(DirectoryReference.Combine(PGODir!, CompileEnvironment.PGOFilenamePrefix!), CompileEnvironment)}\"");
+				Arguments.Add($"-fprofile-instr-use=\"{NormalizeCommandLinePath(DirectoryReference.Combine(PGODir!, CompileEnvironment.PGOFilenamePrefix!), CompileEnvironment.RootPaths)}\"");
 			}
 			else if (CompileEnvironment.bPGOProfile)
 			{
@@ -896,15 +904,6 @@ namespace UnrealBuildTool
 			}
 		}
 
-		protected virtual void GetCompileArguments_VfsOverlay(CppCompileEnvironment CompileEnvironment, List<string> Arguments)
-		{
-			if (CompileEnvironment.VfsOverlayFile != null)
-			{
-				Arguments.Add($"-resource-dir=\"{NormalizeCommandLinePath(GetResourceDir(CompileEnvironment), CompileEnvironment)}\"");
-				Arguments.Add($"-Xclang -ivfsoverlay\"{NormalizeCommandLinePath(CompileEnvironment.VfsOverlayFile)}\"");
-			}
-		}
-
 		/// <summary>
 		/// Common compile arguments for all files in a module.
 		/// Override and call base.GetCompileArguments_Global() in derived classes.
@@ -969,9 +968,6 @@ namespace UnrealBuildTool
 		/// <returns>Path to the target file (such as .o)</returns>
 		protected virtual FileItem GetCompileArguments_FileType(CppCompileEnvironment CompileEnvironment, FileItem SourceFile, DirectoryReference OutputDir, List<string> Arguments, Action CompileAction, CPPOutput CompileResult)
 		{
-			// Add virtual file system overlay arguments if enabled. Will be inserted as the first entry in the list
-			GetCompileArguments_VfsOverlay(CompileEnvironment, Arguments);
-
 			// Add the additional response files
 			foreach (FileItem AdditionalResponseFile in CompileEnvironment.AdditionalResponseFiles)
 			{
@@ -1089,7 +1085,7 @@ namespace UnrealBuildTool
 			if (!PreprocessDepends)
 			{
 				FileItem DependencyListFile = FileItem.GetItemByFileReference(FileReference.Combine(OutputDir, GetFileNameFromExtension(FileName, ".d")));
-				Arguments.Add(GetDepencenciesListFileArgument(DependencyListFile));
+				Arguments.Add(GetDepencenciesListFileArgument(DependencyListFile, CompileEnvironment));
 				CompileAction.DependencyListFile = DependencyListFile;
 				CompileAction.ProducedItems.Add(DependencyListFile);
 			}
@@ -1105,7 +1101,7 @@ namespace UnrealBuildTool
 				}
 
 				// Add the parameters needed to compile the output file to the command-line.
-				Arguments.Add(GetOutputFileArgument(OutputFile));
+				Arguments.Add(GetOutputFileArgument(OutputFile, CompileEnvironment));
 			}
 
 			// Add additional arguments to the argument list, must be the final arguments added
@@ -1331,12 +1327,12 @@ namespace UnrealBuildTool
 				PreprocessGlobalArguments.Remove("-c");
 
 				FileItem DependencyListFile = FileItem.GetItemByFileReference(FileReference.Combine(OutputDir, Path.GetFileName(SourceFile.AbsolutePath) + ".d"));
-				PreprocessFileArguments.Add(GetPreprocessDepencenciesListFileArgument(DependencyListFile));
+				PreprocessFileArguments.Add(GetPreprocessDepencenciesListFileArgument(DependencyListFile, CompileEnvironment));
 				PrepassAction.DependencyListFile = DependencyListFile;
 				PrepassAction.ProducedItems.Add(DependencyListFile);
 
 				PreprocessFileArguments.Remove("-ftime-trace");
-				PreprocessFileArguments.Remove(GetOutputFileArgument(CompileAction.ProducedItems.First()));
+				PreprocessFileArguments.Remove(GetOutputFileArgument(CompileAction.ProducedItems.First(), CompileEnvironment));
 
 				PrepassAction.DeleteItems.UnionWith(PrepassAction.ProducedItems);
 
diff --git a/Engine/Source/Programs/UnrealBuildTool/ToolChain/ISPCToolChain.cs b/Engine/Source/Programs/UnrealBuildTool/ToolChain/ISPCToolChain.cs
--- a/Engine/Source/Programs/UnrealBuildTool/ToolChain/ISPCToolChain.cs
+++ b/Engine/Source/Programs/UnrealBuildTool/ToolChain/ISPCToolChain.cs
@@ -363,26 +363,22 @@ namespace UnrealBuildTool
 		/// Normalize a path for use in a command line, making it relative to Engine/Source if under the root directory
 		/// </summary>
 		/// <param name="Reference">The FileSystemReference to normalize</param>
-		/// <param name="CompileEnvironment">The environment to use for normalization</param>
+		/// <param name="RootPaths">The root paths to use for normalization</param>
 		/// <returns>Normalized path as a string</returns>
-		protected virtual string NormalizeCommandLinePath(FileSystemReference Reference, CppCompileEnvironment? CompileEnvironment = null)
+		protected virtual string NormalizeCommandLinePath(FileSystemReference Reference, CppRootPaths RootPaths)
 		{
-			if (CompileEnvironment?.GetVfsOverlayPath(Reference, out string? vfsPath) == true)
-			{
-				return vfsPath;
-			}
-			return NormalizeCommandLinePath(Reference);
+			return RootPaths.GetVfsOverlayPath(Reference, out string? vfsPath) ? vfsPath : NormalizeCommandLinePath(Reference);
 		}
 
 		/// <summary>
 		/// Normalize a path for use in a command line, making it relative if under the Root Directory
 		/// </summary>
 		/// <param name="Item">The FileItem to normalize</param>
-		/// <param name="CompileEnvironment">The environment to use for normalization</param>
+		/// <param name="RootPaths">The root paths to use for normalization</param>
 		/// <returns>Normalized path as a string</returns>
-		protected virtual string NormalizeCommandLinePath(FileItem Item, CppCompileEnvironment? CompileEnvironment = null)
+		protected virtual string NormalizeCommandLinePath(FileItem Item, CppRootPaths RootPaths)
 		{
-			return NormalizeCommandLinePath(Item.Location, CompileEnvironment);
+			return NormalizeCommandLinePath(Item.Location, RootPaths);
 		}
 
 		protected override CPPOutput GenerateISPCHeaders(CppCompileEnvironment CompileEnvironment, IEnumerable<FileItem> InputFiles, DirectoryReference OutputDir, IActionGraphBuilder Graph)
diff --git a/Engine/Source/Programs/UnrealBuildTool/ToolChain/IWYUToolChain.cs b/Engine/Source/Programs/UnrealBuildTool/ToolChain/IWYUToolChain.cs
--- a/Engine/Source/Programs/UnrealBuildTool/ToolChain/IWYUToolChain.cs
+++ b/Engine/Source/Programs/UnrealBuildTool/ToolChain/IWYUToolChain.cs
@@ -43,6 +43,9 @@ namespace UnrealBuildTool
 			LinuxPlatformSDK PlatformSDK = new LinuxPlatformSDK(Logger);
 			DirectoryReference? BaseLinuxPath = PlatformSDK.GetBaseLinuxPathForArchitecture(LinuxPlatform.DefaultHostArchitecture);
 
+			SourceFileMetadataCache MetadataCache = SourceFileMetadataCache.CreateHierarchy(null, Logger);
+			CppCompileEnvironment CrossCompileEnvironment = new CppCompileEnvironment(UnrealTargetPlatform.Linux, CppConfiguration.Development, new(LinuxPlatform.DefaultHostArchitecture), MetadataCache);
+
 			CrossCompilingArguments.Add($"--sysroot=\"{NormalizeCommandLinePath(BaseLinuxPath!)}\"");
 
 			FileReference ClangPath = FileReference.Combine(BaseLinuxPath!, "bin", $"clang++{BuildHostPlatform.Current.BinarySuffix}");
@@ -51,7 +54,7 @@ namespace UnrealBuildTool
 			// starting with clang 16.x the directory naming changed to include major version only
 			string ClangVersionString = (CompilerToolChainInfo.ClangVersion.Major >= 16) ? CompilerToolChainInfo.ClangVersion.Major.ToString() : CompilerToolChainInfo.ClangVersion.ToString();
 			DirectoryReference SystemPath = DirectoryReference.Combine(BaseLinuxPath!, "lib", "clang", ClangVersionString, "include");
-			CrossCompilingArguments.Add(GetSystemIncludePathArgument(SystemPath, null));
+			CrossCompilingArguments.Add(GetSystemIncludePathArgument(SystemPath, CrossCompileEnvironment));
 
 			string DevPath = ""; //@"include-what-you-use-0.19\vs_projects\bin\RelWithDebInfo";
 
@@ -134,15 +137,15 @@ namespace UnrealBuildTool
         // ".." than there are directories to reduce. This silently fails in odd ways
 
         /// <inheritdoc/>
-        protected override string NormalizeCommandLinePath(FileSystemReference Reference, CppCompileEnvironment? CompileEnvironment)
+        protected override string NormalizeCommandLinePath(FileSystemReference Reference, CppRootPaths RootPath)
 		{
 			return Reference.FullName.Replace('\\', '/');
 		}
 
         /// <inheritdoc/>
-		protected override string NormalizeCommandLinePath(FileItem Item, CppCompileEnvironment? CompileEnvironment)
+		protected override string NormalizeCommandLinePath(FileItem Item, CppRootPaths RootPath)
 		{
-			return NormalizeCommandLinePath(Item.Location, CompileEnvironment);
+			return NormalizeCommandLinePath(Item.Location, RootPath);
 		}
 
 		// Code below here is mostly copy-pasted from LinuxToolChain.
diff --git a/Engine/Source/Programs/UnrealBuildTool/ToolChain/UEToolChain.cs b/Engine/Source/Programs/UnrealBuildTool/ToolChain/UEToolChain.cs
--- a/Engine/Source/Programs/UnrealBuildTool/ToolChain/UEToolChain.cs
+++ b/Engine/Source/Programs/UnrealBuildTool/ToolChain/UEToolChain.cs
@@ -274,6 +274,9 @@ namespace UnrealBuildTool
 
 		public virtual void SetUpGlobalEnvironment(ReadOnlyTargetRules Target, CppCompileEnvironment GlobalCompileEnvironment, LinkEnvironment GlobalLinkEnvironment)
 		{
+			GlobalCompileEnvironment.RootPaths.bUseVfs = Target.bCompilerVFS;
+			GlobalLinkEnvironment.RootPaths.bUseVfs = Target.bCompilerVFS;
+
 			GlobalCompileEnvironment.RootPaths[CppRootPathFolder.Engine] = Unreal.EngineDirectory;
 			GlobalLinkEnvironment.RootPaths[CppRootPathFolder.Engine] = Unreal.EngineDirectory;
 
