[UBA]
* Removed spammy log entry about creating library

diff --git a/Engine/Source/Programs/UnrealBuildTool/Executors/UnrealBuildAccelerator/UBAExecutor.cs b/Engine/Source/Programs/UnrealBuildTool/Executors/UnrealBuildAccelerator/UBAExecutor.cs
--- a/Engine/Source/Programs/UnrealBuildTool/Executors/UnrealBuildAccelerator/UBAExecutor.cs
+++ b/Engine/Source/Programs/UnrealBuildTool/Executors/UnrealBuildAccelerator/UBAExecutor.cs
@@ -239,7 +239,8 @@ namespace UnrealBuildTool
 
 		private static void RemoveLogLineSpam(List<string> logLines)
 		{
-			logLines.RemoveAll((line) => (line.StartsWith("   Creating library ", StringComparison.OrdinalIgnoreCase) && line.EndsWith(".exp", StringComparison.OrdinalIgnoreCase)) || line.EndsWith("file(s) copied.", StringComparison.OrdinalIgnoreCase));
+			logLines.RemoveAll((line) => 
+			((line.StartsWith("   Creating library ", StringComparison.OrdinalIgnoreCase) || line.StartsWith("   Creating object ", StringComparison.OrdinalIgnoreCase)) && line.EndsWith(".exp", StringComparison.OrdinalIgnoreCase)) || line.EndsWith("file(s) copied.", StringComparison.OrdinalIgnoreCase));
 		}
 
 		class UBAActionArtifactCache : IActionArtifactCache
