[UBA]
* Fixed so more paths are virtualized when compilervfs is enabled

diff --git a/Engine/Source/Programs/UnrealBuildTool/ToolChain/ISPCToolChain.cs b/Engine/Source/Programs/UnrealBuildTool/ToolChain/ISPCToolChain.cs
--- a/Engine/Source/Programs/UnrealBuildTool/ToolChain/ISPCToolChain.cs
+++ b/Engine/Source/Programs/UnrealBuildTool/ToolChain/ISPCToolChain.cs
@@ -390,6 +390,8 @@ namespace UnrealBuildTool
 				return Result;
 			}
 
+			CppRootPaths RootPaths = CompileEnvironment.RootPaths;
+
 			List<string> CompileTargets = GetISPCCompileTargets(CompileEnvironment.Platform, CompileEnvironment.Architecture);
 
 			List<string> GlobalArguments = new List<string>();
@@ -424,13 +426,13 @@ namespace UnrealBuildTool
 			// Because ISPC response files don't support white space in arguments, paths with white space need to be passed to the command line directly.
 			foreach (DirectoryReference IncludePath in CompileEnvironment.UserIncludePaths)
 			{
-				GlobalArguments.Add($"-I\"{NormalizeCommandLinePath(IncludePath)}\"");
+				GlobalArguments.Add($"-I\"{NormalizeCommandLinePath(IncludePath, RootPaths)}\"");
 			}
 
 			// System include paths.
 			foreach (DirectoryReference SystemIncludePath in CompileEnvironment.SystemIncludePaths)
 			{
-				GlobalArguments.Add($"-I\"{NormalizeCommandLinePath(SystemIncludePath)}\"");
+				GlobalArguments.Add($"-I\"{NormalizeCommandLinePath(SystemIncludePath, RootPaths)}\"");
 			}
 
 			// Preprocessor definitions.
@@ -467,7 +469,7 @@ namespace UnrealBuildTool
 				List<string> Arguments = new List<string>();
 
 				// Add the ISPC obj file as a prerequisite of the action.
-				Arguments.Add($"\"{NormalizeCommandLinePath(ISPCFile)}\"");
+				Arguments.Add($"\"{NormalizeCommandLinePath(ISPCFile, RootPaths)}\"");
 
 				// Add the ISPC h file to the produced item list.
 				FileItem ISPCIncludeHeaderFile = FileItem.GetItemByFileReference(
@@ -505,11 +507,11 @@ namespace UnrealBuildTool
 				}
 
 				// Add the ISPC file to be compiled.
-				Arguments.Add($"-h \"{NormalizeCommandLinePath(ISPCIncludeHeaderFile)}\"");
+				Arguments.Add($"-h \"{NormalizeCommandLinePath(ISPCIncludeHeaderFile, RootPaths)}\"");
 
 				// Generate the included header dependency list
 				FileItem DependencyListFile = FileItem.GetItemByFileReference(FileReference.Combine(OutputDir, Path.GetFileName(ISPCFile.AbsolutePath) + ".txt"));
-				Arguments.Add($"-MMM \"{NormalizeCommandLinePath(DependencyListFile)}\"");
+				Arguments.Add($"-MMM \"{NormalizeCommandLinePath(DependencyListFile, RootPaths)}\"");
 				CompileAction.DependencyListFile = DependencyListFile;
 				CompileAction.ProducedItems.Add(DependencyListFile);
 
@@ -554,6 +556,8 @@ namespace UnrealBuildTool
 				return Result;
 			}
 
+			CppRootPaths RootPaths = CompileEnvironment.RootPaths;
+
 			List<string> CompileTargets = GetISPCCompileTargets(CompileEnvironment.Platform, CompileEnvironment.Architecture);
 
 			List<string> GlobalArguments = new List<string>();
@@ -624,13 +628,13 @@ namespace UnrealBuildTool
 			// Include paths. Don't use AddIncludePath() here, since it uses the full path and exceeds the max command line length.
 			foreach (DirectoryReference IncludePath in CompileEnvironment.UserIncludePaths)
 			{
-				GlobalArguments.Add($"-I\"{NormalizeCommandLinePath(IncludePath)}\"");
+				GlobalArguments.Add($"-I\"{NormalizeCommandLinePath(IncludePath, RootPaths)}\"");
 			}
 
 			// System include paths.
 			foreach (DirectoryReference SystemIncludePath in CompileEnvironment.SystemIncludePaths)
 			{
-				GlobalArguments.Add($"-I\"{NormalizeCommandLinePath(SystemIncludePath)}\"");
+				GlobalArguments.Add($"-I\"{NormalizeCommandLinePath(SystemIncludePath, RootPaths)}\"");
 			}
 
 			// Preprocessor definitions.
@@ -646,7 +650,7 @@ namespace UnrealBuildTool
 			foreach (FileItem ISPCFile in InputFiles)
 			{
 				Action CompileAction = Graph.CreateAction(ActionType.Compile);
-				CompileAction.RootPaths = CompileEnvironment.RootPaths;
+				CompileAction.RootPaths = RootPaths;
 				CompileAction.CommandDescription = $"Compile [{ISPCArch}]";
 				CompileAction.WorkingDirectory = Unreal.EngineSourceDirectory;
 				CompileAction.CommandPath = new FileReference(GetISPCHostCompilerPath(BuildHostPlatform.Current.Platform));
@@ -671,7 +675,7 @@ namespace UnrealBuildTool
 				List<string> Arguments = new List<string>();
 
 				// Add the ISPC file to be compiled.
-				Arguments.Add($"\"{NormalizeCommandLinePath(ISPCFile)}\"");
+				Arguments.Add($"\"{NormalizeCommandLinePath(ISPCFile, RootPaths)}\"");
 
 				List<FileItem> CompiledISPCObjFiles = new List<FileItem>();
 
@@ -727,7 +731,7 @@ namespace UnrealBuildTool
 				}
 
 				// Add the output ISPC obj file
-				Arguments.Add($"-o \"{NormalizeCommandLinePath(CompiledISPCObjFileNoISA)}\"");
+				Arguments.Add($"-o \"{NormalizeCommandLinePath(CompiledISPCObjFileNoISA, RootPaths)}\"");
 
 				// Generate the timing info
 				if (CompileEnvironment.bPrintTimingInfo)
@@ -758,7 +762,7 @@ namespace UnrealBuildTool
 					CompileAction.ProducedItems.Add(FileItem.GetItemByFileReference(FileReference.Combine(OutputDir, FileName + ".exi")));
 				}
 
-				CompileAction.CommandArguments = $"@\"{ResponseFileName}\"{AdditionalArguments}";
+				CompileAction.CommandArguments = $"@\"{NormalizeCommandLinePath(ResponseFileName)}\"{AdditionalArguments}";
 				CompileAction.PrerequisiteItems.Add(ResponseFileItem);
 
 				// Add the source file and its included files to the prerequisite item list.
@@ -786,10 +790,10 @@ namespace UnrealBuildTool
 							Action PostCompileAction = Graph.CreateAction(ActionType.Compile);
 
 							List<string> PostCompileArgs = new List<string>();
-							PostCompileArgs.Add($"\"{NormalizeCommandLinePath(CompiledBytecodeObjFile)}\"");
+							PostCompileArgs.Add($"\"{NormalizeCommandLinePath(CompiledBytecodeObjFile, RootPaths)}\"");
 							PostCompileArgs.Add("-c");
 							PostCompileArgs.AddRange(CommonArgs);
-							PostCompileArgs.Add($"-o \"{NormalizeCommandLinePath(FinalCompiledISPCObjFile)}\"");
+							PostCompileArgs.Add($"-o \"{NormalizeCommandLinePath(FinalCompiledISPCObjFile, RootPaths)}\"");
 
 							if (bMergeModules)
 							{
@@ -801,7 +805,7 @@ namespace UnrealBuildTool
 							// Write the args to a response file
 							FileReference PostCompileResponseFileName = GetResponseFileName(CompileEnvironment, FinalCompiledISPCObjFile);
 							FileItem PostCompileResponseFileItem = Graph.CreateIntermediateTextFile(PostCompileResponseFileName, PostCompileArgs.Select(x => Utils.ExpandVariables(x)));
-							PostCompileAction.CommandArguments = $"@\"{PostCompileResponseFileName}\"{AdditionalArguments}";
+							PostCompileAction.CommandArguments = $"@\"{NormalizeCommandLinePath(PostCompileResponseFileName)}\"{AdditionalArguments}";
 							PostCompileAction.PrerequisiteItems.Add(PostCompileResponseFileItem);
 
 							PostCompileAction.PrerequisiteItems.Add(CompiledBytecodeObjFile);
@@ -816,7 +820,7 @@ namespace UnrealBuildTool
 								CompileAction.CommandVersion = $"{CompileAction.CommandVersion} Compressed";
 							}
 
-							PostCompileAction.RootPaths = CompileEnvironment.RootPaths;
+							PostCompileAction.RootPaths = RootPaths;
 							PostCompileAction.ArtifactMode = ArtifactMode.Enabled;
 
 							// Disable remote execution to workaround mismatched case on XGE
