Move aditional properties to CppCompileWarning.

diff --git a/Engine/Source/Programs/UnrealBuildTool/Configuration/CppCompileWarnings.cs b/Engine/Source/Programs/UnrealBuildTool/Configuration/CppCompileWarnings.cs
--- a/Engine/Source/Programs/UnrealBuildTool/Configuration/CppCompileWarnings.cs
+++ b/Engine/Source/Programs/UnrealBuildTool/Configuration/CppCompileWarnings.cs
@@ -28,6 +28,22 @@ namespace UnrealBuildTool
 	/// </summary>
 	public class CppCompileWarnings
 	{
+		/*
+		 * To add a new compiler warning, you need to do the following:
+		 * 1. Copy the following stub code:
+		 *		public WarningLevel NEW_WARNING_LEVEL_NAME
+		 *		{
+		 *			get => ResolveWarning(null, ParentCppCompileWarnings?.NEW_WARNING_LEVEL_NAME, _NEW_WARNING_LEVEL_NAME);
+		 *			set => _NEW_WARNING_LEVEL_NAME = value;
+		 *		}
+		 *		private WarningLevel _NEW_WARNING_LEVEL_NAME;
+		 * 2. Update the ReadOnlyCppCompileWarnings class to include a forwarding property:
+		 *		public WarningLevel NEW_WARNING_LEVEL_NAME => _inner.NEW_WARNING_LEVEL_NAME;
+		 * 3. Set your default WarningLevel within appropraite ApplyDefaults(...) setting
+		 *		For Target defaults				- private void ApplyTargetDefaults()
+		 *		For CompileEnvironment defaults	- private void ApplyDefaults()
+		 * 4. Access the new property through the (TargetRules|ModuleRules|CppCompileEnvironment).CppCompileWarnings as needed
+		 */
 		#region -- Compiler Warning Settings --
 
 		/// <summary>
@@ -160,6 +176,26 @@ namespace UnrealBuildTool
 		}
 		private WarningLevel _moduleIncludeSubdirectoryWarningLevel;
 
+		/// <summary>
+		/// How to treat conflicts when a disabled plugin is being enabled by another plugin referencing it
+		/// </summary>
+		public WarningLevel DisablePluginsConflictWarningLevel
+		{
+			get => ResolveWarning(null, ParentCppCompileWarnings?.DisablePluginsConflictWarningLevel, _disablePluginsConflictWarningLevel);
+			set => _disablePluginsConflictWarningLevel = value;
+		}
+		private WarningLevel _disablePluginsConflictWarningLevel;
+
+		/// <summary>
+		/// Enable warnings for when there are .gen.cpp files that could be inlined in a matching handwritten cpp file
+		/// </summary>
+		public WarningLevel NonInlinedGenCppWarningLevel
+		{
+			get => ResolveWarning(null, ParentCppCompileWarnings?.NonInlinedGenCppWarningLevel, _nonInlinedGenCppWarningLevel);
+			set => _nonInlinedGenCppWarningLevel = value;
+		}
+		private WarningLevel _nonInlinedGenCppWarningLevel;
+
 		#endregion -- Rules Validation Warning Settings --
 
 		#region -- Constructors --
@@ -313,6 +349,9 @@ namespace UnrealBuildTool
 			_moduleIncludePathWarningLevel = WarningLevel.Off;
 			_moduleIncludePrivateWarningLevel = WarningLevel.Off;
 			_moduleIncludeSubdirectoryWarningLevel = WarningLevel.Off;
+
+			_disablePluginsConflictWarningLevel = WarningLevel.Default;
+			_nonInlinedGenCppWarningLevel = WarningLevel.Off;
 		}
 
 		/// <summary>
@@ -339,6 +378,9 @@ namespace UnrealBuildTool
 			_moduleIncludePathWarningLevel = WarningLevel.Off;
 			_moduleIncludePrivateWarningLevel = WarningLevel.Off;
 			_moduleIncludeSubdirectoryWarningLevel = WarningLevel.Off;
+
+			_disablePluginsConflictWarningLevel = WarningLevel.Default;
+			_nonInlinedGenCppWarningLevel = WarningLevel.Off;
 		}
 
 		#endregion -- Private API --
diff --git a/Engine/Source/Programs/UnrealBuildTool/Configuration/ReadOnlyCppCompileWarnings.cs b/Engine/Source/Programs/UnrealBuildTool/Configuration/ReadOnlyCppCompileWarnings.cs
--- a/Engine/Source/Programs/UnrealBuildTool/Configuration/ReadOnlyCppCompileWarnings.cs
+++ b/Engine/Source/Programs/UnrealBuildTool/Configuration/ReadOnlyCppCompileWarnings.cs
@@ -53,11 +53,14 @@ namespace UnrealBuildTool
 		#endregion -- Performance Warning Settings --
 
 		#region -- Rules Validation Warning Settings --
+		
 		public WarningLevel ModuleUnsupportedWarningLevel => _inner.ModuleUnsupportedWarningLevel;
 		public WarningLevel PluginModuleUnsupportedWarningLevel => _inner.PluginModuleUnsupportedWarningLevel;
 		public WarningLevel ModuleIncludePathWarningLevel => _inner.ModuleIncludePathWarningLevel;
 		public WarningLevel ModuleIncludePrivateWarningLevel => _inner.ModuleIncludePrivateWarningLevel;
 		public WarningLevel ModuleIncludeSubdirectoryWarningLevel => _inner.ModuleIncludeSubdirectoryWarningLevel;
+		public WarningLevel DisablePluginsConflictWarningLevel => _inner.DisablePluginsConflictWarningLevel;
+		public WarningLevel NonInlinedGenCppWarningLevel => _inner.NonInlinedGenCppWarningLevel;
 
 		#endregion -- Rules Validation Warning Settings --
 #pragma warning restore C1591
diff --git a/Engine/Source/Programs/UnrealBuildTool/Configuration/ReadOnlyTargetRules.cs b/Engine/Source/Programs/UnrealBuildTool/Configuration/ReadOnlyTargetRules.cs
--- a/Engine/Source/Programs/UnrealBuildTool/Configuration/ReadOnlyTargetRules.cs
+++ b/Engine/Source/Programs/UnrealBuildTool/Configuration/ReadOnlyTargetRules.cs
@@ -113,7 +113,7 @@ namespace UnrealBuildTool
 
 		public bool bAllowEnginePluginsEnabledByDefault => Inner.bAllowEnginePluginsEnabledByDefault;
 
-		public WarningLevel DisablePluginsConflictWarningLevel => Inner.DisablePluginsConflictWarningLevel;
+		public WarningLevel DisablePluginsConflictWarningLevel => Inner.CppCompileWarningSettings.DisablePluginsConflictWarningLevel;
 
 		public IReadOnlyList<string> BuildPlugins => Inner.BuildPlugins;
 
diff --git a/Engine/Source/Programs/UnrealBuildTool/Configuration/TargetRules.cs b/Engine/Source/Programs/UnrealBuildTool/Configuration/TargetRules.cs
--- a/Engine/Source/Programs/UnrealBuildTool/Configuration/TargetRules.cs
+++ b/Engine/Source/Programs/UnrealBuildTool/Configuration/TargetRules.cs
@@ -838,7 +838,11 @@ namespace UnrealBuildTool
 		/// <summary>
 		/// How to treat conflicts when a disabled plugin is being enabled by another plugin referencing it
 		/// </summary>
-		public WarningLevel DisablePluginsConflictWarningLevel { get; set; } = WarningLevel.Default;
+		public WarningLevel DisablePluginsConflictWarningLevel
+		{
+			get => CppCompileWarningSettings.DisablePluginsConflictWarningLevel;
+			set => CppCompileWarningSettings.DisablePluginsConflictWarningLevel = value; 
+		}
 
 		/// <summary>
 		/// A list of Plugin names that are allowed to exist as dependencies without being defined in the uplugin descriptor
diff --git a/Engine/Source/Programs/UnrealBuildTool/Configuration/UEBuildModuleCPP.cs b/Engine/Source/Programs/UnrealBuildTool/Configuration/UEBuildModuleCPP.cs
--- a/Engine/Source/Programs/UnrealBuildTool/Configuration/UEBuildModuleCPP.cs
+++ b/Engine/Source/Programs/UnrealBuildTool/Configuration/UEBuildModuleCPP.cs
@@ -599,7 +599,7 @@ namespace UnrealBuildTool
 					}
 					ModuleCompileEnvironment.FileInlineGenCPPMap = CppCompileEnvironment.FileInlineGenCPPMap;
 
-					if (Rules.bEnableNonInlinedGenCppWarnings)
+					if (Rules.CppCompileWarningSettings.NonInlinedGenCppWarningLevel >= WarningLevel.Warning)
 					{
 						Dictionary<string, FileItem> CPPFilesLookup = new Dictionary<string, FileItem>();
 						foreach (FileItem CPPFile in CPPFiles)
diff --git a/Engine/Source/Programs/UnrealBuildTool/Configuration/UEBuildTarget.cs b/Engine/Source/Programs/UnrealBuildTool/Configuration/UEBuildTarget.cs
--- a/Engine/Source/Programs/UnrealBuildTool/Configuration/UEBuildTarget.cs
+++ b/Engine/Source/Programs/UnrealBuildTool/Configuration/UEBuildTarget.cs
@@ -4954,16 +4954,16 @@ namespace UnrealBuildTool
 				// Detect when a plugin in DisablePlugins has been enabled by another plugin referencing it
 				if (Rules.DisablePlugins.Contains(Reference.Name))
 				{
-					if (Rules.DisablePluginsConflictWarningLevel == WarningLevel.Error)
+					if (Rules.CppCompileWarningSettings.DisablePluginsConflictWarningLevel == WarningLevel.Error)
 					{
 						LogValue PluginLogValue = LogValue.SourceFile(Info.File, Info.File.GetFileName());
 						throw new BuildLogEventException("Error: Plugin '{Plugin}' (referenced via {ReferenceChain}) is being enabled but found in DisablePlugins list. Suppress this message by setting DisablePluginsConflictWarningLevel = WarningLevel.Off in '{TargetRulesFile}'", PluginLogValue, ReferenceChain, TargetRulesFile.GetFileName());
 					}
-					else if (Rules.DisablePluginsConflictWarningLevel == WarningLevel.Warning)
+					else if (Rules.CppCompileWarningSettings.DisablePluginsConflictWarningLevel == WarningLevel.Warning)
 					{
 						Logger.LogWarning("Plugin '{ReferenceName}' (referenced via {ReferenceChain}) is being enabled but found in DisablePlugins list. Suppress this message by setting DisablePluginsConflictWarningLevel = WarningLevel.Off in '{TargetRulesFile}'", Reference.Name, ReferenceChain, TargetRulesFile.GetFileName());
 					}
-					else if (Rules.DisablePluginsConflictWarningLevel != WarningLevel.Off)
+					else if (Rules.CppCompileWarningSettings.DisablePluginsConflictWarningLevel != WarningLevel.Off)
 					{
 						Logger.LogDebug("Plugin '{ReferenceName}' (referenced via {ReferenceChain}) is being enabled but found in DisablePlugins list. Suppress this message by setting DisablePluginsConflictWarningLevel = WarningLevel.Off in '{TargetRulesFile}'", Reference.Name, ReferenceChain, TargetRulesFile.GetFileName());
 					}
