[UBT]
* Added so ExtraRootPath from module(s) are propagated to binary when doing linking

diff --git a/Engine/Source/Programs/UnrealBuildTool/Configuration/UEBuildBinary.cs b/Engine/Source/Programs/UnrealBuildTool/Configuration/UEBuildBinary.cs
--- a/Engine/Source/Programs/UnrealBuildTool/Configuration/UEBuildBinary.cs
+++ b/Engine/Source/Programs/UnrealBuildTool/Configuration/UEBuildBinary.cs
@@ -874,7 +874,7 @@ namespace UnrealBuildTool
 				}
 
 				// Allow the module to modify the link environment for the binary.
-				Module.SetupPrivateLinkEnvironment(this, BinaryLinkEnvironment, BinaryDependencies, LinkEnvironmentVisitedModules, ExeDir);
+				Module.SetupPrivateLinkEnvironment(this, BinaryLinkEnvironment, BinaryDependencies, LinkEnvironmentVisitedModules, ExeDir, Logger);
 			}
 
 			if (PerModuleItemToKeep != null)
diff --git a/Engine/Source/Programs/UnrealBuildTool/Configuration/UEBuildModule.cs b/Engine/Source/Programs/UnrealBuildTool/Configuration/UEBuildModule.cs
--- a/Engine/Source/Programs/UnrealBuildTool/Configuration/UEBuildModule.cs
+++ b/Engine/Source/Programs/UnrealBuildTool/Configuration/UEBuildModule.cs
@@ -853,7 +853,9 @@ namespace UnrealBuildTool
 			List<UEBuildBinary> BinaryDependencies,
 			Dictionary<string, HashSet<UnrealArch>> DependenciesToSkip,
 			HashSet<UEBuildModule> VisitedModules,
-			DirectoryReference ExeDir
+			DirectoryReference ExeDir,
+			CppRootPaths RootPaths,
+			ILogger Logger
 			)
 		{
 			// There may be circular dependencies in compile dependencies, so we need to avoid reentrance.
@@ -867,6 +869,8 @@ namespace UnrealBuildTool
 					BinaryDependencies.Add(Binary);
 				}
 
+				RootPaths.SetExtraRootPath(Rules.ExtraRootPath, Logger);
+
 				// If this module belongs to a static library that we are not currently building, recursively add the link environment settings for all of its dependencies too.
 				// Keep doing this until we reach a module that is not part of a static library (or external module, since they have no associated binary).
 				// Static libraries do not contain the symbols for their dependencies, so we need to recursively gather them to be linked into other binary types.
@@ -886,7 +890,7 @@ namespace UnrealBuildTool
 						if (bIsExternalModule || bIsInStaticLibrary)
 						{
 							DependencyModule.SetupPublicLinkEnvironment(SourceBinary, Libraries, SystemLibraryPaths, SystemLibraries, RuntimeLibraryPaths, Frameworks, WeakFrameworks,
-								AdditionalFrameworks, AdditionalBundleResources, DelayLoadDLLs, BinaryDependencies, DependenciesToSkip, VisitedModules, ExeDir);
+								AdditionalFrameworks, AdditionalBundleResources, DelayLoadDLLs, BinaryDependencies, DependenciesToSkip, VisitedModules, ExeDir, RootPaths, Logger);
 						}
 					}
 				}
@@ -922,7 +926,8 @@ namespace UnrealBuildTool
 			LinkEnvironment LinkEnvironment,
 			List<UEBuildBinary> BinaryDependencies,
 			HashSet<UEBuildModule> VisitedModules,
-			DirectoryReference ExeDir
+			DirectoryReference ExeDir,
+			ILogger Logger
 			)
 		{
 			// Add the private rpaths
@@ -930,7 +935,7 @@ namespace UnrealBuildTool
 
 			// Allow the module's public dependencies to add library paths and additional libraries to the link environment.
 			SetupPublicLinkEnvironment(SourceBinary, LinkEnvironment.Libraries, LinkEnvironment.SystemLibraryPaths, LinkEnvironment.SystemLibraries, LinkEnvironment.RuntimeLibraryPaths, LinkEnvironment.Frameworks, LinkEnvironment.WeakFrameworks,
-				LinkEnvironment.AdditionalFrameworks, LinkEnvironment.AdditionalBundleResources, LinkEnvironment.DelayLoadDLLs, BinaryDependencies, LinkEnvironment.DependenciesToSkipPerArchitecture, VisitedModules, ExeDir);
+				LinkEnvironment.AdditionalFrameworks, LinkEnvironment.AdditionalBundleResources, LinkEnvironment.DelayLoadDLLs, BinaryDependencies, LinkEnvironment.DependenciesToSkipPerArchitecture, VisitedModules, ExeDir, LinkEnvironment.RootPaths, Logger);
 
 			// Also allow the module's public and private dependencies to modify the link environment.
 			List<UEBuildModule> AllDependencyModules = new List<UEBuildModule>();
@@ -940,7 +945,7 @@ namespace UnrealBuildTool
 			foreach (UEBuildModule DependencyModule in AllDependencyModules)
 			{
 				DependencyModule.SetupPublicLinkEnvironment(SourceBinary, LinkEnvironment.Libraries, LinkEnvironment.SystemLibraryPaths, LinkEnvironment.SystemLibraries, LinkEnvironment.RuntimeLibraryPaths, LinkEnvironment.Frameworks, LinkEnvironment.WeakFrameworks,
-					LinkEnvironment.AdditionalFrameworks, LinkEnvironment.AdditionalBundleResources, LinkEnvironment.DelayLoadDLLs, BinaryDependencies, LinkEnvironment.DependenciesToSkipPerArchitecture, VisitedModules, ExeDir);
+					LinkEnvironment.AdditionalFrameworks, LinkEnvironment.AdditionalBundleResources, LinkEnvironment.DelayLoadDLLs, BinaryDependencies, LinkEnvironment.DependenciesToSkipPerArchitecture, VisitedModules, ExeDir, LinkEnvironment.RootPaths, Logger);
 			}
 
 			// Add all the additional properties
@@ -962,14 +967,7 @@ namespace UnrealBuildTool
 				ModuleCompileEnvironment.RootPaths[CppRootPathFolder.Project] = Target.ProjectFile.Directory;
 			}
 
-			if (Rules.ExtraRootPath != null)
-			{
-				DirectoryReference? ExtraRootPath = DirectoryReference.FromString(Rules.ExtraRootPath);
-				if (ExtraRootPath != null)
-				{
-					ModuleCompileEnvironment.RootPaths[CppRootPathFolder.Extra] = ExtraRootPath;
-				}
-			}
+			ModuleCompileEnvironment.RootPaths.SetExtraRootPath(Rules.ExtraRootPath, Logger);
 
 			// Generate type libraries for Windows
 			foreach (ModuleRules.TypeLibrary TypeLibrary in Rules.TypeLibraries)
diff --git a/Engine/Source/Programs/UnrealBuildTool/Configuration/UEBuildModuleCPP.cs b/Engine/Source/Programs/UnrealBuildTool/Configuration/UEBuildModuleCPP.cs
--- a/Engine/Source/Programs/UnrealBuildTool/Configuration/UEBuildModuleCPP.cs
+++ b/Engine/Source/Programs/UnrealBuildTool/Configuration/UEBuildModuleCPP.cs
@@ -2471,14 +2471,7 @@ namespace UnrealBuildTool
 				}
 			}
 
-			if (Rules.ExtraRootPath != null)
-			{
-				DirectoryReference? ExtraRootPath = DirectoryReference.FromString(Rules.ExtraRootPath);
-				if (ExtraRootPath != null)
-				{
-					Result.RootPaths[CppRootPathFolder.Extra] = ExtraRootPath;
-				}
-			}
+			Result.RootPaths.SetExtraRootPath(Rules.ExtraRootPath, Logger);
 
 			// Project root path should only be added if the module needs it, to prevent violating shared environment actions
 			if (Target.ProjectFile != null && (!Result.bUseSharedBuildEnvironment || Result.AllIncludePath.Any(x => x.IsUnderDirectory(Target.ProjectFile.Directory))))
diff --git a/Engine/Source/Programs/UnrealBuildTool/System/CppCompileEnvironment.cs b/Engine/Source/Programs/UnrealBuildTool/System/CppCompileEnvironment.cs
--- a/Engine/Source/Programs/UnrealBuildTool/System/CppCompileEnvironment.cs
+++ b/Engine/Source/Programs/UnrealBuildTool/System/CppCompileEnvironment.cs
@@ -7,6 +7,7 @@ using EpicGames.Core;
 using UnrealBuildBase;
 using System;
 using System.Diagnostics.CodeAnalysis;
+using Microsoft.Extensions.Logging;
 
 namespace UnrealBuildTool
 {
@@ -293,6 +294,31 @@ namespace UnrealBuildTool
 		private SortedDictionary<CppRootPathFolder, DirectoryReference> Lookup { get; }
 		public bool bUseVfs { get; set; }
 
+		public void SetExtraRootPath(string? ExtraRootPath, ILogger Logger)
+		{
+			if (ExtraRootPath == null)
+			{
+				return;
+			}
+			DirectoryReference? ExtraRootPathReference = DirectoryReference.FromString(ExtraRootPath);
+			if (ExtraRootPathReference == null)
+			{
+				return;
+			}
+
+			if (Lookup.TryGetValue(CppRootPathFolder.Extra, out DirectoryReference? Root))
+			{
+				if (Root != ExtraRootPathReference)
+				{
+					Logger.LogWarning("Warning: Extra root path already set. Current: {Existing} New: {New}", Root!.FullName, ExtraRootPath);
+				}
+			}
+			else
+			{
+				Lookup[CppRootPathFolder.Extra] = ExtraRootPathReference;
+			}
+		}
+
 		/// <summary>
 		/// Get a virtual filesystem overlay path if in use for a local path
 		/// </summary>
