[UBT]
* Removed warning when having multiple extra root paths.. this is an acceptable scenario and needs to be handled in a different way

diff --git a/Engine/Source/Programs/UnrealBuildTool/System/CppCompileEnvironment.cs b/Engine/Source/Programs/UnrealBuildTool/System/CppCompileEnvironment.cs
--- a/Engine/Source/Programs/UnrealBuildTool/System/CppCompileEnvironment.cs
+++ b/Engine/Source/Programs/UnrealBuildTool/System/CppCompileEnvironment.cs
@@ -310,7 +310,7 @@ namespace UnrealBuildTool
 			{
 				if (Root != ExtraRootPathReference)
 				{
-					Logger.LogWarning("Warning: Extra root path already set. Current: {Existing} New: {New}", Root!.FullName, ExtraRootPath);
+					//Logger.LogWarning("Warning: Extra root path already set. Current: {Existing} New: {New}", Root!.FullName, ExtraRootPath);
 				}
 			}
 			else
