[UBT]
* Added so extra root paths are added from modules that contain TypeLibraries

diff --git a/Engine/Source/Programs/UnrealBuildTool/Configuration/UEBuildModuleCPP.cs b/Engine/Source/Programs/UnrealBuildTool/Configuration/UEBuildModuleCPP.cs
--- a/Engine/Source/Programs/UnrealBuildTool/Configuration/UEBuildModuleCPP.cs
+++ b/Engine/Source/Programs/UnrealBuildTool/Configuration/UEBuildModuleCPP.cs
@@ -2445,6 +2445,21 @@ namespace UnrealBuildTool
 
 			Result.RootPaths.AddExtraPath(Rules.ExtraRootPath);
 
+			// Extra root path is required if module depends on external module that has type libraries.
+			if (PrivateDependencyModules != null)
+			{
+				foreach (UEBuildModule Dependent in PrivateDependencyModules)
+				{
+					if (Dependent.Rules.TypeLibraries.Count != 0)
+					{
+						if (Dependent is UEBuildModuleExternal External)
+						{
+							Result.RootPaths.AddExtraPath(External.Rules.ExtraRootPath);
+						}
+					}
+				}
+			}
+
 			// Project root path should only be added if the module needs it, to prevent violating shared environment actions
 			if (Target.ProjectFile != null && (!Result.bUseSharedBuildEnvironment || Result.AllIncludePath.Any(x => x.IsUnderDirectory(Target.ProjectFile.Directory))))
 			{
