[UBT]
* Changed vfs root to Z:/UEVFS for windows and /UEVFS for posix. This to make path more unique (when uba searches for paths "/UE" could be found inside paths but UEVFS is less likely

diff --git a/Engine/Source/Programs/UnrealBuildTool/System/RootPaths.cs b/Engine/Source/Programs/UnrealBuildTool/System/RootPaths.cs
--- a/Engine/Source/Programs/UnrealBuildTool/System/RootPaths.cs
+++ b/Engine/Source/Programs/UnrealBuildTool/System/RootPaths.cs
@@ -52,7 +52,7 @@ public enum CppRootPathFolder
 /// </summary>
 class CppRootPaths : IEnumerable<(DirectoryReference vfs, DirectoryReference local)>
 {
-	private static readonly DirectoryReference _vfsRootPath = new(OperatingSystem.IsWindows() ? "Z:/UE" : "/UE");
+	private static readonly DirectoryReference _vfsRootPath = new(OperatingSystem.IsWindows() ? "Z:/UEVFS" : "/UEVFS");
 
 	/// <summary>
 	/// Default constructor
