[UBT]
* Fixed logic deciding if shared pch should have project root set

diff --git a/Engine/Source/Programs/UnrealBuildTool/Configuration/UEBuildModuleCPP.cs b/Engine/Source/Programs/UnrealBuildTool/Configuration/UEBuildModuleCPP.cs
--- a/Engine/Source/Programs/UnrealBuildTool/Configuration/UEBuildModuleCPP.cs
+++ b/Engine/Source/Programs/UnrealBuildTool/Configuration/UEBuildModuleCPP.cs
@@ -2514,8 +2514,7 @@ namespace UnrealBuildTool
 			FindModulesInPublicCompileEnvironment(ModuleToIncludePathsOnlyFlag);
 
 			// Set up additional root paths
-			IEnumerable<DirectoryReference> AllIncludePaths = [.. CompileEnvironment.UserIncludePaths, .. CompileEnvironment.SystemIncludePaths];
-			if (Target.ProjectFile != null && (!Target.bUseSharedBuildEnvironment || AllIncludePaths.Any(x => x.IsUnderDirectory(Target.ProjectFile.Directory))))
+			if (Target.ProjectFile != null && (!Target.bUseSharedBuildEnvironment || RulesFile.IsUnderDirectory(Target.ProjectFile.Directory)))
 			{
 				CompileEnvironment.RootPaths[CppRootPathFolder.Project] = Target.ProjectFile.Directory;
 			}
