[UBT]
* Fixed so response files uses vfs path if available

diff --git a/Engine/Source/Programs/UnrealBuildTool/Platform/Windows/VCCompileAction.cs b/Engine/Source/Programs/UnrealBuildTool/Platform/Windows/VCCompileAction.cs
--- a/Engine/Source/Programs/UnrealBuildTool/Platform/Windows/VCCompileAction.cs
+++ b/Engine/Source/Programs/UnrealBuildTool/Platform/Windows/VCCompileAction.cs
@@ -629,7 +629,7 @@ namespace UnrealBuildTool
 				return String.Join(' ', Arguments);
 			}
 
-			string ResponseFileString = VCToolChain.NormalizeCommandLinePath(ResponseFile);
+			string ResponseFileString = VCToolChain.NormalizeCommandLinePath(ResponseFile, RootPaths);
 
 			// cl.exe can't handle response files with a path longer than 260 characters, and relative paths can push it over the limit
 			if (!System.IO.Path.IsPathRooted(ResponseFileString) && System.IO.Path.Combine(WorkingDirectory.FullName, ResponseFileString).Length > 260)
diff --git a/Engine/Source/Programs/UnrealBuildTool/Platform/Windows/VCToolChain.cs b/Engine/Source/Programs/UnrealBuildTool/Platform/Windows/VCToolChain.cs
--- a/Engine/Source/Programs/UnrealBuildTool/Platform/Windows/VCToolChain.cs
+++ b/Engine/Source/Programs/UnrealBuildTool/Platform/Windows/VCToolChain.cs
@@ -2990,12 +2990,12 @@ namespace UnrealBuildTool
 			if (bIsBuildingLibraryOrImportLibrary)
 			{
 				LinkAction.CommandPath = EnvVars.LibraryManagerPath;
-				LinkAction.CommandArguments = $"/LIB @\"{ResponseFileName}\"";
+				LinkAction.CommandArguments = $"/LIB @\"{NormalizeCommandLinePath(ResponseFileName, RootPaths)}\"";
 			}
 			else
 			{
 				LinkAction.CommandPath = EnvVars.LinkerPath;
-				LinkAction.CommandArguments = $"@\"{ResponseFileName}\"";
+				LinkAction.CommandArguments = $"@\"{NormalizeCommandLinePath(ResponseFileName, RootPaths)}\"";
 			}
 			LinkAction.CommandVersion = EnvVars.ToolChainVersion.ToString();
 			LinkAction.ProducedItems.UnionWith(ProducedItems);
