[UBT]
* Added root paths to more normalize calls

diff --git a/Engine/Source/Programs/UnrealBuildTool/Configuration/UEBuildTarget.cs b/Engine/Source/Programs/UnrealBuildTool/Configuration/UEBuildTarget.cs
--- a/Engine/Source/Programs/UnrealBuildTool/Configuration/UEBuildTarget.cs
+++ b/Engine/Source/Programs/UnrealBuildTool/Configuration/UEBuildTarget.cs
@@ -5698,9 +5698,18 @@ namespace UnrealBuildTool
 						OutputFilePath = OutputFilePath.Substring(0, OutputFilePath.LastIndexOf(".app/Contents/MacOS") + 4);
 					}
 
-					DirectoryReference OutputDir = new DirectoryReference(OutputFilePath).ParentDirectory!;
 					DirectoryReference RootDir = Unreal.EngineDirectory.ParentDirectory!;
-					string RootPath = RootDir.MakeRelativeTo(OutputDir);
+					string RootPath;
+					if (GlobalCompileEnvironment.RootPaths.GetVfsOverlayPath(RootDir, out string? VfsPath))
+					{
+						RootPath = VfsPath;
+					}
+					else
+					{
+						DirectoryReference OutputDir = DirectoryReference.FromString(OutputFilePath)!.ParentDirectory!;
+						RootPath = RootDir.MakeRelativeTo(OutputDir);
+					}
+
 					// now get the cleaned engine path (Root/Engine), and make sure to end with /
 					string EnginePath = Utils.CleanDirectorySeparators(Path.Combine(RootPath, "Engine/"), '/');
 					GlobalCompileEnvironment.Definitions.Add(String.Format("UE_ENGINE_DIRECTORY=\"{0}\"", EnginePath));
diff --git a/Engine/Source/Programs/UnrealBuildTool/Platform/Windows/VCCompileAction.cs b/Engine/Source/Programs/UnrealBuildTool/Platform/Windows/VCCompileAction.cs
--- a/Engine/Source/Programs/UnrealBuildTool/Platform/Windows/VCCompileAction.cs
+++ b/Engine/Source/Programs/UnrealBuildTool/Platform/Windows/VCCompileAction.cs
@@ -642,12 +642,12 @@ namespace UnrealBuildTool
 		string GetClFilterArguments()
 		{
 			List<string> Arguments = new List<string>();
-			string DependencyListFileString = VCToolChain.NormalizeCommandLinePath(DependencyListFile!);
+			string DependencyListFileString = VCToolChain.NormalizeCommandLinePath(DependencyListFile!, RootPaths);
 			Arguments.Add(String.Format("-dependencies={0}", Utils.MakePathSafeToUseWithCommandLine(DependencyListFileString)));
 
 			if (TimingFile != null)
 			{
-				string TimingFileString = VCToolChain.NormalizeCommandLinePath(TimingFile);
+				string TimingFileString = VCToolChain.NormalizeCommandLinePath(TimingFile, RootPaths);
 				Arguments.Add(String.Format("-timing={0}", Utils.MakePathSafeToUseWithCommandLine(TimingFileString)));
 			}
 			if (bShowIncludes)
diff --git a/Engine/Source/Programs/UnrealBuildTool/Platform/Windows/VCToolChain.cs b/Engine/Source/Programs/UnrealBuildTool/Platform/Windows/VCToolChain.cs
--- a/Engine/Source/Programs/UnrealBuildTool/Platform/Windows/VCToolChain.cs
+++ b/Engine/Source/Programs/UnrealBuildTool/Platform/Windows/VCToolChain.cs
@@ -1449,7 +1449,7 @@ namespace UnrealBuildTool
 			{
 				Arguments.Add("/MANIFEST:EMBED");
 				FileItem ManifestFile = FileItem.GetItemByPath(Target.WindowsPlatform.ManifestFile);
-				Arguments.Add($"/MANIFESTINPUT:\"{NormalizeCommandLinePath(ManifestFile)}\"");
+				Arguments.Add($"/MANIFESTINPUT:\"{NormalizeCommandLinePath(ManifestFile, LinkEnvironment.RootPaths)}\"");
 			}
 
 			// Prevents the linker from displaying its logo for each invocation.
@@ -2030,7 +2030,7 @@ namespace UnrealBuildTool
 					}
 
 					// Specify the PDB file that the compiler should write to.
-					CompileAction.Arguments.Add($"/Fd\"{NormalizeCommandLinePath(PDBLocation)}\"");
+					CompileAction.Arguments.Add($"/Fd\"{NormalizeCommandLinePath(PDBLocation, CompileEnvironment.RootPaths)}\"");
 
 					// Don't allow remote execution when PDB files are enabled; we need to modify the same files. XGE works around this by generating separate
 					// PDB files per agent, but this functionality is only available with the Visual C++ extension package (via the VCCompiler=true tool option).
@@ -2914,7 +2914,7 @@ namespace UnrealBuildTool
 				{
 					// Write the import library to the output directory for nFringe support.
 					FileItem ImportLibraryFile = FileItem.GetItemByFileReference(ImportLibraryFilePath!);
-					Arguments.Add($"/IMPLIB:\"{NormalizeCommandLinePath(ImportLibraryFilePath!)}\"");
+					Arguments.Add($"/IMPLIB:\"{NormalizeCommandLinePath(ImportLibraryFilePath!, RootPaths!)}\"");
 
 					// Like the export file above, don't add the import library as a produced item when it's cross referenced.
 					if (!LinkEnvironment.bIsCrossReferenced)
@@ -2933,7 +2933,7 @@ namespace UnrealBuildTool
 						if (Target.WindowsPlatform.bStripPrivateSymbols)
 						{
 							FileItem StrippedPDBFile = FileItem.GetItemByFileReference(PDBFilePath);
-							Arguments.Add($"/PDBSTRIPPED:\"{NormalizeCommandLinePath(StrippedPDBFile)}\"");
+							Arguments.Add($"/PDBSTRIPPED:\"{NormalizeCommandLinePath(StrippedPDBFile, RootPaths)}\"");
 							ProducedItems.Add(StrippedPDBFile);
 
 							PDBFilePath = FileReference.Combine(LinkEnvironment.OutputDirectory!, Path.GetFileNameWithoutExtension(OutputFile.AbsolutePath) + ".full.pdb");
diff --git a/Engine/Source/Programs/UnrealBuildTool/ToolChain/ClangToolChain.cs b/Engine/Source/Programs/UnrealBuildTool/ToolChain/ClangToolChain.cs
--- a/Engine/Source/Programs/UnrealBuildTool/ToolChain/ClangToolChain.cs
+++ b/Engine/Source/Programs/UnrealBuildTool/ToolChain/ClangToolChain.cs
@@ -614,9 +614,9 @@ namespace UnrealBuildTool
 			return $"-M -MF\"{NormalizeCommandLinePath(DependencyListFile, CompileEnvironment.RootPaths)}\"";
 		}
 
-		protected virtual string GetResponseFileArgument(FileItem ResponseFile)
+		protected virtual string GetResponseFileArgument(FileItem ResponseFile, CppRootPaths RootPaths)
 		{
-			return $"@\"{NormalizeCommandLinePath(ResponseFile)}\"";
+			return $"@\"{NormalizeCommandLinePath(ResponseFile, RootPaths)}\"";
 		}
 
 		/// <summary>
@@ -971,7 +971,7 @@ namespace UnrealBuildTool
 			// Add the additional response files
 			foreach (FileItem AdditionalResponseFile in CompileEnvironment.AdditionalResponseFiles)
 			{
-				Arguments.Add(GetResponseFileArgument(AdditionalResponseFile));
+				Arguments.Add(GetResponseFileArgument(AdditionalResponseFile, CompileEnvironment.RootPaths));
 			}
 
 			// Add force include paths to the argument list.
@@ -1258,7 +1258,7 @@ namespace UnrealBuildTool
 
 			// Adds the response file to the compiler input.
 			FileItem CompilerResponseFileItem = Graph.CreateIntermediateTextFile(ResponseFileName, ResponseFileContents);
-			string CommandArguments = GetResponseFileArgument(CompilerResponseFileItem);
+			string CommandArguments = GetResponseFileArgument(CompilerResponseFileItem, CompileEnvironment.RootPaths);
 
 			if (bMergeModules && CompileEnvironment.PrecompiledHeaderAction != PrecompiledHeaderAction.Create)
 			{
@@ -1359,7 +1359,7 @@ namespace UnrealBuildTool
 				FileItem PreprocessResponseFileItem = Graph.CreateIntermediateTextFile(PreprocessResponseFileName, PreprocessResponseFileContents);
 				PrepassAction.PrerequisiteItems.Add(PreprocessResponseFileItem);
 
-				PrepassAction.CommandArguments = GetResponseFileArgument(PreprocessResponseFileItem);
+				PrepassAction.CommandArguments = GetResponseFileArgument(PreprocessResponseFileItem, CompileEnvironment.RootPaths);
 				CompileAction.DependencyListFile = DependencyListFile;
 				CompileAction.PrerequisiteItems.Add(DependencyListFile);
 			}
diff --git a/Engine/Source/Programs/UnrealBuildTool/ToolChain/ISPCToolChain.cs b/Engine/Source/Programs/UnrealBuildTool/ToolChain/ISPCToolChain.cs
--- a/Engine/Source/Programs/UnrealBuildTool/ToolChain/ISPCToolChain.cs
+++ b/Engine/Source/Programs/UnrealBuildTool/ToolChain/ISPCToolChain.cs
@@ -519,7 +519,7 @@ namespace UnrealBuildTool
 
 				FileReference ResponseFileName = GetResponseFileName(CompileEnvironment, ISPCIncludeHeaderFile);
 				FileItem ResponseFileItem = Graph.CreateIntermediateTextFile(ResponseFileName, Arguments.Select(x => Utils.ExpandVariables(x)));
-				CompileAction.CommandArguments = $"@\"{NormalizeCommandLinePath(ResponseFileName)}\"";
+				CompileAction.CommandArguments = $"@\"{NormalizeCommandLinePath(ResponseFileName, RootPaths)}\"";
 				CompileAction.PrerequisiteItems.Add(ResponseFileItem);
 
 				// Add the source file and its included files to the prerequisite item list.
@@ -762,7 +762,7 @@ namespace UnrealBuildTool
 					CompileAction.ProducedItems.Add(FileItem.GetItemByFileReference(FileReference.Combine(OutputDir, FileName + ".exi")));
 				}
 
-				CompileAction.CommandArguments = $"@\"{NormalizeCommandLinePath(ResponseFileName)}\"{AdditionalArguments}";
+				CompileAction.CommandArguments = $"@\"{NormalizeCommandLinePath(ResponseFileName, RootPaths)}\"{AdditionalArguments}";
 				CompileAction.PrerequisiteItems.Add(ResponseFileItem);
 
 				// Add the source file and its included files to the prerequisite item list.
@@ -805,7 +805,7 @@ namespace UnrealBuildTool
 							// Write the args to a response file
 							FileReference PostCompileResponseFileName = GetResponseFileName(CompileEnvironment, FinalCompiledISPCObjFile);
 							FileItem PostCompileResponseFileItem = Graph.CreateIntermediateTextFile(PostCompileResponseFileName, PostCompileArgs.Select(x => Utils.ExpandVariables(x)));
-							PostCompileAction.CommandArguments = $"@\"{NormalizeCommandLinePath(PostCompileResponseFileName)}\"{AdditionalArguments}";
+							PostCompileAction.CommandArguments = $"@\"{NormalizeCommandLinePath(PostCompileResponseFileName, RootPaths)}\"{AdditionalArguments}";
 							PostCompileAction.PrerequisiteItems.Add(PostCompileResponseFileItem);
 
 							PostCompileAction.PrerequisiteItems.Add(CompiledBytecodeObjFile);
