[UBT]
* Added system32 as extra root path for actions generating type libraries since they use copy (which is in cmd.exe)

diff --git a/Engine/Source/Programs/UnrealBuildTool/Configuration/UEBuildModule.cs b/Engine/Source/Programs/UnrealBuildTool/Configuration/UEBuildModule.cs
--- a/Engine/Source/Programs/UnrealBuildTool/Configuration/UEBuildModule.cs
+++ b/Engine/Source/Programs/UnrealBuildTool/Configuration/UEBuildModule.cs
@@ -960,6 +960,11 @@ namespace UnrealBuildTool
 		/// </summary>
 		public virtual List<FileItem> Compile(ReadOnlyTargetRules Target, UEToolChain ToolChain, CppCompileEnvironment CompileEnvironment, ISourceFileWorkingSet WorkingSet, IActionGraphBuilder Graph, ILogger Logger)
 		{
+			if (Rules.TypeLibraries.Count == 0)
+			{
+				return [];
+			}
+
 			CppCompileEnvironment ModuleCompileEnvironment = new CppCompileEnvironment(CompileEnvironment);
 			IEnumerable<DirectoryReference> AllIncludePaths = [.. ModuleCompileEnvironment.UserIncludePaths, .. ModuleCompileEnvironment.SystemIncludePaths];
 			if (Target.ProjectFile != null && (Target.BuildEnvironment != TargetBuildEnvironment.Shared || AllIncludePaths.Any(x => x.IsUnderDirectory(Target.ProjectFile.Directory))))
@@ -968,6 +973,7 @@ namespace UnrealBuildTool
 			}
 
 			ModuleCompileEnvironment.RootPaths.AddExtraPath(Rules.ExtraRootPath);
+			ModuleCompileEnvironment.RootPaths.AddExtraPath(("System32", Environment.SystemDirectory)); // Needed for cmd.exe
 
 			// Generate type libraries for Windows
 			foreach (ModuleRules.TypeLibrary TypeLibrary in Rules.TypeLibraries)
@@ -977,7 +983,7 @@ namespace UnrealBuildTool
 				ToolChain.GenerateTypeLibraryHeader(ModuleCompileEnvironment, TypeLibrary, OutputFile, OutputInclude, Graph);
 			}
 
-			return new List<FileItem>();
+			return [];
 		}
 
 		public IEnumerable<FileItem> CopyDebuggerVisualizers(UEToolChain ToolChain, IActionGraphBuilder Graph, ILogger Logger)
