[UBT]
* Changed paths in rsp to be forward slash instead of having to think about single vs double backslash etc.

diff --git a/Engine/Source/Programs/UnrealBuildTool/Platform/Windows/VCToolChain.cs b/Engine/Source/Programs/UnrealBuildTool/Platform/Windows/VCToolChain.cs
--- a/Engine/Source/Programs/UnrealBuildTool/Platform/Windows/VCToolChain.cs
+++ b/Engine/Source/Programs/UnrealBuildTool/Platform/Windows/VCToolChain.cs
@@ -162,10 +162,10 @@ namespace UnrealBuildTool
 			// Try to use a relative path to shorten command line length and to enable remote distribution where absolute paths are not desired
 			if (Reference.IsUnderDirectory(Unreal.EngineDirectory))
 			{
-				return Reference.MakeRelativeTo(Unreal.EngineSourceDirectory);
+				return Reference.MakeRelativeTo(Unreal.EngineSourceDirectory).Replace('\\', '/');
 			}
 
-			return Reference.FullName;
+			return Reference.FullName.Replace('\\', '/');
 		}
 
 		public static new string NormalizeCommandLinePath(FileItem Item)
@@ -2461,7 +2461,7 @@ namespace UnrealBuildTool
 				CompileAction.PrerequisiteItems.Add(FileItem.GetItemByFileReference(IconFile));
 
 				// Setup the compile environment, setting the icon to use via a macro. This is used in Default.rc2.
-				AddDefinition(Arguments, $"BUILD_ICON_FILE_NAME=\"\\\"{NormalizeCommandLinePath(IconFile, RootPaths).Replace("\\", "\\\\")}\\\"\"");
+				AddDefinition(Arguments, $"BUILD_ICON_FILE_NAME=\"\\\"{NormalizeCommandLinePath(IconFile, RootPaths).Replace("/", "\\\\")}\\\"\"");
 
 				// Apply the target settings for the resources
 				if (!CompileEnvironment.bUseSharedBuildEnvironment)
@@ -2576,7 +2576,7 @@ namespace UnrealBuildTool
 
 			// Build the command for touching the output file to update the time stamp
 			string TouchActionCommand;
-			string OutfileFullName = NormalizeCommandLinePath(OutputFile);
+			string OutfileFullName = NormalizeCommandLinePath(OutputFile).Replace('/', '\\'); // Must be backward slash
 			if (BuildHostPlatform.Current.IsRunningOnWine())
 			{
 				TouchActionCommand = $"copy /Y \"{OutfileFullName}\" \"%TMP%/{OutputFile.GetFileName()}\" && type \"%TMP%/{OutputFile.GetFileName()}\" > \"{OutfileFullName}\"";
diff --git a/Engine/Source/Programs/UnrealBuildTool/ToolChain/ISPCToolChain.cs b/Engine/Source/Programs/UnrealBuildTool/ToolChain/ISPCToolChain.cs
--- a/Engine/Source/Programs/UnrealBuildTool/ToolChain/ISPCToolChain.cs
+++ b/Engine/Source/Programs/UnrealBuildTool/ToolChain/ISPCToolChain.cs
@@ -338,14 +338,10 @@ namespace UnrealBuildTool
 			{
 				path = Reference.MakeRelativeTo(Unreal.EngineSourceDirectory);
 			}
-			if (Path.DirectorySeparatorChar == '/')
+			if (Path.DirectorySeparatorChar == '\\')
 			{
 				path = path.Replace("\\", "/");
 			}
-			else
-			{
-				path = path.Replace("\\", "\\\\");
-			}
 			return path;
 		}
 
