[UBT]
* Do not enable vfs on non-windows even though vfs is enabled. Need to fix support in the detouring backend

diff --git a/Engine/Source/Programs/UnrealBuildTool/ToolChain/UEToolChain.cs b/Engine/Source/Programs/UnrealBuildTool/ToolChain/UEToolChain.cs
--- a/Engine/Source/Programs/UnrealBuildTool/ToolChain/UEToolChain.cs
+++ b/Engine/Source/Programs/UnrealBuildTool/ToolChain/UEToolChain.cs
@@ -274,8 +274,11 @@ namespace UnrealBuildTool
 
 		public virtual void SetUpGlobalEnvironment(ReadOnlyTargetRules Target, CppCompileEnvironment GlobalCompileEnvironment, LinkEnvironment GlobalLinkEnvironment)
 		{
-			GlobalCompileEnvironment.RootPaths.bUseVfs = Target.bCompilerVFS;
-			GlobalLinkEnvironment.RootPaths.bUseVfs = Target.bCompilerVFS;
+			if (OperatingSystem.IsWindows())
+			{
+				GlobalCompileEnvironment.RootPaths.bUseVfs = Target.bCompilerVFS;
+				GlobalLinkEnvironment.RootPaths.bUseVfs = Target.bCompilerVFS;
+			}
 
 			GlobalCompileEnvironment.RootPaths[CppRootPathFolder.Root] = Unreal.RootDirectory;
 			GlobalLinkEnvironment.RootPaths[CppRootPathFolder.Root] = Unreal.RootDirectory;
