[UBA]
* Removed log spam from remotely built actions as well

diff --git a/Engine/Source/Programs/UnrealBuildTool/Executors/UnrealBuildAccelerator/UBAExecutor.cs b/Engine/Source/Programs/UnrealBuildTool/Executors/UnrealBuildAccelerator/UBAExecutor.cs
--- a/Engine/Source/Programs/UnrealBuildTool/Executors/UnrealBuildAccelerator/UBAExecutor.cs
+++ b/Engine/Source/Programs/UnrealBuildTool/Executors/UnrealBuildAccelerator/UBAExecutor.cs
@@ -954,7 +954,7 @@ namespace UnrealBuildTool
 					TimeSpan processorTime = e.TotalProcessorTime;
 					TimeSpan executionTime = e.TotalWallTime;
 					List<string> logLines = e.LogLines;
-					logLines.RemoveAll((line) => line.StartsWith("   Creating library ", StringComparison.OrdinalIgnoreCase) && line.EndsWith(".exp", StringComparison.OrdinalIgnoreCase));
+					RemoveLogLineSpam(logLines);
 					ActionFinished(queue, new ExecuteResults(logLines, e.ExitCode, executionTime, processorTime, additionalDescription), action, process);
 				}, action.Weight, knownInputs, knownInputsCount, action.bCanExecuteInUBACrossArchitecture);
 				return Task.CompletedTask;
