[UBT]
* Removed pch timestamp check since it prevents caches from working properly.

diff --git a/Engine/Source/Programs/UnrealBuildTool/ToolChain/ClangToolChain.cs b/Engine/Source/Programs/UnrealBuildTool/ToolChain/ClangToolChain.cs
--- a/Engine/Source/Programs/UnrealBuildTool/ToolChain/ClangToolChain.cs
+++ b/Engine/Source/Programs/UnrealBuildTool/ToolChain/ClangToolChain.cs
@@ -399,6 +399,7 @@ namespace UnrealBuildTool
 			if (CompileEnvironment.PrecompiledHeaderAction != PrecompiledHeaderAction.None && CompilerVersionGreaterOrEqual(11, 0, 0))
 			{
 				// Validate PCH inputs by content if mtime check fails
+				Arguments.Add("-Xclang -fno-pch-timestamp"); // This is needed to prevent check on timestamp stored inside pch
 				Arguments.Add("-fpch-validate-input-files-content");
 			}
 
