[UBT]
* Fixed so root path ids are properly sent in to uba (this is needed to keep uba cache path table smaller)

diff --git a/Engine/Source/Programs/UnrealBuildTool/Executors/UnrealBuildAccelerator/UBAExecutor.cs b/Engine/Source/Programs/UnrealBuildTool/Executors/UnrealBuildAccelerator/UBAExecutor.cs
--- a/Engine/Source/Programs/UnrealBuildTool/Executors/UnrealBuildAccelerator/UBAExecutor.cs
+++ b/Engine/Source/Programs/UnrealBuildTool/Executors/UnrealBuildAccelerator/UBAExecutor.cs
@@ -688,9 +688,9 @@ namespace UnrealBuildTool
 				// And since pch contains absolute paths we unfortunately can't share cache data between machines that have different paths
 				if (action.ArtifactMode.HasFlag(ArtifactMode.AbsolutePath))
 				{
-					foreach ((DirectoryReference vfs, DirectoryReference local) kv in action.RootPaths)
+					foreach (var (id, vfs, local) in action.RootPaths)
 					{
-						hasher.Update(System.Text.Encoding.UTF8.GetBytes(kv.local.FullName));
+						hasher.Update(System.Text.Encoding.UTF8.GetBytes(local.FullName));
 					}
 				}
 
@@ -704,18 +704,18 @@ namespace UnrealBuildTool
 			using (BinaryWriter writer = new(outputsMemory, System.Text.Encoding.UTF8, true))
 			{
 				bool useVfs = action.RootPaths.bUseVfs;
-				foreach (var kv in action.RootPaths)
+				foreach (var (id, vfs, local) in action.RootPaths)
 				{
-					writer.Write((byte)0);// kv.Key);
+					writer.Write((byte)id);
 					if (useVfs)
 					{
-						writer.Write(kv.vfs.FullName);
+						writer.Write(vfs.FullName);
 					}
 					else
 					{
 						writer.Write("");
 					}
-					writer.Write(kv.local.FullName);
+					writer.Write(local.FullName);
 				}
 			}
 			return _session!.RegisterRoots(outputsMemory.GetBuffer(), (uint)outputsMemory.Position);
diff --git a/Engine/Source/Programs/UnrealBuildTool/System/RootPaths.cs b/Engine/Source/Programs/UnrealBuildTool/System/RootPaths.cs
--- a/Engine/Source/Programs/UnrealBuildTool/System/RootPaths.cs
+++ b/Engine/Source/Programs/UnrealBuildTool/System/RootPaths.cs
@@ -50,7 +50,7 @@ public enum CppRootPathFolder
 /// <summary>
 /// Encapsulates the roots used for cache and vfs
 /// </summary>
-class CppRootPaths : IEnumerable<(DirectoryReference vfs, DirectoryReference local)>
+class CppRootPaths : IEnumerable<(uint id, DirectoryReference vfs, DirectoryReference local)>
 {
 	private static readonly DirectoryReference _vfsRootPath = new(OperatingSystem.IsWindows() ? "Z:/UEVFS" : "/UEVFS");
 
@@ -127,17 +127,17 @@ class CppRootPaths : IEnumerable<(DirectoryReference vfs, DirectoryReference loc
 	/// <summary>
 	/// Enumerator for all roots
 	/// </summary>
-	public IEnumerator<(DirectoryReference vfs, DirectoryReference local)> GetEnumerator()
+	public IEnumerator<(uint id, DirectoryReference vfs, DirectoryReference local)> GetEnumerator()
 	{
 		foreach (KeyValuePair<CppRootPathFolder, (DirectoryReference vfs, DirectoryReference local)> kv in _lookup)
 		{
-			yield return new (kv.Value.vfs, kv.Value.local);
+			yield return ((uint)kv.Key, kv.Value.vfs, kv.Value.local);
 		}
 
+		uint index = 7; // Max of CppRootPathFolder
 		foreach (KeyValuePair<string, (DirectoryReference vfs, DirectoryReference local)> kv in _extras)
 		{
-			DirectoryReference vfsPath = DirectoryReference.Combine(_vfsRootPath, kv.Key);
-			yield return new (kv.Value.vfs, kv.Value.local);
+			yield return (index++, kv.Value.vfs, kv.Value.local);
 		}
 		yield break;
 	}
