- added bSkipSdkCheck param to IsPlatformAvailable() to pass along to IsPlatformAvailableForTarget()
- allow Linux CookedEditor cook target without Linux SDK installed

diff --git a/Engine/Source/Developer/CookedEditor/CookedEditor.Build.cs b/Engine/Source/Developer/CookedEditor/CookedEditor.Build.cs
--- a/Engine/Source/Developer/CookedEditor/CookedEditor.Build.cs
+++ b/Engine/Source/Developer/CookedEditor/CookedEditor.Build.cs
@@ -30,7 +30,7 @@ public class CookedEditor : ModuleRules
 		PublicIncludePathModuleNames.Add("WindowsTargetPlatformSettings");
 		PublicIncludePathModuleNames.Add("WindowsTargetPlatformControls");
 
-		if (IsPlatformAvailable(UnrealTargetPlatform.Linux))
+		if (IsPlatformAvailable(UnrealTargetPlatform.Linux, bIgnoreSDKCheck: true))
 		{
 			PublicDefinitions.Add("COOKEDEDITOR_WITH_LINUXTARGETPLATFORM=1");
 			PublicIncludePathModuleNames.Add("LinuxTargetPlatformSettings");
diff --git a/Engine/Source/Programs/UnrealBuildTool/Configuration/ModuleRules.cs b/Engine/Source/Programs/UnrealBuildTool/Configuration/ModuleRules.cs
--- a/Engine/Source/Programs/UnrealBuildTool/Configuration/ModuleRules.cs
+++ b/Engine/Source/Programs/UnrealBuildTool/Configuration/ModuleRules.cs
@@ -1805,8 +1805,12 @@ namespace UnrealBuildTool
 		/// Determines whether a given platform is available in the context of the current Target
 		/// </summary>
 		/// <param name="inPlatform">The platform to check for</param>
+		/// <param name="bIgnoreSDKCheck">Ignore the sdks presence when checking for platform availablity (many platforms can be cooked as a target platform without requiring sdk)</param>
 		/// <returns>True if it's available, false otherwise</returns>
-		protected bool IsPlatformAvailable(UnrealTargetPlatform inPlatform) => UEBuildPlatform.IsPlatformAvailableForTarget(inPlatform, Target);
+		protected bool IsPlatformAvailable(UnrealTargetPlatform inPlatform, bool bIgnoreSDKCheck = false)
+		{
+			return UEBuildPlatform.IsPlatformAvailableForTarget(inPlatform, Target, bIgnoreSDKCheck);
+		}
 
 		/// <summary>
 		/// Returns all the modules that are in the given module group
