UnrealBuildTool: Shorten length of adaptive header build filename

#rnx

diff --git a/Engine/Source/Programs/UnrealBuildTool/System/Unity.cs b/Engine/Source/Programs/UnrealBuildTool/System/Unity.cs
--- a/Engine/Source/Programs/UnrealBuildTool/System/Unity.cs
+++ b/Engine/Source/Programs/UnrealBuildTool/System/Unity.cs
@@ -402,7 +402,7 @@ namespace UnrealBuildTool
 					OutputHeaderCPPWriter.WriteLine("// This file is automatically generated at compile-time to include a modified header file.");
 					OutputHeaderCPPWriter.WriteLine($"#include \"{HeaderFile.AbsolutePath.Replace('\\', '/')}\"");
 
-					string HeaderCPPFileName = $"{ModulePrefix}{BaseName}.Header.{Path.GetFileNameWithoutExtension(HeaderFile.AbsolutePath)}.cpp";
+					string HeaderCPPFileName = $"{HeaderFile.Location.GetFileNameWithoutExtension()}.h.cpp";
 					FileReference HeaderCPPFilePath = FileReference.Combine(IntermediateDirectory, HeaderCPPFileName);
 
 					AdaptiveFiles.Add(Graph.CreateIntermediateTextFile(HeaderCPPFilePath, OutputHeaderCPPWriter.ToString()));
