UnrealBuildTool: Write all linux link arguments to a response file instead of passing on the command line.
Move a few common clang args to be shared

#rnx

diff --git a/Engine/Source/Programs/UnrealBuildTool/Platform/Android/AndroidToolChain.cs b/Engine/Source/Programs/UnrealBuildTool/Platform/Android/AndroidToolChain.cs
--- a/Engine/Source/Programs/UnrealBuildTool/Platform/Android/AndroidToolChain.cs
+++ b/Engine/Source/Programs/UnrealBuildTool/Platform/Android/AndroidToolChain.cs
@@ -547,6 +547,27 @@ namespace UnrealBuildTool
 			}
 		}
 
+		/// <inheritdoc/>
+		protected override void GetCompileArguments_Architecture(CppCompileEnvironment CompileEnvironment, List<string> Arguments)
+		{
+			base.GetCompileArguments_Architecture(CompileEnvironment, Arguments);
+
+			if (CompileEnvironment.Architecture == UnrealArch.Arm64)
+			{
+				Arguments.Add(ToolchainParamsArm64);
+				Arguments.Add("-D__arm64__");            // for some reason this isn't defined and needed for PhysX
+				Arguments.Add("-DPLATFORM_ANDROID_ARM64=1");
+				Arguments.Add("-march=armv8-a");
+			}
+			else if (CompileEnvironment.Architecture == UnrealArch.X64)
+			{
+				Arguments.Add(ToolchainParamsx64);
+				Arguments.Add("-DPLATFORM_ANDROID_X64=1");
+				Arguments.Add("-fno-omit-frame-pointer");
+				Arguments.Add("-march=atom");
+			}
+		}
+
 		/// <inheritdoc/>
 		protected override void GetCompileArguments_Debugging(CppCompileEnvironment CompileEnvironment, List<string> Arguments)
 		{
@@ -656,7 +677,6 @@ namespace UnrealBuildTool
 			Arguments.Add(GetRTTIFlag(CompileEnvironment));
 			Arguments.Add("-no-canonical-prefixes");
 			Arguments.Add("-fno-PIE");
-			Arguments.Add("-funwind-tables");           // Just generates any needed static data, affects no code
 			Arguments.Add("-fPIC");                     // Generates position-independent code (PIC) suitable for use in a shared library
 			Arguments.Add("-fno-strict-aliasing");      // Prevents unwanted or invalid optimizations that could produce incorrect code
 			Arguments.Add("-fno-short-enums");          // Do not allocate to an enum type only as many bytes as it needs for the declared range of possible values
@@ -664,22 +684,6 @@ namespace UnrealBuildTool
 			Arguments.Add("-D_FORTIFY_SOURCE=2");       // FORTIFY default
 			Arguments.Add($"-DPLATFORM_USED_NDK_VERSION_INTEGER={NDKApiLevel64Int}");       // NDK version
 			Arguments.Add("-DPLATFORM_64BITS=1");       // NDK version
-
-			if (CompileEnvironment.Architecture == UnrealArch.Arm64)
-			{
-				Arguments.Add(ToolchainParamsArm64);
-				Arguments.Add("-D__arm64__");            // for some reason this isn't defined and needed for PhysX
-				Arguments.Add("-DPLATFORM_ANDROID_ARM64=1");
-				Arguments.Add("-march=armv8-a");
-				Arguments.Add("-fsigned-char");             // Treat chars as signed //@todo android: any concerns about ABI compatibility with libs here?
-			}
-			else if (CompileEnvironment.Architecture == UnrealArch.X64)
-			{
-				Arguments.Add(ToolchainParamsx64);
-				Arguments.Add("-fno-omit-frame-pointer");
-				Arguments.Add("-march=atom");
-				Arguments.Add("-DPLATFORM_ANDROID_X64=1");
-			}
 		}
 
 		/// <inheritdoc/>
diff --git a/Engine/Source/Programs/UnrealBuildTool/Platform/Linux/LinuxToolChain.cs b/Engine/Source/Programs/UnrealBuildTool/Platform/Linux/LinuxToolChain.cs
--- a/Engine/Source/Programs/UnrealBuildTool/Platform/Linux/LinuxToolChain.cs
+++ b/Engine/Source/Programs/UnrealBuildTool/Platform/Linux/LinuxToolChain.cs
@@ -266,21 +266,6 @@ namespace UnrealBuildTool
 			return Out.ToString();
 		}
 
-		/// <summary>
-		/// Architecture-specific compiler switches
-		/// </summary>
-		static string ArchitectureSpecificSwitches(UnrealArch Architecture)
-		{
-			string Result = "";
-
-			if (Architecture == UnrealArch.Arm64)
-			{
-				Result += "-fsigned-char";
-			}
-
-			return Result;
-		}
-
 		private static bool ShouldUseLibcxx()
 		{
 			// set UE_LINUX_USE_LIBCXX to either 0 or 1. If unset, defaults to 1.
@@ -513,13 +498,6 @@ namespace UnrealBuildTool
 				Arguments.Add("-fbinutils-version=2.36");
 			}
 
-			if (CompileEnvironment.Architecture == UnrealArch.Arm64)
-			{
-				Arguments.Add("-funwind-tables");               // generate unwind tables as they are needed for backtrace (on x86(64) they are generated implicitly)
-			}
-
-			Arguments.Add(ArchitectureSpecificSwitches(CompileEnvironment.Architecture));
-
 			Arguments.Add("-fno-math-errno");               // do not assume that math ops have side effects
 
 			Arguments.Add(GetRTTIFlag(CompileEnvironment)); // flag for run-time type info
@@ -1119,13 +1097,11 @@ namespace UnrealBuildTool
 			LinkAction.RootPaths = LinkEnvironment.RootPaths;
 			LinkAction.WorkingDirectory = Unreal.EngineSourceDirectory;
 
-			string LinkCommandString;
-			LinkCommandString = "\"" + Info.Clang + "\"";
+			string LinkCommandString = $"\"{Info.Clang}\"";
 
 			// Get link arguments.
 			List<string> LinkArguments = new List<string>();
 			GetLinkArguments(LinkEnvironment, LinkArguments);
-			LinkCommandString += " " + String.Join(' ', LinkArguments);
 
 			// Tell the action that we're building an import library here and it should conditionally be
 			// ignored as a prerequisite for other actions
@@ -1172,7 +1148,7 @@ namespace UnrealBuildTool
 			LinkAction.StatusDescription = Path.GetFileName(OutputFile.AbsolutePath);
 
 			// Add the output file to the command-line.
-			LinkCommandString += String.Format(" -o \"{0}\"", OutputFile.AbsolutePath);
+			LinkArguments.Add(GetOutputFileArgument(OutputFile, LinkEnvironment.RootPaths));
 
 			// Add the input files to a response file, and pass the response file on the command-line.
 			List<string> ResponseLines = new List<string>();
@@ -1180,18 +1156,18 @@ namespace UnrealBuildTool
 			{
 				if (InputFile.HasExtension(".dynlist"))
 				{
-					ResponseLines.Add($"--dynamic-list=\"{NormalizeCommandLinePath(InputFile)}\"");
+					ResponseLines.Add($"--dynamic-list=\"{NormalizeCommandLinePath(InputFile, LinkEnvironment.RootPaths)}\"");
 				}
 				else
 				{
-					ResponseLines.Add($"\"{NormalizeCommandLinePath(InputFile)}\"");
+					ResponseLines.Add($"\"{NormalizeCommandLinePath(InputFile, LinkEnvironment.RootPaths)}\"");
 				}
 				LinkAction.PrerequisiteItems.Add(InputFile);
 			}
 
 			if (LinkEnvironment.bIsBuildingDLL)
 			{
-				ResponseLines.Add(String.Format(" -soname=\"{0}\"", OutputFile.Location.GetFileName()));
+				ResponseLines.Add($"-soname=\"{OutputFile.Location.GetFileName()}\"");
 			}
 
 			// Start with the configured LibraryPaths and also add paths to any libraries that
@@ -1232,15 +1208,15 @@ namespace UnrealBuildTool
 					if (CrossCompiling() && RelativePath.StartsWith(Unreal.RootDirectory.FullName))
 					{
 						// do not replace directly, but take care to avoid potential double slashes or missed slashes
-						string PathFromRootDir = RelativePath.Replace(Unreal.RootDirectory.FullName, "");
+						string PathFromRootDir = RelativePath.Replace(Unreal.RootDirectory.FullName, "").Replace('\\', '/').Trim('/');
 						// Path.Combine doesn't combine these properly
-						RelativePath = ((PathFromRootDir.StartsWith("\\") || PathFromRootDir.StartsWith("/")) ? "..\\..\\.." : "..\\..\\..\\") + PathFromRootDir;
+						RelativePath = $"../../../{PathFromRootDir}";
 					}
 
 					if (!RPaths.Contains(RelativePath))
 					{
 						RPaths.Add(RelativePath);
-						ResponseLines.Add(String.Format(" -rpath=\"${{ORIGIN}}/{0}\"", RelativePath.Replace('\\', '/')));
+						ResponseLines.Add($"-rpath=\"${{ORIGIN}}/{RelativePath.Replace('\\', '/')}\"");
 					}
 				}
 			}
@@ -1275,25 +1251,24 @@ namespace UnrealBuildTool
 				if (CrossCompiling() && RelativePath.StartsWith(Unreal.RootDirectory.FullName))
 				{
 					// do not replace directly, but take care to avoid potential double slashes or missed slashes
-					string PathFromRootDir = RelativePath.Replace(Unreal.RootDirectory.FullName, "");
+					string PathFromRootDir = RelativePath.Replace(Unreal.RootDirectory.FullName, "").Replace('\\', '/').Trim('/');
 					// Path.Combine doesn't combine these properly
-					RelativePath = ((PathFromRootDir.StartsWith("\\") || PathFromRootDir.StartsWith("/")) ? "..\\..\\.." : "..\\..\\..\\") + PathFromRootDir;
+					RelativePath = $"../../../{PathFromRootDir}";
 				}
 
 				if (!RPaths.Contains(RelativePath))
 				{
 					RPaths.Add(RelativePath);
-					ResponseLines.Add(String.Format(" -rpath=\"${{ORIGIN}}/{0}\"", RelativePath.Replace('\\', '/')));
+					ResponseLines.Add($"-rpath=\"${{ORIGIN}}/{RelativePath.Replace('\\', '/')}\"");
 				}
 			}
 
-			ResponseLines.Add(String.Format(" -rpath-link=\"{0}\"", Path.GetDirectoryName(OutputFile.AbsolutePath)));
+			ResponseLines.Add($"-rpath-link=\"{NormalizeCommandLinePath(OutputFile.Location.Directory, LinkEnvironment.RootPaths)}\"");
 
 			// Add the library paths to the argument list.
 			foreach (DirectoryReference LibraryPath in AllLibraryPaths)
 			{
-				// use absolute paths because of FixDependencies script again
-				ResponseLines.Add(String.Format(" -L\"{0}\"", LibraryPath.FullName.Replace('\\', '/')));
+				ResponseLines.Add($"-L\"{NormalizeCommandLinePath(LibraryPath, LinkEnvironment.RootPaths)}\"");
 			}
 
 			List<string> EngineAndGameLibrariesLinkFlags = new List<string>();
@@ -1303,17 +1278,17 @@ namespace UnrealBuildTool
 			// when linking with --as-needed.
 			// Move external libraries to a separate --start-group/--end-group to fix it (and also make groups smaller and faster to link).
 			// See https://github.com/EpicGames/UnrealEngine/pull/2778 and https://github.com/EpicGames/UnrealEngine/pull/2793 for discussion
-			string ExternalLibraries = "";
+			List<string> ExternalLibraries = [];
 
 			// add libraries in a library group
-			ResponseLines.Add(String.Format(" --start-group"));
+			ResponseLines.Add("--start-group");
 
 			foreach (string AdditionalLibrary in AdditionalLibraries)
 			{
 				if (String.IsNullOrEmpty(Path.GetDirectoryName(AdditionalLibrary)))
 				{
 					// library was passed just like "jemalloc", turn it into -ljemalloc
-					ExternalLibraries += String.Format(" -l{0}", AdditionalLibrary);
+					ExternalLibraries.Add($"-l{AdditionalLibrary}");
 				}
 				else if (Path.GetExtension(AdditionalLibrary) == ".a")
 				{
@@ -1321,18 +1296,18 @@ namespace UnrealBuildTool
 					string AbsoluteAdditionalLibrary = Path.GetFullPath(AdditionalLibrary);
 					if (AbsoluteAdditionalLibrary.Contains(' '))
 					{
-						AbsoluteAdditionalLibrary = String.Format("\"{0}\"", AbsoluteAdditionalLibrary);
+						AbsoluteAdditionalLibrary = $"\"{AbsoluteAdditionalLibrary}\"";
 					}
 					AbsoluteAdditionalLibrary = AbsoluteAdditionalLibrary.Replace('\\', '/');
 
 					// libcrypto/libssl contain number of functions that are being used in different DSOs. FIXME: generalize?
 					if (LinkEnvironment.bIsBuildingDLL && (AbsoluteAdditionalLibrary.Contains("libcrypto") || AbsoluteAdditionalLibrary.Contains("libssl")))
 					{
-						ResponseLines.Add(" --whole-archive " + AbsoluteAdditionalLibrary + " --no-whole-archive");
+						ResponseLines.Add($"--whole-archive {AbsoluteAdditionalLibrary} --no-whole-archive");
 					}
 					else
 					{
-						ResponseLines.Add(" " + AbsoluteAdditionalLibrary);
+						ResponseLines.Add(AbsoluteAdditionalLibrary);
 					}
 
 					LinkAction.PrerequisiteItems.Add(FileItem.GetItemByPath(AdditionalLibrary));
@@ -1354,9 +1329,9 @@ namespace UnrealBuildTool
 						// The library exists, but it is not prefixed with "lib", so force the
 						// linker to find it without that prefix by prepending a colon to
 						// the file name.
-						LibName = String.Format(":{0}", LibraryDependency.Name);
+						LibName = $":{LibraryDependency.Name}";
 					}
-					string LibLinkFlag = String.Format(" -l{0}", LibName);
+					string LibLinkFlag = $"-l{LibName}";
 
 					if (LinkEnvironment.bIsBuildingDLL && LinkEnvironment.bIsCrossReferenced)
 					{
@@ -1368,71 +1343,76 @@ namespace UnrealBuildTool
 						// it is important to add this exactly to the same place where the missing libraries would have been, it will be replaced later
 						if (!ExternalLibraries.Contains("--allow-shlib-undefined"))
 						{
-							ExternalLibraries += String.Format(" -Wl,--allow-shlib-undefined");
+							ExternalLibraries.Add("-Wl,--allow-shlib-undefined");
 						}
 					}
 					else
 					{
 						LinkAction.PrerequisiteItems.Add(LibraryDependency);
-						ExternalLibraries += LibLinkFlag;
+						ExternalLibraries.Add(LibLinkFlag);
 					}
 				}
 			}
-			ResponseLines.Add(" --end-group");
+			ResponseLines.Add("--end-group");
 
-			FileReference ResponseFileName = GetResponseFileName(LinkEnvironment, OutputFile);
-			FileItem ResponseFileItem = Graph.CreateIntermediateTextFile(ResponseFileName, ResponseLines);
+			{
+				FileReference LinkResponseFileName = GetResponseFileName(LinkEnvironment, FileItem.GetItemByFileReference(FileReference.Combine(OutputFile.Location.Directory, $"{OutputFile.Location.GetFileNameWithoutAnyExtensions()}.input{OutputFile.Location.GetExtension()}")));
+				FileItem LinkResponseFileItem = Graph.CreateIntermediateTextFile(LinkResponseFileName, ResponseLines);
 
-			LinkCommandString += String.Format(" -Wl,@\"{0}\"", ResponseFileName);
-			LinkAction.PrerequisiteItems.Add(ResponseFileItem);
+				LinkArguments.Add($"-Wl,@\"{NormalizeCommandLinePath(LinkResponseFileName, LinkEnvironment.RootPaths)}\"");
+				LinkAction.PrerequisiteItems.Add(LinkResponseFileItem);
+			}
 
-			LinkCommandString += " -Wl,--start-group";
-			LinkCommandString += ExternalLibraries;
+			LinkArguments.Add("-Wl,--start-group");
+			LinkArguments.AddRange(ExternalLibraries);
 
 			// make unresolved symbols an error, unless a) building a cross-referenced DSO  b) we opted out
 			if ((!LinkEnvironment.bIsBuildingDLL || !LinkEnvironment.bIsCrossReferenced) && !LinkEnvironment.bIgnoreUnresolvedSymbols)
 			{
 				// This will make the linker report undefined symbols the current module, but ignore in the dependent DSOs.
 				// It is tempting, but may not be possible to change that report-all - due to circular dependencies between our libs.
-				LinkCommandString += " -Wl,--unresolved-symbols=ignore-in-shared-libs";
+				LinkArguments.Add("-Wl,--unresolved-symbols=ignore-in-shared-libs");
 			}
-			LinkCommandString += " -Wl,--end-group";
+			LinkArguments.Add("-Wl,--end-group");
 
-			LinkCommandString += " -lrt"; // needed for clock_gettime()
-			LinkCommandString += " -lm"; // math
+			LinkArguments.Add("-lrt"); // needed for clock_gettime()
+			LinkArguments.Add("-lm"); // math
 
 			if (ShouldUseLibcxx())
 			{
 				// libc++ and its abi lib
-				LinkCommandString += " -nodefaultlibs";
-				LinkCommandString += " -L" + "ThirdParty/Unix/LibCxx/lib/Unix/" + LinkEnvironment.Architecture.LinuxName + "/";
-				LinkCommandString += " " + "ThirdParty/Unix/LibCxx/lib/Unix/" + LinkEnvironment.Architecture.LinuxName + "/libc++.a";
-				LinkCommandString += " " + "ThirdParty/Unix/LibCxx/lib/Unix/" + LinkEnvironment.Architecture.LinuxName + "/libc++abi.a";
-				LinkCommandString += " -lm";
-				LinkCommandString += " -lc";
-				LinkCommandString += " -lpthread"; // pthread_mutex_trylock is missing from libc stubs
-				LinkCommandString += " -lgcc_s";
-				LinkCommandString += " -lgcc";
+				LinkArguments.Add("-nodefaultlibs");
+				LinkArguments.Add($"-LThirdParty/Unix/LibCxx/lib/Unix/{LinkEnvironment.Architecture.LinuxName}/");
+				LinkArguments.Add($"ThirdParty/Unix/LibCxx/lib/Unix/{LinkEnvironment.Architecture.LinuxName}/libc++.a");
+				LinkArguments.Add($"ThirdParty/Unix/LibCxx/lib/Unix/{LinkEnvironment.Architecture.LinuxName}/libc++abi.a");
+				LinkArguments.Add("-lm");
+				LinkArguments.Add("-lc");
+				LinkArguments.Add("-lpthread"); // pthread_mutex_trylock is missing from libc stubs
+				LinkArguments.Add("-lgcc_s");
+				LinkArguments.Add("-lgcc");
 			}
 
 			// these can be helpful for understanding the order of libraries or library search directories
 			if (PlatformSDK.bVerboseLinker)
 			{
-				LinkCommandString += " -Wl,--verbose";
-				LinkCommandString += " -Wl,--trace";
-				LinkCommandString += " -v";
+				LinkArguments.Add("-Wl,--verbose");
+				LinkArguments.Add("-Wl,--trace");
+				LinkArguments.Add("-v");
 			}
 
 			// Add the additional arguments specified by the environment.
-			LinkCommandString += LinkEnvironment.AdditionalArguments;
-			LinkCommandString = LinkCommandString.Replace("\\\\", "/");
-			LinkCommandString = LinkCommandString.Replace("\\", "/");
+			LinkArguments.Add(LinkEnvironment.AdditionalArguments.Replace("\\\\", "/").Replace("\\", "/"));
 
 			bool bUseCmdExe = BuildHostPlatform.Current.ShellType == ShellType.Cmd;
 			FileReference ShellBinary = BuildHostPlatform.Current.Shell;
 			string ExecuteSwitch = bUseCmdExe ? " /C" : ""; // avoid -c so scripts don't need +x
 
-			// Linux has issues with scripts and parameter expansion from curely brakets
+			FileReference ResponseFileName = GetResponseFileName(LinkEnvironment, OutputFile);
+			FileItem ResponseFileItem = Graph.CreateIntermediateTextFile(ResponseFileName, LinkArguments);
+			LinkCommandString += $" @\"{NormalizeCommandLinePath(ResponseFileName, LinkEnvironment.RootPaths)}\"";
+			LinkAction.PrerequisiteItems.Add(ResponseFileItem);
+
+			// Linux has issues with scripts and parameter expansion from curly braces
 			if (!bUseCmdExe)
 			{
 				LinkCommandString = LinkCommandString.Replace("{", "'{");
diff --git a/Engine/Source/Programs/UnrealBuildTool/ToolChain/ClangToolChain.cs b/Engine/Source/Programs/UnrealBuildTool/ToolChain/ClangToolChain.cs
--- a/Engine/Source/Programs/UnrealBuildTool/ToolChain/ClangToolChain.cs
+++ b/Engine/Source/Programs/UnrealBuildTool/ToolChain/ClangToolChain.cs
@@ -722,7 +722,15 @@ namespace UnrealBuildTool
 				// We have 'this' vs nullptr comparisons that get optimized away for newer versions of Clang, which is undesirable until we refactor these checks.
 				Arguments.Add("-fno-delete-null-pointer-checks");
 			}
+		}
 
+		/// <summary>
+		/// Compile arguments for architecture specific settings
+		/// </summary>
+		/// <param name="CompileEnvironment"></param>
+		/// <param name="Arguments"></param>
+		protected virtual void GetCompileArguments_Architecture(CppCompileEnvironment CompileEnvironment, List<string> Arguments)
+		{
 			// architecture (all but None are AVX)
 			if (CompileEnvironment.Architecture == UnrealArch.X64 && CompileEnvironment.MinCpuArchX64 != MinimumCpuArchitectureX64.None)
 			{
@@ -762,6 +770,12 @@ namespace UnrealBuildTool
 					Arguments.Add("-DPLATFORM_ALWAYS_HAS_AVX_512=1");
 				}
 			}
+
+			// Treat chars as signed on arm64
+			if (CompileEnvironment.Architecture == UnrealArch.Arm64)
+			{
+				Arguments.Add("-fsigned-char");
+			}
 		}
 
 		/// <summary>
@@ -787,6 +801,9 @@ namespace UnrealBuildTool
 			{
 				Arguments.Add("-fextend-lifetimes");
 			}
+
+			// Generate unwind tables as they are needed for backtrace
+			Arguments.Add("-funwind-tables");
 		}
 
 		/// <summary>
@@ -939,6 +956,9 @@ namespace UnrealBuildTool
 			// Add optimization flags to the argument list.
 			GetCompileArguments_Optimizations(CompileEnvironment, Arguments);
 
+			// Add architecture flags to the argument list.
+			GetCompileArguments_Architecture(CompileEnvironment, Arguments);
+
 			// Add debugging flags to the argument list.
 			GetCompileArguments_Debugging(CompileEnvironment, Arguments);
 
