[UBT]
* Changed bCompilerVFS to bUseVFS

diff --git a/Engine/Source/Programs/UnrealBuildTool/Configuration/ReadOnlyTargetRules.cs b/Engine/Source/Programs/UnrealBuildTool/Configuration/ReadOnlyTargetRules.cs
--- a/Engine/Source/Programs/UnrealBuildTool/Configuration/ReadOnlyTargetRules.cs
+++ b/Engine/Source/Programs/UnrealBuildTool/Configuration/ReadOnlyTargetRules.cs
@@ -425,7 +425,7 @@ namespace UnrealBuildTool
 
 		public bool bDeterministic => Inner.bDeterministic;
 
-		public bool bCompilerVFS => Inner.bCompilerVFS;
+		public bool bUseVFS => Inner.bUseVFS;
 
 		public bool bChainPCHs => Inner.bChainPCHs;
 
diff --git a/Engine/Source/Programs/UnrealBuildTool/Platform/Windows/VCToolChain.cs b/Engine/Source/Programs/UnrealBuildTool/Platform/Windows/VCToolChain.cs
--- a/Engine/Source/Programs/UnrealBuildTool/Platform/Windows/VCToolChain.cs
+++ b/Engine/Source/Programs/UnrealBuildTool/Platform/Windows/VCToolChain.cs
@@ -1984,7 +1984,7 @@ namespace UnrealBuildTool
 
 					if (CompileEnvironment.PrecompiledHeaderAction != PrecompiledHeaderAction.None)
 					{
-						if (!Target.bCompilerVFS)
+						if (!Target.bUseVFS)
 						{
 							// Unfortunately we require matching absolute paths for pch to be cached if not using vfs
 							CompileAction.ArtifactMode |= ArtifactMode.AbsolutePath;
diff --git a/Engine/Source/Programs/UnrealBuildTool/System/RootPaths.cs b/Engine/Source/Programs/UnrealBuildTool/System/RootPaths.cs
--- a/Engine/Source/Programs/UnrealBuildTool/System/RootPaths.cs
+++ b/Engine/Source/Programs/UnrealBuildTool/System/RootPaths.cs
@@ -200,7 +200,7 @@ class CppRootPaths : IEnumerable<(uint id, DirectoryReference vfs, DirectoryRefe
 		{
 			if (reference.IsUnderDirectory(item.Value.local))
 			{
-				vfsPath = item.Value.vfs.FullName + '/' + reference.FullName.Substring(item.Value.local.FullName.Length);
+				vfsPath = item.Value.vfs.FullName + reference.FullName.Substring(item.Value.local.FullName.Length);
 				if (OperatingSystem.IsWindows())
 				{
 					vfsPath = vfsPath.Replace('\\', '/');
@@ -213,7 +213,7 @@ class CppRootPaths : IEnumerable<(uint id, DirectoryReference vfs, DirectoryRefe
 		{
 			if (reference.IsUnderDirectory(item.Value.local))
 			{
-				vfsPath = item.Value.vfs.FullName + '/' + reference.FullName.Substring(item.Value.local.FullName.Length);
+				vfsPath = item.Value.vfs.FullName + reference.FullName.Substring(item.Value.local.FullName.Length);
 				if (OperatingSystem.IsWindows())
 				{
 					vfsPath = vfsPath.Replace('\\', '/');
