[UBT]
* Changed bCompilerVFS to bUseVFS

diff --git a/Engine/Source/Programs/UnrealBuildTool/Configuration/TargetRules.cs b/Engine/Source/Programs/UnrealBuildTool/Configuration/TargetRules.cs
--- a/Engine/Source/Programs/UnrealBuildTool/Configuration/TargetRules.cs
+++ b/Engine/Source/Programs/UnrealBuildTool/Configuration/TargetRules.cs
@@ -2121,10 +2121,11 @@ namespace UnrealBuildTool
 		/// <summary>
 		/// Experimental work in progress: Set flags to use toolchain virtual file system support, to generate consistent pathing in outputs
 		/// </summary>
-		[CommandLine("-CompilerVFS", Value = "true")]
+		[CommandLine("-VFS", Value = "true")]
+		[CommandLine("-NoVFS", Value = "false")]
 		[XmlConfigFile(Category = "BuildConfiguration")]
 		[RequiresUniqueBuildEnvironment]
-		public bool bCompilerVFS { get; set; }
+		public bool bUseVFS { get; set; };
 
 		/// <summary>
 		/// Whether PCHs should be chained when compiling with clang.
