[UBT]
* Changed bCompilerVFS to bUseVFS

diff --git a/Engine/Source/Programs/UnrealBuildTool/Configuration/TargetRules.cs b/Engine/Source/Programs/UnrealBuildTool/Configuration/TargetRules.cs
--- a/Engine/Source/Programs/UnrealBuildTool/Configuration/TargetRules.cs
+++ b/Engine/Source/Programs/UnrealBuildTool/Configuration/TargetRules.cs
@@ -2125,7 +2125,7 @@ namespace UnrealBuildTool
 		[CommandLine("-NoVFS", Value = "false")]
 		[XmlConfigFile(Category = "BuildConfiguration")]
 		[RequiresUniqueBuildEnvironment]
-		public bool bUseVFS { get; set; };
+		public bool bUseVFS { get; set; }
 
 		/// <summary>
 		/// Whether PCHs should be chained when compiling with clang.
