[UBT]
* Changed bCompilerVFS to bUseVFS

diff --git a/Engine/Source/Programs/UnrealBuildTool/ToolChain/UEToolChain.cs b/Engine/Source/Programs/UnrealBuildTool/ToolChain/UEToolChain.cs
--- a/Engine/Source/Programs/UnrealBuildTool/ToolChain/UEToolChain.cs
+++ b/Engine/Source/Programs/UnrealBuildTool/ToolChain/UEToolChain.cs
@@ -274,11 +274,8 @@ namespace UnrealBuildTool
 
 		public virtual void SetUpGlobalEnvironment(ReadOnlyTargetRules Target, CppCompileEnvironment GlobalCompileEnvironment, LinkEnvironment GlobalLinkEnvironment)
 		{
-			if (OperatingSystem.IsWindows())
-			{
-				GlobalCompileEnvironment.RootPaths.bUseVfs = Target.bCompilerVFS;
-				GlobalLinkEnvironment.RootPaths.bUseVfs = Target.bCompilerVFS;
-			}
+			GlobalCompileEnvironment.RootPaths.bUseVfs = Target.bUseVFS;
+			GlobalLinkEnvironment.RootPaths.bUseVfs = Target.bUseVFS;
 
 			GlobalCompileEnvironment.RootPaths[CppRootPathFolder.Root] = Unreal.RootDirectory;
 			GlobalLinkEnvironment.RootPaths[CppRootPathFolder.Root] = Unreal.RootDirectory;
