Detect if targets define allocator defines and throw an error to avoid conflicts with new allocator selection via UBT

diff --git a/Engine/Source/Programs/UnrealBuildTool/Configuration/UEBuildTarget.cs b/Engine/Source/Programs/UnrealBuildTool/Configuration/UEBuildTarget.cs
--- a/Engine/Source/Programs/UnrealBuildTool/Configuration/UEBuildTarget.cs
+++ b/Engine/Source/Programs/UnrealBuildTool/Configuration/UEBuildTarget.cs
@@ -5725,6 +5725,11 @@ namespace UnrealBuildTool
 
 			for (int i = 0; i < (int)StaticAllocatorType.None; ++i)
 			{
+				if (GlobalCompileEnvironment.Definitions.Contains(Mapping[i] + "=0") || GlobalCompileEnvironment.Definitions.Contains(Mapping[i] + "=1"))
+				{
+					throw new BuildException("Target `{0}` contains a definition for `{1}` allocator define. It conflicts with the new allocator type declatation via UBT. Please use `StaticAllocator = StaticAllocatorType.{2};` in your *.Target.cs instead", TargetName, Mapping[i], ((StaticAllocatorType)i).ToString());
+				}
+
 				string EnabledValue = Rules.StaticAllocator == (StaticAllocatorType)i ? "=1" : "=0";
 				GlobalCompileEnvironment.Definitions.Add(Mapping[i] + EnabledValue);
 			}
