UnrealBuildTool: Move a few common clang args to be shared

#rnx

diff --git a/Engine/Source/Programs/UnrealBuildTool/Platform/Android/AndroidToolChain.cs b/Engine/Source/Programs/UnrealBuildTool/Platform/Android/AndroidToolChain.cs
--- a/Engine/Source/Programs/UnrealBuildTool/Platform/Android/AndroidToolChain.cs
+++ b/Engine/Source/Programs/UnrealBuildTool/Platform/Android/AndroidToolChain.cs
@@ -547,6 +547,27 @@ namespace UnrealBuildTool
 			}
 		}
 
+		/// <inheritdoc/>
+		protected override void GetCompileArguments_Architecture(CppCompileEnvironment CompileEnvironment, List<string> Arguments)
+		{
+			base.GetCompileArguments_Architecture(CompileEnvironment, Arguments);
+
+			if (CompileEnvironment.Architecture == UnrealArch.Arm64)
+			{
+				Arguments.Add(ToolchainParamsArm64);
+				Arguments.Add("-D__arm64__");            // for some reason this isn't defined and needed for PhysX
+				Arguments.Add("-DPLATFORM_ANDROID_ARM64=1");
+				Arguments.Add("-march=armv8-a");
+			}
+			else if (CompileEnvironment.Architecture == UnrealArch.X64)
+			{
+				Arguments.Add(ToolchainParamsx64);
+				Arguments.Add("-DPLATFORM_ANDROID_X64=1");
+				Arguments.Add("-fno-omit-frame-pointer");
+				Arguments.Add("-march=atom");
+			}
+		}
+
 		/// <inheritdoc/>
 		protected override void GetCompileArguments_Debugging(CppCompileEnvironment CompileEnvironment, List<string> Arguments)
 		{
@@ -656,7 +677,6 @@ namespace UnrealBuildTool
 			Arguments.Add(GetRTTIFlag(CompileEnvironment));
 			Arguments.Add("-no-canonical-prefixes");
 			Arguments.Add("-fno-PIE");
-			Arguments.Add("-funwind-tables");           // Just generates any needed static data, affects no code
 			Arguments.Add("-fPIC");                     // Generates position-independent code (PIC) suitable for use in a shared library
 			Arguments.Add("-fno-strict-aliasing");      // Prevents unwanted or invalid optimizations that could produce incorrect code
 			Arguments.Add("-fno-short-enums");          // Do not allocate to an enum type only as many bytes as it needs for the declared range of possible values
@@ -664,22 +684,6 @@ namespace UnrealBuildTool
 			Arguments.Add("-D_FORTIFY_SOURCE=2");       // FORTIFY default
 			Arguments.Add($"-DPLATFORM_USED_NDK_VERSION_INTEGER={NDKApiLevel64Int}");       // NDK version
 			Arguments.Add("-DPLATFORM_64BITS=1");       // NDK version
-
-			if (CompileEnvironment.Architecture == UnrealArch.Arm64)
-			{
-				Arguments.Add(ToolchainParamsArm64);
-				Arguments.Add("-D__arm64__");            // for some reason this isn't defined and needed for PhysX
-				Arguments.Add("-DPLATFORM_ANDROID_ARM64=1");
-				Arguments.Add("-march=armv8-a");
-				Arguments.Add("-fsigned-char");             // Treat chars as signed //@todo android: any concerns about ABI compatibility with libs here?
-			}
-			else if (CompileEnvironment.Architecture == UnrealArch.X64)
-			{
-				Arguments.Add(ToolchainParamsx64);
-				Arguments.Add("-fno-omit-frame-pointer");
-				Arguments.Add("-march=atom");
-				Arguments.Add("-DPLATFORM_ANDROID_X64=1");
-			}
 		}
 
 		/// <inheritdoc/>
diff --git a/Engine/Source/Programs/UnrealBuildTool/Platform/Linux/LinuxToolChain.cs b/Engine/Source/Programs/UnrealBuildTool/Platform/Linux/LinuxToolChain.cs
--- a/Engine/Source/Programs/UnrealBuildTool/Platform/Linux/LinuxToolChain.cs
+++ b/Engine/Source/Programs/UnrealBuildTool/Platform/Linux/LinuxToolChain.cs
@@ -266,21 +266,6 @@ namespace UnrealBuildTool
 			return Out.ToString();
 		}
 
-		/// <summary>
-		/// Architecture-specific compiler switches
-		/// </summary>
-		static string ArchitectureSpecificSwitches(UnrealArch Architecture)
-		{
-			string Result = "";
-
-			if (Architecture == UnrealArch.Arm64)
-			{
-				Result += "-fsigned-char";
-			}
-
-			return Result;
-		}
-
 		private static bool ShouldUseLibcxx()
 		{
 			// set UE_LINUX_USE_LIBCXX to either 0 or 1. If unset, defaults to 1.
@@ -513,13 +498,6 @@ namespace UnrealBuildTool
 				Arguments.Add("-fbinutils-version=2.36");
 			}
 
-			if (CompileEnvironment.Architecture == UnrealArch.Arm64)
-			{
-				Arguments.Add("-funwind-tables");               // generate unwind tables as they are needed for backtrace (on x86(64) they are generated implicitly)
-			}
-
-			Arguments.Add(ArchitectureSpecificSwitches(CompileEnvironment.Architecture));
-
 			Arguments.Add("-fno-math-errno");               // do not assume that math ops have side effects
 
 			Arguments.Add(GetRTTIFlag(CompileEnvironment)); // flag for run-time type info
diff --git a/Engine/Source/Programs/UnrealBuildTool/ToolChain/ClangToolChain.cs b/Engine/Source/Programs/UnrealBuildTool/ToolChain/ClangToolChain.cs
--- a/Engine/Source/Programs/UnrealBuildTool/ToolChain/ClangToolChain.cs
+++ b/Engine/Source/Programs/UnrealBuildTool/ToolChain/ClangToolChain.cs
@@ -722,7 +722,15 @@ namespace UnrealBuildTool
 				// We have 'this' vs nullptr comparisons that get optimized away for newer versions of Clang, which is undesirable until we refactor these checks.
 				Arguments.Add("-fno-delete-null-pointer-checks");
 			}
+		}
 
+		/// <summary>
+		/// Compile arguments for architecture specific settings
+		/// </summary>
+		/// <param name="CompileEnvironment"></param>
+		/// <param name="Arguments"></param>
+		protected virtual void GetCompileArguments_Architecture(CppCompileEnvironment CompileEnvironment, List<string> Arguments)
+		{
 			// architecture (all but None are AVX)
 			if (CompileEnvironment.Architecture == UnrealArch.X64 && CompileEnvironment.MinCpuArchX64 != MinimumCpuArchitectureX64.None)
 			{
@@ -762,6 +770,12 @@ namespace UnrealBuildTool
 					Arguments.Add("-DPLATFORM_ALWAYS_HAS_AVX_512=1");
 				}
 			}
+
+			// Treat chars as signed on arm64
+			if (CompileEnvironment.Architecture == UnrealArch.Arm64)
+			{
+				Arguments.Add("-fsigned-char");
+			}
 		}
 
 		/// <summary>
@@ -787,6 +801,9 @@ namespace UnrealBuildTool
 			{
 				Arguments.Add("-fextend-lifetimes");
 			}
+
+			// Generate unwind tables as they are needed for backtrace
+			Arguments.Add("-funwind-tables");
 		}
 
 		/// <summary>
@@ -939,6 +956,9 @@ namespace UnrealBuildTool
 			// Add optimization flags to the argument list.
 			GetCompileArguments_Optimizations(CompileEnvironment, Arguments);
 
+			// Add architecture flags to the argument list.
+			GetCompileArguments_Architecture(CompileEnvironment, Arguments);
+
 			// Add debugging flags to the argument list.
 			GetCompileArguments_Debugging(CompileEnvironment, Arguments);
 
