[UBT]
* Fixed so extra root paths are inherited from private dependency modules.
* Fixed up modules to have proper extra root paths

diff --git a/Engine/Plugins/Experimental/GameInputWindows/Source/GameInputWindows/GameInputWindows.build.cs b/Engine/Plugins/Experimental/GameInputWindows/Source/GameInputWindows/GameInputWindows.build.cs
--- a/Engine/Plugins/Experimental/GameInputWindows/Source/GameInputWindows/GameInputWindows.build.cs
+++ b/Engine/Plugins/Experimental/GameInputWindows/Source/GameInputWindows/GameInputWindows.build.cs
@@ -8,7 +8,7 @@ namespace UnrealBuildTool.Rules
 	{
 		public GameInputWindows(ReadOnlyTargetRules Target) : base(Target)
 		{
-			PublicDependencyModuleNames.AddRange(
+			PrivateDependencyModuleNames.AddRange(
 				new string[]
 				{
 					"Core",
diff --git a/Engine/Plugins/Runtime/GameInput/Source/GameInputBase/GameInputBase.build.cs b/Engine/Plugins/Runtime/GameInput/Source/GameInputBase/GameInputBase.build.cs
--- a/Engine/Plugins/Runtime/GameInput/Source/GameInputBase/GameInputBase.build.cs
+++ b/Engine/Plugins/Runtime/GameInput/Source/GameInputBase/GameInputBase.build.cs
@@ -56,7 +56,8 @@ namespace UnrealBuildTool.Rules
 			}
 			
 			string GRDKLatestPath = Environment.GetEnvironmentVariable("GRDKLatest");
-			
+			ExtraRootPath = ("GRDKLatest", GRDKLatestPath);
+
 			// GameInput.h is from    <GRDKLatest>\GameKit\Include\GameInput.h
 			
 			string IncludePath = Path.Combine(GRDKLatestPath, "GameKit", "Include");
diff --git a/Engine/Source/Programs/UnrealBuildTool/Configuration/UEBuildModuleCPP.cs b/Engine/Source/Programs/UnrealBuildTool/Configuration/UEBuildModuleCPP.cs
--- a/Engine/Source/Programs/UnrealBuildTool/Configuration/UEBuildModuleCPP.cs
+++ b/Engine/Source/Programs/UnrealBuildTool/Configuration/UEBuildModuleCPP.cs
@@ -2445,18 +2445,12 @@ namespace UnrealBuildTool
 
 			Result.RootPaths.AddExtraPath(Rules.ExtraRootPath);
 
-			// Extra root path is required if module depends on external module that has type libraries.
+			// Inherit extra root paths from dependencies
 			if (PrivateDependencyModules != null)
 			{
 				foreach (UEBuildModule Dependent in PrivateDependencyModules)
 				{
-					if (Dependent.Rules.TypeLibraries.Count != 0)
-					{
-						if (Dependent is UEBuildModuleExternal External)
-						{
-							Result.RootPaths.AddExtraPath(External.Rules.ExtraRootPath);
-						}
-					}
+					Result.RootPaths.AddExtraPath(Dependent.Rules.ExtraRootPath);
 				}
 			}
 
