UnrealBuildTool: Fix typo with cross-architecture paths

#rnx

diff --git a/Engine/Source/Programs/UnrealBuildTool/Platform/Windows/VCEnvironment.cs b/Engine/Source/Programs/UnrealBuildTool/Platform/Windows/VCEnvironment.cs
--- a/Engine/Source/Programs/UnrealBuildTool/Platform/Windows/VCEnvironment.cs
+++ b/Engine/Source/Programs/UnrealBuildTool/Platform/Windows/VCEnvironment.cs
@@ -302,7 +302,7 @@ namespace UnrealBuildTool
 			foreach (DirectoryReference ToolPath in ToolPaths)
 			{
 				DirectoryReference HostBinDir = DirectoryReference.Combine(ToolPath, MicrosoftPlatformSDK.MSVCHostUnrealArch.bIsX64 ? "x64" : "arm64");
-				DirectoryReference CrossBinDir = DirectoryReference.Combine(ToolPath, !MicrosoftPlatformSDK.MSVCHostUnrealArch.bIsX64 ? "arm64" : "x64");
+				DirectoryReference CrossBinDir = DirectoryReference.Combine(ToolPath, !MicrosoftPlatformSDK.MSVCHostUnrealArch.bIsX64 ? "x64" : "arm64");
 				if (FileReference.Exists(HostBinDir) && FileReference.Exists(CrossBinDir))
 				{
 					EpicGames.UBA.Utils.RegisterCrossArchitecturePath(HostBinDir.FullName, CrossBinDir.FullName);
