Fix iOS Remote Toolchain build issue where blueprint projects do not use the project settings for the remote Mac IP and username.

diff --git a/Engine/Source/Programs/UnrealBuildTool/Modes/BuildMode.cs b/Engine/Source/Programs/UnrealBuildTool/Modes/BuildMode.cs
--- a/Engine/Source/Programs/UnrealBuildTool/Modes/BuildMode.cs
+++ b/Engine/Source/Programs/UnrealBuildTool/Modes/BuildMode.cs
@@ -206,7 +206,9 @@ namespace UnrealBuildTool
 						FileReference BaseLogFile = Log.OutputFile ?? new FileReference(BaseLogFileName);
 						FileReference RemoteLogFile = FileReference.Combine(BaseLogFile.Directory, BaseLogFile.GetFileNameWithoutExtension() + "_Remote.txt");
 
-						RemoteMac RemoteMac = new RemoteMac(TargetDesc.ProjectFile, Logger);
+						int RemoteIniIndex = TargetDesc.AdditionalArguments.FindIndex(x => x.ToLower().StartsWith("-remoteini="));
+						string? RemoteIni = RemoteIniIndex == -1 ? null : new string(TargetDesc.AdditionalArguments[RemoteIniIndex].Substring(11));
+						RemoteMac RemoteMac = new RemoteMac(TargetDesc.ProjectFile, Logger, RemoteIni);
 						if (!RemoteMac.Build(TargetDesc, RemoteLogFile, bSkipPreBuildTargets, Logger))
 						{
 							return (int)CompilationResult.Unknown;
diff --git a/Engine/Source/Programs/UnrealBuildTool/Platform/IOS/IOSExports.cs b/Engine/Source/Programs/UnrealBuildTool/Platform/IOS/IOSExports.cs
--- a/Engine/Source/Programs/UnrealBuildTool/Platform/IOS/IOSExports.cs
+++ b/Engine/Source/Programs/UnrealBuildTool/Platform/IOS/IOSExports.cs
@@ -1,4 +1,4 @@
-﻿// Copyright Epic Games, Inc. All Rights Reserved.
+// Copyright Epic Games, Inc. All Rights Reserved.
 
 using System;
 using System.Diagnostics;
@@ -185,10 +185,10 @@ namespace UnrealBuildTool
 
 		public static void SetSecondaryRemoteMac(string ClientPlatform, FileReference ProjectFile, ILogger Logger)
 		{
-			RemoteMac Remote = new RemoteMac(ProjectFile, Logger, true, true);
+			RemoteMac Remote = new RemoteMac(ProjectFile, Logger, null, true, true);
 			Remote.RetrieveFilesGeneratedOnPrimaryMac(ProjectFile, Logger, ClientPlatform);
 
-			RemoteMac SecondaryRemote = new RemoteMac(ProjectFile, Logger, false);
+			RemoteMac SecondaryRemote = new RemoteMac(ProjectFile, Logger, null, false);
 			SecondaryRemote.UploadToSecondaryMac(ProjectFile, Logger);
 		}
 
diff --git a/Engine/Source/Programs/UnrealBuildTool/ToolChain/RemoteMac.cs b/Engine/Source/Programs/UnrealBuildTool/ToolChain/RemoteMac.cs
--- a/Engine/Source/Programs/UnrealBuildTool/ToolChain/RemoteMac.cs
+++ b/Engine/Source/Programs/UnrealBuildTool/ToolChain/RemoteMac.cs
@@ -126,9 +126,10 @@ namespace UnrealBuildTool
 		/// </summary>
 		/// <param name="ProjectFile">Project to read settings from</param>
 		/// <param name="Logger">Logger for output</param>
+		/// <param name="RemoteIniPath">Remote ini path, or null to use UnrealBuildTool.GetRemoteIniPath()</param>
 		/// <param name="bIsPrimary">Is the primary Remote Mac or not. True by default.</param>
 		/// <param name="bPrepareForSecondaryMac">Added arguments when preparing a secondary Mac for debug. False by default.</param>
-		public RemoteMac(FileReference? ProjectFile, ILogger Logger, bool bIsPrimary = true, bool bPrepareForSecondaryMac = false)
+		public RemoteMac(FileReference? ProjectFile, ILogger Logger, string? RemoteIniPath = null, bool bIsPrimary = true, bool bPrepareForSecondaryMac = false)
 		{
 			RsyncExe = FileReference.Combine(Unreal.EngineDirectory, "Extras", "ThirdPartyNotUE", "cwrsync", "bin", "rsync.exe");
 			SshExe = FileReference.Combine(Unreal.EngineDirectory, "Extras", "ThirdPartyNotUE", "cwrsync", "bin", "ssh.exe");
@@ -149,9 +150,13 @@ namespace UnrealBuildTool
 
 			// Get the project config file path
 			DirectoryReference? EngineIniPath = ProjectFile?.Directory;
-			if (EngineIniPath == null && UnrealBuildTool.GetRemoteIniPath() != null)
+			if (RemoteIniPath == null)
 			{
-				EngineIniPath = new DirectoryReference(UnrealBuildTool.GetRemoteIniPath()!);
+				RemoteIniPath = UnrealBuildTool.GetRemoteIniPath();
+			}
+			if (EngineIniPath == null && RemoteIniPath != null)
+			{
+				EngineIniPath = new DirectoryReference(RemoteIniPath!);
 			}
 			ConfigHierarchy Ini = ConfigCache.ReadHierarchy(ConfigHierarchyType.Engine, EngineIniPath, UnrealTargetPlatform.IOS);
 
