UnrealBuildTool: Suppress a few clang 19 warnings for now

#rnx

diff --git a/Engine/Source/Programs/UnrealBuildTool/ToolChain/ClangWarnings.cs b/Engine/Source/Programs/UnrealBuildTool/ToolChain/ClangWarnings.cs
--- a/Engine/Source/Programs/UnrealBuildTool/ToolChain/ClangWarnings.cs
+++ b/Engine/Source/Programs/UnrealBuildTool/ToolChain/ClangWarnings.cs
@@ -57,6 +57,11 @@ namespace UnrealBuildTool
 																																										// It could be easily fixed in the engine, but would create the risk of Win64 building fine and the code failing just on Clang-based platforms.
 																																										// We tend to use NAN just as a bit pattern, and not rely on it in calculations, so it should be reasonably safe to disable.
 			}
+			if (ClangVersion >= new VersionNumber(19))
+			{
+				Arguments.Add("-Wno-extra-qualification");              // https://clang.llvm.org/docs/DiagnosticsReference.html#wextra-qualification					// new warning for clang 19
+				Arguments.Add("-Wno-shorten-64-to-32");                 // https://clang.llvm.org/docs/DiagnosticsReference.html#wshorten-64-to-32						// new warning for clang 19
+			}
 
 			Arguments.Add("-Wno-gnu-string-literal-operator-template"); // https://clang.llvm.org/docs/DiagnosticsReference.html#wgnu-string-literal-operator-template	// We use this feature to allow static FNames.
 			Arguments.Add("-Wno-inconsistent-missing-override");        // https://clang.llvm.org/docs/DiagnosticsReference.html#winconsistent-missing-override			// ?? no reason given
