PR #12547: Added support of the analysis-paths flag in integration with PVS-Studio

diff --git a/Engine/Source/Programs/UnrealBuildTool/Platform/Windows/PVSToolChain.cs b/Engine/Source/Programs/UnrealBuildTool/Platform/Windows/PVSToolChain.cs
--- a/Engine/Source/Programs/UnrealBuildTool/Platform/Windows/PVSToolChain.cs
+++ b/Engine/Source/Programs/UnrealBuildTool/Platform/Windows/PVSToolChain.cs
@@ -913,6 +913,11 @@ namespace UnrealBuildTool
 					ConfigFileContents.Append("new-output-format=yes\n");
 				}
 
+				if (Target.bStaticAnalyzerProjectOnly && AnalyzerVersion.CompareTo(new Version("7.34")) >= 0)
+				{
+					ConfigFileContents.Append($"analysis-paths=skip={Unreal.EngineSourceDirectory}\n");
+				}
+
 				string BaseFileName = PreprocessedFileItem.Location.GetFileName();
 
 				FileReference ConfigFileLocation = FileReference.Combine(OutputDir, BaseFileName + ".cfg");
