UnrealBuildTool: Allow remapping VFS folder names for clarity

#rnx

diff --git a/Engine/Source/Programs/UnrealBuildTool/Platform/Windows/PVSToolChain.cs b/Engine/Source/Programs/UnrealBuildTool/Platform/Windows/PVSToolChain.cs
--- a/Engine/Source/Programs/UnrealBuildTool/Platform/Windows/PVSToolChain.cs
+++ b/Engine/Source/Programs/UnrealBuildTool/Platform/Windows/PVSToolChain.cs
@@ -645,8 +645,17 @@ namespace UnrealBuildTool
 			base.SetUpGlobalEnvironment(Target, GlobalCompileEnvironment, GlobalLinkEnvironment);
 			InnerToolChain.SetUpGlobalEnvironment(Target, GlobalCompileEnvironment, GlobalLinkEnvironment);
 
-			GlobalCompileEnvironment.RootPaths[CppRootPathFolder.Analyzer] = AnalyzerFile.Directory;
-			GlobalLinkEnvironment.RootPaths[CppRootPathFolder.Analyzer] = AnalyzerFile.Directory;
+			if (!AnalyzerFile.IsUnderDirectory(Unreal.RootDirectory))
+			{
+				GlobalCompileEnvironment.RootPaths.AddExtraPath(("PVSAnalyzer", AnalyzerFile.Directory.FullName));
+				GlobalLinkEnvironment.RootPaths.AddExtraPath(("PVSAnalyzer", AnalyzerFile.Directory.FullName));
+			}
+
+			if (LicenseFile != null && !LicenseFile.IsUnderDirectory(Unreal.RootDirectory))
+			{
+				GlobalCompileEnvironment.RootPaths.AddExtraPath(("PVSLicense", LicenseFile.Directory.FullName));
+				GlobalLinkEnvironment.RootPaths.AddExtraPath(("PVSLicense", LicenseFile.Directory.FullName));
+			}
 		}
   
 		public override void SetEnvironmentVariables()
diff --git a/Engine/Source/Programs/UnrealBuildTool/Platform/Windows/VCToolChain.cs b/Engine/Source/Programs/UnrealBuildTool/Platform/Windows/VCToolChain.cs
--- a/Engine/Source/Programs/UnrealBuildTool/Platform/Windows/VCToolChain.cs
+++ b/Engine/Source/Programs/UnrealBuildTool/Platform/Windows/VCToolChain.cs
@@ -3394,10 +3394,16 @@ namespace UnrealBuildTool
 		{
 			base.SetUpGlobalEnvironment(Target, GlobalCompileEnvironment, GlobalLinkEnvironment);
 
+			GlobalCompileEnvironment.RootPaths.AddFolderName(CppRootPathFolder.Compiler, EnvVars.Compiler.IsMSVC() ? "MSVC" : EnvVars.Compiler.ToString());
+			GlobalLinkEnvironment.RootPaths.AddFolderName(CppRootPathFolder.Compiler, EnvVars.Compiler.IsMSVC() ? "MSVC" : EnvVars.Compiler.ToString());
+
 			GlobalCompileEnvironment.RootPaths[CppRootPathFolder.Compiler] = EnvVars.CompilerDir;
 			GlobalLinkEnvironment.RootPaths[CppRootPathFolder.Compiler] = EnvVars.CompilerDir;
 			if (EnvVars.CompilerDir != EnvVars.ToolChainDir)
 			{
+				GlobalCompileEnvironment.RootPaths.AddFolderName(CppRootPathFolder.Toolchain, EnvVars.ToolChain.IsMSVC() ? "MSVC" : EnvVars.ToolChain.ToString());
+				GlobalLinkEnvironment.RootPaths.AddFolderName(CppRootPathFolder.Toolchain, EnvVars.ToolChain.IsMSVC() ? "MSVC" : EnvVars.ToolChain.ToString());
+
 				GlobalCompileEnvironment.RootPaths[CppRootPathFolder.Toolchain] = EnvVars.ToolChainDir;
 				GlobalLinkEnvironment.RootPaths[CppRootPathFolder.Toolchain] = EnvVars.ToolChainDir;
 			}
diff --git a/Engine/Source/Programs/UnrealBuildTool/System/Action.cs b/Engine/Source/Programs/UnrealBuildTool/System/Action.cs
--- a/Engine/Source/Programs/UnrealBuildTool/System/Action.cs
+++ b/Engine/Source/Programs/UnrealBuildTool/System/Action.cs
@@ -577,7 +577,7 @@ namespace UnrealBuildTool
 				Action.DeleteItems.UnionWith(DeleteItems.Select(x => FileItem.GetItemByPath(x)));
 			}
 
-			if (Object.TryGetObjectArrayField("RootPaths", out JsonObject[]? RootPaths))
+			if (Object.TryGetObjectField("RootPaths", out JsonObject? RootPaths))
 			{
 				Action.RootPaths.Read(RootPaths);
 			}
@@ -683,9 +683,9 @@ namespace UnrealBuildTool
 			}
 			Writer.WriteArrayEnd();
 
-			Writer.WriteArrayStart("RootPaths");
+			Writer.WriteObjectStart("RootPaths");
 			Action.RootPaths.Write(Writer);
-			Writer.WriteArrayEnd();
+			Writer.WriteObjectEnd();
 
 			if (Action.DependencyListFile != null)
 			{
diff --git a/Engine/Source/Programs/UnrealBuildTool/System/RootPaths.cs b/Engine/Source/Programs/UnrealBuildTool/System/RootPaths.cs
--- a/Engine/Source/Programs/UnrealBuildTool/System/RootPaths.cs
+++ b/Engine/Source/Programs/UnrealBuildTool/System/RootPaths.cs
@@ -4,6 +4,7 @@ using EpicGames.Core;
 using System;
 using System.Collections.Generic;
 using System.Diagnostics.CodeAnalysis;
+using System.Linq;
 using UnrealBuildBase;
 
 /// <summary>
@@ -37,14 +38,14 @@ public enum CppRootPathFolder
 	WinSDK,
 
 	/// <summary>
-	///  Additional platform SDK folder (if needed)
+	/// Additional platform SDK folder (if needed)
 	/// </summary>
 	PlatformSDK,
 
 	/// <summary>
-	/// Analyzer folder (if needed)
+	/// Number of entries in enum
 	/// </summary>
-	Analyzer,
+	Count,
 }
 
 /// <summary>
@@ -61,6 +62,7 @@ class CppRootPaths : IEnumerable<(uint id, DirectoryReference vfs, DirectoryRefe
 	{
 		_lookup = [];
 		_extras = [];
+		_folderNames = [];
 	}
 
 	/// <summary>
@@ -70,7 +72,17 @@ class CppRootPaths : IEnumerable<(uint id, DirectoryReference vfs, DirectoryRefe
 	{
 		{
 			int count = Reader.ReadInt();
-			_lookup = new SortedDictionary<CppRootPathFolder, (DirectoryReference vfs, DirectoryReference local)>();
+			_folderNames = [];
+			for (int idx = 0; idx < count; idx++)
+			{
+				CppRootPathFolder folder = (CppRootPathFolder)Reader.ReadInt();
+				string name = Reader.ReadString()!;
+				_folderNames[folder] = name;
+			}
+		}
+		{
+			int count = Reader.ReadInt();
+			_lookup = [];
 			for (int idx = 0; idx < count; idx++)
 			{
 				CppRootPathFolder folder = (CppRootPathFolder)Reader.ReadInt();
@@ -80,7 +92,7 @@ class CppRootPaths : IEnumerable<(uint id, DirectoryReference vfs, DirectoryRefe
 		}
 		{
 			int count = Reader.ReadInt();
-			_extras = new SortedDictionary<string, (DirectoryReference vfs, DirectoryReference local)>();
+			_extras = [];
 			for (int idx = 0; idx < count; idx++)
 			{
 				string id = Reader.ReadString()!;
@@ -100,6 +112,7 @@ class CppRootPaths : IEnumerable<(uint id, DirectoryReference vfs, DirectoryRefe
 	{
 		_lookup = new(Other._lookup);
 		_extras = new(Other._extras);
+		_folderNames = new(Other._folderNames);
 		bUseVfs = Other.bUseVfs;
 	}
 
@@ -108,6 +121,7 @@ class CppRootPaths : IEnumerable<(uint id, DirectoryReference vfs, DirectoryRefe
 	/// </summary>
 	public void Write(BinaryArchiveWriter Writer)
 	{
+		Writer.WriteDictionary(_folderNames, Key => Writer.WriteInt((int)Key), Value => Writer.WriteString(Value));
 		Writer.WriteDictionary(_lookup, Key => Writer.WriteInt((int)Key), Value => Writer.WriteCompactDirectoryReference(Value.local));
 		Writer.WriteDictionary(_extras, Key => Writer.WriteString(Key), Value => Writer.WriteCompactDirectoryReference(Value.local));
 		Writer.WriteBool(bUseVfs);
@@ -120,6 +134,10 @@ class CppRootPaths : IEnumerable<(uint id, DirectoryReference vfs, DirectoryRefe
 	{
 		set
 		{
+			if (key == CppRootPathFolder.Project)
+			{
+				AddFolderName(CppRootPathFolder.Project, value.GetDirectoryName());
+			}
 			_lookup[key] = (AsVfsPath(key), value);
 		}
 	}
@@ -134,7 +152,7 @@ class CppRootPaths : IEnumerable<(uint id, DirectoryReference vfs, DirectoryRefe
 			yield return ((uint)kv.Key, kv.Value.vfs, kv.Value.local);
 		}
 
-		uint index = 7; // Max of CppRootPathFolder
+		uint index = (uint)CppRootPathFolder.Count; // Max of CppRootPathFolder
 		foreach (KeyValuePair<string, (DirectoryReference vfs, DirectoryReference local)> kv in _extras)
 		{
 			yield return (index++, kv.Value.vfs, kv.Value.local);
@@ -228,22 +246,39 @@ class CppRootPaths : IEnumerable<(uint id, DirectoryReference vfs, DirectoryRefe
 	/// <summary>
 	/// Populate from json objects
 	/// </summary>
-	public void Read(JsonObject[] rootPaths)
+	public void Read(JsonObject Object)
 	{
-		foreach (JsonObject root in rootPaths)
+		if (Object.TryGetObjectArrayField("Lookup", out JsonObject[]? lookupItems))
 		{
-			DirectoryReference local = DirectoryReference.FromString(root.GetStringField("Path"))!;
-
-			if (root.TryGetEnumField<CppRootPathFolder>("Root", out CppRootPathFolder folder))
+			foreach (JsonObject item in lookupItems)
 			{
-				_lookup.Add(folder, (AsVfsPath(folder), local));
+				DirectoryReference local = new(item.GetStringField("Path"));
+				if (item.TryGetEnumField("Root", out CppRootPathFolder folder))
+				{
+					_lookup.Add(folder, (AsVfsPath(folder), local));
+
+					if (item.TryGetStringField("Root", out string? folderName))
+					{
+						_folderNames.Add(folder, folderName);
+					}
+				}
 			}
-			else
+		}
+
+		if (Object.TryGetObjectArrayField("Extras", out JsonObject[]? extraItems))
+		{
+			foreach (JsonObject item in extraItems)
 			{
-				string id = root.GetStringField("Id");
+				string id = item.GetStringField("Id");
+				DirectoryReference local = new(item.GetStringField("Path"));
 				_extras.Add(id, (AsVfsPath(id), local));
 			}
 		}
+
+		if (Object.TryGetBoolField("bUseVfs", out bool useVfs))
+		{
+			bUseVfs = useVfs;
+		}
 	}
 
 	/// <summary>
@@ -251,31 +286,76 @@ class CppRootPaths : IEnumerable<(uint id, DirectoryReference vfs, DirectoryRefe
 	/// </summary>
 	public void Write(JsonWriter writer)
 	{
-		foreach (KeyValuePair<CppRootPathFolder, (DirectoryReference vfs, DirectoryReference local)> item in _lookup)
+		if (_lookup.Any())
 		{
-			writer.WriteObjectStart();
-			writer.WriteEnumValue("Root", item.Key);
-			writer.WriteValue("Path", item.Value.local.FullName);
-			writer.WriteObjectEnd();
+			writer.WriteArrayStart("Lookup");
+			foreach (KeyValuePair<CppRootPathFolder, (DirectoryReference vfs, DirectoryReference local)> item in _lookup)
+			{
+				writer.WriteObjectStart();
+				writer.WriteEnumValue("Root", item.Key);
+				writer.WriteValue("Path", item.Value.local.FullName);
+				if (!_folderNames.TryGetValue(item.Key, out string? folderName))
+				{
+					writer.WriteValue("Name", folderName);
+				}
+				writer.WriteObjectEnd();
+			}
+			writer.WriteArrayEnd();
 		}
-		foreach (KeyValuePair<string, (DirectoryReference vfs, DirectoryReference local)> item in _extras)
+
+		if (_extras.Any())
+		{
+			writer.WriteArrayStart("Extras");
+			foreach (KeyValuePair<string, (DirectoryReference vfs, DirectoryReference local)> item in _extras)
+			{
+				writer.WriteObjectStart();
+				writer.WriteValue("Id", item.Key);
+				writer.WriteValue("Path", item.Value.local.FullName);
+				writer.WriteObjectEnd();
+			}
+			writer.WriteArrayEnd();
+		}
+
+		if (_folderNames.Any())
 		{
-			writer.WriteObjectStart();
-			writer.WriteValue("Id", item.Key);
-			writer.WriteValue("Path", item.Value.local.FullName);
-			writer.WriteObjectEnd();
+			writer.WriteArrayStart("Names");
+			foreach (KeyValuePair<CppRootPathFolder, string> item in _folderNames)
+			{
+				writer.WriteObjectStart();
+				writer.WriteEnumValue("Root", item.Key);
+				writer.WriteValue("Name", item.Value);
+				writer.WriteObjectEnd();
+			}
+			writer.WriteArrayEnd();
 		}
+
+		writer.WriteValue("bUseVfs", bUseVfs);
 	}
 
 	/// <summary>
-	/// Set to true to make sure all created paths are using virtual path
+	/// Add an optional folder name to use in VFS paths rather than the enum name
 	/// </summary>
-	public bool bUseVfs { get; set; }
-
+	/// <param name="folder">The folder enum to update</param>
+	/// <param name="name">The new string name</param>
+	public void AddFolderName(CppRootPathFolder folder, string name) => _folderNames[folder] = name;
 
-	private static DirectoryReference AsVfsPath(CppRootPathFolder folder) => DirectoryReference.Combine(_vfsRootPath, folder.ToString());
+	private DirectoryReference AsVfsPath(CppRootPathFolder folder)
+	{
+		string? folderName;
+		if (!_folderNames.TryGetValue(folder, out folderName))
+		{
+			folderName = folder.ToString();
+		}
+		return DirectoryReference.Combine(_vfsRootPath, folderName);
+	}
 	private static DirectoryReference AsVfsPath(string id) => DirectoryReference.Combine(_vfsRootPath, id);
 
+	private Dictionary<CppRootPathFolder, string> _folderNames { get; }
 	private SortedDictionary<CppRootPathFolder, (DirectoryReference vfs, DirectoryReference local)> _lookup { get; }
 	private SortedDictionary<string, (DirectoryReference vfs, DirectoryReference local)> _extras { get; }
+
+	/// <summary>
+	/// Set to true to make sure all created paths are using virtual path
+	/// </summary>
+	public bool bUseVfs { get; set; }
 }
diff --git a/Engine/Source/Programs/UnrealBuildTool/ToolChain/ClangToolChain.cs b/Engine/Source/Programs/UnrealBuildTool/ToolChain/ClangToolChain.cs
--- a/Engine/Source/Programs/UnrealBuildTool/ToolChain/ClangToolChain.cs
+++ b/Engine/Source/Programs/UnrealBuildTool/ToolChain/ClangToolChain.cs
@@ -239,6 +239,8 @@ namespace UnrealBuildTool
 
 			if (LazyInfo.Value.BasePath != null)
 			{
+				GlobalCompileEnvironment.RootPaths.AddFolderName(CppRootPathFolder.PlatformSDK, $"{Target.Platform}SDK");
+				GlobalLinkEnvironment.RootPaths.AddFolderName(CppRootPathFolder.PlatformSDK, $"{Target.Platform}SDK");
 				GlobalCompileEnvironment.RootPaths[CppRootPathFolder.PlatformSDK] = LazyInfo.Value.BasePath;
 				GlobalLinkEnvironment.RootPaths[CppRootPathFolder.PlatformSDK] = LazyInfo.Value.BasePath;
 			}
