[UBT]
* Removed -resource-dir in clang toolchain when running with vfs since we are not using vfsoverlay anymore
* Fixes for macos when running with -vfs

diff --git a/Engine/Source/Programs/UnrealBuildTool/Platform/Mac/MacToolChain.cs b/Engine/Source/Programs/UnrealBuildTool/Platform/Mac/MacToolChain.cs
--- a/Engine/Source/Programs/UnrealBuildTool/Platform/Mac/MacToolChain.cs
+++ b/Engine/Source/Programs/UnrealBuildTool/Platform/Mac/MacToolChain.cs
@@ -185,6 +185,8 @@ namespace UnrealBuildTool
 		{
 			base.GetCompileArguments_Global(CompileEnvironment, Arguments);
 
+			CppRootPaths RootPaths = CompileEnvironment.RootPaths;
+
 			Arguments.Add("-fasm-blocks");
 
 			// Disable FMA contraction on arm builds in order to preserve consistency with x64
@@ -199,7 +201,7 @@ namespace UnrealBuildTool
 
 			// Pass through architecture and OS info
 			Arguments.Add("" + FormatArchitectureArg(CompileEnvironment.Architectures));
-			Arguments.Add($"-isysroot \"{Settings.GetSDKPath(CompileEnvironment.Architecture)}\"");
+			Arguments.Add($"-isysroot \"{NormalizeCommandLinePath(Settings.GetSDKPath(CompileEnvironment.Architecture), RootPaths)}\"");
 
 			List<string> FrameworksSearchPaths = new List<string>();
 			foreach (UEBuildFramework Framework in CompileEnvironment.AdditionalFrameworks)
@@ -207,7 +209,7 @@ namespace UnrealBuildTool
 				FileReference FrameworkPath = new FileReference(Path.GetFullPath(Framework.Name));
 				if (!FrameworksSearchPaths.Contains(FrameworkPath.Directory.FullName))
 				{
-					Arguments.Add($"-F \"{NormalizeCommandLinePath(FrameworkPath.Directory, CompileEnvironment.RootPaths)}\"");
+					Arguments.Add($"-F \"{NormalizeCommandLinePath(FrameworkPath.Directory, RootPaths)}\"");
 					FrameworksSearchPaths.Add(FrameworkPath.Directory.FullName);
 				}
 			}
diff --git a/Engine/Source/Programs/UnrealBuildTool/ToolChain/AppleToolChain.cs b/Engine/Source/Programs/UnrealBuildTool/ToolChain/AppleToolChain.cs
--- a/Engine/Source/Programs/UnrealBuildTool/ToolChain/AppleToolChain.cs
+++ b/Engine/Source/Programs/UnrealBuildTool/ToolChain/AppleToolChain.cs
@@ -312,7 +312,7 @@ namespace UnrealBuildTool
 			foreach (FileItem SourceFile in InputFiles)
 			{
 				Action CompileAction = CompileCPPFile(CompileEnvironment, SourceFile, OutputDir, ModuleName, Graph, GlobalArguments, Result);
-				if (CompileEnvironment.PrecompiledHeaderAction == PrecompiledHeaderAction.Create)
+				if (CompileEnvironment.PrecompiledHeaderAction == PrecompiledHeaderAction.Create && !CompileEnvironment.RootPaths.bUseVfs)
 				{
 					CompileAction.bCanExecuteInUBA = false;
 				}
diff --git a/Engine/Source/Programs/UnrealBuildTool/ToolChain/ClangToolChain.cs b/Engine/Source/Programs/UnrealBuildTool/ToolChain/ClangToolChain.cs
--- a/Engine/Source/Programs/UnrealBuildTool/ToolChain/ClangToolChain.cs
+++ b/Engine/Source/Programs/UnrealBuildTool/ToolChain/ClangToolChain.cs
@@ -527,11 +527,6 @@ namespace UnrealBuildTool
 
 		protected virtual void GetCompileArguments_IncludePaths(CppCompileEnvironment CompileEnvironment, List<string> Arguments)
 		{
-            if (CompileEnvironment.RootPaths.bUseVfs)
-            {
-                Arguments.Add($"-resource-dir=\"{NormalizeCommandLinePath(GetResourceDir(CompileEnvironment), CompileEnvironment.RootPaths)}\"");
-            }
-
             Arguments.AddRange(CompileEnvironment.UserIncludePaths.Select(IncludePath => GetUserIncludePathArgument(IncludePath, CompileEnvironment)));
 			Arguments.AddRange(CompileEnvironment.SystemIncludePaths.Select(IncludePath => GetSystemIncludePathArgument(IncludePath, CompileEnvironment)));
 
