[UBT]
* Disable vfs when running with live coding

diff --git a/Engine/Source/Programs/UnrealBuildTool/Executors/UnrealBuildAccelerator/UBAExecutor.cs b/Engine/Source/Programs/UnrealBuildTool/Executors/UnrealBuildAccelerator/UBAExecutor.cs
--- a/Engine/Source/Programs/UnrealBuildTool/Executors/UnrealBuildAccelerator/UBAExecutor.cs
+++ b/Engine/Source/Programs/UnrealBuildTool/Executors/UnrealBuildAccelerator/UBAExecutor.cs
@@ -126,7 +126,10 @@ namespace UnrealBuildTool
 			foreach (TargetDescriptor targetDescriptor in targetDescriptors)
 			{
 				if (targetDescriptor.HotReloadMode != HotReloadMode.Disabled)
+				{
 					UBAConfig.bStoreObjFilesCompressed = false;
+					UBAConfig.CacheServer = null;
+				}
 				targetDescriptor.AdditionalArguments.ApplyTo(this);
 				targetDescriptor.AdditionalArguments.ApplyTo(UBAConfig);
 				hordeAgentCoordinators.Add(new UBAAgentCoordinatorHorde(logger, UBAConfig, targetDescriptor.AdditionalArguments, targetDescriptor.ProjectFile?.Directory));
@@ -239,8 +242,10 @@ namespace UnrealBuildTool
 
 		private static void RemoveLogLineSpam(List<string> logLines)
 		{
-			logLines.RemoveAll((line) => 
-			((line.StartsWith("   Creating library ", StringComparison.OrdinalIgnoreCase) || line.StartsWith("   Creating object ", StringComparison.OrdinalIgnoreCase)) && line.EndsWith(".exp", StringComparison.OrdinalIgnoreCase)) || line.EndsWith("file(s) copied.", StringComparison.OrdinalIgnoreCase));
+			logLines.RemoveAll(line => 
+				line.StartsWith("   Creating library ", StringComparison.Ordinal) ||
+				line.StartsWith("   Creating object ", StringComparison.Ordinal) ||
+				line.EndsWith("file(s) copied.", StringComparison.Ordinal));
 		}
 
 		class UBAActionArtifactCache : IActionArtifactCache
diff --git a/Engine/Source/Programs/UnrealBuildTool/System/HotReload.cs b/Engine/Source/Programs/UnrealBuildTool/System/HotReload.cs
--- a/Engine/Source/Programs/UnrealBuildTool/System/HotReload.cs
+++ b/Engine/Source/Programs/UnrealBuildTool/System/HotReload.cs
@@ -976,7 +976,7 @@ namespace UnrealBuildTool
 						}
 					}
 
-					FileReference OldResponseFile = new FileReference(Arguments[ResponseFileIdx].Substring(1).Trim('\"'));
+					FileReference OldResponseFile = Action.RootPaths.GetLocalPath(new FileReference(Arguments[ResponseFileIdx].Substring(1).Trim('\"')));
 					FileReference NewResponseFile = new FileReference(OldResponseFile.FullName + responseFileExtension);
 
 					NewAction.PrerequisiteItems.Remove(FileItem.GetItemByFileReference(OldResponseFile));
@@ -990,7 +990,7 @@ namespace UnrealBuildTool
 						string ResponseLine = ResponseLines[Idx];
 						if (ResponseLine.StartsWith(OutputFilePrefix, StringComparison.Ordinal))
 						{
-							FileReference OldOutputFile = new FileReference(ResponseLine.Substring(OutputFilePrefix.Length).Trim('\"'));
+							FileReference OldOutputFile = Action.RootPaths.GetLocalPath(new FileReference(ResponseLine.Substring(OutputFilePrefix.Length).Trim('\"')));
 							FileItem OldOutputFileItem = Action.ProducedItems.First(x => x.Location == OldOutputFile);
 							NewAction.ProducedItems.Remove(OldOutputFileItem);
 
@@ -1017,7 +1017,7 @@ namespace UnrealBuildTool
 							string ResponseLine = ResponseLines[Idx];
 							if (ResponseLine.StartsWith(SourceDependencyPrefix, StringComparison.Ordinal))
 							{
-								FileReference OldSourceDependencyFile = new FileReference(ResponseLine.Substring(SourceDependencyPrefix.Length).Trim().Trim('\"'));
+								FileReference OldSourceDependencyFile = Action.RootPaths.GetLocalPath(new FileReference(ResponseLine.Substring(SourceDependencyPrefix.Length).Trim().Trim('\"')));
 								FileItem OldSourceDependencyFileItem = Action.ProducedItems.First(x => x.Location == OldSourceDependencyFile);
 								NewAction.ProducedItems.Remove(OldSourceDependencyFileItem);
 
diff --git a/Engine/Source/Programs/UnrealBuildTool/System/RootPaths.cs b/Engine/Source/Programs/UnrealBuildTool/System/RootPaths.cs
--- a/Engine/Source/Programs/UnrealBuildTool/System/RootPaths.cs
+++ b/Engine/Source/Programs/UnrealBuildTool/System/RootPaths.cs
@@ -243,6 +243,35 @@ class CppRootPaths : IEnumerable<(uint id, DirectoryReference vfs, DirectoryRefe
 		return false;
 	}
 
+	/// <summary>
+	/// Always return a local file reference (convert from vfs if needed)
+	/// </summary>
+	public FileReference GetLocalPath(FileReference reference)
+	{
+		if (!bUseVfs)
+		{
+			return reference;
+		}
+
+		foreach (KeyValuePair<CppRootPathFolder, (DirectoryReference vfs, DirectoryReference local)> item in _lookup)
+		{
+			if (reference.IsUnderDirectory(item.Value.vfs))
+			{
+				return FileReference.FromString(item.Value.local.FullName + reference.FullName.Substring(item.Value.vfs.FullName.Length));
+			}
+		}
+
+		foreach (KeyValuePair<string, (DirectoryReference vfs, DirectoryReference local)> item in _extras)
+		{
+			if (reference.IsUnderDirectory(item.Value.vfs))
+			{
+				return FileReference.FromString(item.Value.local.FullName + reference.FullName.Substring(item.Value.vfs.FullName.Length));
+			}
+		}
+
+		return reference;
+	}
+
 	/// <summary>
 	/// Populate from json objects
 	/// </summary>
