UnrealBuildTool: -Wimplicit-int-conversion also controls -Wshorten-64-to-32 so it needs to be disabled afterwards for clang 19

#rnx

diff --git a/Engine/Source/Programs/UnrealBuildTool/ToolChain/ClangWarnings.cs b/Engine/Source/Programs/UnrealBuildTool/ToolChain/ClangWarnings.cs
--- a/Engine/Source/Programs/UnrealBuildTool/ToolChain/ClangWarnings.cs
+++ b/Engine/Source/Programs/UnrealBuildTool/ToolChain/ClangWarnings.cs
@@ -60,7 +60,6 @@ namespace UnrealBuildTool
 			if (ClangVersion >= new VersionNumber(19))
 			{
 				Arguments.Add("-Wno-extra-qualification");              // https://clang.llvm.org/docs/DiagnosticsReference.html#wextra-qualification					// new warning for clang 19
-				Arguments.Add("-Wno-shorten-64-to-32");                 // https://clang.llvm.org/docs/DiagnosticsReference.html#wshorten-64-to-32						// new warning for clang 19
 			}
 
 			Arguments.Add("-Wno-gnu-string-literal-operator-template"); // https://clang.llvm.org/docs/DiagnosticsReference.html#wgnu-string-literal-operator-template	// We use this feature to allow static FNames.
@@ -176,6 +175,12 @@ namespace UnrealBuildTool
 				}
 			}
 
+			// -Wimplicit-int-conversion also controls -Wshorten-64-to-32 as of clang 19
+			if (ClangVersion >= new VersionNumber(19))
+			{
+				Arguments.Add("-Wno-shorten-64-to-32");                 // https://clang.llvm.org/docs/DiagnosticsReference.html#wshorten-64-to-32						// new warning for clang 19
+			}
+
 			// The code base contains lots of places where we do "extern template X_API class ..." and we want to keep doing that to reduce compile times.
 			Arguments.Add("-Wno-dllexport-explicit-instantiation-decl");
 		}
