- Fixed rpaths to handle remapped runtime dependencies, and made the code cleaner

diff --git a/Engine/Source/Programs/UnrealBuildTool/Platform/Mac/MacToolChain.cs b/Engine/Source/Programs/UnrealBuildTool/Platform/Mac/MacToolChain.cs
--- a/Engine/Source/Programs/UnrealBuildTool/Platform/Mac/MacToolChain.cs
+++ b/Engine/Source/Programs/UnrealBuildTool/Platform/Mac/MacToolChain.cs
@@ -748,8 +748,65 @@ namespace UnrealBuildTool
 			}
 			else if (!LinkEnvironment.bIsBuildingLibrary)
 			{
-				// Add the additional libraries to the argument list.
 				IEnumerable<string> AdditionalLibraries = Enumerable.Concat(LinkEnvironment.SystemLibraries, LinkEnvironment.Libraries.Select(x => x.FullName));
+
+				// @todo set this to false if we have any rpath issues, or maybe set it to false when making Monolithic
+				// we will remove the old path entirely after a lot of wide testing
+				bool bCanUseSimpleRPaths = true;
+				if (bCanUseSimpleRPaths)
+				{
+					string StagedBundlePrefix = bIsBuildingAppBundle ? "../UE/"  : "";
+					string UnstagedBundlePrefix = bIsBuildingAppBundle ? "../../../" : "";
+
+					HashSet<DirectoryReference> UsedDestDirs = new();
+					HashSet<FileReference> UsedSourceLibs = new();
+					var ProcessLib = (FileReference SourceLib, FileReference DestLib, FileReference ReferencingFile, ref string LinkCommand) => 
+					{
+						if (string.Equals(SourceLib.GetExtension(), ".dylib", StringComparison.InvariantCultureIgnoreCase))
+						{
+							DirectoryReference DestDir = DestLib.Directory;
+							if (UsedSourceLibs.Contains(SourceLib) || UsedDestDirs.Contains(DestDir))
+							{
+								return;
+							}
+
+							UsedDestDirs.Add(DestDir);
+							UsedSourceLibs.Add(SourceLib);
+
+							string RPathEntry;
+							// handle staged rpaths
+							if (ProjectFile != null && DestDir.IsUnderDirectory(ProjectFile.Directory))
+							{
+								RPathEntry = StagedBundlePrefix + DestDir.MakeRelativeTo(ProjectFile.Directory.ParentDirectory!);
+							}
+							else
+							{
+								RPathEntry = StagedBundlePrefix + DestDir.MakeRelativeTo(Unreal.EngineDirectory.ParentDirectory!);
+							}
+							LinkCommand += $" -rpath \"@executable_path/{RPathEntry}\"";
+
+							// handle unstaged paths
+							RPathEntry = UnstagedBundlePrefix + DestDir.MakeRelativeTo(ReferencingFile.Directory);
+							LinkCommand += $" -rpath \"@loader_path/{RPathEntry}\"";
+						}
+					};
+
+					// first look in the RuntimeDependency list to look for remapped locations, anything not remapped we will deal with 
+					// after this loop
+					foreach (ModuleRules.RuntimeDependency SrcDep in LinkEnvironment.RuntimeDependencies)
+					{
+						FileReference DestPath = new(SrcDep.Path);
+						FileReference SourcePath = SrcDep.SourcePath == null ? DestPath : new FileReference(SrcDep.SourcePath);
+						ProcessLib(SourcePath, DestPath, MultiArchOutputFile, ref LinkCommand);
+					}
+					foreach (string AdditionalLibrary in AdditionalLibraries)
+					{
+						FileReference LibPath = new(AdditionalLibrary);
+						ProcessLib(LibPath, LibPath, MultiArchOutputFile, ref LinkCommand);
+					}
+				}
+
+				// Add the additional libraries to the argument list.
 				foreach (string AdditionalLibrary in AdditionalLibraries)
 				{
 					if (!String.IsNullOrEmpty(Path.GetDirectoryName(AdditionalLibrary)) &&
@@ -783,14 +840,20 @@ namespace UnrealBuildTool
 						LinkCommand += String.Format(" \"{0}\"", Path.GetFullPath(AdditionalLibrary));
 					}
 
-					AddLibraryPathToRPaths(AdditionalLibrary, AbsolutePath, ref RPaths, ref LinkCommand, bIsBuildingAppBundle, Logger);
+					if (!bCanUseSimpleRPaths)
+					{
+						AddLibraryPathToRPaths(AdditionalLibrary, AbsolutePath, ref RPaths, ref LinkCommand, bIsBuildingAppBundle, Logger);
+					}
 				}
 
 				foreach (string AdditionalLibrary in LinkEnvironment.DelayLoadDLLs)
 				{
 					LinkCommand += String.Format(" -weak_library \"{0}\"", Path.GetFullPath(AdditionalLibrary));
 
-					AddLibraryPathToRPaths(AdditionalLibrary, AbsolutePath, ref RPaths, ref LinkCommand, bIsBuildingAppBundle, Logger);
+					if (!bCanUseSimpleRPaths)
+					{
+						AddLibraryPathToRPaths(AdditionalLibrary, AbsolutePath, ref RPaths, ref LinkCommand, bIsBuildingAppBundle, Logger);
+					}
 				}
 
 				// Add frameworks
