UBT not utilizing project level buildconfiguration.xml in normal execution.

diff --git a/Engine/Source/Programs/UnrealBuildTool/Modes/GenerateProjectFilesMode.cs b/Engine/Source/Programs/UnrealBuildTool/Modes/GenerateProjectFilesMode.cs
--- a/Engine/Source/Programs/UnrealBuildTool/Modes/GenerateProjectFilesMode.cs
+++ b/Engine/Source/Programs/UnrealBuildTool/Modes/GenerateProjectFilesMode.cs
@@ -79,7 +79,7 @@ namespace UnrealBuildTool
 
 			// Parse rocket-specific arguments.
 			FileReference? ProjectFile;
-			TryParseProjectFileArgument(Arguments, Logger, out ProjectFile);
+			Utils.TryParseProjectFileArgument(Arguments, Logger, out ProjectFile);
 
 			// Apply the XML config again with a project specific BuildConfiguration.xml 
 			XmlConfig.ReadConfigFiles(null, ProjectFile?.Directory, Logger);
@@ -270,74 +270,5 @@ namespace UnrealBuildTool
 			}
 			return Task.FromResult((int)CompilationResult.Succeeded);
 		}
-
-		/// <summary>
-		/// Try to parse the project file from the command line
-		/// </summary>
-		/// <param name="Arguments">The command line arguments</param>
-		/// <param name="Logger">Logger for output</param>
-		/// <param name="ProjectFile">The project file that was parsed</param>
-		/// <returns>True if the project file was parsed, false otherwise</returns>
-		public static bool TryParseProjectFileArgument(CommandLineArguments Arguments, ILogger Logger, [NotNullWhen(true)] out FileReference? ProjectFile)
-		{
-			string? CandidateProjectPath = null;
-
-			// look for -project=<path>, if it does not exist check arguments for anything that has .uproject in it
-			if (!Arguments.TryGetValue("-Project=", out CandidateProjectPath))
-			{
-				// Go through the argument list and try to match poorly (or well..) formed arguments like 
-				// EngineTest, EngineTest.uproject
-				// Collaboration/FooProject
-				// by checking for those in the native project list
-				for (int Idx = 0; Idx < Arguments.Count; Idx++)
-				{
-					if (Arguments[Idx][0] != '-' && Arguments[Idx].EndsWith(".uproject", StringComparison.OrdinalIgnoreCase))
-					{
-						CandidateProjectPath = Arguments[Idx];
-						Arguments.MarkAsUsed(Idx);
-						break;
-					}
-				}
-			}
-
-			// We have a project file either via -project= or because there was something called .uproject in the arg list
-			// so now validate it
-			if (!String.IsNullOrEmpty(CandidateProjectPath))
-			{
-				FileReference? CandidateProjectFile = FileReference.FindCorrectCase(new FileReference(CandidateProjectPath));
-
-				// if the path doesn't exist then check native paths (ueprojectdirs)
-				if (!FileReference.Exists(CandidateProjectFile))
-				{
-					// clean everything the user provided to just the name and make sure it has the expected extension
-					string ProjectName = CandidateProjectFile.ChangeExtension("uproject").GetFileName();
-
-					// check native project paths (uprojectdirs)
-					IEnumerable<FileReference> NativeProjectFiles = NativeProjects.EnumerateProjectFiles(Logger);
-
-					CandidateProjectFile = NativeProjectFiles.Where(F => F.GetFileName().Equals(ProjectName, StringComparison.InvariantCultureIgnoreCase)).FirstOrDefault();
-				}
-
-				if (CandidateProjectFile == null || !FileReference.Exists(CandidateProjectFile))
-				{
-					// if we didn't find anything then throw an error as the user explicitly provided a uproject
-					throw new Exception(String.Format("Unable to find project file based on argument {0}", CandidateProjectPath));
-				}
-
-				Logger.LogDebug("Resolved project argument {CandidateProjectPath} to {CandidateProjectFile}", CandidateProjectPath, CandidateProjectFile);
-				ProjectFile = CandidateProjectFile;
-				return true;
-			}
-
-			FileReference? InstalledProjectFile = Unreal.GetInstalledProjectFile();
-			if (InstalledProjectFile != null)
-			{
-				ProjectFile = InstalledProjectFile;
-				return true;
-			}
-
-			ProjectFile = null;
-			return false;
-		}
 	}
 }
diff --git a/Engine/Source/Programs/UnrealBuildTool/Modes/QueryMode.cs b/Engine/Source/Programs/UnrealBuildTool/Modes/QueryMode.cs
--- a/Engine/Source/Programs/UnrealBuildTool/Modes/QueryMode.cs
+++ b/Engine/Source/Programs/UnrealBuildTool/Modes/QueryMode.cs
@@ -200,7 +200,7 @@ namespace UnrealBuildTool
 		{
 			try
 			{
-				GenerateProjectFilesMode.TryParseProjectFileArgument(Arguments, Logger, out FileReference? ProjectFileArg);
+				Utils.TryParseProjectFileArgument(Arguments, Logger, out FileReference? ProjectFileArg);
 
 				List<UnrealTargetPlatform> Platforms = new();
 				foreach (UnrealTargetPlatform Platform in UnrealTargetPlatform.GetValidPlatforms())
@@ -333,7 +333,7 @@ namespace UnrealBuildTool
 				return 1;
 			}
 
-			GenerateProjectFilesMode.TryParseProjectFileArgument(Arguments, Logger, out FileReference? ProjectFileArg);
+			Utils.TryParseProjectFileArgument(Arguments, Logger, out FileReference? ProjectFileArg);
 			List<string> RawArgs = new List<string> { TargetName, TargetConfiguration, TargetPlatform };
 			if (ProjectFileArg != null)
 			{
diff --git a/Engine/Source/Programs/UnrealBuildTool/System/Utils.cs b/Engine/Source/Programs/UnrealBuildTool/System/Utils.cs
--- a/Engine/Source/Programs/UnrealBuildTool/System/Utils.cs
+++ b/Engine/Source/Programs/UnrealBuildTool/System/Utils.cs
@@ -3,6 +3,7 @@
 using System;
 using System.Collections.Generic;
 using System.Diagnostics;
+using System.Diagnostics.CodeAnalysis;
 using System.IO;
 using System.Linq;
 using System.Runtime.InteropServices;
@@ -1517,6 +1518,80 @@ namespace UnrealBuildTool
 			return Arguments;
 		}
 
+		/// <summary>
+		/// Try to parse the project file from the command line
+		/// </summary>
+		/// <param name="Arguments">The command line arguments</param>
+		/// <param name="Logger">Logger for output</param>
+		/// <param name="ProjectFile">The project file that was parsed</param>
+		/// <param name="MarkArgumentAsUsed">Marks the Project argument as used if set to true.</param>
+		/// <returns>True if the project file was parsed, false otherwise</returns>
+		public static bool TryParseProjectFileArgument(CommandLineArguments Arguments, ILogger Logger, [NotNullWhen(true)] out FileReference? ProjectFile, bool MarkArgumentAsUsed = true)
+		{
+			string? CandidateProjectPath = null;
+
+			// look for -project=<path>, if it does not exist check arguments for anything that has .uproject in it
+			if (!Arguments.TryGetValue("-Project=", out CandidateProjectPath))
+			{
+				// Go through the argument list and try to match poorly (or well..) formed arguments like 
+				// EngineTest, EngineTest.uproject
+				// Collaboration/FooProject
+				// by checking for those in the native project list
+				for (int Idx = 0; Idx < Arguments.Count; Idx++)
+				{
+					if (Arguments[Idx][0] != '-' && Arguments[Idx].EndsWith(".uproject", StringComparison.OrdinalIgnoreCase))
+					{
+						CandidateProjectPath = Arguments[Idx];
+
+						if (MarkArgumentAsUsed)
+						{
+							Arguments.MarkAsUsed(Idx);
+						}
+						break;
+					}
+				}
+			}
+
+			// We have a project file either via -project= or because there was something called .uproject in the arg list
+			// so now validate it
+			if (!String.IsNullOrEmpty(CandidateProjectPath))
+			{
+				FileReference? CandidateProjectFile = FileReference.FindCorrectCase(new FileReference(CandidateProjectPath));
+
+				// if the path doesn't exist then check native paths (ueprojectdirs)
+				if (!FileReference.Exists(CandidateProjectFile))
+				{
+					// clean everything the user provided to just the name and make sure it has the expected extension
+					string ProjectName = CandidateProjectFile.ChangeExtension("uproject").GetFileName();
+
+					// check native project paths (uprojectdirs)
+					IEnumerable<FileReference> NativeProjectFiles = NativeProjects.EnumerateProjectFiles(Logger);
+
+					CandidateProjectFile = NativeProjectFiles.Where(F => F.GetFileName().Equals(ProjectName, StringComparison.InvariantCultureIgnoreCase)).FirstOrDefault();
+				}
+
+				if (CandidateProjectFile == null || !FileReference.Exists(CandidateProjectFile))
+				{
+					// if we didn't find anything then throw an error as the user explicitly provided a uproject
+					throw new Exception(String.Format("Unable to find project file based on argument {0}", CandidateProjectPath));
+				}
+
+				Logger.LogDebug("Resolved project argument {CandidateProjectPath} to {CandidateProjectFile}", CandidateProjectPath, CandidateProjectFile);
+				ProjectFile = CandidateProjectFile;
+				return true;
+			}
+
+			FileReference? InstalledProjectFile = Unreal.GetInstalledProjectFile();
+			if (InstalledProjectFile != null)
+			{
+				ProjectFile = InstalledProjectFile;
+				return true;
+			}
+
+			ProjectFile = null;
+			return false;
+		}
+
 		/// <summary>
 		/// Formats a list of arguments as a command line, inserting quotes as necessary
 		/// </summary>
diff --git a/Engine/Source/Programs/UnrealBuildTool/UnrealBuildTool.cs b/Engine/Source/Programs/UnrealBuildTool/UnrealBuildTool.cs
--- a/Engine/Source/Programs/UnrealBuildTool/UnrealBuildTool.cs
+++ b/Engine/Source/Programs/UnrealBuildTool/UnrealBuildTool.cs
@@ -507,7 +507,9 @@ namespace UnrealBuildTool
 						using (GlobalSingleInstanceMutex XmlConfigMutex = new GlobalSingleInstanceMutex(XmlConfigMutexName, true))
 						{
 							FileReference? XmlConfigCache = Arguments.GetFileReferenceOrDefault("-XmlConfigCache=", null);
-							XmlConfig.ReadConfigFiles(XmlConfigCache, null, Logger);
+							FileReference? ProjectFile;
+							Utils.TryParseProjectFileArgument(Arguments, Logger, out ProjectFile, false);
+							XmlConfig.ReadConfigFiles(XmlConfigCache, ProjectFile?.Directory, Logger);
 						}
 					}
 
