Make the AutoRTFM compiler build process also build Linux binaries, and check them in too.

#rnx

diff --git a/Engine/Source/Programs/UnrealBuildTool/Platform/Linux/LinuxToolChain.cs b/Engine/Source/Programs/UnrealBuildTool/Platform/Linux/LinuxToolChain.cs
--- a/Engine/Source/Programs/UnrealBuildTool/Platform/Linux/LinuxToolChain.cs
+++ b/Engine/Source/Programs/UnrealBuildTool/Platform/Linux/LinuxToolChain.cs
@@ -128,9 +128,17 @@ namespace UnrealBuildTool
 				if (Options.HasFlag(ClangToolChainOptions.UseAutoRTFMCompiler))
 				{
 					DirectoryReference AutoRTFMDir = DirectoryReference.Combine(Unreal.EngineDirectory, "Restricted", "NotForLicensees", "Binaries", BuildHostPlatform.Current.Platform.ToString(), "AutoRTFM", "bin");
-
+					
 					// set up the path to our toolchain
-					ClangPath = FileReference.Combine(AutoRTFMDir, $"verse-clang-cl{BuildHostPlatform.Current.BinarySuffix}");
+					if (BuildHostPlatform.Current.Platform == UnrealTargetPlatform.Linux)
+					{
+						ClangPath = FileReference.Combine(AutoRTFMDir, $"verse-clang++{BuildHostPlatform.Current.BinarySuffix}");
+					}
+					else
+					{
+						// We use verse-clang-cl when the host is Windows just to keep the number of compilers deployed to a minimum.
+						ClangPath = FileReference.Combine(AutoRTFMDir, $"verse-clang-cl{BuildHostPlatform.Current.BinarySuffix}");
+					}
 				}
 
 				// When cross-compiling on Windows, use old FixDeps. It is slow, but it does not have timing issues
