Fix builds from Xcode. Use String.StartsWith instead of indexing into command line argument strings. Indexing was causing IndexOutOfBounds exception when invoked from XcodeBuild.sh which forwards empty string command line arguments (which is an issue in of itself).

diff --git a/Engine/Source/Programs/UnrealBuildTool/System/Utils.cs b/Engine/Source/Programs/UnrealBuildTool/System/Utils.cs
--- a/Engine/Source/Programs/UnrealBuildTool/System/Utils.cs
+++ b/Engine/Source/Programs/UnrealBuildTool/System/Utils.cs
@@ -1539,7 +1539,7 @@ namespace UnrealBuildTool
 				// by checking for those in the native project list
 				for (int Idx = 0; Idx < Arguments.Count; Idx++)
 				{
-					if (Arguments[Idx][0] != '-' && Arguments[Idx].EndsWith(".uproject", StringComparison.OrdinalIgnoreCase))
+					if (!Arguments[Idx].StartsWith("-") && Arguments[Idx].EndsWith(".uproject", StringComparison.OrdinalIgnoreCase))
 					{
 						CandidateProjectPath = Arguments[Idx];
 
