UnrealBuildTool: Handle json error when checking if rules files require compilation

#rnx

diff --git a/Engine/Source/Programs/UnrealBuildTool/System/DynamicCompilation.cs b/Engine/Source/Programs/UnrealBuildTool/System/DynamicCompilation.cs
--- a/Engine/Source/Programs/UnrealBuildTool/System/DynamicCompilation.cs
+++ b/Engine/Source/Programs/UnrealBuildTool/System/DynamicCompilation.cs
@@ -67,7 +67,16 @@ namespace UnrealBuildTool
 				return true;
 			}
 
-			JsonObject Manifest = JsonObject.Read(AssemblyManifestFilePath);
+			JsonObject Manifest;
+			try
+			{
+				Manifest = JsonObject.Read(AssemblyManifestFilePath);
+			}
+			catch (Exception)
+			{
+				Logger.LogDebug("Compiling {OutputAssemblyPath}: Error reading source file list ({AssemblyManifestFilePath})", OutputAssemblyPath, AssemblyManifestFilePath);
+				return true;
+			}
 
 			// check if the engine version is different
 			string EngineVersionManifest = Manifest.GetStringField("EngineVersion");
