UnrealBuildTool: Fix registration of host arm64 clang binaries. Allow using x64 binaries on arm64

#rnx

diff --git a/Engine/Source/Programs/UnrealBuildTool/Platform/Windows/MicrosoftPlatformSDK.cs b/Engine/Source/Programs/UnrealBuildTool/Platform/Windows/MicrosoftPlatformSDK.cs
--- a/Engine/Source/Programs/UnrealBuildTool/Platform/Windows/MicrosoftPlatformSDK.cs
+++ b/Engine/Source/Programs/UnrealBuildTool/Platform/Windows/MicrosoftPlatformSDK.cs
@@ -638,6 +638,18 @@ namespace UnrealBuildTool
 				.ThenByDescending(x => x.Version)
 				.FirstOrDefault();
 
+			// Allow x64 host toolchains on arm64 if no native arm64 toolchains are available, and not running via wine
+			if (!BuildHostPlatform.Current.IsRunningOnWine() && MicrosoftPlatformSDK.MSVCHostUnrealArch == UnrealArch.Arm64 && (ToolChain == null || ToolChain.Error != null))
+			{
+				ToolChain = Preference(ToolChains.Where(x => x.Architecture == Architecture && x.HostArchitecture != MicrosoftPlatformSDK.MSVCHostUnrealArch)
+					.OrderByDescending(x => x.Error == null))
+					.ThenBy(x => x.ReleaseChannel)
+					.ThenBy(x => x.FamilyRank)
+					.ThenByDescending(x => x.IsAutoSdk)
+					.ThenByDescending(x => x.Version)
+					.FirstOrDefault();
+			}
+
 			if (ToolChain?.Error != null)
 			{
 				if (!IgnoreToolchainErrors)
@@ -871,17 +883,20 @@ namespace UnrealBuildTool
 
 						if (CrossToolChain != null)
 						{
-							DirectoryReference HostBinDir = DirectoryReference.Combine(ToolChain.BaseDir, "bin", Compiler.IsMSVC() ? UnrealArchToMSVCHostDirectoryName(ToolChain.HostArchitecture) : String.Empty);
-							DirectoryReference CrossBinDir = DirectoryReference.Combine(CrossToolChain.BaseDir, "bin", Compiler.IsMSVC() ? UnrealArchToMSVCHostDirectoryName(CrossToolChain.HostArchitecture) : String.Empty);
-							if (FileReference.Exists(HostBinDir) && FileReference.Exists(CrossBinDir))
+							DirectoryReference HostBinDir;
+							DirectoryReference CrossBinDir;
+
+							if (Compiler.IsMSVC())
 							{
-								EpicGames.UBA.Utils.RegisterCrossArchitecturePath(HostBinDir.FullName, CrossBinDir.FullName);
+								HostBinDir = DirectoryReference.Combine(ToolChain.BaseDir, "bin", UnrealArchToMSVCHostDirectoryName(ToolChain.HostArchitecture));
+								CrossBinDir = DirectoryReference.Combine(CrossToolChain.BaseDir, "bin", UnrealArchToMSVCHostDirectoryName(CrossToolChain.HostArchitecture));
 							}
-						}
-						else if (Compiler.IsClang())
-						{
-							DirectoryReference HostBinDir = DirectoryReference.Combine(ToolChain.BaseDir, MSVCHostUnrealArch.bIsX64 ? "bin" : "bin-woa64");
-							DirectoryReference CrossBinDir = DirectoryReference.Combine(ToolChain.BaseDir, !MSVCHostUnrealArch.bIsX64 ? "bin" : "bin-woa64");
+							else
+							{
+								HostBinDir = DirectoryReference.Combine(ToolChain.BaseDir, MSVCHostUnrealArch.bIsX64 ? "bin" : "bin-woa64");
+								CrossBinDir = DirectoryReference.Combine(CrossToolChain.BaseDir, !MSVCHostUnrealArch.bIsX64 ? "bin" : "bin-woa64");
+							}
+
 							if (FileReference.Exists(HostBinDir) && FileReference.Exists(CrossBinDir))
 							{
 								EpicGames.UBA.Utils.RegisterCrossArchitecturePath(HostBinDir.FullName, CrossBinDir.FullName);
@@ -1112,13 +1127,12 @@ namespace UnrealBuildTool
 		/// </summary>
 		/// <param name="Compiler"></param>
 		/// <param name="ToolChainDir"></param>
+		/// <param name="CompilerFile"></param>
 		/// <param name="ToolChains"></param>
 		/// <param name="IsAutoSdk"></param>
 		/// <param name="Logger"></param>
-		static void AddClangToolChain(WindowsCompiler Compiler, DirectoryReference ToolChainDir, List<ToolChainInstallation> ToolChains, bool IsAutoSdk, ILogger Logger)
+		static void AddClangToolChain(WindowsCompiler Compiler, DirectoryReference ToolChainDir, FileReference CompilerFile, List<ToolChainInstallation> ToolChains, bool IsAutoSdk, ILogger Logger)
 		{
-			string CompilerFilename = Compiler == WindowsCompiler.ClangRTFM ? "verse-clang-cl.exe" : "clang-cl.exe";
-			FileReference CompilerFile = FileReference.Combine(ToolChainDir, "bin", CompilerFilename);
 			if (FileReference.Exists(CompilerFile))
 			{
 				FileVersionInfo VersionInfo = FileVersionInfo.GetVersionInfo(CompilerFile.FullName);
@@ -1156,6 +1170,21 @@ namespace UnrealBuildTool
 			}
 		}
 
+		/// <summary>
+		/// Add a Clang toolchain
+		/// </summary>
+		/// <param name="Compiler"></param>
+		/// <param name="ToolChainDir"></param>
+		/// <param name="ToolChains"></param>
+		/// <param name="IsAutoSdk"></param>
+		/// <param name="Logger"></param>
+		static void AddClangToolChain(WindowsCompiler Compiler, DirectoryReference ToolChainDir, List<ToolChainInstallation> ToolChains, bool IsAutoSdk, ILogger Logger)
+		{
+			string CompilerFilename = Compiler == WindowsCompiler.ClangRTFM ? "verse-clang-cl.exe" : "clang-cl.exe";
+			AddClangToolChain(Compiler, ToolChainDir, FileReference.Combine(ToolChainDir, "bin", CompilerFilename), ToolChains, IsAutoSdk, Logger);
+			AddClangToolChain(Compiler, ToolChainDir, FileReference.Combine(ToolChainDir, "bin-woa64", CompilerFilename), ToolChains, IsAutoSdk, Logger);
+		}
+
 		/// <summary>
 		/// Add an Intel OneAPI toolchain
 		/// </summary>
diff --git a/Engine/Source/Programs/UnrealBuildTool/ToolChain/ISPCToolChain.cs b/Engine/Source/Programs/UnrealBuildTool/ToolChain/ISPCToolChain.cs
--- a/Engine/Source/Programs/UnrealBuildTool/ToolChain/ISPCToolChain.cs
+++ b/Engine/Source/Programs/UnrealBuildTool/ToolChain/ISPCToolChain.cs
@@ -4,6 +4,7 @@ using System;
 using System.Collections.Generic;
 using System.IO;
 using System.Linq;
+using System.Runtime.InteropServices;
 using System.Text;
 using EpicGames.Core;
 using Microsoft.Extensions.Logging;
@@ -456,6 +457,11 @@ namespace UnrealBuildTool
 				CompileAction.bCanExecuteRemotely = true;
 				CompileAction.bCanExecuteInUBACrossArchitecture = false;
 
+				if (RuntimeInformation.ProcessArchitecture == Architecture.Arm64)
+				{
+					CompileAction.bCanExecuteInUBA = false;
+				}
+
 				// Disable remote execution to workaround mismatched case on XGE
 				CompileAction.bCanExecuteRemotelyWithXGE = false;
 
@@ -660,6 +666,11 @@ namespace UnrealBuildTool
 				CompileAction.bCanExecuteRemotely = true;
 				CompileAction.bCanExecuteInUBACrossArchitecture = false;
 
+				if (RuntimeInformation.ProcessArchitecture == Architecture.Arm64)
+				{
+					CompileAction.bCanExecuteInUBA = false;
+				}
+
 				// Disable remote execution to workaround mismatched case on XGE
 				CompileAction.bCanExecuteRemotelyWithXGE = false;
 
