[UBA]
* Added so log lines are written to log from failed remote process before retrying locally

diff --git a/Engine/Source/Programs/UnrealBuildTool/Executors/UnrealBuildAccelerator/UBAExecutor.cs b/Engine/Source/Programs/UnrealBuildTool/Executors/UnrealBuildAccelerator/UBAExecutor.cs
--- a/Engine/Source/Programs/UnrealBuildTool/Executors/UnrealBuildAccelerator/UBAExecutor.cs
+++ b/Engine/Source/Programs/UnrealBuildTool/Executors/UnrealBuildAccelerator/UBAExecutor.cs
@@ -922,32 +922,32 @@ namespace UnrealBuildTool
 					}
 					else if ((uint)e.ExitCode == 0xC0000005)
 					{
-						RemoteActionFailedCrash(queue, action, e.ExitCode, e.ExecutingHost ?? "Unknown", "Access violation");
+						RemoteActionFailedCrash(queue, action, e.ExitCode, e.ExecutingHost ?? "Unknown", "Access violation", e.LogLines);
 						return;
 					}
 					else if ((uint)e.ExitCode == 0xC0000409)
 					{
-						RemoteActionFailedCrash(queue, action, e.ExitCode, e.ExecutingHost ?? "Unknown", "Stack buffer overflow");
+						RemoteActionFailedCrash(queue, action, e.ExitCode, e.ExecutingHost ?? "Unknown", "Stack buffer overflow", e.LogLines);
 						return;
 					}
 					else if ((uint)e.ExitCode == 0xC0000602)
 					{
-						RemoteActionFailedCrash(queue, action, e.ExitCode, e.ExecutingHost ?? "Unknown", "Fail Fast Exception");
+						RemoteActionFailedCrash(queue, action, e.ExitCode, e.ExecutingHost ?? "Unknown", "Fail Fast Exception", e.LogLines);
 						return;
 					}
 					else if (e.ExitCode != 0 && e.LogLines.Any(x => x.Contains(" C1001: ", StringComparison.Ordinal)))
 					{
-						RemoteActionFailedCrash(queue, action, e.ExitCode, e.ExecutingHost ?? "Unknown", "C1001");
+						RemoteActionFailedCrash(queue, action, e.ExitCode, e.ExecutingHost ?? "Unknown", "C1001", e.LogLines);
 						return;
 					}
 					else if (e.ExitCode >= 9000 && e.ExitCode < 10000)
 					{
-						RemoteActionFailedCrash(queue, action, e.ExitCode, e.ExecutingHost ?? "Unknown", "UBA error");
+						RemoteActionFailedCrash(queue, action, e.ExitCode, e.ExecutingHost ?? "Unknown", "UBA error", e.LogLines);
 						return;
 					}
 					else if (e.ExitCode != 0 && UBAConfig.bForcedRetryRemote)
 					{
-						RemoteActionFailedCrash(queue, action, e.ExitCode, e.ExecutingHost ?? "Unknown", "Force local retry");
+						RemoteActionFailedCrash(queue, action, e.ExitCode, e.ExecutingHost ?? "Unknown", "Force local retry", e.LogLines);
 						return;
 					}
 
@@ -1043,7 +1043,7 @@ namespace UnrealBuildTool
 			}
 		}
 
-		void RemoteActionFailedCrash(ImmediateActionQueue queue, LinkedAction action, int exitCode, string executingHost, string error)
+		void RemoteActionFailedCrash(ImmediateActionQueue queue, LinkedAction action, int exitCode, string executingHost, string error, List<string> logLines)
 		{
 			if (_bIsCancelled)
 			{
@@ -1052,6 +1052,11 @@ namespace UnrealBuildTool
 			}
 
 			_threadedLogger.LogInformation("{Description} {StatusDescription} [RemoteExecutor: {ExecutingHost}]: Exited with error code {ExitCode} ({Error}). This action will retry locally", action.CommandDescription, action.StatusDescription, executingHost, exitCode, error);
+			foreach (string line in logLines)
+			{
+				_threadedLogger.LogInformation(line);
+			}
+
 			_localRetryActions.AddOrUpdate(action, false, (k, v) => false);
 			queue.RequeueAction(action);
 
