Marking compile settings DisablePluginsConflictWarningLevels & bEnableNonInlinedGenCppWarnings as obsolete in respective Rules class.

diff --git a/Engine/Source/Programs/UnrealBuildTool/Configuration/ModuleRules.cs b/Engine/Source/Programs/UnrealBuildTool/Configuration/ModuleRules.cs
--- a/Engine/Source/Programs/UnrealBuildTool/Configuration/ModuleRules.cs
+++ b/Engine/Source/Programs/UnrealBuildTool/Configuration/ModuleRules.cs
@@ -1171,6 +1171,7 @@ namespace UnrealBuildTool
 		/// <summary>
 		/// Enable warnings for when there are .gen.cpp files that could be inlined in a matching handwritten cpp file
 		/// </summary>
+		[Obsolete("Deprecated in UE5.6 - Replace with ModuleRules.CppCompileWarningSettings.NonInlinedGenCppWarningLevel")]
 		public bool bEnableNonInlinedGenCppWarnings { get; set; }
 
 		/// <summary>
diff --git a/Engine/Source/Programs/UnrealBuildTool/Configuration/TargetRules.cs b/Engine/Source/Programs/UnrealBuildTool/Configuration/TargetRules.cs
--- a/Engine/Source/Programs/UnrealBuildTool/Configuration/TargetRules.cs
+++ b/Engine/Source/Programs/UnrealBuildTool/Configuration/TargetRules.cs
@@ -838,6 +838,7 @@ namespace UnrealBuildTool
 		/// <summary>
 		/// How to treat conflicts when a disabled plugin is being enabled by another plugin referencing it
 		/// </summary>
+		[Obsolete("Deprecated in UE5.6 - Replace with TargetRules.CppCompileWarningSettings.DisablePluginsConflictWarningLevel")]
 		public WarningLevel DisablePluginsConflictWarningLevel
 		{
 			get => CppCompileWarningSettings.DisablePluginsConflictWarningLevel;
