- Using old rpath code on modular builds while we try to track down an issue on horde

diff --git a/Engine/Source/Programs/UnrealBuildTool/Platform/Mac/MacToolChain.cs b/Engine/Source/Programs/UnrealBuildTool/Platform/Mac/MacToolChain.cs
--- a/Engine/Source/Programs/UnrealBuildTool/Platform/Mac/MacToolChain.cs
+++ b/Engine/Source/Programs/UnrealBuildTool/Platform/Mac/MacToolChain.cs
@@ -752,7 +752,7 @@ namespace UnrealBuildTool
 
 				// @todo set this to false if we have any rpath issues, or maybe set it to false when making Monolithic
 				// we will remove the old path entirely after a lot of wide testing
-				bool bCanUseSimpleRPaths = true;
+				bool bCanUseSimpleRPaths = LinkEnvironment.LinkType == TargetLinkType.Monolithic;
 				if (bCanUseSimpleRPaths)
 				{
 					string StagedBundlePrefix = bIsBuildingAppBundle ? "../UE/"  : "";
