[UBT]
* Fixed LinuxToolChain so it is not overwriting any files in the same build session. When we know a relink is happning we will write the first '.so' file to intermediate folder instead. All link steps are using the intermediate version of the .so file if it exists (we solve this by creating a lookup in SetupBundleDependencies)

Now there should be no weird race conditions in the LinuxToolChain and we should be able to do distributed linking if we want (not enabled in this changelist)

diff --git a/Engine/Source/Programs/UnrealBuildTool/Platform/Linux/LinuxToolChain.cs b/Engine/Source/Programs/UnrealBuildTool/Platform/Linux/LinuxToolChain.cs
--- a/Engine/Source/Programs/UnrealBuildTool/Platform/Linux/LinuxToolChain.cs
+++ b/Engine/Source/Programs/UnrealBuildTool/Platform/Linux/LinuxToolChain.cs
@@ -98,10 +98,6 @@ namespace UnrealBuildTool
 					throw new BuildException("Unable to find system llvm-objcopy; cannot instantiate Linux toolchain");
 				}
 
-				// When compiling on Linux, use a faster way to relink circularly dependent libraries.
-				// Race condition between actions linking to the .so and action overwriting it is avoided thanks to inodes
-				bUseFixdeps = false;
-
 				bIsCrossCompiling = false;
 
 				return new LinuxToolChainInfo(null, null, ClangPath, LlvmArPath, ObjcopyPath, Logger);
@@ -141,9 +137,6 @@ namespace UnrealBuildTool
 					}
 				}
 
-				// When cross-compiling on Windows, use old FixDeps. It is slow, but it does not have timing issues
-				bUseFixdeps = BuildHostPlatform.Current.Platform == UnrealTargetPlatform.Win64;
-
 				if (BuildHostPlatform.Current.Platform == UnrealTargetPlatform.Linux)
 				{
 					Environment.SetEnvironmentVariable("LC_ALL", "C");
@@ -164,9 +157,10 @@ namespace UnrealBuildTool
 		{
 			bool bUseCmdExe = BuildHostPlatform.Current.Platform == UnrealTargetPlatform.Win64;
 			string DumpCommand = bUseCmdExe ? "\"{0}\" \"{1}\" \"{2}\" 2>NUL" : "\"{0}\" -c -o \"{2}\" \"{1}\"";
+			string DeleteCommand = bUseCmdExe ? "del /F" : "rm";
 			FileItem EncodedBinarySymbolsFile = FileItem.GetItemByPath(Path.Combine(LinkEnvironment.OutputDirectory!.FullName, OutputFile.Location.GetFileNameWithoutExtension() + ".sym"));
-			FileItem SymbolsFile = FileItem.GetItemByPath(Path.Combine(LinkEnvironment.LocalShadowDirectory!.FullName, OutputFile.Location.GetFileName() + ".psym"));
-			FileItem StrippedFile = FileItem.GetItemByPath(Path.Combine(LinkEnvironment.LocalShadowDirectory.FullName, OutputFile.Location.GetFileName() + "_nodebug"));
+			FileItem SymbolsFile = FileItem.GetItemByPath(Path.Combine(LinkEnvironment.IntermediateDirectory!.FullName, OutputFile.Location.GetFileName() + ".psym"));
+			FileItem StrippedFile = FileItem.GetItemByPath(Path.Combine(LinkEnvironment.IntermediateDirectory.FullName, OutputFile.Location.GetFileName() + "_nodebug"));
 			FileItem DebugFile = FileItem.GetItemByPath(Path.Combine(LinkEnvironment.OutputDirectory.FullName, OutputFile.Location.GetFileNameWithoutExtension() + ".debug"));
 
 			if (Options.HasFlag(ClangToolChainOptions.PreservePSYM))
@@ -192,11 +186,12 @@ namespace UnrealBuildTool
 					SymbolsFile.AbsolutePath,
 					EncodedBinarySymbolsFile.AbsolutePath
 				);
+
+				Out.WriteLine($"{DeleteCommand} \"{SymbolsFile.AbsolutePath}\"");
 			}
 			else
 			{
 				// we have to create dummy files to prevent packaging errors
-				Out.WriteLine("echo DummyPSym >> \"{0}\"", SymbolsFile.AbsolutePath);
 				Out.WriteLine("echo DummySyms>> \"{0}\"", EncodedBinarySymbolsFile.AbsolutePath);
 			}
 
@@ -226,21 +221,17 @@ namespace UnrealBuildTool
 				);
 
 				// objcopy link debug file to final so
-				Out.WriteLine("\"{0}\" --add-gnu-debuglink=\"{1}\" \"{2}\" \"{3}.temp\"",
+				Out.WriteLine("\"{0}\" --add-gnu-debuglink=\"{1}\" \"{2}\" \"{3}\"",
 					LinuxInfo.Objcopy,
 					DebugFile.AbsolutePath,
 					StrippedFile.AbsolutePath,
 					OutputFile.AbsolutePath
 				);
 
+				Out.WriteLine($"{DeleteCommand} \"{StrippedFile.AbsolutePath}\"");
+
 				if (bUseCmdExe)
 				{
-					// Only move the temp final elf file once its done being linked by objcopy
-					Out.WriteLine("move /Y \"{0}.temp\" \"{1}\"",
-						OutputFile.AbsolutePath,
-						OutputFile.AbsolutePath
-					);
-
 					if (MaxBinarySizeOverrideForObjcopy > 0)
 					{
 						// If we have an override size, then we need to create a dummy file if that size is exceeded
@@ -251,12 +242,6 @@ namespace UnrealBuildTool
 				}
 				else
 				{
-					// Only move the temp final elf file once its done being linked by objcopy
-					Out.WriteLine("mv \"{0}.temp\" \"{1}\"",
-						OutputFile.AbsolutePath,
-						OutputFile.AbsolutePath
-					);
-
 					// Change the debug file to normal permissions. It was taking on the +x rights from the output file
 					Out.WriteLine("chmod 644 \"{0}\"",
 						DebugFile.AbsolutePath
@@ -832,21 +817,12 @@ namespace UnrealBuildTool
 
 		// cache the location of NDK tools
 		protected bool bIsCrossCompiling;
-		/// <summary>
-		/// Whether to use old, slower way to relink circularly dependent libraries.
-		/// It makes sense to use it when cross-compiling on Windows due to race conditions between actions reading and modifying the libs.
-		/// </summary>
-		private bool bUseFixdeps = false;
 
 		/// <summary>
-		/// Track which scripts need to be deleted before appending to
+		/// Contains converter from output file to temporary output file (the one that will be relinked)
 		/// </summary>
-		private bool bHasWipedFixDepsScript = false;
+		private Dictionary<FileItem, FileItem> LinkToPrelinkFileList = new();
 
-		/// <summary>
-		/// Holds all the binaries for a particular target (except maybe the executable itself).
-		/// </summary>
-		private static Dictionary<ReadOnlyTargetRules, List<FileItem>> AllBinaries = new();
 
 		/// <summary>
 		/// Tracks that information about used C++ library is only printed once
@@ -882,16 +858,6 @@ namespace UnrealBuildTool
 				Logger.LogInformation("{SanitizerInfo}", SanitizerInfo);
 			}
 
-			// Also print other once-per-build information
-			if (bUseFixdeps)
-			{
-				Logger.LogInformation("Using old way to relink circularly dependent libraries (with a FixDeps step).");
-			}
-			else
-			{
-				Logger.LogInformation("Using fast way to relink  circularly dependent libraries (no FixDeps).");
-			}
-
 			if (Options.HasFlag(ClangToolChainOptions.CompressDebugFile))
 			{
 				Logger.LogInformation("Compressing debug files");
@@ -970,6 +936,11 @@ namespace UnrealBuildTool
 			return Result;
 		}
 
+		public FileItem GetPrelinkOutput(DirectoryReference intermediateDir, FileItem output)
+		{
+			return FileItem.GetItemByFileReference(FileReference.Combine(intermediateDir, "Prelink", output.Name));
+		}
+
 		/// <summary>
 		/// Creates an action to archive all the .o files into single .a file
 		/// </summary>
@@ -1029,71 +1000,6 @@ namespace UnrealBuildTool
 			return OutputFile;
 		}
 
-		public FileItem? FixDependencies(ReadOnlyTargetRules Target, LinkEnvironment LinkEnvironment, FileItem Executable, IActionGraphBuilder Graph, ILogger Logger)
-		{
-			if (bUseFixdeps)
-			{
-				Logger.LogDebug("Adding postlink step");
-
-				bool bUseCmdExe = BuildHostPlatform.Current.ShellType == ShellType.Cmd;
-				FileReference ShellBinary = BuildHostPlatform.Current.Shell;
-				string ExecuteSwitch = bUseCmdExe ? " /C" : ""; // avoid -c so scripts don't need +x
-				string ScriptName = bUseCmdExe ? "FixDependencies.bat" : "FixDependencies.sh";
-
-				FileItem FixDepsScript = FileItem.GetItemByFileReference(FileReference.Combine(LinkEnvironment.LocalShadowDirectory!, ScriptName));
-
-				// if we never generated one we did not have a circular depends that needed fixing up
-				if (!FixDepsScript.Exists)
-				{
-					return null;
-				}
-
-				Action PostLinkAction = Graph.CreateAction(ActionType.Link);
-				PostLinkAction.WorkingDirectory = Unreal.EngineSourceDirectory;
-				PostLinkAction.CommandPath = ShellBinary;
-				PostLinkAction.StatusDescription = String.Format("{0}", Path.GetFileName(Executable.AbsolutePath));
-				PostLinkAction.CommandDescription = "FixDeps";
-				PostLinkAction.bCanExecuteRemotely = false;
-				PostLinkAction.bCanExecuteInUBA = OperatingSystem.IsWindows(); // Linker on native linux uses vfork/exec which is not handled in uba right now
-				PostLinkAction.CommandArguments = ExecuteSwitch;
-
-				PostLinkAction.CommandArguments += bUseCmdExe ? " \"" : " -c '";
-
-				FileItem OutputFile = FileItem.GetItemByFileReference(FileReference.Combine(LinkEnvironment.LocalShadowDirectory!, Path.GetFileNameWithoutExtension(Executable.AbsolutePath) + ".link"));
-
-				// Make sure we don't run this script until the all executables and shared libraries
-				// have been built.
-				PostLinkAction.PrerequisiteItems.Add(Executable);
-				foreach (FileItem Dependency in AllBinaries.GetValueOrDefault(Target, new()))
-				{
-					PostLinkAction.PrerequisiteItems.Add(Dependency);
-				}
-
-				PostLinkAction.CommandArguments += ShellBinary + ExecuteSwitch + " \"" + FixDepsScript.AbsolutePath + "\" && ";
-
-				// output file should not be empty or it will be rebuilt next time
-				string Touch = bUseCmdExe ? "echo \"Dummy\" >> \"{0}\" && copy /b \"{0}\" +,," : "echo \"Dummy\" >> \"{0}\"";
-
-				PostLinkAction.CommandArguments += String.Format(Touch, OutputFile.AbsolutePath);
-				PostLinkAction.CommandArguments += bUseCmdExe ? "\"" : "'";
-
-				System.Console.WriteLine("{0} {1}", PostLinkAction.CommandPath, PostLinkAction.CommandArguments);
-
-				PostLinkAction.ProducedItems.Add(OutputFile);
-				return OutputFile;
-			}
-			else
-			{
-				return null;
-			}
-		}
-
-		// allow sub-platforms to modify the name of the output file
-		protected virtual FileItem GetLinkOutputFile(LinkEnvironment LinkEnvironment)
-		{
-			return FileItem.GetItemByFileReference(LinkEnvironment.OutputFilePath);
-		}
-
 		public override FileItem LinkFiles(LinkEnvironment LinkEnvironment, bool bBuildImportLibraryOnly, IActionGraphBuilder Graph)
 		{
 			Debug.Assert(!bBuildImportLibraryOnly);
@@ -1122,9 +1028,19 @@ namespace UnrealBuildTool
 			// ignored as a prerequisite for other actions
 			LinkAction.bProducesImportLibrary = LinkEnvironment.bIsBuildingDLL;
 
+			bool needsRelink = LinkEnvironment.bIsBuildingDLL && LinkEnvironment.bIsCrossReferenced;
+
+			// If there will be a relink we let the link step write to a file under intermediate folder
+			FileItem RealOutputFile = FileItem.GetItemByFileReference(LinkEnvironment.OutputFilePath);
+			FileItem OutputFile = RealOutputFile;
+			if (needsRelink)
+			{
+				OutputFile = GetPrelinkOutput(LinkEnvironment.IntermediateDirectory!, OutputFile);
+			}
+
 			// Add the output file as a production of the link action.
-			FileItem OutputFile = GetLinkOutputFile(LinkEnvironment);
 			LinkAction.ProducedItems.Add(OutputFile);
+
 			// LTO/PGO can take a lot of time, make it clear for the user
 			if (LinkEnvironment.bPGOProfile)
 			{
@@ -1203,9 +1119,9 @@ namespace UnrealBuildTool
 					}
 				}
 
-				if ((AdditionalLibrary.Contains("Plugins") || AdditionalLibrary.Contains("Binaries/ThirdParty") || AdditionalLibrary.Contains("Binaries\\ThirdParty")) && Path.GetDirectoryName(AdditionalLibrary) != Path.GetDirectoryName(OutputFile.AbsolutePath))
+				if ((AdditionalLibrary.Contains("Plugins") || AdditionalLibrary.Contains("Binaries/ThirdParty") || AdditionalLibrary.Contains("Binaries\\ThirdParty")) && Path.GetDirectoryName(AdditionalLibrary) != Path.GetDirectoryName(RealOutputFile.AbsolutePath))
 				{
-					string RelativePath = new FileReference(AdditionalLibrary).Directory.MakeRelativeTo(OutputFile.Location.Directory);
+					string RelativePath = new FileReference(AdditionalLibrary).Directory.MakeRelativeTo(RealOutputFile.Location.Directory);
 
 					if (LinkEnvironment.bIsBuildingDLL)
 					{
@@ -1213,7 +1129,7 @@ namespace UnrealBuildTool
 						string AdditionalLibraryRootPath = new FileReference(AdditionalLibrary).Directory.MakeRelativeTo(Unreal.RootDirectory);
 
 						// Figure out how many dirs we need to go back
-						string RelativeRootPath = Unreal.RootDirectory.MakeRelativeTo(OutputFile.Location.Directory);
+						string RelativeRootPath = Unreal.RootDirectory.MakeRelativeTo(RealOutputFile.Location.Directory);
 
 						// Combine the two together ie. number of ../ + the path after the root
 						RelativePath = Path.Combine(RelativeRootPath, AdditionalLibraryRootPath);
@@ -1248,7 +1164,7 @@ namespace UnrealBuildTool
 						string RuntimeLibraryRootPath = new DirectoryReference(RuntimeLibaryPath).MakeRelativeTo(Unreal.RootDirectory);
 
 						// Figure out how many dirs we need to go back
-						string RelativeRootPath = Unreal.RootDirectory.MakeRelativeTo(OutputFile.Location.Directory);
+						string RelativeRootPath = Unreal.RootDirectory.MakeRelativeTo(RealOutputFile.Location.Directory);
 
 						// Combine the two together ie. number of ../ + the path after the root
 						RelativePath = Path.Combine(RelativeRootPath, RuntimeLibraryRootPath);
@@ -1278,7 +1194,7 @@ namespace UnrealBuildTool
 				}
 			}
 
-			ResponseLines.Add(String.Format(" -rpath-link=\"{0}\"", Path.GetDirectoryName(OutputFile.AbsolutePath)));
+			ResponseLines.Add(String.Format(" -rpath-link=\"{0}\"", Path.GetDirectoryName(RealOutputFile.AbsolutePath)));
 
 			// Add the library paths to the argument list.
 			foreach (DirectoryReference LibraryPath in AllLibraryPaths)
@@ -1334,6 +1250,12 @@ namespace UnrealBuildTool
 					// dependencies.
 					FileItem LibraryDependency = FileItem.GetItemByPath(AdditionalLibrary);
 
+					if (LinkToPrelinkFileList.TryGetValue(LibraryDependency, out FileItem? ToRelinkItem))
+					{
+						LibraryDependency = ToRelinkItem;
+						ExternalLibraries += $" -L\"{ToRelinkItem.Directory.FullName.Replace("\\", "/")}\"";
+					}
+
 					string LibName = Path.GetFileNameWithoutExtension(AdditionalLibrary);
 					if (LibName.StartsWith("lib"))
 					{
@@ -1349,7 +1271,7 @@ namespace UnrealBuildTool
 					}
 					string LibLinkFlag = String.Format(" -l{0}", LibName);
 
-					if (LinkEnvironment.bIsBuildingDLL && LinkEnvironment.bIsCrossReferenced)
+					if (needsRelink)
 					{
 						// We are building a cross referenced DLL so we can't actually include
 						// dependencies at this point. Instead we add it to the list of
@@ -1431,11 +1353,14 @@ namespace UnrealBuildTool
 				LinkCommandString = LinkCommandString.Replace("$'{", "'${");    // fixing $'{ORIGIN}' to be '${ORIGIN}'
 			}
 
-			string LinkScriptName = String.Format((bUseCmdExe ? "Link-{0}.link.bat" : "Link-{0}.link.sh"), OutputFile.Location.GetFileName());
-			string LinkScriptFullPath = Path.Combine(LinkEnvironment.LocalShadowDirectory!.FullName, LinkScriptName);
-			Logger.LogDebug("Creating link script: {LinkScriptFullPath}", LinkScriptFullPath);
-			Directory.CreateDirectory(Path.GetDirectoryName(LinkScriptFullPath)!);
-			using (StreamWriter LinkWriter = File.CreateText(LinkScriptFullPath))
+			string LinkScriptName = String.Format((bUseCmdExe ? "Link-{0}.bat" : "Link-{0}.sh"), OutputFile.Location.GetFileName());
+			FileReference LinkScriptFullPath = FileReference.Combine(LinkEnvironment.IntermediateDirectory!, LinkScriptName);
+
+			LinkAction.PrerequisiteItems.Add(FileItem.GetItemByFileReference(LinkScriptFullPath));
+
+			Logger.LogDebug("Creating link script: {LinkScriptFullPath}", LinkScriptFullPath.FullName);
+			Directory.CreateDirectory(Path.GetDirectoryName(LinkScriptFullPath.FullName)!);
+			using (StreamWriter LinkWriter = File.CreateText(LinkScriptFullPath.FullName))
 			{
 				if (bUseCmdExe)
 				{
@@ -1449,7 +1374,10 @@ namespace UnrealBuildTool
 					LinkWriter.WriteLine("if %Retries% GEQ 10 goto failedtorelink");
 					LinkWriter.WriteLine(LinkCommandString);
 					LinkWriter.WriteLine("if %errorlevel% neq 0 goto sleepandretry");
-					LinkWriter.WriteLine(GetDumpEncodeDebugCommand(LinkEnvironment, OutputFile));
+					if (!needsRelink)
+					{
+						LinkWriter.WriteLine(GetDumpEncodeDebugCommand(LinkEnvironment, OutputFile));
+					}
 					LinkWriter.WriteLine("exit 0");
 					LinkWriter.WriteLine(":sleepandretry");
 					LinkWriter.WriteLine("ping 127.0.0.1 -n 1 -w 5000 >NUL 2>NUL");     // timeout complains about lack of redirection
@@ -1468,241 +1396,120 @@ namespace UnrealBuildTool
 					LinkWriter.WriteLine();
 					LinkWriter.WriteLine("set -o errexit");
 					LinkWriter.WriteLine(LinkCommandString);
-					LinkWriter.WriteLine(GetDumpEncodeDebugCommand(LinkEnvironment, OutputFile));
+					if (!needsRelink)
+					{
+						LinkWriter.WriteLine(GetDumpEncodeDebugCommand(LinkEnvironment, OutputFile));
+					}
 				}
 			}
 
 			LinkAction.CommandPath = ShellBinary;
 
 			// This must maintain the quotes around the LinkScriptFullPath
-			LinkAction.CommandArguments = ExecuteSwitch + " \"" + LinkScriptFullPath + "\"";
-
-			// prepare a linker script
-			FileReference LinkerScriptPath = FileReference.Combine(LinkEnvironment.LocalShadowDirectory, "remove-sym.ldscript");
-			if (!DirectoryReference.Exists(LinkEnvironment.LocalShadowDirectory))
-			{
-				DirectoryReference.CreateDirectory(LinkEnvironment.LocalShadowDirectory);
-			}
-			if (FileReference.Exists(LinkerScriptPath))
-			{
-				FileReference.Delete(LinkerScriptPath);
-			}
+			LinkAction.CommandArguments = ExecuteSwitch + " \"" + LinkScriptFullPath.FullName + "\"";
 
 			// Only execute linking on the local PC.
 			LinkAction.bCanExecuteRemotely = false;
 
 			// Prepare a script that will run later, once all shared libraries and the executable
 			// are created. This script will be called by action created in FixDependencies()
-			if (LinkEnvironment.bIsCrossReferenced && LinkEnvironment.bIsBuildingDLL)
-			{
-				if (bUseFixdeps)
+			if (needsRelink)
+			{
+				// Create the action to relink the library. This actions does not overwrite the source file so it can be executed in parallel
+				Action RelinkAction = Graph.CreateAction(ActionType.Link);
+				RelinkAction.RootPaths = LinkEnvironment.RootPaths;
+				RelinkAction.WorkingDirectory = LinkAction.WorkingDirectory;
+				RelinkAction.StatusDescription = LinkAction.StatusDescription;
+				RelinkAction.CommandDescription = "Relink";
+				RelinkAction.bCanExecuteRemotely = false;
+				RelinkAction.bCanExecuteInUBA = OperatingSystem.IsWindows(); // Linker on native linux uses vfork/exec which is not handled in uba right now
+				RelinkAction.ProducedItems.Add(RealOutputFile);
+
+				foreach (FileItem Dependency in EngineAndGameLibrariesFiles)
 				{
-					string ScriptName = bUseCmdExe ? "FixDependencies.bat" : "FixDependencies.sh";
+					RelinkAction.PrerequisiteItems.Add(Dependency);
+				}
+				RelinkAction.PrerequisiteItems.Add(OutputFile); // also depend on the first link action's output
 
-					string FixDepsScriptPath = Path.Combine(LinkEnvironment.LocalShadowDirectory.FullName, ScriptName);
-					if (!bHasWipedFixDepsScript)
-					{
-						bHasWipedFixDepsScript = true;
-						Logger.LogDebug("Creating script: {FixDepsScriptPath}", FixDepsScriptPath);
-						Directory.CreateDirectory(Path.GetDirectoryName(FixDepsScriptPath)!);
-						using (StreamWriter Writer = File.CreateText(FixDepsScriptPath))
-						{
-							if (bUseCmdExe)
-							{
-								Writer.NewLine = "\r\n";
-								Writer.WriteLine("@echo off");
-								Writer.WriteLine("rem Automatically generated by UnrealBuildTool");
-								Writer.WriteLine("rem *DO NOT EDIT*");
-								Writer.WriteLine();
-							}
-							else
-							{
-								Writer.NewLine = "\n";
-								Writer.WriteLine("#!/bin/sh");
-								Writer.WriteLine("# Automatically generated by UnrealBuildTool");
-								Writer.WriteLine("# *DO NOT EDIT*");
-								Writer.WriteLine();
-								Writer.WriteLine("set -o errexit");
-							}
-						}
-					}
+				string LinkOutputFileForwardSlashes = OutputFile.AbsolutePath.Replace("\\", "/");
+				string RelinkedFileForwardSlashes = RealOutputFile.AbsolutePath.Replace("\\", "/");
 
-					StreamWriter FixDepsScript = File.AppendText(FixDepsScriptPath);
-					FixDepsScript.NewLine = bUseCmdExe ? "\r\n" : "\n";
+				string EngineAndGameLibrariesString = "";
+				foreach (string Library in EngineAndGameLibrariesLinkFlags)
+				{
+					EngineAndGameLibrariesString += Library;
+				}
 
-					string EngineAndGameLibrariesString = "";
-					foreach (string Library in EngineAndGameLibrariesLinkFlags)
-					{
-						EngineAndGameLibrariesString += Library;
-					}
+				// create the relinking step
+				string RelinkScriptName = String.Format((bUseCmdExe ? "Relink-{0}.bat" : "Relink-{0}.sh"), RealOutputFile.Location.GetFileName());
+				FileReference RelinkScriptFullPath = FileReference.Combine(LinkEnvironment.IntermediateDirectory!, RelinkScriptName);
 
-					FixDepsScript.WriteLine("echo Fixing {0}", Path.GetFileName(OutputFile.AbsolutePath));
-					if (!bUseCmdExe)
-					{
-						FixDepsScript.WriteLine("TIMESTAMP=`stat --format %y \"{0}\"`", OutputFile.AbsolutePath);
-					}
-					string FixDepsLine = LinkCommandString;
+				RelinkAction.PrerequisiteItems.Add(FileItem.GetItemByFileReference(RelinkScriptFullPath));
+
+				using (StreamWriter RelinkWriter = File.CreateText(RelinkScriptFullPath.FullName))
+				{
+					string RelinkInvocation = LinkCommandString;
 					string Replace = "-Wl,--allow-shlib-undefined";
+					RelinkInvocation = RelinkInvocation.Replace(Replace, EngineAndGameLibrariesString);
+
+					// should be the same as RelinkedFileRef
+					RelinkInvocation = RelinkInvocation.Replace(LinkOutputFileForwardSlashes, RelinkedFileForwardSlashes);
+					RelinkInvocation = RelinkInvocation.Replace("$", "\\$");
 
-					FixDepsLine = FixDepsLine.Replace(Replace, EngineAndGameLibrariesString);
-					string OutputFileForwardSlashes = OutputFile.AbsolutePath.Replace("\\", "/");
-					FixDepsLine = FixDepsLine.Replace(OutputFileForwardSlashes, OutputFileForwardSlashes + ".fixed");
-					FixDepsLine = FixDepsLine.Replace("$", "\\$");
-					FixDepsScript.WriteLine(FixDepsLine);
 					if (bUseCmdExe)
 					{
-						FixDepsScript.WriteLine("move /Y \"{0}.fixed\" \"{0}\"", OutputFile.AbsolutePath);
+						RelinkWriter.WriteLine("@echo off");
+						RelinkWriter.WriteLine("rem Automatically generated by UnrealBuildTool");
+						RelinkWriter.WriteLine("rem *DO NOT EDIT*");
+						RelinkWriter.WriteLine();
+						RelinkWriter.WriteLine("set Retries=0");
+						RelinkWriter.WriteLine(":relinkloop");
+						RelinkWriter.WriteLine("if %Retries% GEQ 10 goto failedtorelink");
+						RelinkWriter.WriteLine(RelinkInvocation);
+						RelinkWriter.WriteLine("if %errorlevel% neq 0 goto sleepandretry");
+						RelinkWriter.WriteLine(GetDumpEncodeDebugCommand(LinkEnvironment, RealOutputFile));
+						RelinkWriter.WriteLine("exit 0");
+						RelinkWriter.WriteLine(":sleepandretry");
+						RelinkWriter.WriteLine("ping 127.0.0.1 -n 1 -w 5000 >NUL 2>NUL");     // timeout complains about lack of redirection
+						RelinkWriter.WriteLine("set /a Retries+=1");
+						RelinkWriter.WriteLine("goto relinkloop");
+						RelinkWriter.WriteLine(":failedtorelink");
+						RelinkWriter.WriteLine("echo Failed to relink {0} after %Retries% retries", OutputFile.AbsolutePath);
+						RelinkWriter.WriteLine("exit 1");
 					}
 					else
 					{
-						FixDepsScript.WriteLine("mv \"{0}.fixed\" \"{0}\"", OutputFile.AbsolutePath);
-						FixDepsScript.WriteLine("touch -d \"$TIMESTAMP\" \"{0}\"", OutputFile.AbsolutePath);
-						FixDepsScript.WriteLine();
+						RelinkWriter.NewLine = "\n";
+						RelinkWriter.WriteLine("#!/bin/sh");
+						RelinkWriter.WriteLine("# Automatically generated by UnrealBuildTool");
+						RelinkWriter.WriteLine("# *DO NOT EDIT*");
+						RelinkWriter.WriteLine();
+						RelinkWriter.WriteLine("set -o errexit");
+						RelinkWriter.WriteLine(RelinkInvocation);
+						RelinkWriter.WriteLine(GetDumpEncodeDebugCommand(LinkEnvironment, RealOutputFile));
 					}
-					FixDepsScript.Close();
 				}
-				else
-				{
-					// Create the action to relink the library. This actions does not overwrite the source file so it can be executed in parallel
-					Action RelinkAction = Graph.CreateAction(ActionType.Link);
-					RelinkAction.RootPaths = LinkEnvironment.RootPaths;
-					RelinkAction.WorkingDirectory = LinkAction.WorkingDirectory;
-					RelinkAction.StatusDescription = LinkAction.StatusDescription;
-					RelinkAction.CommandDescription = "Relink";
-					RelinkAction.bCanExecuteRemotely = false;
-					RelinkAction.bCanExecuteInUBA = OperatingSystem.IsWindows(); // Linker on native linux uses vfork/exec which is not handled in uba right now
-					RelinkAction.ProducedItems.Clear();
-					RelinkAction.PrerequisiteItems = new SortedSet<FileItem>(LinkAction.PrerequisiteItems);
-					foreach (FileItem Dependency in EngineAndGameLibrariesFiles)
-					{
-						RelinkAction.PrerequisiteItems.Add(Dependency);
-					}
-					RelinkAction.PrerequisiteItems.Add(OutputFile); // also depend on the first link action's output
-
-					string LinkOutputFileForwardSlashes = OutputFile.AbsolutePath.Replace("\\", "/");
-					string RelinkedFileForwardSlashes = Path.Combine(LinkEnvironment.LocalShadowDirectory.FullName, OutputFile.Location.GetFileName()) + ".relinked";
-
-					// cannot use the real product because we need to maintain the timestamp on it
-					FileReference RelinkActionDummyProductRef = FileReference.Combine(LinkEnvironment.LocalShadowDirectory, LinkEnvironment.OutputFilePath.GetFileNameWithoutExtension() + ".relinked_action_ran");
-					RelinkAction.ProducedItems.Add(FileItem.GetItemByFileReference(RelinkActionDummyProductRef));
-
-					string EngineAndGameLibrariesString = "";
-					foreach (string Library in EngineAndGameLibrariesLinkFlags)
-					{
-						EngineAndGameLibrariesString += Library;
-					}
-
-					// create the relinking step
-					string RelinkScriptName = String.Format((bUseCmdExe ? "Relink-{0}.bat" : "Relink-{0}.sh"), OutputFile.Location.GetFileName());
-					string RelinkScriptFullPath = Path.Combine(LinkEnvironment.LocalShadowDirectory.FullName, RelinkScriptName);
 
-					Logger.LogDebug("Creating script: {RelinkScriptFullPath}", RelinkScriptFullPath);
-					Directory.CreateDirectory(Path.GetDirectoryName(RelinkScriptFullPath)!);
-					using (StreamWriter RelinkWriter = File.CreateText(RelinkScriptFullPath))
-					{
-						string RelinkInvocation = LinkCommandString;
-						string Replace = "-Wl,--allow-shlib-undefined";
-						RelinkInvocation = RelinkInvocation.Replace(Replace, EngineAndGameLibrariesString);
-
-						// should be the same as RelinkedFileRef
-						RelinkInvocation = RelinkInvocation.Replace(LinkOutputFileForwardSlashes, RelinkedFileForwardSlashes);
-						RelinkInvocation = RelinkInvocation.Replace("$", "\\$");
-
-						if (bUseCmdExe)
-						{
-							RelinkWriter.WriteLine("@echo off");
-							RelinkWriter.WriteLine("rem Automatically generated by UnrealBuildTool");
-							RelinkWriter.WriteLine("rem *DO NOT EDIT*");
-							RelinkWriter.WriteLine();
-							RelinkWriter.WriteLine("set Retries=0");
-							RelinkWriter.WriteLine(":relinkloop");
-							RelinkWriter.WriteLine("if %Retries% GEQ 10 goto failedtorelink");
-							RelinkWriter.WriteLine(RelinkInvocation);
-							RelinkWriter.WriteLine("if %errorlevel% neq 0 goto sleepandretry");
-							RelinkWriter.WriteLine("copy /B \"{0}\" \"{1}.temp\" >NUL 2>NUL", RelinkedFileForwardSlashes, OutputFile.AbsolutePath);
-							RelinkWriter.WriteLine("if %errorlevel% neq 0 goto sleepandretry");
-							RelinkWriter.WriteLine("move /Y \"{0}.temp\" \"{1}\" >NUL 2>NUL", OutputFile.AbsolutePath, OutputFile.AbsolutePath);
-							RelinkWriter.WriteLine("if %errorlevel% neq 0 goto sleepandretry");
-							RelinkWriter.WriteLine(GetDumpEncodeDebugCommand(LinkEnvironment, OutputFile));
-							RelinkWriter.WriteLine("echo \"Dummy\" >> \"{0}\" && copy /b \"{0}\" +,,", RelinkActionDummyProductRef.FullName);
-							RelinkWriter.WriteLine("echo Relinked {0} successfully after %Retries% retries", OutputFile.AbsolutePath);
-							RelinkWriter.WriteLine("exit 0");
-							RelinkWriter.WriteLine(":sleepandretry");
-							RelinkWriter.WriteLine("ping 127.0.0.1 -n 1 -w 5000 >NUL 2>NUL");     // timeout complains about lack of redirection
-							RelinkWriter.WriteLine("set /a Retries+=1");
-							RelinkWriter.WriteLine("goto relinkloop");
-							RelinkWriter.WriteLine(":failedtorelink");
-							RelinkWriter.WriteLine("echo Failed to relink {0} after %Retries% retries", OutputFile.AbsolutePath);
-							RelinkWriter.WriteLine("exit 1");
-						}
-						else
-						{
-							RelinkWriter.NewLine = "\n";
-							RelinkWriter.WriteLine("#!/bin/sh");
-							RelinkWriter.WriteLine("# Automatically generated by UnrealBuildTool");
-							RelinkWriter.WriteLine("# *DO NOT EDIT*");
-							RelinkWriter.WriteLine();
-							RelinkWriter.WriteLine("set -o errexit");
-							RelinkWriter.WriteLine(RelinkInvocation);
-							RelinkWriter.WriteLine("TIMESTAMP=`stat --format %y \"{0}\"`", OutputFile.AbsolutePath);
-							RelinkWriter.WriteLine("cp \"{0}\" \"{1}.temp\"", RelinkedFileForwardSlashes, OutputFile.AbsolutePath);
-							RelinkWriter.WriteLine("mv \"{0}.temp\" \"{1}\"", OutputFile.AbsolutePath, OutputFile.AbsolutePath);
-							RelinkWriter.WriteLine(GetDumpEncodeDebugCommand(LinkEnvironment, OutputFile));
-							RelinkWriter.WriteLine("touch -d \"$TIMESTAMP\" \"{0}\"", OutputFile.AbsolutePath);
-							RelinkWriter.WriteLine();
-							RelinkWriter.WriteLine("echo \"Dummy\" >> \"{0}\"", RelinkActionDummyProductRef.FullName);
-						}
-					}
-
-					RelinkAction.CommandPath = ShellBinary;
-					RelinkAction.CommandArguments = ExecuteSwitch + " \"" + RelinkScriptFullPath + "\"";
-				}
+				RelinkAction.CommandPath = ShellBinary;
+				RelinkAction.CommandArguments = ExecuteSwitch + " \"" + RelinkScriptFullPath + "\"";
 			}
-			return OutputFile;
+			return RealOutputFile;
 		}
 
 		public override void SetupBundleDependencies(ReadOnlyTargetRules Target, IEnumerable<UEBuildBinary> Binaries, string GameName)
 		{
-			if (bUseFixdeps)
+			// Populate the lookup table from real output to prelink output
+			lock (LinkToPrelinkFileList)
 			{
-				if (!AllBinaries.ContainsKey(Target))
-				{
-					AllBinaries.Add(Target, new());
-				}
-				foreach (UEBuildBinary Binary in Binaries)
+				foreach (UEBuildBinary binary in Binaries)
 				{
-					AllBinaries[Target].Add(FileItem.GetItemByFileReference(Binary.OutputFilePath));
-				}
-			}
-		}
-
-		public override ICollection<FileItem> PostBuild(ReadOnlyTargetRules Target, FileItem Executable, LinkEnvironment BinaryLinkEnvironment, IActionGraphBuilder Graph)
-		{
-			ICollection<FileItem> OutputFiles = base.PostBuild(Target, Executable, BinaryLinkEnvironment, Graph);
-
-			if (bUseFixdeps)
-			{
-				if (BinaryLinkEnvironment.bIsBuildingDLL || BinaryLinkEnvironment.bIsBuildingLibrary)
-				{
-					return OutputFiles;
-				}
-
-				FileItem? FixDepsOutputFile = FixDependencies(Target, BinaryLinkEnvironment, Executable, Graph, Logger);
-				if (FixDepsOutputFile != null)
-				{
-					OutputFiles.Add(FixDepsOutputFile);
-				}
-			}
-			else
-			{
-				// make build product of cross-referenced DSOs to be *.relinked_action_ran, so the relinking steps are executed
-				if (BinaryLinkEnvironment.bIsBuildingDLL && BinaryLinkEnvironment.bIsCrossReferenced)
-				{
-					FileReference RelinkedMapRef = FileReference.Combine(BinaryLinkEnvironment.LocalShadowDirectory!, BinaryLinkEnvironment.OutputFilePath.GetFileNameWithoutExtension() + ".relinked_action_ran");
-					OutputFiles.Add(FileItem.GetItemByFileReference(RelinkedMapRef));
+					if (binary.bCreateImportLibrarySeparately)
+					{
+						FileItem output = FileItem.GetItemByFileReference(binary.OutputFilePath);
+						LinkToPrelinkFileList.Add(output, GetPrelinkOutput(binary.IntermediateDirectory, output));
+					}
 				}
 			}
-			return OutputFiles;
 		}
 
 		public void StripSymbols(FileReference SourceFile, FileReference TargetFile, ILogger Logger)
