Shell quotes around uproject argments to handle when UE projets root has spaces. Otherwise project gets passed as multiple args to XcodeBuild and in turn to UBT resulting in errors building.

#rn Fix Xcode build for projects whose full paths may contain spaces.

diff --git a/Engine/Source/Programs/UnrealBuildTool/ProjectFiles/Xcode/XcodeProject.cs b/Engine/Source/Programs/UnrealBuildTool/ProjectFiles/Xcode/XcodeProject.cs
--- a/Engine/Source/Programs/UnrealBuildTool/ProjectFiles/Xcode/XcodeProject.cs
+++ b/Engine/Source/Programs/UnrealBuildTool/ProjectFiles/Xcode/XcodeProject.cs
@@ -1226,7 +1226,7 @@ namespace UnrealBuildTool.XcodeProjectXcconfig
 				string BuildToolPath = UEDir + "/Engine/Build/BatchFiles/Mac/XcodeBuild.sh";
 
 				// insert elements to call UBT when building
-				Content.WriteLine(3, $"buildArgumentsString = \"$(ACTION) {UnrealTargetName} $(PLATFORM_NAME) $(CONFIGURATION) {UProjectParam}\";");
+				Content.WriteLine(3, $"buildArgumentsString = \"$(ACTION) {UnrealTargetName} $(PLATFORM_NAME) $(CONFIGURATION) \\\"{UProjectParam}\\\"\";");
 				Content.WriteLine(3, $"buildToolPath = \"{BuildToolPath}\";");
 				Content.WriteLine(3, $"buildWorkingDirectory = \"{UEDir}\";");
 			}
