[UBT]
* Fixed bug where psym file is deleted when it was supposed to be kept
* Also added produced files to list of produced files

diff --git a/Engine/Source/Programs/UnrealBuildTool/Platform/Linux/LinuxToolChain.cs b/Engine/Source/Programs/UnrealBuildTool/Platform/Linux/LinuxToolChain.cs
--- a/Engine/Source/Programs/UnrealBuildTool/Platform/Linux/LinuxToolChain.cs
+++ b/Engine/Source/Programs/UnrealBuildTool/Platform/Linux/LinuxToolChain.cs
@@ -153,7 +153,7 @@ namespace UnrealBuildTool
 			return bIsCrossCompiling;
 		}
 
-		protected internal virtual string GetDumpEncodeDebugCommand(LinkEnvironment LinkEnvironment, FileItem OutputFile)
+		protected internal virtual string GetDumpEncodeDebugCommand(LinkEnvironment LinkEnvironment, FileItem OutputFile, Action action)
 		{
 			bool bUseCmdExe = BuildHostPlatform.Current.Platform == UnrealTargetPlatform.Win64;
 			string DumpCommand = bUseCmdExe ? "\"{0}\" \"{1}\" \"{2}\" 2>NUL" : "\"{0}\" -c -o \"{2}\" \"{1}\"";
@@ -163,7 +163,8 @@ namespace UnrealBuildTool
 			FileItem StrippedFile = FileItem.GetItemByPath(Path.Combine(LinkEnvironment.IntermediateDirectory.FullName, OutputFile.Location.GetFileName() + "_nodebug"));
 			FileItem DebugFile = FileItem.GetItemByPath(Path.Combine(LinkEnvironment.OutputDirectory.FullName, OutputFile.Location.GetFileNameWithoutExtension() + ".debug"));
 
-			if (Options.HasFlag(ClangToolChainOptions.PreservePSYM))
+			bool preservePsym = Options.HasFlag(ClangToolChainOptions.PreservePSYM);
+			if (preservePsym)
 			{
 				SymbolsFile = FileItem.GetItemByPath(Path.Combine(LinkEnvironment.OutputDirectory.FullName, OutputFile.Location.GetFileNameWithoutExtension() + ".psym"));
 			}
@@ -187,7 +188,14 @@ namespace UnrealBuildTool
 					EncodedBinarySymbolsFile.AbsolutePath
 				);
 
-				Out.WriteLine($"{DeleteCommand} \"{SymbolsFile.AbsolutePath}\"");
+				if (preservePsym)
+				{
+					action.ProducedItems.Add(SymbolsFile);
+				}
+				else
+				{
+					Out.WriteLine($"{DeleteCommand} \"{SymbolsFile.AbsolutePath}\"");
+				}
 			}
 			else
 			{
@@ -195,6 +203,8 @@ namespace UnrealBuildTool
 				Out.WriteLine("echo DummySyms>> \"{0}\"", EncodedBinarySymbolsFile.AbsolutePath);
 			}
 
+			action.ProducedItems.Add(EncodedBinarySymbolsFile);
+
 			if (LinkEnvironment.bCreateDebugInfo)
 			{
 				if (MaxBinarySizeOverrideForObjcopy > 0 && bUseCmdExe)
@@ -256,6 +266,8 @@ namespace UnrealBuildTool
 				);
 			}
 
+			action.ProducedItems.Add(DebugFile);
+
 			return Out.ToString();
 		}
 
@@ -1376,7 +1388,7 @@ namespace UnrealBuildTool
 					LinkWriter.WriteLine("if %errorlevel% neq 0 goto sleepandretry");
 					if (!needsRelink)
 					{
-						LinkWriter.WriteLine(GetDumpEncodeDebugCommand(LinkEnvironment, OutputFile));
+						LinkWriter.WriteLine(GetDumpEncodeDebugCommand(LinkEnvironment, OutputFile, LinkAction));
 					}
 					LinkWriter.WriteLine("exit 0");
 					LinkWriter.WriteLine(":sleepandretry");
@@ -1398,7 +1410,7 @@ namespace UnrealBuildTool
 					LinkWriter.WriteLine(LinkCommandString);
 					if (!needsRelink)
 					{
-						LinkWriter.WriteLine(GetDumpEncodeDebugCommand(LinkEnvironment, OutputFile));
+						LinkWriter.WriteLine(GetDumpEncodeDebugCommand(LinkEnvironment, OutputFile, LinkAction));
 					}
 				}
 			}
@@ -1467,7 +1479,7 @@ namespace UnrealBuildTool
 						RelinkWriter.WriteLine("if %Retries% GEQ 10 goto failedtorelink");
 						RelinkWriter.WriteLine(RelinkInvocation);
 						RelinkWriter.WriteLine("if %errorlevel% neq 0 goto sleepandretry");
-						RelinkWriter.WriteLine(GetDumpEncodeDebugCommand(LinkEnvironment, RealOutputFile));
+						RelinkWriter.WriteLine(GetDumpEncodeDebugCommand(LinkEnvironment, RealOutputFile, RelinkAction));
 						RelinkWriter.WriteLine("exit 0");
 						RelinkWriter.WriteLine(":sleepandretry");
 						RelinkWriter.WriteLine("ping 127.0.0.1 -n 1 -w 5000 >NUL 2>NUL");     // timeout complains about lack of redirection
@@ -1486,7 +1498,7 @@ namespace UnrealBuildTool
 						RelinkWriter.WriteLine();
 						RelinkWriter.WriteLine("set -o errexit");
 						RelinkWriter.WriteLine(RelinkInvocation);
-						RelinkWriter.WriteLine(GetDumpEncodeDebugCommand(LinkEnvironment, RealOutputFile));
+						RelinkWriter.WriteLine(GetDumpEncodeDebugCommand(LinkEnvironment, RealOutputFile, RelinkAction));
 					}
 				}
 
