[UBT]
* Fixed so swift actions always disable vfs. Will revisit and figure out what code paths swiftc uses that is not properly handled

diff --git a/Engine/Source/Programs/UnrealBuildTool/ToolChain/AppleToolChain.cs b/Engine/Source/Programs/UnrealBuildTool/ToolChain/AppleToolChain.cs
--- a/Engine/Source/Programs/UnrealBuildTool/ToolChain/AppleToolChain.cs
+++ b/Engine/Source/Programs/UnrealBuildTool/ToolChain/AppleToolChain.cs
@@ -604,6 +604,10 @@ namespace UnrealBuildTool
 
 		private void GetSharedSwiftArguments(CppCompileEnvironment CompileEnvironment, List<string> Arguments, string ModuleName)
 		{
+			// TODO: Fix so swiftc path supports vfs
+			CppRootPaths rootPaths = new(CompileEnvironment.RootPaths);
+			rootPaths.bUseVfs = false;
+
 			Arguments.Add($"-target {ToolChainSettings.Value.GetTargetTuple(CompileEnvironment.Architecture, Target!.Platform, Target!.Type)}");
 			Arguments.Add($"-sdk {ToolChainSettings.Value.GetSDKPath(CompileEnvironment.Architecture)}");
 
@@ -613,7 +617,7 @@ namespace UnrealBuildTool
 
 			if (CompileEnvironment.ForceIncludeFiles.Count() > 0)
 			{
-				Arguments.Add($"-import-objc-header \"{NormalizeCommandLinePath(CompileEnvironment.ForceIncludeFiles.First(), CompileEnvironment.RootPaths)}\"");
+				Arguments.Add($"-import-objc-header \"{NormalizeCommandLinePath(CompileEnvironment.ForceIncludeFiles.First(), rootPaths)}\"");
 			}
 
 			// tell swift to define PLATFORM_FOO, and tell C++ to define PLATFORM_FOO=1
