[UBT]
* Fixed so if project is called "Clang", then it will not use that name and instead use the default (Project). Fixes compile errors when building the project "Clang" with vfs enabled

diff --git a/Engine/Source/Programs/UnrealBuildTool/System/RootPaths.cs b/Engine/Source/Programs/UnrealBuildTool/System/RootPaths.cs
--- a/Engine/Source/Programs/UnrealBuildTool/System/RootPaths.cs
+++ b/Engine/Source/Programs/UnrealBuildTool/System/RootPaths.cs
@@ -136,7 +136,11 @@ class CppRootPaths : IEnumerable<(uint id, DirectoryReference vfs, DirectoryRefe
 		{
 			if (key == CppRootPathFolder.Project)
 			{
-				AddFolderName(CppRootPathFolder.Project, value.GetDirectoryName());
+				string projectName = value.GetDirectoryName();
+				if (!projectName.StartsWith("Clang")) // Lol, there is a project called clang.. which collides with the compiler clang causing two roots to have the same path
+				{
+					AddFolderName(CppRootPathFolder.Project, value.GetDirectoryName());
+				}
 			}
 			_lookup[key] = (AsVfsPath(key), value);
 		}
@@ -366,7 +370,12 @@ class CppRootPaths : IEnumerable<(uint id, DirectoryReference vfs, DirectoryRefe
 	/// </summary>
 	/// <param name="folder">The folder enum to update</param>
 	/// <param name="name">The new string name</param>
-	public void AddFolderName(CppRootPathFolder folder, string name) => _folderNames[folder] = name;
+	public void AddFolderName(CppRootPathFolder folder, string name)
+	{
+		if (name == "Clang")
+			Console.WriteLine();
+		_folderNames[folder] = name;
+	}
 
 	private DirectoryReference AsVfsPath(CppRootPathFolder folder)
 	{
