[UBT]
* Moved ispc header creation logic to after pch to make sure pch is not depending on ispc header creation.
* Added validation logic in pch wrapper file to verify that there are no ispc generated headers in pch

diff --git a/Engine/Source/Programs/UnrealBuildTool/Configuration/UEBuildModuleCPP.cs b/Engine/Source/Programs/UnrealBuildTool/Configuration/UEBuildModuleCPP.cs
--- a/Engine/Source/Programs/UnrealBuildTool/Configuration/UEBuildModuleCPP.cs
+++ b/Engine/Source/Programs/UnrealBuildTool/Configuration/UEBuildModuleCPP.cs
@@ -464,12 +464,6 @@ namespace UnrealBuildTool
 				MinFilesUsingPrecompiledHeader = 1;
 			}
 
-			// Generate ISPC headers first so C++ can consume them
-			if (InputFiles.ISPCFiles.Count > 0)
-			{
-				CreateHeadersForISPC(ToolChain, ModuleCompileEnvironment, InputFiles.ISPCFiles, IntermediateDirectory, Graph);
-			}
-
 			// Set up the environment with which to compile the CPP files
 			CppCompileEnvironment CppCompileEnvironment = new(ModuleCompileEnvironment);
 
@@ -501,6 +495,13 @@ namespace UnrealBuildTool
 				Logger.LogDebug("Module '{ModuleName}' uses PCH '{PCHIncludeFilename}'", Name, CppCompileEnvironment.PrecompiledHeaderFile);
 			}
 
+			// Generate ISPC headers before normal c++ actions but after pch since we don't want pch actions to wait for ispc (ispc headers are not allowed in pch)
+			if (InputFiles.ISPCFiles.Count > 0)
+			{
+				CreateHeadersForISPC(ToolChain, CppCompileEnvironment, InputFiles.ISPCFiles, IntermediateDirectory, Graph);
+			}
+
+
 			// Write all the definitions to a separate file
 			CreateHeaderForDefinitions(CppCompileEnvironment, IntermediateDirectory, null, Graph);
 
@@ -2021,6 +2022,9 @@ namespace UnrealBuildTool
 				WrapperContents.AppendLine("// PCH for {0}", IncludedFile.Name);
 				WrapperContents.AppendLine("#include \"{0}\"", DefinitionsFile.Name);
 				WrapperContents.AppendLine("#include \"{0}\"", IncludedFile.Name);
+				WrapperContents.AppendLine("#ifdef __ISPC_ALIGN__");
+				WrapperContents.AppendLine("#error ispc.generated.h files are not allowed in precompiled headers (This is a build time optimization to reduce action dependencies)");
+				WrapperContents.AppendLine("#endif");
 			}
 
 			// Create the item
