[UBT]
* Fixed so ispc compile actions depend on ispc header generation actions.. got lost when changing so pch actions are not depending on ispc header generation actions

diff --git a/Engine/Source/Programs/UnrealBuildTool/Configuration/UEBuildModuleCPP.cs b/Engine/Source/Programs/UnrealBuildTool/Configuration/UEBuildModuleCPP.cs
--- a/Engine/Source/Programs/UnrealBuildTool/Configuration/UEBuildModuleCPP.cs
+++ b/Engine/Source/Programs/UnrealBuildTool/Configuration/UEBuildModuleCPP.cs
@@ -496,9 +496,12 @@ namespace UnrealBuildTool
 			}
 
 			// Generate ISPC headers before normal c++ actions but after pch since we don't want pch actions to wait for ispc (ispc headers are not allowed in pch)
+			IEnumerable<FileItem> ISPCHeaders = Enumerable.Empty<FileItem>();
 			if (InputFiles.ISPCFiles.Count > 0)
 			{
-				CreateHeadersForISPC(ToolChain, CppCompileEnvironment, InputFiles.ISPCFiles, IntermediateDirectory, Graph);
+				ISPCHeaders = ToolChain.GenerateAllISPCHeaders(CppCompileEnvironment, InputFiles.ISPCFiles, IntermediateDirectory, Graph).GeneratedHeaderFiles;
+				CppCompileEnvironment.AdditionalPrerequisites.AddRange(ISPCHeaders);
+				CppCompileEnvironment.UserIncludePaths.Add(IntermediateDirectory);
 			}
 
 
@@ -833,7 +836,10 @@ namespace UnrealBuildTool
 			// Compile ISPC files directly
 			if (InputFiles.ISPCFiles.Count > 0)
 			{
-				LinkInputFiles.AddRange(ToolChain.CompileAllISPCFiles(ModuleCompileEnvironment, InputFiles.ISPCFiles, IntermediateDirectory, Graph).ObjectFiles);
+				CppCompileEnvironment IspcCompileEnvironment = new(ModuleCompileEnvironment);
+				IspcCompileEnvironment.AdditionalPrerequisites.AddRange(ISPCHeaders);
+				IspcCompileEnvironment.UserIncludePaths.Add(IntermediateDirectory);
+				LinkInputFiles.AddRange(ToolChain.CompileAllISPCFiles(IspcCompileEnvironment, InputFiles.ISPCFiles, IntermediateDirectory, Graph).ObjectFiles);
 			}
 
 			// Compile C files. Do not use a PCH here, because a C++ PCH is not compatible with C source files.
@@ -1988,22 +1994,6 @@ namespace UnrealBuildTool
 			return null;
 		}
 
-		/// <summary>
-		/// Creates header files from ISPC for inclusion and adds them as dependencies.
-		/// </summary>
-		/// <param name="ToolChain">The toolchain to generate the PCH</param>
-		/// <param name="CompileEnvironment">Compile environment</param>
-		/// <param name="InputFiles">List of ISPC source files</param>
-		/// <param name="IntermediateDirectory">Directory to create the intermediate file</param>
-		/// <param name="Graph">List of actions to be executed. Additional actions will be added to this list.</param>
-		static void CreateHeadersForISPC(UEToolChain ToolChain, CppCompileEnvironment CompileEnvironment, List<FileItem> InputFiles, DirectoryReference IntermediateDirectory, IActionGraphBuilder Graph)
-		{
-			CPPOutput Output = ToolChain.GenerateAllISPCHeaders(CompileEnvironment, InputFiles, IntermediateDirectory, Graph);
-
-			CompileEnvironment.AdditionalPrerequisites.AddRange(Output.GeneratedHeaderFiles);
-			CompileEnvironment.UserIncludePaths.Add(IntermediateDirectory);
-		}
-
 		/// <summary>
 		/// Create a header file containing the module definitions, which also includes the PCH itself. Including through another file is necessary on 
 		/// Clang, since we get warnings about #pragma once otherwise, but it also allows us to consistently define the preprocessor state on all 
