[UBT]
* Disable vfs when running pvs-studio build since post build action is not detoured and will not be able to understand vfs paths

diff --git a/Engine/Source/Programs/UnrealBuildTool/Configuration/TargetRules.cs b/Engine/Source/Programs/UnrealBuildTool/Configuration/TargetRules.cs
--- a/Engine/Source/Programs/UnrealBuildTool/Configuration/TargetRules.cs
+++ b/Engine/Source/Programs/UnrealBuildTool/Configuration/TargetRules.cs
@@ -3180,6 +3180,13 @@ namespace UnrealBuildTool
 			{
 				// Always override environment for analyzing, regardless of other settings
 				IntermediateEnvironment = UnrealIntermediateEnvironment.Analyze;
+
+				// Need to disable vfs for pvsstudio since there is a postbuild action that reads output which is not detoured
+				// and will not be able to interpret virtual file paths
+				if (StaticAnalyzer == StaticAnalyzer.PVSStudio)
+				{
+					bUseVFS = false;
+				}
 			}
 			else if (IntermediateEnvironment == UnrealIntermediateEnvironment.Default)
 			{
