Added circular dependency information in modules to tell UBT which things that needs to relink on linux

#rnx

diff --git a/Engine/Source/Developer/AssetTools/AssetTools.Build.cs b/Engine/Source/Developer/AssetTools/AssetTools.Build.cs
--- a/Engine/Source/Developer/AssetTools/AssetTools.Build.cs
+++ b/Engine/Source/Developer/AssetTools/AssetTools.Build.cs
@@ -97,6 +97,9 @@ public class AssetTools : ModuleRules
 		CircularlyReferencedDependentModules.AddRange(
 			new string[] {
 				"UnrealEd",
+				"Landscape",
+				"Kismet",
+				"MaterialEditor",
 			}
 		);
 
diff --git a/Engine/Source/Editor/AnimGraph/AnimGraph.Build.cs b/Engine/Source/Editor/AnimGraph/AnimGraph.Build.cs
--- a/Engine/Source/Editor/AnimGraph/AnimGraph.Build.cs
+++ b/Engine/Source/Editor/AnimGraph/AnimGraph.Build.cs
@@ -42,10 +42,15 @@ public class AnimGraph : ModuleRules
 		);
 
         CircularlyReferencedDependentModules.AddRange(
-			new string[] {
+			[
                 "UnrealEd",
                 "GraphEditor",
-            }
+				"KismetCompiler",
+				"Kismet",
+				"BlueprintGraph",
+				"ContentBrowser",
+				"PropertyEditor",
+			]
 		);
 
         PrivateIncludePathModuleNames.AddRange(
diff --git a/Engine/Source/Editor/Documentation/Documentation.Build.cs b/Engine/Source/Editor/Documentation/Documentation.Build.cs
--- a/Engine/Source/Editor/Documentation/Documentation.Build.cs
+++ b/Engine/Source/Editor/Documentation/Documentation.Build.cs
@@ -37,7 +37,11 @@ namespace UnrealBuildTool.Rules
 				}
 			);
 
-			CircularlyReferencedDependentModules.Add("SourceControl");
+			CircularlyReferencedDependentModules.AddRange(
+				[
+					"ContentBrowser",
+					"SourceControl",
+				]);
 
 			DynamicallyLoadedModuleNames.AddRange(
 				new string[]
diff --git a/Engine/Source/Editor/GraphEditor/GraphEditor.Build.cs b/Engine/Source/Editor/GraphEditor/GraphEditor.Build.cs
--- a/Engine/Source/Editor/GraphEditor/GraphEditor.Build.cs
+++ b/Engine/Source/Editor/GraphEditor/GraphEditor.Build.cs
@@ -51,9 +51,10 @@ public class GraphEditor : ModuleRules
 
 		// Circular references that need to be cleaned up
 		CircularlyReferencedDependentModules.AddRange(
-		   new string[] {
-				"Kismet"
-		   }
+		   [
+				"Kismet",
+				"Persona",
+		   ]
 	   );
 	}
 }
diff --git a/Engine/Source/Editor/Kismet/Kismet.Build.cs b/Engine/Source/Editor/Kismet/Kismet.Build.cs
--- a/Engine/Source/Editor/Kismet/Kismet.Build.cs
+++ b/Engine/Source/Editor/Kismet/Kismet.Build.cs
@@ -76,11 +76,13 @@ public class Kismet : ModuleRules
 
 		// Circular references that need to be cleaned up
 		CircularlyReferencedDependentModules.AddRange(
-			new string[] {
+			[
 				"BlueprintGraph",
 				"UMGEditor",
-				"Merge"
-            }
-        ); 
+				"Merge",
+				"KismetCompiler",
+				"SubobjectEditor",
+			]
+		); 
 	}
 }
diff --git a/Engine/Source/Editor/Sequencer/Sequencer.Build.cs b/Engine/Source/Editor/Sequencer/Sequencer.Build.cs
--- a/Engine/Source/Editor/Sequencer/Sequencer.Build.cs
+++ b/Engine/Source/Editor/Sequencer/Sequencer.Build.cs
@@ -58,6 +58,8 @@ public class Sequencer : ModuleRules
 		CircularlyReferencedDependentModules.AddRange(
 			new string[] {
 				"ViewportInteraction",
+				"MovieSceneCaptureDialog",
+				"UniversalObjectLocatorEditor",
 				}
 			);
 
diff --git a/Engine/Source/Editor/UMGEditor/UMGEditor.Build.cs b/Engine/Source/Editor/UMGEditor/UMGEditor.Build.cs
--- a/Engine/Source/Editor/UMGEditor/UMGEditor.Build.cs
+++ b/Engine/Source/Editor/UMGEditor/UMGEditor.Build.cs
@@ -68,5 +68,11 @@ public class UMGEditor : ModuleRules
 				"WorkspaceMenuStructure"
 			}
 			);
+
+		// Circular references that need to be cleaned up
+		CircularlyReferencedDependentModules.AddRange(
+		   [
+				"Blutility",
+		   ]);
 	}
 }
diff --git a/Engine/Source/Editor/UnrealEd/UnrealEd.Build.cs b/Engine/Source/Editor/UnrealEd/UnrealEd.Build.cs
--- a/Engine/Source/Editor/UnrealEd/UnrealEd.Build.cs
+++ b/Engine/Source/Editor/UnrealEd/UnrealEd.Build.cs
@@ -272,6 +272,7 @@ public class UnrealEd : ModuleRules
 				"GraphEditor",
 				"Kismet",
 				"AudioEditor",
+				"LocalizationService",
 				"ViewportInteraction",
 				"VREditor",
 				"MeshPaint",
@@ -279,7 +280,33 @@ public class UnrealEd : ModuleRules
 				"ToolMenusEditor",
 				"ClothingSystemEditor",
 				"MaterialShaderQualitySettings",
-				"EditorInteractiveToolsFramework"
+				"EditorInteractiveToolsFramework",
+				"AddContentDialog",
+				"AnimationBlueprintLibrary",
+				"AutomationController",
+				"CommonMenuExtensions",
+				"CurveEditor",
+				"HierarchicalLODUtilities",
+				"MaterialUtilities",
+				"StatsViewer",
+				"SubobjectDataInterface",
+				"LevelEditor",
+				"LevelSequence",
+				"MaterialEditor",
+				"GameProjectGeneration",
+				"DataLayerEditor",
+				"VirtualizationEditor",
+				"EditorWidgets",
+				"WidgetRegistration",
+				"StatusBar",
+				"MeshBuilder",
+				"UATHelper",
+				"UncontrolledChangelists",
+				"SubobjectEditor",
+				"FoliageEdit",
+				"UniversalObjectLocatorEditor",
+				"MaterialBaking",
+				"MovieSceneTracks",
 			}
 		);
 
diff --git a/Engine/Source/Programs/UnrealBuildTool/Configuration/UEBuildModuleCPP.cs b/Engine/Source/Programs/UnrealBuildTool/Configuration/UEBuildModuleCPP.cs
--- a/Engine/Source/Programs/UnrealBuildTool/Configuration/UEBuildModuleCPP.cs
+++ b/Engine/Source/Programs/UnrealBuildTool/Configuration/UEBuildModuleCPP.cs
@@ -123,75 +123,119 @@ namespace UnrealBuildTool
 		/// List of allowed circular dependencies. Please do NOT add new modules here; refactor to allow the modules to be decoupled instead.
 		/// </summary>
 		static readonly KeyValuePair<string, string>[] CircularDependenciesAllowList =
-		{
-			new KeyValuePair<string, string>("AIModule", "AITestSuite"),
-			new KeyValuePair<string, string>("AnimGraph", "GraphEditor"),
-			new KeyValuePair<string, string>("AnimGraph", "UnrealEd"),
-			new KeyValuePair<string, string>("AudioEditor", "DetailCustomizations"),
-			new KeyValuePair<string, string>("AudioMixer", "NonRealtimeAudioRenderer"),
-			new KeyValuePair<string, string>("AudioMixer", "SoundFieldRendering"),
-			new KeyValuePair<string, string>("AssetTools", "UnrealEd"),
-			new KeyValuePair<string, string>("BlueprintGraph", "CinematicCamera"),
-			new KeyValuePair<string, string>("BlueprintGraph", "GraphEditor"),
-			new KeyValuePair<string, string>("BlueprintGraph", "Kismet"),
-			new KeyValuePair<string, string>("BlueprintGraph", "KismetCompiler"),
-			new KeyValuePair<string, string>("BlueprintGraph", "UnrealEd"),
-			new KeyValuePair<string, string>("ConfigEditor", "PropertyEditor"),
-			new KeyValuePair<string, string>("Documentation", "SourceControl"),
-			new KeyValuePair<string, string>("Engine", "AudioMixer"),
-			new KeyValuePair<string, string>("Engine", "CinematicCamera"),
-			new KeyValuePair<string, string>("Engine", "CollisionAnalyzer"),
-			new KeyValuePair<string, string>("Engine", "GameplayTags"),
-			new KeyValuePair<string, string>("Engine", "Kismet"),
-			new KeyValuePair<string, string>("Engine", "Landscape"),
-			new KeyValuePair<string, string>("Engine", "LogVisualizer"),
-			new KeyValuePair<string, string>("Engine", "MaterialShaderQualitySettings"),
-			new KeyValuePair<string, string>("Engine", "UMG"),
-			new KeyValuePair<string, string>("Engine", "UnrealEd"),
-			new KeyValuePair<string, string>("FoliageEdit", "ViewportInteraction"),
-			new KeyValuePair<string, string>("FoliageEdit", "VREditor"),
-			new KeyValuePair<string, string>("FunctionalTesting", "UnrealEd"),
-			new KeyValuePair<string, string>("GameplayAbilitiesEditor", "BlueprintGraph"),
-			new KeyValuePair<string, string>("GameplayDebugger", "AIModule"),
-			new KeyValuePair<string, string>("GameplayDebugger", "GameplayTasks"),
-			new KeyValuePair<string, string>("GameplayTasks", "UnrealEd"),
-			new KeyValuePair<string, string>("GraphEditor", "Kismet"),
-			new KeyValuePair<string, string>("HierarchicalLODOutliner", "UnrealEd"),
-			new KeyValuePair<string, string>("Kismet", "BlueprintGraph"),
-			new KeyValuePair<string, string>("Kismet", "Merge"),
-			new KeyValuePair<string, string>("Kismet", "UMGEditor"),
-			new KeyValuePair<string, string>("KismetWidgets", "BlueprintGraph"),
-			new KeyValuePair<string, string>("Landscape", "MaterialUtilities"),
-			new KeyValuePair<string, string>("Landscape", "UnrealEd"),
-			new KeyValuePair<string, string>("LandscapeEditor", "ViewportInteraction"),
-			new KeyValuePair<string, string>("LandscapeEditor", "VREditor"),
-			new KeyValuePair<string, string>("LocalizationDashboard", "LocalizationService"),
-			new KeyValuePair<string, string>("LocalizationDashboard", "MainFrame"),
-			new KeyValuePair<string, string>("LocalizationDashboard", "TranslationEditor"),
-			new KeyValuePair<string, string>("MaterialUtilities", "Landscape"),
-			new KeyValuePair<string, string>("MovieSceneTools", "Sequencer"),
-			new KeyValuePair<string, string>("NavigationSystem", "UnrealEd"),
-			new KeyValuePair<string, string>("PacketHandler", "ReliabilityHandlerComponent"),
-			new KeyValuePair<string, string>("PIEPreviewDeviceProfileSelector", "UnrealEd"),
-			new KeyValuePair<string, string>("PIEPreviewDeviceProfileSelector", "MaterialShaderQualitySettings"),
-			new KeyValuePair<string, string>("PixelInspectorModule", "UnrealEd"),
-			new KeyValuePair<string, string>("Sequencer", "MovieSceneTools"),
-			new KeyValuePair<string, string>("Sequencer", "ViewportInteraction"),
-			new KeyValuePair<string, string>("SourceControl", "UnrealEd"),
-			new KeyValuePair<string, string>("UnrealEd", "AudioEditor"),
-			new KeyValuePair<string, string>("UnrealEd", "ClothingSystemEditor"),
-			new KeyValuePair<string, string>("UnrealEd", "Documentation"),
-			new KeyValuePair<string, string>("UnrealEd", "EditorInteractiveToolsFramework"),
-			new KeyValuePair<string, string>("UnrealEd", "GraphEditor"),
-			new KeyValuePair<string, string>("UnrealEd", "Kismet"),
-			new KeyValuePair<string, string>("UnrealEd", "MeshPaint"),
-			new KeyValuePair<string, string>("UnrealEd", "PropertyEditor"),
-			new KeyValuePair<string, string>("UnrealEd", "ToolMenusEditor"),
-			new KeyValuePair<string, string>("UnrealEd", "ViewportInteraction"),
-			new KeyValuePair<string, string>("UnrealEd", "VREditor"),
-			new KeyValuePair<string, string>("UnrealEd", "MaterialShaderQualitySettings"),
-			new KeyValuePair<string, string>("WebBrowser", "WebBrowserTexture"),
-		};
+		[
+			new("AIModule", "AITestSuite"),
+			new("AnimGraph", "BlueprintGraph"),
+			new("AnimGraph", "ContentBrowser"),
+			new("AnimGraph", "GraphEditor"),
+			new("AnimGraph", "Kismet"),
+			new("AnimGraph", "KismetCompiler"),
+			new("AnimGraph", "PropertyEditor"),
+			new("AnimGraph", "UnrealEd"),
+			new("AssetTools", "Kismet"),
+			new("AssetTools", "Landscape"),
+			new("AssetTools", "MaterialEditor"),
+			new("AssetTools", "UnrealEd"),
+			new("AudioEditor", "DetailCustomizations"),
+			new("AudioMixer", "NonRealtimeAudioRenderer"),
+			new("AudioMixer", "SoundFieldRendering"),
+			new("BlueprintGraph", "CinematicCamera"),
+			new("BlueprintGraph", "GraphEditor"),
+			new("BlueprintGraph", "Kismet"),
+			new("BlueprintGraph", "KismetCompiler"),
+			new("BlueprintGraph", "UnrealEd"),
+			new("ConfigEditor", "PropertyEditor"),
+			new("Documentation", "ContentBrowser"),
+			new("Documentation", "SourceControl"),
+			new("Engine", "AudioMixer"),
+			new("Engine", "CinematicCamera"),
+			new("Engine", "CollisionAnalyzer"),
+			new("Engine", "GameplayTags"),
+			new("Engine", "Kismet"),
+			new("Engine", "Landscape"),
+			new("Engine", "LogVisualizer"),
+			new("Engine", "MaterialShaderQualitySettings"),
+			new("Engine", "MaterialUtilities"),
+			new("Engine", "UMG"),
+			new("Engine", "UnrealEd"),
+			new("FoliageEdit", "ViewportInteraction"),
+			new("FoliageEdit", "VREditor"),
+			new("FunctionalTesting", "UnrealEd"),
+			new("GameplayAbilitiesEditor", "BlueprintGraph"),
+			new("GameplayDebugger", "AIModule"),
+			new("GameplayDebugger", "GameplayTasks"),
+			new("GameplayTasks", "UnrealEd"),
+			new("GraphEditor", "Kismet"),
+			new("GraphEditor", "Persona"),
+			new("HierarchicalLODOutliner", "UnrealEd"),
+			new("Kismet", "BlueprintGraph"),
+			new("Kismet", "KismetCompiler"),
+			new("Kismet", "Merge"),
+			new("Kismet", "SubobjectEditor"),
+			new("Kismet", "UMGEditor"),
+			new("KismetWidgets", "BlueprintGraph"),
+			new("Landscape", "MaterialUtilities"),
+			new("Landscape", "MeshUtilities"),
+			new("Landscape", "UnrealEd"),
+			new("LandscapeEditor", "ViewportInteraction"),
+			new("LandscapeEditor", "VREditor"),
+			new("LocalizationDashboard", "LocalizationService"),
+			new("LocalizationDashboard", "MainFrame"),
+			new("LocalizationDashboard", "TranslationEditor"),
+			new("MaterialUtilities", "Landscape"),
+			new("MovieSceneTools", "Sequencer"),
+			new("NavigationSystem", "UnrealEd"),
+			new("PacketHandler", "ReliabilityHandlerComponent"),
+			new("PIEPreviewDeviceProfileSelector", "MaterialShaderQualitySettings"),
+			new("PIEPreviewDeviceProfileSelector", "UnrealEd"),
+			new("PixelInspectorModule", "UnrealEd"),
+			new("Sequencer", "MovieSceneCaptureDialog"),
+			new("Sequencer", "MovieSceneTools"),
+			new("Sequencer", "UniversalObjectLocatorEditor"),
+			new("Sequencer", "ViewportInteraction"),
+			new("SourceControl", "UnrealEd"),
+			new("UMGEditor", "Blutility"),
+			new("UnrealEd", "AddContentDialog"),
+			new("UnrealEd", "AnimationBlueprintLibrary"),
+			new("UnrealEd", "AudioEditor"),
+			new("UnrealEd", "AutomationController"),
+			new("UnrealEd", "ClothingSystemEditor"),
+			new("UnrealEd", "CommonMenuExtensions"),
+			new("UnrealEd", "CurveEditor"),
+			new("UnrealEd", "DataLayerEditor"),
+			new("UnrealEd", "Documentation"),
+			new("UnrealEd", "EditorInteractiveToolsFramework"),
+			new("UnrealEd", "EditorWidgets"),
+			new("UnrealEd", "FoliageEdit"),
+			new("UnrealEd", "GameProjectGeneration"),
+			new("UnrealEd", "GraphEditor"),
+			new("UnrealEd", "HierarchicalLODUtilities"),
+			new("UnrealEd", "Kismet"),
+			new("UnrealEd", "LevelEditor"),
+			new("UnrealEd", "LevelSequence"),
+			new("UnrealEd", "LocalizationService"),
+			new("UnrealEd", "MaterialBaking"),
+			new("UnrealEd", "MaterialEditor"),
+			new("UnrealEd", "MaterialShaderQualitySettings"),
+			new("UnrealEd", "MaterialUtilities"),
+			new("UnrealEd", "MeshBuilder"),
+			new("UnrealEd", "MeshPaint"),
+			new("UnrealEd", "MovieSceneTracks"),
+			new("UnrealEd", "PropertyEditor"),
+			new("UnrealEd", "StatsViewer"),
+			new("UnrealEd", "StatusBar"),
+			new("UnrealEd", "SubobjectDataInterface"),
+			new("UnrealEd", "SubobjectEditor"),
+			new("UnrealEd", "ToolMenusEditor"),
+			new("UnrealEd", "UATHelper"),
+			new("UnrealEd", "UncontrolledChangelists"),
+			new("UnrealEd", "UniversalObjectLocatorEditor"),
+			new("UnrealEd", "ViewportInteraction"),
+			new("UnrealEd", "VirtualizationEditor"),
+			new("UnrealEd", "VREditor"),
+			new("UnrealEd", "WidgetRegistration"),
+			new("WebBrowser", "WebBrowserTexture"),
+		];
 
 		public UEBuildModuleCPP(ModuleRules Rules, DirectoryReference IntermediateDirectory, DirectoryReference IntermediateDirectoryNoArch, DirectoryReference? GeneratedCodeDirectory, ILogger Logger)
 			: base(Rules, IntermediateDirectory, IntermediateDirectoryNoArch, Logger)
diff --git a/Engine/Source/Runtime/Engine/Engine.Build.cs b/Engine/Source/Runtime/Engine/Engine.Build.cs
--- a/Engine/Source/Runtime/Engine/Engine.Build.cs
+++ b/Engine/Source/Runtime/Engine/Engine.Build.cs
@@ -234,12 +234,16 @@ public class Engine : ModuleRules
 			);
 		}
 
-		CircularlyReferencedDependentModules.Add("GameplayTags");
-		CircularlyReferencedDependentModules.Add("Landscape");
-		CircularlyReferencedDependentModules.Add("UMG");
-		CircularlyReferencedDependentModules.Add("MaterialShaderQualitySettings");
-		CircularlyReferencedDependentModules.Add("CinematicCamera");
-		CircularlyReferencedDependentModules.Add("AudioMixer");
+		CircularlyReferencedDependentModules.AddRange(
+			[
+				"GameplayTags",
+				"Landscape",
+				"UMG",
+				"MaterialShaderQualitySettings",
+				"CinematicCamera",
+				"AudioMixer",
+				"MaterialUtilities",
+			]);
 		
 
 		if (Target.bCompileAgainstEditor)
diff --git a/Engine/Source/Runtime/Landscape/Landscape.Build.cs b/Engine/Source/Runtime/Landscape/Landscape.Build.cs
--- a/Engine/Source/Runtime/Landscape/Landscape.Build.cs
+++ b/Engine/Source/Runtime/Landscape/Landscape.Build.cs
@@ -74,6 +74,7 @@ public class Landscape : ModuleRules
 				new string[] {
 					"UnrealEd",
 					"MaterialUtilities",
+					"MeshUtilities",
 				}
 			);
 		}
