UnrealBuildTool: Improve circular dependency check

#rnx

diff --git a/Engine/Source/Programs/Shared/EpicGames.Core/KnownLogEvents.cs b/Engine/Source/Programs/Shared/EpicGames.Core/KnownLogEvents.cs
--- a/Engine/Source/Programs/Shared/EpicGames.Core/KnownLogEvents.cs
+++ b/Engine/Source/Programs/Shared/EpicGames.Core/KnownLogEvents.cs
@@ -45,6 +45,11 @@ namespace EpicGames.Core
 		/// </summary>
 		public static EventId RulesAssembly { get; } = new EventId(102);
 
+		/// <summary>
+		/// Validation messages
+		/// </summary>
+		public static EventId Validation { get; } = new EventId(103);
+
 		/// <summary>
 		/// UHT compiler message
 		/// </summary>
diff --git a/Engine/Source/Programs/UnrealBuildTool/Configuration/UEBuildModule.cs b/Engine/Source/Programs/UnrealBuildTool/Configuration/UEBuildModule.cs
--- a/Engine/Source/Programs/UnrealBuildTool/Configuration/UEBuildModule.cs
+++ b/Engine/Source/Programs/UnrealBuildTool/Configuration/UEBuildModule.cs
@@ -5,6 +5,7 @@ using System.Collections.Generic;
 using System.Configuration;
 using System.IO;
 using System.Linq;
+using System.Reflection;
 using System.Text;
 using System.Threading.Tasks;
 using EpicGames.Core;
@@ -432,9 +433,9 @@ namespace UnrealBuildTool
 		/// <summary>
 		/// Determines whether this module has a circular dependency on the given module
 		/// </summary>
-		public bool HasCircularDependencyOn(string ModuleName)
+		public bool HasCircularDependencyOn(UEBuildModule Module)
 		{
-			return Rules.CircularlyReferencedDependentModules.Contains(ModuleName);
+			return Rules.CircularlyReferencedDependentModules.Contains(Module.Name);
 		}
 
 		/// <summary>
@@ -1051,7 +1052,7 @@ namespace UnrealBuildTool
 			foreach (UEBuildModule DependencyModule in AllDependencyModules)
 			{
 				// Don't follow circular back-references!
-				if (bForceCircular || !HasCircularDependencyOn(DependencyModule.Name))
+				if (bForceCircular || !HasCircularDependencyOn(DependencyModule))
 				{
 					if (IgnoreReferencedModules.Add(DependencyModule))
 					{
@@ -1110,7 +1111,7 @@ namespace UnrealBuildTool
 			foreach (UEBuildModule DependencyModule in AllDependencyModules.Distinct())
 			{
 				// Don't follow circular back-references!
-				if (bForceCircular || !HasCircularDependencyOn(DependencyModule.Name))
+				if (bForceCircular || !HasCircularDependencyOn(DependencyModule))
 				{
 					if (IgnoreReferencedModules.Add(DependencyModule))
 					{
@@ -1142,7 +1143,7 @@ namespace UnrealBuildTool
 		/// </summary>
 		/// <param name="CreateModule">Delegate to create a module with a given name</param>
 		/// <param name="ReferenceChain">Chain of references before reaching this module</param>
-		/// <param name="ReferenceStack">Stack of module dependencies that led to this module</param>
+		/// <param name="ReferenceStack">ParentStack of module dependencies that led to this module</param>
 		/// <param name="Logger">Logger for output</param>
 		protected void RecursivelyCreateModules(CreateModuleDelegate CreateModule, string ReferenceChain, List<UEBuildModule> ReferenceStack, ILogger Logger)
 		{
@@ -1201,6 +1202,8 @@ namespace UnrealBuildTool
 					LogDependencyNameList("Private Include Paths:", Rules.PrivateIncludePathModuleNames, Logger);
 				}
 
+				ReferenceStackParentModules ??= new();
+
 				// Create all the dependency modules - pass through the reference stack so we can check for cycles
 				RecursivelyCreateModulesByName(Rules.PublicDependencyModuleNames, ref PublicDependencyModules, ref bDependsOnVerse, CreateModule, NextReferenceChain, ReferenceStack, Logger);
 				if (Rules.Target.IsTestTarget)
@@ -1221,21 +1224,55 @@ namespace UnrealBuildTool
 
 				// Dynamic loads aren't considered a reference chain so start with an empty stack
 				RecursivelyCreateModulesByName(Rules.DynamicallyLoadedModuleNames, ref DynamicallyLoadedModules, ref bDependsOnVerse, CreateModule, NextReferenceChain, new List<UEBuildModule>(), Logger);
+
+				PublicDependencyModules?.ForEach(x => x.ReferenceStackParentModules?.Add(this));
+				PrivateDependencyModules?.ForEach(x => x.ReferenceStackParentModules?.Add(this));
 			}
 
 			// pop us off the current stack
 			ReferenceStack.RemoveAt(ReferenceStack.Count - 1);
 		}
 
+		record CircularReference
+		{
+			public List<UEBuildModule> ReferenceStack { get; }
+			public Dictionary<UEBuildModule, UEBuildModule> GuiltyModules { get; } = new();
+			public bool IsError => !GuiltyModules.Any();
+
+			public CircularReference(List<UEBuildModule> referenceStack) => ReferenceStack = referenceStack;
+
+			public void AddGuiltyModule(UEBuildModule guiltyModule, UEBuildModule victimModule) => GuiltyModules.TryAdd(guiltyModule, victimModule);
+		}
+
 		/// <summary>
 		/// Public entry point to validate a module
 		/// </summary>
-		/// <param name="ReferenceChain"></param>
-		/// <param name="Logger"></param>
-		public bool ValidateModule(string ReferenceChain, ILogger Logger)
+		/// <param name="logger"></param>
+		public bool ValidateModule(ILogger logger)
 		{
-			List<UEBuildModule> ReferenceStack = new List<UEBuildModule>();
-			RecursivelyCheckForCircularReferences(ReferenceChain, ReferenceStack, Logger);
+			bool anyErrors = false;
+
+			// Report any [Obsolete("...")] messages from referenced modules.
+			IEnumerable<ObsoleteAttribute> obsoleteAttributes = this.Rules.GetType().GetCustomAttributes<ObsoleteAttribute>();
+			if (obsoleteAttributes.Any())
+			{
+				logger.LogWarning(KnownLogEvents.Validation, "Referenced Module '{Name}' is obsolete: '{Message}'", Name, String.Join(", ", obsoleteAttributes.Select(x => x.Message ?? "<unknown reason>")));
+			}
+
+			anyErrors |= ValidateCircularReferences(logger);
+
+			return anyErrors;
+		}
+
+		/// <summary>
+		/// Validate a module for circular references
+		/// </summary>
+		/// <param name="logger"></param>
+		/// <returns>If any circular references are found</returns>
+		bool ValidateCircularReferences(ILogger logger)
+		{
+			List<CircularReference> circularReferences = new();
+			CheckForCircularReferences(this, new(), new(), circularReferences);
 
 			// Check if this module is invalid for the current target, based on Module attributes
 			if (Rules.CppCompileWarningSettings.ModuleUnsupportedWarningLevel == WarningLevel.Error && !ModuleRules.IsValidForTarget(Rules.GetType(), Rules.Target, out string? _))
@@ -1251,85 +1288,85 @@ namespace UnrealBuildTool
 				}
 			}
 
-			return false; // No fatal errors
+			IEnumerable<CircularReference> allowed = circularReferences.Where(x => !x.IsError).OrderBy(x => x.ReferenceStack.Count);
+			IEnumerable<CircularReference> errors = circularReferences.Where(x => x.IsError).OrderBy(x => x.ReferenceStack.Count);
+
+			if (errors.Any())
+			{
+				logger.LogError(KnownLogEvents.Validation, "Circular dependency for {RulesFile} detected:", RulesFile.MakeRelativeTo(Unreal.EngineSourceDirectory));
+				foreach (CircularReference circularReference in errors)
+				{
+					logger.LogError(KnownLogEvents.Validation, "* {ReferenceStack}", String.Join(" -> ", circularReference.ReferenceStack));
+				}
+				logger.LogInformation(KnownLogEvents.Validation, "Break this loop by moving dependencies into a separate module or using Private/PublicIncludePathModuleNames to reference declarations");
+			}
+
+			if (allowed.Any())
+			{
+				logger.LogDebug(KnownLogEvents.Validation, "Allowed circular dependency for {RulesFile} detected:", RulesFile.MakeRelativeTo(Unreal.EngineSourceDirectory));
+				foreach (CircularReference circularReference in allowed)
+				{
+					logger.LogDebug(KnownLogEvents.Validation, "* {ReferenceStack} (allowed via {GuiltyModules})", String.Join(" -> ", circularReference.ReferenceStack), String.Join(", ", circularReference.GuiltyModules.Select(x => $"{x.Key} -> {x.Value}")));
+				}
+			}
+			return errors.Any();
 		}
 
 		/// <summary>
-		/// Recursively check a module for circular dependencies
+		/// Recursively check a module for circular references
 		/// </summary>
-		/// <param name="ReferenceChain">Chain of references before reaching this module</param>
-		/// <param name="ReferenceStack">Stack of module dependencies that led to this module</param>
-		/// <param name="Logger">Logger for output</param>
-		protected void RecursivelyCheckForCircularReferences(string ReferenceChain, List<UEBuildModule> ReferenceStack, ILogger Logger)
+		/// <param name="module">The module currently being checked</param>
+		/// <param name="processed">Set of modules already processed</param>
+		/// <param name="parentStack">The current parent reference stack (in reverse order)</param>
+		/// <param name="circularReferences">List of found circular reference stacks</param>
+		void CheckForCircularReferences(UEBuildModule module, HashSet<UEBuildModule> processed, List<UEBuildModule> parentStack, List<CircularReference> circularReferences)
 		{
-			// Name of this reference
-			string ThisRefName = (RulesFile == null) ? Name : RulesFile.GetFileName();
-
-			// Set the reference chain for anything referenced by this module
-			string NextReferenceChain = String.Format("{0} -> {1}", ReferenceChain, ThisRefName);
-
-			// We need to check for cycles if this module has already been created and its name is in the stack.
-			bool CheckForCycles = PrivateIncludePathModules != null && ReferenceStack.Contains(this);
-
-			UEBuildModule? PreviousModule = ReferenceStack.LastOrDefault();
-			ReferenceStackParentModules ??= new();
-			if (PreviousModule != null && !ReferenceStackParentModules.Add(PreviousModule))
+			// If this module has no references or has been processed already, nothing to do
+			if (ReferenceStackParentModules == null || !processed.Add(this))
 			{
 				return;
 			}
 
-			// Add us to the reference stack - note do this before checking for cycles as we allow them if an element in the stack
-			// declares the next element at leading to a dependency cycle.
-			ReferenceStack.Add(this);
-
-			if (CheckForCycles)
+			// Check for a circular reference
+			if (ReferenceStackParentModules.Contains(module))
 			{
-				// We are FeatureModuleA and we know we're already in the list so there must be a circular dependency like this:
-				//   Target -> FeatureModuleA -> FeatureModuleB -> BaseModule -> FeatureModuleA
-				// In this case it's up to BaseModule to confess that it's going to introduce a 
-				// cycle by referencing FeatureModuleA so let's check it did that.
-				// (Note: it's *bad* that it is doing this, but it's even worse not to flag these
-				// cycles when they're introduced and have a way of tracking them!)
+				List<UEBuildModule> referenceStack = new(parentStack);
+				referenceStack.Add(this);
+				referenceStack.Add(module);
+				referenceStack.Reverse();
 
-				string? GuiltyModule = null;
-				string? VictimModule = null;
+				CircularReference circularReference = new(referenceStack);
 
-				for (int i = 0; i < ReferenceStack.Count - 1; i++)
+				for (int i = 0; i < referenceStack.Count - 1; i++)
 				{
-					UEBuildModule ReferringModule = ReferenceStack.ElementAt(i);
-					UEBuildModule TargetModule = ReferenceStack.ElementAt(i + 1);
+					UEBuildModule referringModule = referenceStack.ElementAt(i);
+					UEBuildModule targetModule = referenceStack.ElementAt(i + 1);
 
-					if (ReferringModule.HasCircularDependencyOn(TargetModule.Name))
+					// If module has confessed its guilt this is not an error
+					if (referringModule.HasCircularDependencyOn(targetModule))
 					{
-						GuiltyModule = ReferringModule.Name;
-						VictimModule = TargetModule.Name;
-						break;
+						circularReference.AddGuiltyModule(referringModule, targetModule);
+					}
+					else if (targetModule.HasCircularDependencyOn(referringModule))
+					{
+						circularReference.AddGuiltyModule(targetModule, referringModule);
 					}
 				}
-
-				// No module has confessed its guilt, so this is an error.
-				if (String.IsNullOrEmpty(GuiltyModule))
-				{
-					string CycleChain = String.Join(" -> ", ReferenceStack.SkipWhile(x => x != this));
-					Logger.LogError("Circular dependency on {Name} detected:\n" +
-						"\tFull Route: {FullRoute}\n" +
-						"\tCycled Route: is {CycleRoute}\n" +
-						"Break this loop by moving dependencies into a separate module or using Private/PublicIncludePathModuleNames to reference declarations\n",
-						ThisRefName, NextReferenceChain, CycleChain);
-				}
-				else
-				{
-					Logger.LogDebug("Found circular reference to {ThisRefName}, but {GuiltyModule} declares a cycle on {VictimModule} which breaks the chain", ThisRefName, GuiltyModule, VictimModule);
-				}
+				circularReferences.Add(circularReference);
+				return;
 			}
-			else
+
+			// Push this module to the stack
+			parentStack.Add(this);
+
+			// Check very module that references this module until all are processed
+			foreach (UEBuildModule ParentModule in ReferenceStackParentModules.OrderBy(x => x.Name))
 			{
-				// Continue checking for circular dependencies until a cycle is found or the graph is exhausted
-				PublicDependencyModules?.ForEach(x => x.RecursivelyCheckForCircularReferences(NextReferenceChain, ReferenceStack, Logger));
-				PrivateDependencyModules?.ForEach(x => x.RecursivelyCheckForCircularReferences(NextReferenceChain, ReferenceStack, Logger));
+				ParentModule.CheckForCircularReferences(module, processed, parentStack, circularReferences);
 			}
-			// pop us off the current stack
-			ReferenceStack.RemoveAt(ReferenceStack.Count - 1);
+
+			// Pop this module from the stack
+			parentStack.RemoveAt(parentStack.Count - 1);
 		}
 
 		private void MoveTestsRunnerDependencyToLastPosition()
diff --git a/Engine/Source/Programs/UnrealBuildTool/Configuration/UEBuildTarget.cs b/Engine/Source/Programs/UnrealBuildTool/Configuration/UEBuildTarget.cs
--- a/Engine/Source/Programs/UnrealBuildTool/Configuration/UEBuildTarget.cs
+++ b/Engine/Source/Programs/UnrealBuildTool/Configuration/UEBuildTarget.cs
@@ -1485,11 +1485,6 @@ namespace UnrealBuildTool
 			{
 				UEBuildPlatform.GetBuildPlatform(Platform).ValidateModule(module, Rules);
 				anyErrors |= UEBuildPlatform.GetBuildPlatform(Platform).ValidateModuleIncludePaths(module, Rules, Modules.Values);
-			}
-
-			// Intentionally not ordered by name to maintain the reference include order
-			foreach (UEBuildModule module in Modules.Values)
-			{
 				anyErrors |= ValidateModule(module, logger);
 			}
 
@@ -1504,15 +1499,7 @@ namespace UnrealBuildTool
 		/// <returns>true if there are any fatal errors</returns>
 		bool ValidateModule(UEBuildModule module, ILogger logger)
 		{
-			bool anyErrors = module.ValidateModule("Target", logger);
-
-			// Report any [Obsolete("...")] messages from referenced modules.
-			IEnumerable<ObsoleteAttribute> ObsoleteAttributes = module.Rules.GetType().GetCustomAttributes<ObsoleteAttribute>();
-			if (ObsoleteAttributes.Any())
-			{
-				string Message = String.Join(", ", ObsoleteAttributes.Select(x => x.Message ?? "<unknown reason>"));
-				logger.LogWarning($"Warning: Referenced Module '{module.Name}' is obsolete: '{Message}'");
-			}
+			bool anyErrors = module.ValidateModule(logger);
 
 			Lazy<HashSet<UEBuildModule>> referencedModules = new(() => module.GetDependencies(bWithIncludePathModules: true, bWithDynamicallyLoadedModules: true));
 
@@ -1523,7 +1510,7 @@ namespace UnrealBuildTool
 				{
 					if (!module.Rules.Context.Scope.Contains(referencedModule.Rules.Context.Scope) && !IsEngineToPluginReferenceAllowed(module.Name, referencedModule.Name))
 					{
-						logger.LogError("Module '{ModuleName}' ({ScopeName}) should not reference module '{DependencyModuleName}' ({DependencyScopeName}). Hierarchy is {Hierarchy}.", module.Name, module.Rules.Context.Scope.Name, referencedModule.Name, referencedModule.Rules.Context.Scope.Name, referencedModule.Rules.Context.Scope.FormatHierarchy());
+						logger.LogError(KnownLogEvents.Validation, "Module '{ModuleName}' ({ScopeName}) should not reference module '{DependencyModuleName}' ({DependencyScopeName}). Hierarchy is {Hierarchy}.", module.Name, module.Rules.Context.Scope.Name, referencedModule.Name, referencedModule.Rules.Context.Scope.Name, referencedModule.Rules.Context.Scope.FormatHierarchy());
 						anyErrors = true;
 					}
 				}
@@ -1556,19 +1543,19 @@ namespace UnrealBuildTool
 						{
 							if (!projectPluginReference.bEnabled)
 							{
-								logger.LogWarning("Warning: {ProjectName} plugin dependency '{DependencyPluginName}' is not enabled, but module '{ModuleName}' depends on '{DependencyModuleName}'.", projectName, dependencyPluginInfo.Name, module.Name, dependencyModule.Name);
+								logger.LogWarning(KnownLogEvents.Validation, "{ProjectName} plugin dependency '{DependencyPluginName}' is not enabled, but module '{ModuleName}' depends on '{DependencyModuleName}'.", projectName, dependencyPluginInfo.Name, module.Name, dependencyModule.Name);
 							}
 							else if (!projectPluginReference.IsEnabledForTarget(TargetType))
 							{
-								logger.LogWarning("Warning: {ProjectName} plugin dependency '{DependencyPluginName}' is not enabled for target '{TargetType}', but module '{ModuleName}' depends on '{DependencyModuleName}'.", projectName, dependencyPluginInfo.Name, TargetType.ToString(), module.Name, dependencyModule.Name);
+								logger.LogWarning(KnownLogEvents.Validation, "{ProjectName} plugin dependency '{DependencyPluginName}' is not enabled for target '{TargetType}', but module '{ModuleName}' depends on '{DependencyModuleName}'.", projectName, dependencyPluginInfo.Name, TargetType.ToString(), module.Name, dependencyModule.Name);
 							}
 							else if (!projectPluginReference.IsEnabledForPlatform(Platform))
 							{
-								logger.LogWarning("Warning: {ProjectName} plugin dependency '{DependencyPluginName}' is not enabled for platform '{Platform}', but module '{ModuleName}' depends on '{DependencyModuleName}'.", projectName, dependencyPluginInfo.Name, Platform.ToString(), module.Name, dependencyModule.Name);
+								logger.LogWarning(KnownLogEvents.Validation, "{ProjectName} plugin dependency '{DependencyPluginName}' is not enabled for platform '{Platform}', but module '{ModuleName}' depends on '{DependencyModuleName}'.", projectName, dependencyPluginInfo.Name, Platform.ToString(), module.Name, dependencyModule.Name);
 							}
 							else if (!projectPluginReference.IsEnabledForTargetConfiguration(Configuration))
 							{
-								logger.LogWarning("Warning: {ProjectName} plugin dependency '{DependencyPluginName}' is not enabled for target configuration '{Configuration}', but module '{ModuleName}' depends on '{DependencyModuleName}'.", projectName, dependencyPluginInfo.Name, Configuration.ToString(), module.Name, dependencyModule.Name);
+								logger.LogWarning(KnownLogEvents.Validation, "{ProjectName} plugin dependency '{DependencyPluginName}' is not enabled for target configuration '{Configuration}', but module '{ModuleName}' depends on '{DependencyModuleName}'.", projectName, dependencyPluginInfo.Name, Configuration.ToString(), module.Name, dependencyModule.Name);
 							}
 						}
 						else
@@ -1576,7 +1563,7 @@ namespace UnrealBuildTool
 							// No project plugin reference exists. Check inclusion of a plugin was not made as part of a Target.cs file or Build.cs file
 							if (!Rules.EnablePlugins.Contains(dependencyPluginInfo.Name) && !Rules.InternalPluginDependencies.Contains(dependencyPluginInfo.Name))
 							{
-								logger.LogWarning("Warning: {ProjectName} does not list plugin '{DependencyPluginName}' as a dependency, but module '{ModuleName}' depends on '{DependencyModuleName}'.", projectName, dependencyPluginInfo.Name, module.Name, dependencyModule.Name);
+								logger.LogWarning(KnownLogEvents.Validation, "{ProjectName} does not list plugin '{DependencyPluginName}' as a dependency, but module '{ModuleName}' depends on '{DependencyModuleName}'.", projectName, dependencyPluginInfo.Name, module.Name, dependencyModule.Name);
 							}
 						}
 					}
diff --git a/Engine/Source/Programs/UnrealBuildTool/Modes/PrintBuildGraphInfo.cs b/Engine/Source/Programs/UnrealBuildTool/Modes/PrintBuildGraphInfo.cs
--- a/Engine/Source/Programs/UnrealBuildTool/Modes/PrintBuildGraphInfo.cs
+++ b/Engine/Source/Programs/UnrealBuildTool/Modes/PrintBuildGraphInfo.cs
@@ -248,7 +248,7 @@ namespace UnrealBuildTool
 			foreach (UEBuildModule DependencyModule in AllDependencyModules)
 			{
 				// Don't follow circular back-references!
-				if (!CurrentBuildModule.HasCircularDependencyOn(DependencyModule.Name))
+				if (!CurrentBuildModule.HasCircularDependencyOn(DependencyModule))
 				{
 					if (IgnoreReferencedModules.Add(DependencyModule))
 					{
