[UBT]
* Fixed so ispc generated headers are added as dependency to compile environment used by generated headers when compiling with -includeheaders

diff --git a/Engine/Source/Programs/UnrealBuildTool/Configuration/UEBuildModuleCPP.cs b/Engine/Source/Programs/UnrealBuildTool/Configuration/UEBuildModuleCPP.cs
--- a/Engine/Source/Programs/UnrealBuildTool/Configuration/UEBuildModuleCPP.cs
+++ b/Engine/Source/Programs/UnrealBuildTool/Configuration/UEBuildModuleCPP.cs
@@ -1024,6 +1024,12 @@ namespace UnrealBuildTool
 				if (HeaderFileItems.Count > 0)
 				{
 					CppCompileEnvironment HeaderCompileEnvironment = new CppCompileEnvironment(ModuleCompileEnvironment);
+					if (InputFiles.ISPCFiles.Count > 0)
+					{
+						HeaderCompileEnvironment.AdditionalPrerequisites.AddRange(ISPCHeaders);
+						HeaderCompileEnvironment.UserIncludePaths.Add(IntermediateDirectory);
+					}
+
 					DirectoryReference HeaderIntermediateDirectory = DirectoryReference.Combine(IntermediateDirectory, "H");
 
 					CreateHeaderForDefinitions(HeaderCompileEnvironment, HeaderIntermediateDirectory, "h", Graph);
