Allow invoking QueryTargets without including all of the target's parent's assemblies.

This can be controlled by bIncludeParentAssembly, under the engine config section [/Script/BuildSettings.BuildSettings].

This ultimately allows LiveLinkHub to launch without generating a TargetInfo.json file that contains 200 entries

#rnx

diff --git a/Engine/Source/Developer/DesktopPlatform/Private/DesktopPlatformBase.cpp b/Engine/Source/Developer/DesktopPlatform/Private/DesktopPlatformBase.cpp
--- a/Engine/Source/Developer/DesktopPlatform/Private/DesktopPlatformBase.cpp
+++ b/Engine/Source/Developer/DesktopPlatform/Private/DesktopPlatformBase.cpp
@@ -1470,6 +1470,12 @@ const TArray<FTargetInfo>& FDesktopPlatformBase::GetTargetsForProject(const FStr
 	Arguments += FString::Printf(TEXT(" -Output=\"%s\""), *IFileManager::Get().ConvertToAbsolutePathForExternalAppForRead(*InfoFileName));
 	Arguments += TEXT(" -IncludeAllTargets");
 
+	const bool bIncludeParentAssembly = GConfig->GetBoolOrDefault(TEXT("/Script/BuildSettings.BuildSettings"), TEXT("bIncludeParentAssembly"), true, GEngineIni);
+	if (!bIncludeParentAssembly)
+	{
+		Arguments += TEXT(" -DontIncludeParentAssembly");
+	}
+
 	// Run UBT to update the list of targets. Try to run it without building first.
 	FString Output;
 	int32 ReturnCode = 0;
diff --git a/Engine/Source/Programs/LiveLinkHub/Config/DefaultEngine.ini b/Engine/Source/Programs/LiveLinkHub/Config/DefaultEngine.ini
--- a/Engine/Source/Programs/LiveLinkHub/Config/DefaultEngine.ini
+++ b/Engine/Source/Programs/LiveLinkHub/Config/DefaultEngine.ini
@@ -38,6 +38,8 @@ Slate.EnableGamepadEditorNavigation=false
 
 [/Script/BuildSettings.BuildSettings]
 DefaultEditorTarget=LiveLinkHub
+; Disabling this to ensure that DesktopPlatformBase doesn't generate a TargetInfo.json file with UE's hundreds of targets
+bIncludeParentAssembly=false
 
 ; Not needed at the moment, can be removed if we enable the TextureFormatOodle plugin
 [AlternateTextureCompression]
diff --git a/Engine/Source/Programs/UnrealBuildTool/Modes/QueryTargetsMode.cs b/Engine/Source/Programs/UnrealBuildTool/Modes/QueryTargetsMode.cs
--- a/Engine/Source/Programs/UnrealBuildTool/Modes/QueryTargetsMode.cs
+++ b/Engine/Source/Programs/UnrealBuildTool/Modes/QueryTargetsMode.cs
@@ -34,6 +34,13 @@ namespace UnrealBuildTool
 		[CommandLine("-IncludeAllTargets")]
 		bool bIncludeAllTargets = false;
 
+		/// <summary>
+		/// Write all targets from parent assemblies (useful for EngineRules which can be chained)
+		/// </summary>
+		[CommandLine("-IncludeParentAssembly", Value ="true")]
+		[CommandLine("-DontIncludeParentAssembly", Value = "false")]
+		bool bIncludeParentAssembly = true;
+
 		/// <summary>
 		/// Execute the mode
 		/// </summary>
@@ -67,7 +74,7 @@ namespace UnrealBuildTool
 			}
 
 			// Write information about these targets
-			WriteTargetInfo(ProjectFile, Assembly, OutputFile, Arguments, Logger, bIncludeAllTargets, bIncludeParentAssembly: true);
+			WriteTargetInfo(ProjectFile, Assembly, OutputFile, Arguments, Logger, bIncludeAllTargets, bIncludeParentAssembly);
 			Logger.LogInformation("Written {OutputFile}", OutputFile);
 			return Task.FromResult(0);
 		}
