UnrealBuildTool: Correctly handle generating project files for a Win64 non-default arch

#rnx

diff --git a/Engine/Source/Programs/UnrealBuildTool/ProjectFiles/ProjectFileGenerator.cs b/Engine/Source/Programs/UnrealBuildTool/ProjectFiles/ProjectFileGenerator.cs
--- a/Engine/Source/Programs/UnrealBuildTool/ProjectFiles/ProjectFileGenerator.cs
+++ b/Engine/Source/Programs/UnrealBuildTool/ProjectFiles/ProjectFileGenerator.cs
@@ -2976,7 +2976,8 @@ namespace UnrealBuildTool
 													x => UEBuildPlatform.TryGetBuildPlatform(x, out _) &&
 													(TargetRulesObject.LinkType != TargetLinkType.Modular || !UEBuildPlatform.PlatformRequiresMonolithicBuilds(x, TargetRulesObject.Configuration))
 													).ToArray(),
-							CreateRulesDelegate: (Platform, Configuration) => RulesAssembly.CreateTargetRules(TargetName, Platform, Configuration, null, CheckProjectFile, new CommandLineArguments(GetTargetArguments(Arguments)), Logger)
+							CreateRulesDelegate: (Platform, Configuration) => RulesAssembly.CreateTargetRules(TargetName, Platform, Configuration, null, CheckProjectFile, new CommandLineArguments(GetTargetArguments(Arguments)), Logger),
+							CreateRulesDelegateWithArch: (Platform, Configuration, Arch) => RulesAssembly.CreateTargetRules(TargetName, Platform, Configuration, new(Arch), CheckProjectFile, new CommandLineArguments(GetTargetArguments(Arguments)), Logger)
 						);
 
 					ProjectFile.ProjectTargets.Add(ProjectTarget);
diff --git a/Engine/Source/Programs/UnrealBuildTool/ProjectFiles/ProjectTarget.cs b/Engine/Source/Programs/UnrealBuildTool/ProjectFiles/ProjectTarget.cs
--- a/Engine/Source/Programs/UnrealBuildTool/ProjectFiles/ProjectTarget.cs
+++ b/Engine/Source/Programs/UnrealBuildTool/ProjectFiles/ProjectTarget.cs
@@ -20,6 +20,9 @@ namespace UnrealBuildTool
 		/// Delegate for creating a rules instance for a given platform/configuration
 		public Func<UnrealTargetPlatform, UnrealTargetConfiguration, TargetRules> CreateRulesDelegate;
 
+		/// Delegate for creating a rules instance for a given platform/configuration/architecture
+		public Func<UnrealTargetPlatform, UnrealTargetConfiguration, UnrealArch, TargetRules> CreateRulesDelegateWithArch;
+
 		public string Name => TargetFilePath.GetFileNameWithoutAnyExtensions();
 
 		public override string ToString()
@@ -33,7 +36,8 @@ namespace UnrealBuildTool
 			FileReference? UnrealProjectFilePath,
 			TargetRules? TargetRules,
 			UnrealTargetPlatform[] SupportedPlatforms,
-			Func<UnrealTargetPlatform, UnrealTargetConfiguration, TargetRules> CreateRulesDelegate
+			Func<UnrealTargetPlatform, UnrealTargetConfiguration, TargetRules> CreateRulesDelegate,
+			Func<UnrealTargetPlatform, UnrealTargetConfiguration, UnrealArch, TargetRules> CreateRulesDelegateWithArch
 			)
 			: base(SupportedPlatforms)
 		{
@@ -42,6 +46,7 @@ namespace UnrealBuildTool
 			this.UnrealProjectFilePath = UnrealProjectFilePath;
 			this.TargetRules = TargetRules;
 			this.CreateRulesDelegate = CreateRulesDelegate;
+			this.CreateRulesDelegateWithArch = CreateRulesDelegateWithArch;
 		}
 	}
 }
diff --git a/Engine/Source/Programs/UnrealBuildTool/ProjectFiles/Rider/RiderProjectFileGenerator.cs b/Engine/Source/Programs/UnrealBuildTool/ProjectFiles/Rider/RiderProjectFileGenerator.cs
--- a/Engine/Source/Programs/UnrealBuildTool/ProjectFiles/Rider/RiderProjectFileGenerator.cs
+++ b/Engine/Source/Programs/UnrealBuildTool/ProjectFiles/Rider/RiderProjectFileGenerator.cs
@@ -403,6 +403,9 @@ namespace UnrealBuildTool
 							.Where(x => UEBuildPlatform.TryGetBuildPlatform(x, out _)).ToArray(),
 						CreateRulesDelegate: (Platform, Configuration) =>
 							RulesAssembly.CreateTargetRules(TargetName, Platform, Configuration, null, CheckProjectFile,
+								null, Logger),
+						CreateRulesDelegateWithArch: (Platform, Configuration, Arch) =>
+							RulesAssembly.CreateTargetRules(TargetName, Platform, Configuration, new(Arch), CheckProjectFile,
 								null, Logger)
 					);
 
diff --git a/Engine/Source/Programs/UnrealBuildTool/ProjectFiles/VisualStudio/VCProject.cs b/Engine/Source/Programs/UnrealBuildTool/ProjectFiles/VisualStudio/VCProject.cs
--- a/Engine/Source/Programs/UnrealBuildTool/ProjectFiles/VisualStudio/VCProject.cs
+++ b/Engine/Source/Programs/UnrealBuildTool/ProjectFiles/VisualStudio/VCProject.cs
@@ -2336,7 +2336,14 @@ namespace UnrealBuildTool
 					TargetRules TargetRulesObject;
 					try
 					{
-						TargetRulesObject = Combination.ProjectTarget!.CreateRulesDelegate(Platform, Configuration);
+						if (Combination.Architecture != null)
+						{
+							TargetRulesObject = Combination.ProjectTarget!.CreateRulesDelegateWithArch(Platform, Configuration, (UnrealArch)Combination.Architecture);
+						}
+						else
+						{
+							TargetRulesObject = Combination.ProjectTarget!.CreateRulesDelegate(Platform, Configuration);
+						}
 					}
 					catch (BuildException)
 					{
