Fix for mac packaged game version not following MacRuntimeSettings in Engine.ini

#rnx

diff --git a/Engine/Source/Programs/UnrealBuildTool/ProjectFiles/Xcode/XcodeProject.cs b/Engine/Source/Programs/UnrealBuildTool/ProjectFiles/Xcode/XcodeProject.cs
--- a/Engine/Source/Programs/UnrealBuildTool/ProjectFiles/Xcode/XcodeProject.cs
+++ b/Engine/Source/Programs/UnrealBuildTool/ProjectFiles/Xcode/XcodeProject.cs
@@ -1859,8 +1859,11 @@ namespace UnrealBuildTool.XcodeProjectXcconfig
 				DeploymentTarget = MacToolChainSettings.MinMacDeploymentVersion(UnrealData.TargetRules.Type);
 				BundleIdentifier = bIsEditor ? "com.epicgames.UnrealEditor" : UnrealData.BundleIdentifier;
 
-				// @todo: get a version for  games, like IOS has
-				MarketingVersion = MacToolChain.LoadEngineDisplayVersion();
+				PlatformIni.TryGetValue($"/Script/MacRuntimeSettings.MacRuntimeSettings", "VersionInfo", out MarketingVersion);
+				if (String.IsNullOrEmpty(MarketingVersion))
+				{
+					MarketingVersion = MacToolChain.LoadEngineDisplayVersion();
+				}
 			}
 			else
 			{
@@ -2796,4 +2799,4 @@ namespace UnrealBuildTool.XcodeProjectXcconfig
 			return true;
 		}
 	}
-}
\ No newline at end of file
+}
