UnrealBuildTool: Add query to get module reference lists

#rnx

diff --git a/Engine/Source/Programs/UnrealBuildTool/Modes/QueryMode.cs b/Engine/Source/Programs/UnrealBuildTool/Modes/QueryMode.cs
--- a/Engine/Source/Programs/UnrealBuildTool/Modes/QueryMode.cs
+++ b/Engine/Source/Programs/UnrealBuildTool/Modes/QueryMode.cs
@@ -22,7 +22,9 @@ namespace UnrealBuildTool
 		Capabilities,
 		AvailableTargets,
 		TargetDetails,
+		TargetModuleReferences,
 	}
+
 	internal class TargetIntellisenseInfo
 	{
 		internal class CompileSettings
@@ -53,6 +55,7 @@ namespace UnrealBuildTool
 			}
 			return null;
 		}
+
 		public UEBuildModule? FindModuleForDirectory(DirectoryReference Directory)
 		{
 			DirectoryReference? Dir = Directory;
@@ -169,6 +172,9 @@ namespace UnrealBuildTool
 				case QueryType.TargetDetails:
 					Logger.LogInformation("QueryTargetDetails");
 					return QueryTargetDetails(Arguments, Logger, ResponseOptions);
+				case QueryType.TargetModuleReferences:
+					Logger.LogInformation("QueryTargetModuleReferences");
+					return QueryTargetModuleReferences(Arguments, Logger, ResponseOptions);
 			}
 
 			return Task.FromResult(0);
@@ -190,7 +196,7 @@ namespace UnrealBuildTool
 		{
 			var Reply = new
 			{
-				Queries = new List<string> { QueryType.Capabilities.ToString(), QueryType.AvailableTargets.ToString(), QueryType.TargetDetails.ToString() }
+				Queries = new List<string> { QueryType.Capabilities.ToString(), QueryType.AvailableTargets.ToString(), QueryType.TargetDetails.ToString(), QueryType.TargetModuleReferences.ToString() }
 			};
 			await WriteResultsAsync(Reply, JsonOptions, Logger);
 			return 0;
@@ -492,5 +498,124 @@ namespace UnrealBuildTool
 				return 1;
 			}
 		}
+
+		private async Task<int> QueryTargetModuleReferences(CommandLineArguments Arguments, ILogger Logger, JsonSerializerOptions JsonOptions)
+		{
+			if (TargetName == null)
+			{
+				Logger.LogError("Missing argument Target");
+				return 1;
+			}
+			else if (TargetConfiguration == null)
+			{
+				Logger.LogError("Missing argument Configuration");
+				return 1;
+			}
+			else if (TargetPlatform == null)
+			{
+				Logger.LogError("Missing argument Platform");
+				return 1;
+			}
+
+			Utils.TryParseProjectFileArgument(Arguments, Logger, out FileReference? ProjectFileArg);
+			List<string> RawArgs = new List<string> { TargetName, TargetConfiguration, TargetPlatform };
+			if (ProjectFileArg != null)
+			{
+				RawArgs.Add(ProjectFileArg.ToString());
+			}
+			RawArgs.AddRange(Arguments.GetUnusedArguments());
+			RawArgs.Add("-NoPCHChain"); // Currently unsupported
+			CommandLineArguments Args = new CommandLineArguments(RawArgs.ToArray());
+			List<TargetDescriptor> TargetDescriptors = new();
+
+			TargetDescriptor.ParseSingleCommandLine(Args, false, false, false, TargetDescriptors, Logger);
+			if (TargetDescriptors.Count == 0)
+			{
+				// TOOD: Error 
+				return 1;
+			}
+
+			// Ensure the intermediate environment does not conflict with normal builds
+			TargetDescriptors[0].IntermediateEnvironment = UnrealIntermediateEnvironment.Query;
+
+			try
+			{
+				UEBuildTarget CurrentTarget;
+				using (GlobalTracer.Instance.BuildSpan("UEBuildTarget.Create()").StartActive())
+				{
+					bool bUsePrecompiled = false;
+
+					// Prevent multiple conflicting processes building TargetRules at the same time
+					string MutexName = GlobalSingleInstanceMutex.GetUniqueMutexForPath("UnrealBuildTool_QueryMode_UEBuildTarget-Create", Unreal.RootDirectory);
+					using (new GlobalSingleInstanceMutex(MutexName, true))
+					{
+						CurrentTarget = UEBuildTarget.Create(TargetDescriptors[0], false, false, bUsePrecompiled, TargetDescriptors[0].IntermediateEnvironment, Logger);
+					}
+				}
+
+				UEBuildBinary? LaunchBinary = CurrentTarget.Binaries.FirstOrDefault(Binary => Binary.Modules.Any(Module => Module.Name == CurrentTarget.Rules.LaunchModuleName));
+				if (LaunchBinary == null)
+				{
+					throw new BuildException("Unable to find launch binary for target");
+				}
+
+				// Create the makefile for the target and export the module information
+				{
+					using ISourceFileWorkingSet WorkingSet = new EmptySourceFileWorkingSet();
+					TargetMakefile Makefile;
+					try
+					{
+						Makefile = await BuildMode.CreateMakefileAsync(BuildConfiguration, TargetDescriptors[0], WorkingSet, Logger);
+					}
+					finally
+					{
+						SourceFileMetadataCache.SaveAll();
+					}
+				}
+
+				ConcurrentDictionary<UEBuildModuleCPP, object> ModuleInfo = [];
+
+				await Parallel.ForEachAsync(CurrentTarget.Binaries, async (Binary, CancellationToken) =>
+				{
+					IEnumerable<UEBuildModuleCPP> Modules = Binary.Modules.OfType<UEBuildModuleCPP>().Where(x => x.Binary == Binary);
+
+					await Parallel.ForEachAsync(Modules, (Module, CancellationToken) =>
+					{
+						if (ModuleInfo.ContainsKey(Module))
+						{
+							return ValueTask.CompletedTask;
+						}
+
+						var Info = new
+						{
+							Public = new SortedSet<string>(Module.PublicDependencyModules?.Where(x => x is UEBuildModuleCPP).Select(x => x.Name) ?? []),
+							Private = new SortedSet<string>(Module.PrivateDependencyModules?.Where(x => x is UEBuildModuleCPP).Select(x => x.Name) ?? []),
+							PublicInclude = new SortedSet<string>(Module.PublicIncludePathModules?.Where(x => x is UEBuildModuleCPP).Select(x => x.Name) ?? []),
+							PrivateInclude = new SortedSet<string>(Module.PrivateIncludePathModules?.Where(x => x is UEBuildModuleCPP).Select(x => x.Name) ?? []),
+							DynamicallyLoaded = new SortedSet<string>(Module.DynamicallyLoadedModules?.Where(x => x is UEBuildModuleCPP).Select(x => x.Name) ?? []),
+							Circular = new SortedSet<string>(Module.Rules.CircularlyReferencedDependentModules),
+						};
+
+						ModuleInfo.TryAdd(Module, Info);
+
+						return ValueTask.CompletedTask;
+					});
+				});
+
+				var Result = new
+				{
+					Modules = ModuleInfo.ToImmutableSortedDictionary(x => x.Key.Name, x => x.Value)
+				};
+
+				await WriteResultsAsync(Result, JsonOptions, Logger);
+				return 0;
+			}
+			catch (Exception Ex)
+			{
+				Logger.LogError("Failed to query available targets: {Error}", Ex.Message);
+				Logger.LogDebug(Ex, "Unhandled exception: {Ex}", ExceptionUtils.FormatExceptionDetails(Ex));
+				return 1;
+			}
+		}
 	}
 }
