UnrealBuildTool: Ensure generated headers are added to the prereqs of adaptive files

#rnx

diff --git a/Engine/Source/Programs/UnrealBuildTool/Configuration/UEBuildModuleCPP.cs b/Engine/Source/Programs/UnrealBuildTool/Configuration/UEBuildModuleCPP.cs
--- a/Engine/Source/Programs/UnrealBuildTool/Configuration/UEBuildModuleCPP.cs
+++ b/Engine/Source/Programs/UnrealBuildTool/Configuration/UEBuildModuleCPP.cs
@@ -1820,6 +1820,9 @@ namespace UnrealBuildTool
 			{
 				// Create the new compile environment. Always turn off PCH due to different compiler settings.
 				CppCompileEnvironment AdaptiveUnityEnvironment = new CppCompileEnvironment(ModuleCompileEnvironment);
+				// Ensure any generated headers are added to the adapative environment
+				AdaptiveUnityEnvironment.AdditionalPrerequisites = new(CompileEnvironment.AdditionalPrerequisites);
+				AdaptiveUnityEnvironment.UserIncludePaths = new(CompileEnvironment.UserIncludePaths);
 				if (Target.bAdaptiveUnityDisablesOptimizations)
 				{
 					AdaptiveUnityEnvironment.bOptimizeCode = false;
