Conditionally add VerseVM linker dependencies to all .Build.cs files with native Verse code
- new utility function SetupVerse() in ModuleRules that sets VersePath and VerseScope, as well as adds required linker dependencies
- refactored all applicable Build.cs files to use SetupVerse() instead of setting VersePath and VerseScope directly

#rnx

diff --git a/Engine/Source/Programs/UnrealBuildTool/Configuration/ModuleRules.cs b/Engine/Source/Programs/UnrealBuildTool/Configuration/ModuleRules.cs
--- a/Engine/Source/Programs/UnrealBuildTool/Configuration/ModuleRules.cs
+++ b/Engine/Source/Programs/UnrealBuildTool/Configuration/ModuleRules.cs
@@ -1455,12 +1455,12 @@ namespace UnrealBuildTool
 		/// <summary>
 		/// If this module has associated Verse code, this is the Verse root path of it
 		/// </summary>
-		public string? VersePath { get; set; }
+		public string? VersePath { get; private set; }
 
 		/// <summary>
 		/// Visibility of Verse code in this module's Source/Verse folder
 		/// </summary>
-		public VerseScope VerseScope { get; set; } = VerseScope.PublicUser;
+		public VerseScope VerseScope { get; private set; } = VerseScope.PublicUser;
 
 		/// <summary>
 		/// Whether this module qualifies included headers from other modules relative to the root of their 'Public' folder. This reduces the number
@@ -1602,6 +1602,34 @@ namespace UnrealBuildTool
 			}
 		}
 
+		/// <summary>
+		/// Make sure the necessary linker dependencies are set for compiling code using VerseVM
+		/// </summary>
+		public void SetupVerse()
+		{
+			if (Target.bUseVerseVM)
+			{
+				PublicDependencyModuleNames.AddRange(
+					new string[]
+					{
+						"libpas",
+						"CorePreciseFP",
+					}
+				);
+			}
+		}
+
+		/// <summary>
+		/// Setup this module for Verse
+		/// </summary>
+		public void SetupVerse(string versePath, VerseScope verseScope)
+		{
+			VersePath = versePath;
+			VerseScope = verseScope;
+
+			SetupVerse();
+		}
+
 		/// <summary>
 		/// Setup this module for Mesh Editor support (based on the settings in UEBuildConfiguration)
 		/// </summary>
diff --git a/Engine/Source/Runtime/CoreUObject/CoreUObject.Build.cs b/Engine/Source/Runtime/CoreUObject/CoreUObject.Build.cs
--- a/Engine/Source/Runtime/CoreUObject/CoreUObject.Build.cs
+++ b/Engine/Source/Runtime/CoreUObject/CoreUObject.Build.cs
@@ -5,21 +5,6 @@ using UnrealBuildTool;
 
 public class CoreUObject : ModuleRules
 {
-	public static void AddVerseVMDependencies(ModuleRules Rules, ReadOnlyTargetRules Target)
-	{
-		// If using the new VM either by default or directly, then add in a dependency to 
-		// the core VerseVM which contains special compile flags not compatible with CoreUObject
-		if (Target.bUseVerseVM)
-		{
-			Rules.PublicDependencyModuleNames.AddRange(
-				new string[]
-				{
-					"libpas",
-				}
-			);
-		}
-	}
-
 	public CoreUObject(ReadOnlyTargetRules Target) : base(Target)
 	{
 		// Autogenerate headers for our bytecode.
@@ -45,6 +30,8 @@ public class CoreUObject : ModuleRules
 			}
 		);
 
+		SetupVerse();
+
 		PrivateDependencyModuleNames.AddRange(
 			new string[]
 			{
@@ -53,8 +40,6 @@ public class CoreUObject : ModuleRules
 			}
 		);
 
-		AddVerseVMDependencies(this, Target);
-
 		if (Target.bBuildWithEditorOnlyData || Target.Type == TargetType.Server)
 		{
 			PublicDefinitions.Add("WITH_VERSE_COMPILER=1");
