[UBT]
* Fixed so clang single file actions works with vfs

diff --git a/Engine/Source/Programs/UnrealBuildTool/ToolChain/ClangSpecificFileAction.cs b/Engine/Source/Programs/UnrealBuildTool/ToolChain/ClangSpecificFileAction.cs
--- a/Engine/Source/Programs/UnrealBuildTool/ToolChain/ClangSpecificFileAction.cs
+++ b/Engine/Source/Programs/UnrealBuildTool/ToolChain/ClangSpecificFileAction.cs
@@ -122,7 +122,7 @@ namespace UnrealBuildTool
 				Action.ProducedItems.Add(PreprocessedItem);
 			}
 
-			FileItem ResponseFile = FileItem.GetItemByPath(Action.CommandArguments.Substring(1).Trim('"'));
+			FileReference ResponseFile = RootPaths.GetLocalPath(FileReference.FromString(Action.CommandArguments.Substring(1).Trim('"')));
 			File.WriteAllLines(ResponseFile.FullName, NewRspLines);
 
 			return Action;
