Legacy Xcode: Setup CFBundleSupportedPlatforms as string array as expected. Fixes newly added check in Gauntlet.IOSBuildSource.cs.

diff --git a/Engine/Source/Programs/UnrealBuildTool/Platform/IOS/UEDeployIOS.cs b/Engine/Source/Programs/UnrealBuildTool/Platform/IOS/UEDeployIOS.cs
--- a/Engine/Source/Programs/UnrealBuildTool/Platform/IOS/UEDeployIOS.cs
+++ b/Engine/Source/Programs/UnrealBuildTool/Platform/IOS/UEDeployIOS.cs
@@ -827,7 +827,8 @@ namespace UnrealBuildTool
 
 			Text.AppendLine("\t<key>CFBundleSupportedPlatforms</key>");
 			Text.AppendLine("\t<array>");
-			Text.AppendLine("\t\t<string>iPhoneOS iPhoneSimulator</string>");
+			Text.AppendLine("\t\t<string>iPhoneOS</string>");
+			Text.AppendLine("\t\t<string>iPhoneSimulator</string>");
 			Text.AppendLine("\t</array>");
 			Text.AppendLine("\t<key>MinimumOSVersion</key>");
 			Text.AppendLine(String.Format("\t<string>{0}</string>", MinVersion));
