UnrealBuildTool: Suppress default toolchain warnings when generating project files

#rnx

diff --git a/Engine/Source/Programs/UnrealBuildTool/Platform/Windows/UEBuildWindows.cs b/Engine/Source/Programs/UnrealBuildTool/Platform/Windows/UEBuildWindows.cs
--- a/Engine/Source/Programs/UnrealBuildTool/Platform/Windows/UEBuildWindows.cs
+++ b/Engine/Source/Programs/UnrealBuildTool/Platform/Windows/UEBuildWindows.cs
@@ -1175,6 +1175,8 @@ namespace UnrealBuildTool
 				bool bCompilingForArm = !Target.Architecture.bIsX64;
 				if (bCompilingForArm && Target.Name != "UnrealHeaderTool")
 				{
+					// TODO: Force unique build environment due to disabled plugins
+					Target.BuildEnvironment = TargetBuildEnvironment.Unique;
 					if (Target.WindowsPlatform.Architecture == UnrealArch.Arm64ec)
 					{
 						Target.DisablePlugins.AddRange(new string[]
@@ -1453,7 +1455,7 @@ namespace UnrealBuildTool
 				return WindowsCompiler.VisualStudio2022;
 			}
 
-			if (!bSkipWarning)
+			if (!bSkipWarning && !MicrosoftPlatformSDK.IgnoreToolchainErrors)
 			{
 				UEBuildPlatformSDK? SDK = GetSDK(UnrealTargetPlatform.Win64);
 				VersionNumber? MinimumMsvcVersion = SDK?.GetVersionNumberFromConfig("MinimumVisualCppVersion");
