UnrealBuildTool: Update -SingleFileBuildDependents to include generated code

#rnx

diff --git a/Engine/Source/Programs/UnrealBuildTool/Modes/BuildMode.cs b/Engine/Source/Programs/UnrealBuildTool/Modes/BuildMode.cs
--- a/Engine/Source/Programs/UnrealBuildTool/Modes/BuildMode.cs
+++ b/Engine/Source/Programs/UnrealBuildTool/Modes/BuildMode.cs
@@ -939,7 +939,7 @@ namespace UnrealBuildTool
 		}
 
 		/// <summary>
-		/// If there are headers in the specific list, this function will exand the list to include all files that include those headers
+		/// If there are headers in the specific list, this function will expand the list to include all files that include those headers
 		/// </summary>
 		/// <param name="SpecificFilesToCompile">List of files to compile</param>
 		/// <param name="ProjectFile">Project file if there is one</param>
@@ -948,12 +948,18 @@ namespace UnrealBuildTool
 		internal static List<FileReference> GetAllSourceFilesIncludingHeader(List<FileReference> SpecificFilesToCompile, FileReference? ProjectFile, List<IExternalAction> Actions, ILogger Logger)
 		{
 			Func<FileReference, bool> IsCode = x => IsHeader(x) || x.HasExtension(".cpp") || x.HasExtension(".c");
-			List<FileReference> SpecificHeaderFiles = SpecificFilesToCompile.Where(IsHeader).ToList();
-			if (SpecificHeaderFiles.Count == 0)
+			IEnumerable<FileReference> SpecificHeaderFiles = SpecificFilesToCompile.Where(IsHeader);
+			if (!SpecificHeaderFiles.Any())
 			{
 				return SpecificFilesToCompile;
 			}
 
+			Logger.LogInformation("Building dependency cache for specified headers");
+			foreach (FileReference HeaderFile in SpecificHeaderFiles)
+			{
+				Logger.LogDebug("  {HeaderFile}", HeaderFile);
+			}
+
 			// TODO: This code below works very similar to the code before this changelist but is not a great solution for figuring out which files to compile.
 			// Suggested approach is to:
 			// 1. Write out a file with additional include paths per module.
@@ -961,9 +967,7 @@ namespace UnrealBuildTool
 			// 3. Fix so CppIncludeLookup is including headers... Current code does not handle this: A.h <- B.h <- C.cpp... C.cpp is not added.
 			// 4. Cache which files unity files include to reduce number of files needed to be read
 
-			List<DirectoryReference> BaseDirs = new List<DirectoryReference>();
-			BaseDirs.Add(Unreal.EngineDirectory);
-
+			List<DirectoryReference> BaseDirs = [Unreal.EngineDirectory];
 			DirectoryReference CacheDir = Unreal.EngineDirectory;
 			if (ProjectFile != null)
 			{
@@ -971,74 +975,33 @@ namespace UnrealBuildTool
 				CacheDir = ProjectFile.Directory;
 			}
 
-			FileReference IncludeCache = FileReference.Combine(CacheDir, "Intermediate", "HeaderLookup.dat");
-
-			Logger.LogInformation("Building dependency cache for specified headers");
-			foreach (FileReference HeaderFile in SpecificHeaderFiles)
-			{
-				Logger.LogDebug("  {HeaderFile}", HeaderFile);
-			}
+			IEnumerable<FileItem> CodeFiles = Actions.SelectMany(x => x.PrerequisiteItems).Distinct().Where(x => IsCode(x.Location) && !x.Name.StartsWith(Unity.ModulePrefix));
 
+			FileReference IncludeCache = FileReference.Combine(CacheDir, "Intermediate", "HeaderLookup.dat");
 			CppIncludeLookup IncludeLookup = new CppIncludeLookup(IncludeCache);
 			IncludeLookup.Load();
-			IncludeLookup.Update(BaseDirs);
+			IncludeLookup.Update(BaseDirs, CodeFiles);
 			IncludeLookup.Save();
 
 			// Find all files that can be part of this build
-			ConcurrentDictionary<FileReference, int> UsedFiles = new();
-			Parallel.ForEach(Actions, (Action) =>
-			{
-				foreach (FileItem File in Action.PrerequisiteItems)
-				{
-					if (!IsCode(File.Location))
-					{
-						continue;
-					}
-
-					// If file is a unity file we want to traverse the internal files and add them as UsedFiles
-					if (File.Name.StartsWith(Unity.ModulePrefix))
-					{
-						foreach (string Line in System.IO.File.ReadAllLines(File.FullName))
-						{
-							int IncludeStart = Line.IndexOf('"') + 1;
-							if (IncludeStart != 0)
-							{
-								string Include = Line.Substring(IncludeStart, Line.Length - IncludeStart - 1);
-								FileItem SourceFile = FileItem.GetItemByFileReference(FileReference.Combine(Unreal.EngineSourceDirectory, Include));
-								UsedFiles.TryAdd(SourceFile.Location, 0);
-							}
-						}
-					}
-					else
-					{
-						UsedFiles.TryAdd(File.Location, 0);
-					}
-				}
-			});
-
-			SortedSet<FileReference> NewFiles = new();
-			List<FileReference> FoundFiles = IncludeLookup.FindFiles(SpecificFilesToCompile.Select(x => x.GetFileName())).Select(x => x.Location).ToList();
+			IEnumerable<FileReference> UsedFiles = CodeFiles.Select(x => x.Location);
+			IEnumerable<FileReference> FoundFiles = IncludeLookup.FindFiles(SpecificFilesToCompile.Select(x => x.GetFileName()).Distinct()).Select(x => x.Location);
 
 			// Found files can be a lot (Commandline.h gives 130.000+ files). And we only want to build the ones that are in existing actions
-			Parallel.ForEach(FoundFiles, (File) =>
-			{
-				if (UsedFiles.ContainsKey(File))
-				{
-					lock (NewFiles)
-					{
-						NewFiles.Add(File);
-					}
-				}
-			});
+			IEnumerable<FileReference> NewFiles = FoundFiles.Intersect(UsedFiles).Distinct();
 
-			Logger.LogInformation("Found {NewFilesCount} source files", NewFiles.Count);
-			foreach (FileReference NewFile in NewFiles)
+			int genFileCount = NewFiles.Where(x => x.GetFileName().EndsWith(".gen.cpp", StringComparison.OrdinalIgnoreCase)).Count();
+			Logger.LogInformation("Found {NewFilesCount} source files", NewFiles.Count() - genFileCount);
+			if (genFileCount > 0)
+			{
+				Logger.LogInformation("Found {NewFilesCount} generated source files", genFileCount);
+			}
+			foreach (FileReference NewFile in NewFiles.Order())
 			{
 				Logger.LogDebug("  {NewFile}", NewFile);
 			}
 
-			NewFiles.UnionWith(SpecificFilesToCompile);
-			return NewFiles.ToList();
+			return NewFiles.Concat(SpecificFilesToCompile).Distinct().Order().ToList();
 		}
 
 		/// <summary>
diff --git a/Engine/Source/Programs/UnrealBuildTool/System/CppIncludeLookup.cs b/Engine/Source/Programs/UnrealBuildTool/System/CppIncludeLookup.cs
--- a/Engine/Source/Programs/UnrealBuildTool/System/CppIncludeLookup.cs
+++ b/Engine/Source/Programs/UnrealBuildTool/System/CppIncludeLookup.cs
@@ -146,10 +146,10 @@ namespace UnrealBuildTool
 			return NameInfoSet.SelectMany(x => x.Files).Select(x => x.File);
 		}
 
-		public void Update(IEnumerable<DirectoryReference> RootDirs)
+		public void Update(IEnumerable<DirectoryReference> RootDirs, IEnumerable<FileItem> AdditionalFiles)
 		{
 			// Find all the source files in the given root directory
-			ConcurrentBag<FileItem> KeepFileItems = new ConcurrentBag<FileItem>();
+			ConcurrentBag<FileItem> KeepFileItems = new ConcurrentBag<FileItem>(AdditionalFiles);
 			using (ThreadPoolWorkQueue Queue = new ThreadPoolWorkQueue())
 			{
 				foreach (DirectoryReference RootDir in RootDirs)
@@ -203,6 +203,10 @@ namespace UnrealBuildTool
 
 		void ScanDependencies(CppIncludeFileInfo FileInfo, HashSet<CppIncludeFileInfo> Files, ThreadPoolWorkQueue Queue)
 		{
+			if (!FileInfo.File.Exists)
+			{
+				return;
+			}
 			// Update names for this file
 			if (FileInfo.IncludedNames == null || FileInfo.File.LastWriteTimeUtc != FileInfo.LastWriteTimeUtc)
 			{
