- Fixed an issue with platform sdk validity for IOS if it gets registered before Mac

diff --git a/Engine/Source/Programs/UnrealBuildTool/Platform/Mac/ApplePlatformSDK.cs b/Engine/Source/Programs/UnrealBuildTool/Platform/Mac/ApplePlatformSDK.cs
--- a/Engine/Source/Programs/UnrealBuildTool/Platform/Mac/ApplePlatformSDK.cs
+++ b/Engine/Source/Programs/UnrealBuildTool/Platform/Mac/ApplePlatformSDK.cs
@@ -144,7 +144,7 @@ namespace UnrealBuildTool
 				{
 					string? DllVersion = FileVersionInfo.GetVersionInfo(DllPath).FileVersion;
 					// Only return the DLL version as the SDK version if we can correctly parse it
-					UEBuildPlatformSDK? Sdk = UEBuildPlatformSDK.GetSDKForPlatform("Mac");
+					ApplePlatformSDK? Sdk = UEBuildPlatformSDK.GetSDKForPlatformOrMakeTemp<ApplePlatformSDK>("Mac");
 					if (Sdk != null && Sdk.TryConvertVersionToInt(DllVersion, out _, null))
 					{
 						return DllVersion;
