[UBT]
* Fixed so generated header cpp files (source files that are generated to be able to compile headers in isolation) are using vfs paths if vfs is enabled

diff --git a/Engine/Source/Programs/UnrealBuildTool/Configuration/UEBuildModuleCPP.cs b/Engine/Source/Programs/UnrealBuildTool/Configuration/UEBuildModuleCPP.cs
--- a/Engine/Source/Programs/UnrealBuildTool/Configuration/UEBuildModuleCPP.cs
+++ b/Engine/Source/Programs/UnrealBuildTool/Configuration/UEBuildModuleCPP.cs
@@ -1043,7 +1043,11 @@ namespace UnrealBuildTool
 						foreach (FileItem Header in HeaderFileItems.Where(x => x.Directory == ParentDir))
 						{
 							string IncludeFileString = Header.AbsolutePath;
-							if (Header.Location.IsUnderDirectory(Unreal.RootDirectory))
+							if (HeaderCompileEnvironment.RootPaths.GetVfsOverlayPath(Header.Location, out string? VfsPath))
+							{
+								IncludeFileString = VfsPath;
+							}
+							else if (Header.Location.IsUnderDirectory(Unreal.RootDirectory))
 							{
 								IncludeFileString = Header.Location.MakeRelativeTo(Unreal.EngineSourceDirectory);
 							}
