[UBT]
* Reverted back so UE_ENGINE_DIR is a relative path from executable regardless if vfs is on or not. This hard coded path should be stored in BuildSettings instead.. and need to handle executable and engine dirs being on different drives

diff --git a/Engine/Source/Programs/UnrealBuildTool/Configuration/UEBuildTarget.cs b/Engine/Source/Programs/UnrealBuildTool/Configuration/UEBuildTarget.cs
--- a/Engine/Source/Programs/UnrealBuildTool/Configuration/UEBuildTarget.cs
+++ b/Engine/Source/Programs/UnrealBuildTool/Configuration/UEBuildTarget.cs
@@ -5685,17 +5685,9 @@ namespace UnrealBuildTool
 						OutputFilePath = OutputFilePath.Substring(0, OutputFilePath.LastIndexOf(".app/Contents/MacOS") + 4);
 					}
 
+					DirectoryReference OutputDir = new DirectoryReference(OutputFilePath).ParentDirectory!;
 					DirectoryReference RootDir = Unreal.EngineDirectory.ParentDirectory!;
-					string RootPath;
-					if (GlobalCompileEnvironment.RootPaths.GetVfsOverlayPath(RootDir, out string? VfsPath))
-					{
-						RootPath = VfsPath;
-					}
-					else
-					{
-						DirectoryReference OutputDir = DirectoryReference.FromString(OutputFilePath)!.ParentDirectory!;
-						RootPath = RootDir.MakeRelativeTo(OutputDir);
-					}
+					string RootPath = RootDir.MakeRelativeTo(OutputDir);
 
 					// now get the cleaned engine path (Root/Engine), and make sure to end with /
 					string EnginePath = Utils.CleanDirectorySeparators(Path.Combine(RootPath, "Engine/"), '/');
