UnrealBuildTool: Fix issue where Internal paths were being provided to modules that should not have access

#rnx

diff --git a/Engine/Source/Programs/UnrealBuildTool/Configuration/UEBuildModule.cs b/Engine/Source/Programs/UnrealBuildTool/Configuration/UEBuildModule.cs
--- a/Engine/Source/Programs/UnrealBuildTool/Configuration/UEBuildModule.cs
+++ b/Engine/Source/Programs/UnrealBuildTool/Configuration/UEBuildModule.cs
@@ -659,8 +659,8 @@ namespace UnrealBuildTool
 				AddIncludePaths(IncludePaths, LegacyPublicIncludePaths);
 			}
 
-			// Add this module's internal include paths, only if the scope contains the same as the SourceModule's scope or if this module is an engine module
-			if (SourceModule != null && (Rules.Context.Scope.Contains(SourceModule.Rules.Context.Scope) || Rules.bTreatAsEngineModule))
+			// Add this module's internal include paths, only if the scope contains the same as the SourceModule's scope or if the source module is an engine module
+			if (SourceModule != null && (Rules.Context.Scope.Contains(SourceModule.Rules.Context.Scope) || SourceModule.Rules.bTreatAsEngineModule))
 			{
 				AddIncludePaths(IncludePaths, InternalIncludePaths);
 			}
