Cleanup ResolveWarning API to remove buildconfiguration delegate, and instead apply target defaults in new API.

diff --git a/Engine/Source/Programs/Shared/UnrealBuildTool.Tests/CppCompileWarningsTests.cs b/Engine/Source/Programs/Shared/UnrealBuildTool.Tests/CppCompileWarningsTests.cs
--- a/Engine/Source/Programs/Shared/UnrealBuildTool.Tests/CppCompileWarningsTests.cs
+++ b/Engine/Source/Programs/Shared/UnrealBuildTool.Tests/CppCompileWarningsTests.cs
@@ -17,6 +17,36 @@ namespace UnrealBuildTool.Tests
 			}
 		}
 
+		[DataTestMethod]
+		[DataRow(WarningLevel.Default, null, true, WarningLevel.Default)]					// Null applied, no overwrite
+		[DataRow(WarningLevel.Default, WarningLevel.Off, true, WarningLevel.Off)]			// Overwrite with Off
+		[DataRow(WarningLevel.Default, WarningLevel.Warning, true, WarningLevel.Warning)]	// Overwrite with Warning
+		[DataRow(WarningLevel.Default, WarningLevel.Error, true, WarningLevel.Error)]		// Overwrite with Error
+		[DataRow(WarningLevel.Default, WarningLevel.Default, true, WarningLevel.Default)]	// Applied is Default, no change
+
+		[DataRow(WarningLevel.Default, WarningLevel.Warning, false, WarningLevel.Warning)]	// Default, non-overwrite, change
+		[DataRow(WarningLevel.Default, WarningLevel.Error, false, WarningLevel.Error)]		// Default, non-overwrite, change
+
+		[DataRow(WarningLevel.Off, null, true, WarningLevel.Off)]							// Null applied, no change
+		[DataRow(WarningLevel.Off, WarningLevel.Default, true, WarningLevel.Off)]			// Applied Default, no change
+		[DataRow(WarningLevel.Off, WarningLevel.Warning, true, WarningLevel.Warning)]		// Overwrite with Warning
+		[DataRow(WarningLevel.Off, WarningLevel.Error, true, WarningLevel.Error)]			// Overwrite with Error
+		[DataRow(WarningLevel.Off, WarningLevel.Warning, false, WarningLevel.Off)]			// No overwrite, no change
+
+		[DataRow(WarningLevel.Warning, WarningLevel.Error, true, WarningLevel.Error)]		// Overwrite with Error
+		[DataRow(WarningLevel.Warning, WarningLevel.Off, true, WarningLevel.Off)]			// Overwrite with Off
+		[DataRow(WarningLevel.Warning, WarningLevel.Off, false, WarningLevel.Warning)]		// No overwrite, remains Warning
+
+		[DataRow(WarningLevel.Error, WarningLevel.Off, true, WarningLevel.Off)]				// Overwrite with Off
+		[DataRow(WarningLevel.Error, WarningLevel.Warning, true, WarningLevel.Warning)]		// Overwrite with Warning
+		[DataRow(WarningLevel.Error, WarningLevel.Off, false, WarningLevel.Error)]			// No overwrite, remains Error
+		public void TestApplyWarnings(WarningLevel inSourceWarningLevel, WarningLevel? appliedWarningLevel, bool overwriteSourceValue, WarningLevel expected)
+		{
+			WarningLevel sourceWarningLevel = inSourceWarningLevel;
+			sourceWarningLevel.ApplyWarning(appliedWarningLevel, overwriteSourceValue);
+			Assert.AreEqual(expected, sourceWarningLevel);
+		}
+
 		[TestMethod]
 		public void TestDefaultBuildSettings()
 		{
@@ -76,9 +106,12 @@ namespace UnrealBuildTool.Tests
 
 				// Boundary analysis
 				buildContextProvider._testBuildSettings = BuildSettingsVersion.V2;
+				CppCompileWarnings.ApplyTargetDefaults(newWarnings, true);
 				Assert.AreEqual(WarningLevel.Error, newWarnings.ShadowVariableWarningLevel);
 
 				buildContextProvider._testBuildSettings = BuildSettingsVersion.V3;
+
+				CppCompileWarnings.ApplyTargetDefaults(newWarnings, true);
 				Assert.AreEqual(WarningLevel.Error, newWarnings.ShadowVariableWarningLevel);
 			}
 		}
diff --git a/Engine/Source/Programs/UnrealBuildTool/Configuration/CppCompileWarnings.cs b/Engine/Source/Programs/UnrealBuildTool/Configuration/CppCompileWarnings.cs
--- a/Engine/Source/Programs/UnrealBuildTool/Configuration/CppCompileWarnings.cs
+++ b/Engine/Source/Programs/UnrealBuildTool/Configuration/CppCompileWarnings.cs
@@ -13,10 +13,11 @@ namespace UnrealBuildTool
 		/// </summary>
 		/// <param name="inSourceWarningLevel">The source warning level to override.</param>
 		/// <param name="appliedWarningLevel">The warning level to apply to the source.</param>
+		/// <param name ="overwriteSourceValue">Flag to overwrite the source warning with new value. If false, will only apply setting if value is unset.</param>
 		/// <remarks>If <paramref name="appliedWarningLevel"/> with value of null or <see cref="WarningLevel.Default"/> will be ignored as an override.</remarks>
-		internal static void ApplyWarning(this ref WarningLevel inSourceWarningLevel, WarningLevel? appliedWarningLevel)
+		internal static void ApplyWarning(this ref WarningLevel inSourceWarningLevel, WarningLevel? appliedWarningLevel, bool overwriteSourceValue = true)
 		{
-			if (appliedWarningLevel.HasValue && appliedWarningLevel.Value != WarningLevel.Default)
+			if ((appliedWarningLevel.HasValue && appliedWarningLevel.Value != WarningLevel.Default) && (overwriteSourceValue ||  (!overwriteSourceValue && inSourceWarningLevel == WarningLevel.Default)))
 			{
 				inSourceWarningLevel = appliedWarningLevel.Value;
 			}
@@ -33,15 +34,15 @@ namespace UnrealBuildTool
 		 * 1. Copy the following stub code:
 		 *		public WarningLevel NEW_WARNING_LEVEL_NAME
 		 *		{
-		 *			get => ResolveWarning(null, ParentCppCompileWarnings?.NEW_WARNING_LEVEL_NAME, _NEW_WARNING_LEVEL_NAME);
+		 *			get => ResolveWarning(ParentCppCompileWarnings?.NEW_WARNING_LEVEL_NAME, _NEW_WARNING_LEVEL_NAME);
 		 *			set => _NEW_WARNING_LEVEL_NAME = value;
 		 *		}
 		 *		private WarningLevel _NEW_WARNING_LEVEL_NAME;
 		 * 2. Update the ReadOnlyCppCompileWarnings class to include a forwarding property:
 		 *		public WarningLevel NEW_WARNING_LEVEL_NAME => _inner.NEW_WARNING_LEVEL_NAME;
 		 * 3. Set your default WarningLevel within appropraite ApplyDefaults(...) setting
-		 *		For Target defaults				- private void ApplyTargetDefaults()
-		 *		For CompileEnvironment defaults	- private void ApplyDefaults()
+		 *		For Target defaults				- internal void ApplyTargetDefaults()
+		 *			Note: This is the most appropriatel place for BuildSettings version specific defaults
 		 * 4. Access the new property through the (TargetRules|ModuleRules|CppCompileEnvironment).CppCompileWarnings as needed
 		 */
 		#region -- Compiler Warning Settings --
@@ -53,7 +54,7 @@ namespace UnrealBuildTool
 		[RequiresUniqueBuildEnvironment]
 		public WarningLevel ShadowVariableWarningLevel
 		{
-			get => ResolveWarning((BuildSettingsVersion x) => (x >= BuildSettingsVersion.V2) ? WarningLevel.Error : (ParentCppCompileWarnings != null ? ParentCppCompileWarnings.ShadowVariableWarningLevel : WarningLevel.Default), ParentCppCompileWarnings?.ShadowVariableWarningLevel, _shadowVariableWarningLevel);
+			get => ResolveWarning(ParentCppCompileWarnings?.ShadowVariableWarningLevel, _shadowVariableWarningLevel);
 			set => _shadowVariableWarningLevel = value;
 		}
 		private WarningLevel _shadowVariableWarningLevel;
@@ -65,7 +66,7 @@ namespace UnrealBuildTool
 		[RequiresUniqueBuildEnvironment]
 		public WarningLevel UnsafeTypeCastWarningLevel
 		{
-			get => ResolveWarning(null, ParentCppCompileWarnings?.UnsafeTypeCastWarningLevel, _unsafeTypeCastWarningLevel);
+			get => ResolveWarning(ParentCppCompileWarnings?.UnsafeTypeCastWarningLevel, _unsafeTypeCastWarningLevel);
 			set => _unsafeTypeCastWarningLevel = value;
 		}
 		private WarningLevel _unsafeTypeCastWarningLevel;
@@ -78,7 +79,7 @@ namespace UnrealBuildTool
 		[RequiresUniqueBuildEnvironment]
 		public WarningLevel UndefinedIdentifierWarningLevel
 		{
-			get => ResolveWarning(null, ParentCppCompileWarnings?.UndefinedIdentifierWarningLevel, _undefinedIdentifierWarningLevel);
+			get => ResolveWarning(ParentCppCompileWarnings?.UndefinedIdentifierWarningLevel, _undefinedIdentifierWarningLevel);
 			set => _undefinedIdentifierWarningLevel = value;
 		}
 		private WarningLevel _undefinedIdentifierWarningLevel;
@@ -91,7 +92,7 @@ namespace UnrealBuildTool
 		[RequiresUniqueBuildEnvironment]
 		public WarningLevel SwitchUnhandledEnumeratorWarningLevel
 		{
-			get => ResolveWarning(null, ParentCppCompileWarnings?.SwitchUnhandledEnumeratorWarningLevel, _switchUnhandledEnumeratorWarningLevel);
+			get => ResolveWarning(ParentCppCompileWarnings?.SwitchUnhandledEnumeratorWarningLevel, _switchUnhandledEnumeratorWarningLevel);
 			set => _switchUnhandledEnumeratorWarningLevel = value;
 		}
 		private WarningLevel _switchUnhandledEnumeratorWarningLevel;
@@ -107,7 +108,7 @@ namespace UnrealBuildTool
 		[CommandLine("-PCHPerformanceIssueWarningLevel=")]
 		public WarningLevel PCHPerformanceIssueWarningLevel
 		{
-			get => ResolveWarning(null, ParentCppCompileWarnings?.PCHPerformanceIssueWarningLevel, _pchPerformanceIssueWarningLevel);
+			get => ResolveWarning(ParentCppCompileWarnings?.PCHPerformanceIssueWarningLevel, _pchPerformanceIssueWarningLevel);
 			set => _pchPerformanceIssueWarningLevel = value;
 		}
 		private WarningLevel _pchPerformanceIssueWarningLevel;
@@ -123,7 +124,7 @@ namespace UnrealBuildTool
 		[CommandLine("-ModuleUnsupportedWarningLevel=")]
 		public WarningLevel ModuleUnsupportedWarningLevel
 		{
-			get => ResolveWarning(null, ParentCppCompileWarnings?.ModuleUnsupportedWarningLevel, _moduleUnsupportedWarningLevel);
+			get => ResolveWarning(ParentCppCompileWarnings?.ModuleUnsupportedWarningLevel, _moduleUnsupportedWarningLevel);
 			set => _moduleUnsupportedWarningLevel = value;
 		}
 		private WarningLevel _moduleUnsupportedWarningLevel;
@@ -135,7 +136,7 @@ namespace UnrealBuildTool
 		[CommandLine("-PluginModuleUnsupportedWarningLevel=")]
 		public WarningLevel PluginModuleUnsupportedWarningLevel
 		{
-			get => ResolveWarning(null, ParentCppCompileWarnings?.PluginModuleUnsupportedWarningLevel, _pluginModuleUnsupportedWarningLevel);
+			get => ResolveWarning(ParentCppCompileWarnings?.PluginModuleUnsupportedWarningLevel, _pluginModuleUnsupportedWarningLevel);
 			set => _pluginModuleUnsupportedWarningLevel = value;
 		}
 		private WarningLevel _pluginModuleUnsupportedWarningLevel;
@@ -147,7 +148,7 @@ namespace UnrealBuildTool
 		[CommandLine("-ModuleIncludePathWarningLevel=")]
 		public WarningLevel ModuleIncludePathWarningLevel
 		{
-			get => ResolveWarning(null, ParentCppCompileWarnings?.ModuleIncludePathWarningLevel, _moduleIncludePathWarningLevel);
+			get => ResolveWarning(ParentCppCompileWarnings?.ModuleIncludePathWarningLevel, _moduleIncludePathWarningLevel);
 			set => _moduleIncludePathWarningLevel = value;
 		}
 		private WarningLevel _moduleIncludePathWarningLevel;
@@ -159,7 +160,7 @@ namespace UnrealBuildTool
 		[CommandLine("-ModuleIncludePrivateWarningLevel=")]
 		public WarningLevel ModuleIncludePrivateWarningLevel
 		{
-			get => ResolveWarning(null, ParentCppCompileWarnings?.ModuleIncludePrivateWarningLevel, _moduleIncludePrivateWarningLevel);
+			get => ResolveWarning(ParentCppCompileWarnings?.ModuleIncludePrivateWarningLevel, _moduleIncludePrivateWarningLevel);
 			set => _moduleIncludePrivateWarningLevel = value;
 		}
 		private WarningLevel _moduleIncludePrivateWarningLevel;
@@ -171,7 +172,7 @@ namespace UnrealBuildTool
 		[CommandLine("-ModuleIncludeSubdirectoryWarningLevel=")]
 		public WarningLevel ModuleIncludeSubdirectoryWarningLevel
 		{
-			get => ResolveWarning(null, ParentCppCompileWarnings?.ModuleIncludeSubdirectoryWarningLevel, _moduleIncludeSubdirectoryWarningLevel);
+			get => ResolveWarning(ParentCppCompileWarnings?.ModuleIncludeSubdirectoryWarningLevel, _moduleIncludeSubdirectoryWarningLevel);
 			set => _moduleIncludePrivateWarningLevel = value;
 		}
 		private WarningLevel _moduleIncludeSubdirectoryWarningLevel;
@@ -181,7 +182,7 @@ namespace UnrealBuildTool
 		/// </summary>
 		public WarningLevel DisablePluginsConflictWarningLevel
 		{
-			get => ResolveWarning(null, ParentCppCompileWarnings?.DisablePluginsConflictWarningLevel, _disablePluginsConflictWarningLevel);
+			get => ResolveWarning(ParentCppCompileWarnings?.DisablePluginsConflictWarningLevel, _disablePluginsConflictWarningLevel);
 			set => _disablePluginsConflictWarningLevel = value;
 		}
 		private WarningLevel _disablePluginsConflictWarningLevel;
@@ -191,7 +192,7 @@ namespace UnrealBuildTool
 		/// </summary>
 		public WarningLevel NonInlinedGenCppWarningLevel
 		{
-			get => ResolveWarning(null, ParentCppCompileWarnings?.NonInlinedGenCppWarningLevel, _nonInlinedGenCppWarningLevel);
+			get => ResolveWarning(ParentCppCompileWarnings?.NonInlinedGenCppWarningLevel, _nonInlinedGenCppWarningLevel);
 			set => _nonInlinedGenCppWarningLevel = value;
 		}
 		private WarningLevel _nonInlinedGenCppWarningLevel;
@@ -211,8 +212,6 @@ namespace UnrealBuildTool
 			BuildContextProvider = new BuildSystemContext(new TargetRulesBuildSettingsProvider(targetRulesContextObject), targetRulesContextObject, null);
 			ParentCppCompileWarnings = null;
 			_logger = logger;
-
-			ApplyTargetDefaults();
 		}
 
 		/// <summary>
@@ -238,7 +237,6 @@ namespace UnrealBuildTool
 		{
 			BuildContextProvider = new BuildSystemContext(contextObject);
 			ParentCppCompileWarnings = parentCompilerWarnings;
-
 			_logger = logger;
 		}
 
@@ -300,58 +298,56 @@ namespace UnrealBuildTool
 
 		#endregion -- Public API --
 
-		#region -- Private API --
+		#region -- Internal API --
 
 		/// <summary>
-		/// General purpose resolver that will resolve a warning based off of the precedence of 1. locally set in the <see cref="CppCompileWarnings"/> object; 2. utilize the parent if it exists; 3. a default as evaluated by the provided build settings delegate.
+		/// Applies the Target context defaults to the provided <see cref="CppCompileWarnings"/>.
 		/// </summary>
-		/// <param name="buildSettingsEvaluator">The build settings evaluator to use to ascibe a build default if there is no instance setting.</param>
-		/// <param name="parentWarningLevel">The parent <see cref="CppCompileWarnings"/> <see cref="WarningLevel"/>.</param>
-		/// <param name="instanceWarningLevel">The instance <see cref="WarningLevel"/>.</param>
-		/// <returns>The resolved warning level.</returns>
-		private WarningLevel ResolveWarning(Func<BuildSettingsVersion, WarningLevel>? buildSettingsEvaluator, WarningLevel? parentWarningLevel, WarningLevel instanceWarningLevel)
+		/// <param name="cppCompileWarnings">The CppCompileWarnings object to apply target defaults to.</param>
+		/// <param name="overwriteSourceValue">Whether to overwrite warning setting values, or set if default.</param>
+		internal static void ApplyTargetDefaults(CppCompileWarnings cppCompileWarnings, bool overwriteSourceValue = false)
 		{
-			WarningLevel resolvedWarning = WarningLevel.Default;
+			cppCompileWarnings._logger?.LogTrace("Applying default C++ Compiler warnings for the Target({TargetName}) context.", cppCompileWarnings.BuildContextProvider?._targetRulesPrivate?.Name);
 
-			if (buildSettingsEvaluator != null && instanceWarningLevel == WarningLevel.Default)
-			{
-				resolvedWarning.ApplyWarning(buildSettingsEvaluator(BuildContextProvider?._buildContext != null ? BuildContextProvider._buildContext.GetBuildSettings() : BuildSettingsVersion.Latest));
-			}
-			else
-			{
-				resolvedWarning.ApplyWarning(parentWarningLevel);
-				resolvedWarning.ApplyWarning(instanceWarningLevel);
-			}
+			// Compiler Warning Settings
+			cppCompileWarnings._shadowVariableWarningLevel.ApplyWarning((cppCompileWarnings.BuildContextProvider?._buildContext.GetBuildSettings() >= BuildSettingsVersion.V2) ? WarningLevel.Error : WarningLevel.Warning, overwriteSourceValue);
+			cppCompileWarnings._unsafeTypeCastWarningLevel.ApplyWarning(WarningLevel.Off, overwriteSourceValue);
+			cppCompileWarnings._undefinedIdentifierWarningLevel.ApplyWarning(WarningLevel.Error, overwriteSourceValue);
+			cppCompileWarnings._switchUnhandledEnumeratorWarningLevel.ApplyWarning(WarningLevel.Off, overwriteSourceValue);
 
-			return resolvedWarning;
+			// Performance Warning Settings
+			cppCompileWarnings._pchPerformanceIssueWarningLevel.ApplyWarning(WarningLevel.Off, overwriteSourceValue);
+
+			// Rules Validation Warning Settings
+			cppCompileWarnings._moduleUnsupportedWarningLevel.ApplyWarning(WarningLevel.Warning, overwriteSourceValue);
+			cppCompileWarnings._pluginModuleUnsupportedWarningLevel.ApplyWarning(WarningLevel.Off, overwriteSourceValue);
+
+			cppCompileWarnings._moduleIncludePathWarningLevel.ApplyWarning(WarningLevel.Off, overwriteSourceValue);
+			cppCompileWarnings._moduleIncludePrivateWarningLevel.ApplyWarning(WarningLevel.Off, overwriteSourceValue);
+			cppCompileWarnings._moduleIncludeSubdirectoryWarningLevel.ApplyWarning(WarningLevel.Off, overwriteSourceValue);
+
+			cppCompileWarnings._disablePluginsConflictWarningLevel.ApplyWarning(WarningLevel.Default,overwriteSourceValue);
+			cppCompileWarnings._nonInlinedGenCppWarningLevel.ApplyWarning(WarningLevel.Off, false);
 		}
 
+		#endregion
+
+		#region -- Private API --
+
 		/// <summary>
-		/// Default settings for the target context.
+		/// General purpose resolver that will resolve a warning based off of the precedence of 1. locally set in the <see cref="CppCompileWarnings"/> object; 2. utilize the parent if it exists;
 		/// </summary>
-		private void ApplyTargetDefaults()
+		/// <param name="parentWarningLevel">The parent <see cref="CppCompileWarnings"/> <see cref="WarningLevel"/>.</param>
+		/// <param name="instanceWarningLevel">The instance <see cref="WarningLevel"/>.</param>
+		/// <returns>The resolved warning level.</returns>
+		private static WarningLevel ResolveWarning(WarningLevel? parentWarningLevel, WarningLevel instanceWarningLevel)
 		{
-			_logger?.LogTrace("Applying default C++ Compiler warnings for the Target({TargetName}) context.", BuildContextProvider?._targetRulesPrivate?.Name);
-
-			// Compiler Warning Settings
-			_shadowVariableWarningLevel = WarningLevel.Warning;
-			_unsafeTypeCastWarningLevel = WarningLevel.Off;
-			_undefinedIdentifierWarningLevel = WarningLevel.Error;
-			_switchUnhandledEnumeratorWarningLevel = WarningLevel.Off;
-
-			// Performance Warning Settings
-			_pchPerformanceIssueWarningLevel = WarningLevel.Off;
-
-			// Rules Validation Warning Settings
-			_moduleUnsupportedWarningLevel = WarningLevel.Warning;
-			_pluginModuleUnsupportedWarningLevel = WarningLevel.Off;
+			WarningLevel resolvedWarning = WarningLevel.Default;
 
-			_moduleIncludePathWarningLevel = WarningLevel.Off;
-			_moduleIncludePrivateWarningLevel = WarningLevel.Off;
-			_moduleIncludeSubdirectoryWarningLevel = WarningLevel.Off;
+			resolvedWarning.ApplyWarning(parentWarningLevel);
+			resolvedWarning.ApplyWarning(instanceWarningLevel);
 
-			_disablePluginsConflictWarningLevel = WarningLevel.Default;
-			_nonInlinedGenCppWarningLevel = WarningLevel.Off;
+			return resolvedWarning;
 		}
 
 		/// <summary>
@@ -388,7 +384,7 @@ namespace UnrealBuildTool
 		#region -- Internal & Private Instance State --
 
 		private readonly ILogger? _logger;
-		private BuildSystemContext BuildContextProvider { get; set; }
+		private BuildSystemContext BuildContextProvider { get; init; }
 		internal ReadOnlyCppCompileWarnings? ParentCppCompileWarnings { get; private set; }
 
 		#endregion -- Internal & Private Instance State --
diff --git a/Engine/Source/Programs/UnrealBuildTool/Configuration/TargetRules.cs b/Engine/Source/Programs/UnrealBuildTool/Configuration/TargetRules.cs
--- a/Engine/Source/Programs/UnrealBuildTool/Configuration/TargetRules.cs
+++ b/Engine/Source/Programs/UnrealBuildTool/Configuration/TargetRules.cs
@@ -3118,6 +3118,7 @@ namespace UnrealBuildTool
 			try
 			{
 				constructor.Invoke(rules, new object[] { targetInfo });
+				CppCompileWarnings.ApplyTargetDefaults(rules.CppCompileWarningSettings);
 			}
 			catch (Exception ex)
 			{
