SwitchUnhandledEnumeratorWarningLevel not correctly referenced from VC Toolchain

diff --git a/Engine/Source/Programs/UnrealBuildTool/Platform/Windows/VCToolChain.cs b/Engine/Source/Programs/UnrealBuildTool/Platform/Windows/VCToolChain.cs
--- a/Engine/Source/Programs/UnrealBuildTool/Platform/Windows/VCToolChain.cs
+++ b/Engine/Source/Programs/UnrealBuildTool/Platform/Windows/VCToolChain.cs
@@ -1148,7 +1148,7 @@ namespace UnrealBuildTool
 
 			if (CompileEnvironment.CppCompileWarnings.SwitchUnhandledEnumeratorWarningLevel != WarningLevel.Off)
 			{
-				if (CompileEnvironment.CppCompileWarnings.UndefinedIdentifierWarningLevel == WarningLevel.Error)
+				if (CompileEnvironment.CppCompileWarnings.SwitchUnhandledEnumeratorWarningLevel == WarningLevel.Error)
 				{
 					Arguments.Add("/we4061");
 				}
