UnrealBuildTool: Load BuildConfiguration.xml settings from the environment if in the format "UnrealBuildTool_<Category>__<Field>"

#rnx

diff --git a/Engine/Source/Programs/UnrealBuildTool/System/XmlConfig.cs b/Engine/Source/Programs/UnrealBuildTool/System/XmlConfig.cs
--- a/Engine/Source/Programs/UnrealBuildTool/System/XmlConfig.cs
+++ b/Engine/Source/Programs/UnrealBuildTool/System/XmlConfig.cs
@@ -1,6 +1,7 @@
 // Copyright Epic Games, Inc. All Rights Reserved.
 
 using System;
+using System.Collections;
 using System.Collections.Generic;
 using System.Diagnostics.CodeAnalysis;
 using System.IO;
@@ -8,6 +9,7 @@ using System.Linq;
 using System.Reflection;
 using System.Text;
 using System.Xml;
+using System.Xml.Linq;
 using System.Xml.Schema;
 using System.Xml.Serialization;
 using EpicGames.Core;
@@ -121,6 +123,11 @@ namespace UnrealBuildTool
 					temporaryInputFileLocations.Add(projectRootConfigLocation);
 				}
 
+				// Write any configuration environment variables to a file so they can be loaded
+				FileReference environmentXml = FileReference.Combine(Unreal.WritableEngineDirectory, "Intermediate", "Build", "UnrealBuildTool.Env.BuildConfiguration.xml");
+				WriteEnvironmentXml(environmentXml, logger);
+				temporaryInputFileLocations.Add(environmentXml);
+
 				FileReference[] inputFileLocations = temporaryInputFileLocations.ToArray();
 
 				// Get the path to the schema
@@ -776,6 +783,72 @@ namespace UnrealBuildTool
 			return true;
 		}
 
+		/// <summary>
+		/// Writes environment variables as a BuildConfiguration xml
+		/// </summary>
+		/// <param name="location">Path to write to</param>
+		/// <param name="logger">Logger for output</param>
+		/// <returns>True if the environment was written successfully</returns>
+		static bool WriteEnvironmentXml(FileReference location, ILogger logger)
+		{
+			// Read the environment variables
+			Dictionary<string, Dictionary<string, string>> contents = [];
+			foreach (DictionaryEntry EnvironmentVariable in Environment.GetEnvironmentVariables())
+			{
+				string key = EnvironmentVariable.Key?.ToString() ?? String.Empty;
+				string value = EnvironmentVariable.Value?.ToString() ?? String.Empty;
+				if (!key.StartsWith("UnrealBuildTool_") || !key.Contains("__"))
+				{
+					continue;
+				}
+				string[] split = key.Substring("UnrealBuildTool_".Length).Split("__");
+				string category = split.First();
+				string field = split.Last();
+				if (!contents.TryGetValue(category, out Dictionary<string, string>? fields))
+				{
+					fields = new();
+					contents.Add(category, fields);
+				}
+
+				fields.Add(field, value);
+			}
+			StringBuilder output = new StringBuilder();
+			output.AppendLine("<?xml version=\"1.0\" encoding=\"UTF-8\"?>");
+			XmlWriterSettings xmlSettings = new XmlWriterSettings();
+			xmlSettings.Encoding = new UTF8Encoding(false);
+			xmlSettings.Indent = true;
+			xmlSettings.OmitXmlDeclaration = true;
+			using (XmlWriter writer = XmlWriter.Create(output, xmlSettings))
+			{
+				XNamespace ns = XNamespace.Get("https://www.unrealengine.com/BuildConfiguration");
+				XElement root = new XElement(ns + "Configuration");
+				foreach (KeyValuePair<string, Dictionary<string, string>> CategoryItem in contents)
+				{
+					XElement categoryElement = new XElement(ns + CategoryItem.Key);
+					foreach (KeyValuePair<string, string> FieldItem in CategoryItem.Value)
+					{
+						XElement fieldElement = new XElement(ns + FieldItem.Key);
+						if (FieldItem.Value.Contains(';'))
+						{
+							foreach (string item in FieldItem.Value.Split(';', StringSplitOptions.RemoveEmptyEntries | StringSplitOptions.TrimEntries))
+							{
+								fieldElement.Add(new XElement(ns + "Item", item));
+							}
+						}
+						else
+						{
+							fieldElement.Add(FieldItem.Value);
+						}
+						categoryElement.Add(fieldElement);
+					}
+					root.Add(categoryElement);
+				}
+				new XDocument(root).Save(writer);
+			}
+			Utils.WriteFileIfChanged(location, output.ToString(), logger);
+			return true;
+		}
+
 		/// <summary>
 		/// Parse the value for a field from its text based representation in an XML file
 		/// </summary>
diff --git a/Engine/Source/Programs/UnrealBuildTool/UnrealBuildTool.csproj b/Engine/Source/Programs/UnrealBuildTool/UnrealBuildTool.csproj
--- a/Engine/Source/Programs/UnrealBuildTool/UnrealBuildTool.csproj
+++ b/Engine/Source/Programs/UnrealBuildTool/UnrealBuildTool.csproj
@@ -132,5 +132,6 @@
       <None Condition="Exists('$(AppData)\Unreal Engine\UnrealBuildTool\BuildConfiguration.xml')" Include="$(AppData)\Unreal Engine\UnrealBuildTool\BuildConfiguration.xml"><Link>BuildConfiguration\BuildConfiguration.xml (Global AppData)</Link></None>
       <None Condition="Exists('$(LocalAppData)\Unreal Engine\UnrealBuildTool\BuildConfiguration.xml')" Include="$(LocalAppData)\Unreal Engine\UnrealBuildTool\BuildConfiguration.xml"><Link>BuildConfiguration\BuildConfiguration.xml (Global LocalAppData)</Link></None>
       <None Condition="Exists('$(UserProfile)\Documents\Unreal Engine\UnrealBuildTool\BuildConfiguration.xml')" Include="$(UserProfile)\Documents\Unreal Engine\UnrealBuildTool\BuildConfiguration.xml"><Link>BuildConfiguration\BuildConfiguration.xml (Global Documents)</Link></None>
+      <None Condition="Exists('..\..\..\Intermediate\Build\UnrealBuildTool.Env.BuildConfiguration.xml')" Include="..\..\..\Intermediate\Build\UnrealBuildTool.Env.BuildConfiguration.xml"><Link>BuildConfiguration\BuildConfiguration.xml (Environment)</Link></None>
   </ItemGroup>
 </Project>
