Add Android's CompileSDK as UPL global variable and move CompileSDK discovery earlier so it is available to nodes that may affect GameActivity template.
This also fixes the sdk level used during GenerateManifest.

diff --git a/Engine/Source/Programs/UnrealBuildTool/Platform/Android/UEDeployAndroid.cs b/Engine/Source/Programs/UnrealBuildTool/Platform/Android/UEDeployAndroid.cs
--- a/Engine/Source/Programs/UnrealBuildTool/Platform/Android/UEDeployAndroid.cs
+++ b/Engine/Source/Programs/UnrealBuildTool/Platform/Android/UEDeployAndroid.cs
@@ -3198,6 +3198,7 @@ namespace UnrealBuildTool
 				throw new BuildException("AndroidManifest.xml is invalid {0}\n{1}", e, Text.ToString());
 			}
 
+			UPL!.SetGlobalContextVariable("CompileSDK", SDKLevelInt.ToString());
 			UPL!.ProcessPluginNode(Arch, "androidManifestUpdates", "", ref XDoc);
 			return XDoc.ToString();
 		}
@@ -4284,7 +4285,6 @@ namespace UnrealBuildTool
 
 			// do this here so we'll stop early if there is a problem with the SDK API level (cached so later calls will return the same)
 			string SDKAPILevel = GetSdkApiLevel(ToolChain);
-			int SDKLevelInt = GetApiLevelInt(SDKAPILevel);
 			string BuildToolsVersion = GetBuildToolsVersion();
 
 			// cache some tools paths
@@ -4585,20 +4585,61 @@ namespace UnrealBuildTool
 			UPL.SetGlobalContextVariable("AndroidPackageName", PackageName);
 			UPL.SetGlobalContextVariable("CookFlavor", CookFlavor);
 
+			Dictionary<string, int> NDKArchToCompileSDK = new Dictionary<string, int>();
+			foreach (UnrealArch Arch in Architectures.Architectures)
+			{
+				string NDKArch = GetNDKArch(Arch);
+				// check if any plugins want to increase the required compile SDK version
+				string CompileSDKVersion = SDKAPILevel.Replace("android-", "");
+				string CompileSDKMin = UPL.ProcessPluginNode(NDKArch, "minimumSDKAPI", "");
+				int CompileSDKVersionInt = MinimumSDKLevel;
+				Int32.TryParse(CompileSDKVersion, out CompileSDKVersionInt);
+				if (!String.IsNullOrEmpty(CompileSDKMin))
+				{
+					if (!Int32.TryParse(CompileSDKVersion, out CompileSDKVersionInt))
+					{
+						CompileSDKVersionInt = MinimumSDKLevel;
+					}
+					bool bUpdatedCompileSDK = false;
+					string[] CompileSDKLines = CompileSDKMin.Split(new[] { "\r\n", "\r", "\n" }, StringSplitOptions.None);
+					foreach (string CompileLine in CompileSDKLines)
+					{
+						//string VersionString = CompileLine.Replace("android-", "");
+						int VersionInt;
+						if (Int32.TryParse(CompileLine, out VersionInt))
+						{
+							if (VersionInt > CompileSDKVersionInt)
+							{
+								CompileSDKVersionInt = VersionInt;
+								bUpdatedCompileSDK = true;
+							}
+						}
+					}
+					if (bUpdatedCompileSDK)
+					{
+						CompileSDKVersion = CompileSDKVersionInt.ToString();
+						Logger.LogInformation("Building Java with SDK API Level 'android-{CompileSDKVersion}' for arch {NDKArch} due to enabled plugin requirements", CompileSDKVersion, NDKArch);
+					}
+				}
+				NDKArchToCompileSDK.Add(NDKArch, CompileSDKVersionInt);
+			}
+
 			IEnumerable<Tuple<UnrealArch, string>>? BuildList = null;
 
 			bool bRequiresOBB = RequiresOBB(bDisallowPackagingDataInApk, ObbFileLocation);
 			if (!bBuildSettingsMatch)
 			{
 				BuildList = from Arch in Architectures.Architectures
-							let manifest = GenerateManifest(ToolChain, ProjectName, InTargetType, EngineDirectory, bForDistribution, bPackageDataInsideApk, GameBuildFilesPath, bRequiresOBB, bDisableVerifyOBBOnStartUp, Arch, CookFlavor, bUseExternalFilesDir, Configuration.ToString(), SDKLevelInt, bSkipGradleBuild, bEnableBundle)
+							let CompileSDKLevelForArch = NDKArchToCompileSDK[GetNDKArch(Arch)]
+							let manifest = GenerateManifest(ToolChain, ProjectName, InTargetType, EngineDirectory, bForDistribution, bPackageDataInsideApk, GameBuildFilesPath, bRequiresOBB, bDisableVerifyOBBOnStartUp, Arch, CookFlavor, bUseExternalFilesDir, Configuration.ToString(), CompileSDKLevelForArch, bSkipGradleBuild, bEnableBundle)
 							select Tuple.Create(Arch, manifest);
 			}
 			else
 			{
 				BuildList = from Arch in Architectures.Architectures
+							let CompileSDKLevelForArch = NDKArchToCompileSDK[GetNDKArch(Arch)]
 							let manifestFile = Path.Combine(IntermediateAndroidPath, Arch + "_AndroidManifest.xml")
-							let manifest = GenerateManifest(ToolChain, ProjectName, InTargetType, EngineDirectory, bForDistribution, bPackageDataInsideApk, GameBuildFilesPath, bRequiresOBB, bDisableVerifyOBBOnStartUp, Arch, CookFlavor, bUseExternalFilesDir, Configuration.ToString(), SDKLevelInt, bSkipGradleBuild, bEnableBundle)
+							let manifest = GenerateManifest(ToolChain, ProjectName, InTargetType, EngineDirectory, bForDistribution, bPackageDataInsideApk, GameBuildFilesPath, bRequiresOBB, bDisableVerifyOBBOnStartUp, Arch, CookFlavor, bUseExternalFilesDir, Configuration.ToString(), CompileSDKLevelForArch, bSkipGradleBuild, bEnableBundle)
 							let OldManifest = File.Exists(manifestFile) ? File.ReadAllText(manifestFile) : ""
 							where manifest != OldManifest
 							select Tuple.Create(Arch, manifest);
@@ -4913,6 +4954,11 @@ namespace UnrealBuildTool
 				string Manifest = build.Item2;
 				string NDKArch = GetNDKArch(Arch);
 
+				int CompileSDKLevelForArch = NDKArchToCompileSDK[NDKArch];
+				string CompileSDKVersion = CompileSDKLevelForArch.ToString();
+
+				UPL.SetGlobalContextVariable("CompileSDK", CompileSDKVersion);
+
 				Logger.LogInformation("\n===={Time}====PREPARING NATIVE CODE====={Arch}============================================================", DateTime.Now.ToString(), Arch);
 
 				string UnrealBuildPath = Path.Combine(IntermediateAndroidPath, Arch.ToString());
@@ -5007,7 +5053,6 @@ namespace UnrealBuildTool
 				// use Gradle for compile/package
 				string UnrealBuildGradleAppPath = Path.Combine(UnrealBuildGradlePath, "app");
 				string UnrealBuildGradleMainPath = Path.Combine(UnrealBuildGradleAppPath, "src", "main");
-				string CompileSDKVersion = SDKAPILevel.Replace("android-", "");
 
 				// Write the manifest to the correct locations (cache and real)
 				string ManifestFile = Path.Combine(IntermediateAndroidPath, Arch + "_AndroidManifest.xml");
@@ -5206,39 +5251,6 @@ popd
 
 				Logger.LogInformation("\n===={Time}====PERFORMING FINAL APK PACKAGE OPERATION====={Arch}===========================================", DateTime.Now.ToString(), Arch);
 
-				// check if any plugins want to increase the required compile SDK version
-				string CompileSDKMin = UPL.ProcessPluginNode(NDKArch, "minimumSDKAPI", "");
-				if (!String.IsNullOrEmpty(CompileSDKMin))
-				{
-					int CompileSDKVersionInt;
-					if (!Int32.TryParse(CompileSDKVersion, out CompileSDKVersionInt))
-					{
-						CompileSDKVersionInt = MinimumSDKLevel;
-					}
-
-					bool bUpdatedCompileSDK = false;
-					string[] CompileSDKLines = CompileSDKMin.Split(new[] { "\r\n", "\r", "\n" }, StringSplitOptions.None);
-					foreach (string CompileLine in CompileSDKLines)
-					{
-						//string VersionString = CompileLine.Replace("android-", "");
-						int VersionInt;
-						if (Int32.TryParse(CompileLine, out VersionInt))
-						{
-							if (VersionInt > CompileSDKVersionInt)
-							{
-								CompileSDKVersionInt = VersionInt;
-								bUpdatedCompileSDK = true;
-							}
-						}
-					}
-
-					if (bUpdatedCompileSDK)
-					{
-						CompileSDKVersion = CompileSDKVersionInt.ToString();
-						Logger.LogInformation("Building Java with SDK API Level 'android-{CompileSDKVersion}' due to enabled plugin requirements", CompileSDKVersion);
-					}
-				}
-
 				// stage files into gradle app directory
 				string GradleManifest = Path.Combine(UnrealBuildGradleMainPath, "AndroidManifest.xml");
 				MakeDirectoryIfRequired(GradleManifest);
