[UBT]
* Fixed so hot reloading supports vfs
* Changed some file names for hot reloading to play nicer with UBA (.dep.json for dependency files on msvc toolchain)
* Added vfs information in LiveCoding manifest file

diff --git a/Engine/Source/Programs/UnrealBuildTool/System/HotReload.cs b/Engine/Source/Programs/UnrealBuildTool/System/HotReload.cs
--- a/Engine/Source/Programs/UnrealBuildTool/System/HotReload.cs
+++ b/Engine/Source/Programs/UnrealBuildTool/System/HotReload.cs
@@ -901,11 +901,9 @@ namespace UnrealBuildTool
 		/// <param name="Logger"></param>
 		public static void PatchActionGraphForLiveCoding(IEnumerable<LinkedAction> Actions, Dictionary<FileReference, FileReference> OriginalFileToPatchedFile, HotReloadMode hotReloadMode, ILogger Logger)
 		{
+			string extraExt = hotReloadMode == HotReloadMode.LiveCoding ? ".lc" : ".lcpt";
 			string dependencyFileExtension = hotReloadMode == HotReloadMode.LiveCoding ? ".lc.response" : ".lcpt.response";
-			string responseFileExtension = hotReloadMode == HotReloadMode.LiveCoding ? ".lc" : ".lcpt";
 			string objectFileExtension = hotReloadMode == HotReloadMode.LiveCoding ? ".lc.obj" : ".lcpt.obj";
-			string clSourceDepFileExtension = hotReloadMode == HotReloadMode.LiveCoding ? ".lc.json" : ".lcpt.json";
-			string clangSourceDepFileExtension = hotReloadMode == HotReloadMode.LiveCoding ? ".lc.d" : ".lcpt.d";
 
 			foreach (LinkedAction Action in Actions)
 			{
@@ -977,7 +975,7 @@ namespace UnrealBuildTool
 					}
 
 					FileReference OldResponseFile = Action.RootPaths.GetLocalPath(new FileReference(Arguments[ResponseFileIdx].Substring(1).Trim('\"')));
-					FileReference NewResponseFile = new FileReference(OldResponseFile.FullName + responseFileExtension);
+					FileReference NewResponseFile = OldResponseFile.ChangeExtension(".lc.rsp");
 
 					NewAction.PrerequisiteItems.Remove(FileItem.GetItemByFileReference(OldResponseFile));
 					NewAction.PrerequisiteItems.Add(FileItem.GetItemByFileReference(NewResponseFile));
@@ -1000,7 +998,13 @@ namespace UnrealBuildTool
 
 							OriginalFileToPatchedFile[OldOutputFile] = NewOutputFile;
 
-							ResponseLines[Idx] = OutputFilePrefix + "\"" + NewOutputFile.FullName + "\"";
+							string NewOutputPath = NewOutputFile.FullName;
+							if (Action.RootPaths.GetVfsOverlayPath(NewOutputFile, out string? virtualPath))
+							{
+								NewOutputPath = virtualPath;
+							}
+
+							ResponseLines[Idx] = OutputFilePrefix + "\"" + NewOutputPath + "\"";
 							break;
 						}
 					}
@@ -1011,7 +1015,8 @@ namespace UnrealBuildTool
 						|| Action.CommandPath.GetFileName().Equals("verse-clang-cl.exe", StringComparison.OrdinalIgnoreCase))
 					{
 						string SourceDependencyPrefix = Action.CommandPath.GetFileName().Equals("cl.exe", StringComparison.OrdinalIgnoreCase) ? "/sourceDependencies" : "/clang:-MD /clang:-MF";
-						string NewExtension = Action.CommandPath.GetFileName().Equals("cl.exe", StringComparison.OrdinalIgnoreCase) ? clSourceDepFileExtension : clangSourceDepFileExtension;
+						string extension = Action.CommandPath.GetFileName().Equals("cl.exe", StringComparison.OrdinalIgnoreCase) ? ".dep.json" : ".d";
+
 						for (int Idx = 0; Idx < ResponseLines.Length; Idx++)
 						{
 							string ResponseLine = ResponseLines[Idx];
@@ -1021,12 +1026,15 @@ namespace UnrealBuildTool
 								FileItem OldSourceDependencyFileItem = Action.ProducedItems.First(x => x.Location == OldSourceDependencyFile);
 								NewAction.ProducedItems.Remove(OldSourceDependencyFileItem);
 
-								FileReference NewSourceDependencyFile = OldSourceDependencyFile.ChangeExtension(NewExtension);
+								FileReference NewSourceDependencyFile = FileReference.FromString(OldSourceDependencyFile.FullName.Substring(0, OldSourceDependencyFile.FullName.Length - extension.Length) + extraExt + extension);
 								FileItem NewSourceDependencyFileItem = FileItem.GetItemByFileReference(NewSourceDependencyFile);
 								NewAction.ProducedItems.Add(NewSourceDependencyFileItem);
 								NewAction.DependencyListFile = NewSourceDependencyFileItem;
 
-								ResponseLines[Idx] = SourceDependencyPrefix + "\"" + NewSourceDependencyFile.FullName + "\"";
+								string NewSourceDependencyPath = NewSourceDependencyFile.FullName;
+								if (Action.RootPaths.GetVfsOverlayPath(NewSourceDependencyFile, out string? VirtualPath))
+									NewSourceDependencyPath = VirtualPath;
+								ResponseLines[Idx] = SourceDependencyPrefix + "\"" + NewSourceDependencyPath + "\"";
 								break;
 							}
 						}
@@ -1439,16 +1447,32 @@ namespace UnrealBuildTool
 				}
 				Writer.WriteObjectEnd();
 
+				HashSet<DirectoryReference> vfsPathsAdded = new();
+				List<(DirectoryReference virtualPath, DirectoryReference localPath)> vfsPaths = new();
+
 				Writer.WriteArrayStart("Modules");
 				foreach (IExternalAction Action in Actions)
 				{
 					if (Action.ActionType == ActionType.Link)
 					{
+						string GetPath(FileReference fileRef)
+						{
+							if (Action.RootPaths.GetVfsOverlayPath(fileRef, out string? VirtualPath))
+							{
+								return VirtualPath.Replace('/', '\\');
+							}
+							else
+							{
+								return fileRef.FullName;
+							}
+						}
+
+
 						FileItem? OutputFile = Action.ProducedItems.FirstOrDefault(x => x.HasExtension(".exe") || x.HasExtension(".dll"));
 						if (OutputFile != null && Action.PrerequisiteItems.Any(x => OriginalFileToPatchedFile.ContainsKey(x.Location)))
 						{
 							Writer.WriteObjectStart();
-							Writer.WriteValue("Output", OutputFile.Location.FullName);
+							Writer.WriteValue("Output", GetPath(OutputFile.Location));
 
 							Writer.WriteArrayStart("Inputs");
 							foreach (FileItem InputFile in Action.PrerequisiteItems)
@@ -1456,7 +1480,7 @@ namespace UnrealBuildTool
 								FileReference? PatchedFile;
 								if (OriginalFileToPatchedFile.TryGetValue(InputFile.Location, out PatchedFile))
 								{
-									Writer.WriteValue(PatchedFile.FullName);
+									Writer.WriteValue(GetPath(PatchedFile));
 								}
 							}
 							Writer.WriteArrayEnd();
@@ -1466,17 +1490,39 @@ namespace UnrealBuildTool
 							{
 								if (InputFile.HasExtension(".lib"))
 								{
-									Writer.WriteValue(InputFile.FullName);
+									Writer.WriteValue(GetPath(InputFile.Location));
 								}
 							}
 							Writer.WriteArrayEnd();
 
 							Writer.WriteObjectEnd();
+
+							if (Action.RootPaths.bUseVfs)
+							{
+								foreach (var RootPath in Action.RootPaths)
+								{
+									if (vfsPathsAdded.Add(RootPath.vfs))
+									{
+										vfsPaths.Add((RootPath.vfs, RootPath.local));
+									}
+								}
+							}
 						}
 					}
 				}
 				Writer.WriteArrayEnd();
 
+				Writer.WriteArrayStart("Vfs");
+				foreach (var kv in vfsPaths)
+				{
+					Writer.WriteObjectStart();
+					Writer.WriteValue("Virtual", kv.virtualPath.FullName);
+					Writer.WriteValue("Local", kv.localPath.FullName);
+					Writer.WriteObjectEnd();
+				}
+				Writer.WriteArrayEnd();
+
+
 				Writer.WriteObjectEnd();
 			}
 		}
