LLT fixes:
- Android linker error where ldscript cannot be found
- PlainProps deactivated flag reflected in xml
- Remove SetThreadPriority from TestRunner which causes PS4 tests to hang
- TDS helpers tests doesn't exist

#rnx

diff --git a/Engine/Build/LowLevelTests/PlainProps.xml b/Engine/Build/LowLevelTests/PlainProps.xml
--- a/Engine/Build/LowLevelTests/PlainProps.xml
+++ b/Engine/Build/LowLevelTests/PlainProps.xml
@@ -3,6 +3,6 @@
   <Option Name="RunPlainPropsTests" DefaultValue="" Description="Run PlainProps Tests" />
   <Property Name="TestNames" Value="$(TestNames);PlainProps" />
   <Extend Name="RunAllTests">
-    <Expand Name="DeployAndTest" Platform="Win64" TestName="PlainProps" ShortName="PlainProps" TargetName="PlainPropsTests" BinaryRelativePath="Engine\Binaries" ReportType="xml" ExtraCompilationArgs="-allmodules" Deactivated="True" />
+    <Expand Name="DeployAndTest" Platform="Win64" Deactivated="True" TestName="PlainProps" ShortName="PlainProps" TargetName="PlainPropsTests" BinaryRelativePath="Engine\Binaries" ReportType="xml" ExtraCompilationArgs="-allmodules" />
   </Extend>
 </BuildGraph>
\ No newline at end of file
diff --git a/Engine/Source/Developer/LowLevelTestsRunner/Private/TestRunner.cpp b/Engine/Source/Developer/LowLevelTestsRunner/Private/TestRunner.cpp
--- a/Engine/Source/Developer/LowLevelTestsRunner/Private/TestRunner.cpp
+++ b/Engine/Source/Developer/LowLevelTestsRunner/Private/TestRunner.cpp
@@ -211,7 +211,6 @@ void FTestRunner::GlobalSetup()
 
 	// Finish setting up the Game Thread, which requires the command line.
 	FPlatformProcess::SetThreadAffinityMask(FPlatformAffinity::GetMainGameMask());
-	FPlatformProcess::SetThreadPriority(FPlatformAffinity::GetGameThreadPriority());
 	FPlatformProcess::SetupGameThread();
 
 	// Always set up GError to handle FatalError, failed assertions, and crashes and other fatal errors.
diff --git a/Engine/Source/Programs/UnrealBuildTool/Configuration/TestTargetRules.cs b/Engine/Source/Programs/UnrealBuildTool/Configuration/TestTargetRules.cs
--- a/Engine/Source/Programs/UnrealBuildTool/Configuration/TestTargetRules.cs
+++ b/Engine/Source/Programs/UnrealBuildTool/Configuration/TestTargetRules.cs
@@ -188,21 +188,26 @@ namespace UnrealBuildTool
 				}
 				if (File != null)
 				{
-					DirectoryReference outputDirectory;
-					if (ProjectFile != null && File.IsUnderDirectory(ProjectFile.Directory))
-					{
-						outputDirectory = UEBuildTarget.GetOutputDirectoryForExecutable(ProjectFile.Directory, File);
-					}
-					else
-					{
-						outputDirectory = UEBuildTarget.GetOutputDirectoryForExecutable(Unreal.EngineDirectory, File);
-					}
-					FileReference outputTestFile = FileReference.Combine(outputDirectory, "Binaries", target.Platform.ToString(), ExeBinariesSubFolder, $"{outputName}.exe.is_unreal_test");
+					FileReference outputTestFile = FileReference.Combine(GetOutputDirectory(), "Binaries", target.Platform.ToString(), ExeBinariesSubFolder, $"{outputName}.exe.is_unreal_test");
 					PostBuildSteps.Add("echo > " + outputTestFile.FullName);
 				}
 			}
 		}
 
+		private DirectoryReference GetOutputDirectory()
+		{
+			DirectoryReference outputDirectory;
+			if (ProjectFile != null && File!.IsUnderDirectory(ProjectFile.Directory))
+			{
+				outputDirectory = UEBuildTarget.GetOutputDirectoryForExecutable(ProjectFile.Directory, File);
+			}
+			else
+			{
+				outputDirectory = UEBuildTarget.GetOutputDirectoryForExecutable(Unreal.EngineDirectory, File!);
+			}
+			return outputDirectory;
+		}
+
 		/// <summary>
 		/// Constructs a valid TestTargetRules instance from an existent TargetRules instance.
 		/// </summary>
@@ -342,7 +347,8 @@ namespace UnrealBuildTool
 
 				// Addresses linker errors when building LowLevelTests for Android
 				// TODO: This should be written to the intermediate directory, somewhere else not when TargetRules are being created
-				FileReference versionScriptFile = new FileReference(Path.GetTempPath() + $"LLTWorkaroundScrip-{Name}.ldscript");
+				string IntermediatePath = Path.Combine(Unreal.EngineDirectory.FullName, "Intermediate", "Android");
+				FileReference versionScriptFile = new FileReference(Path.Combine(IntermediatePath, $"LLTWorkaroundScript-{Name}.ldscript"));
 				string SymbolsLinkScriptContent = @"{
 					global: LowLevelTestsMain;
 					global: JNI_OnLoad;
@@ -352,8 +358,12 @@ namespace UnrealBuildTool
 					*Catch*;
 					local: *;
 				};";
+				if (!Directory.Exists(IntermediatePath))
+				{
+					Directory.CreateDirectory(IntermediatePath);
+				}
 				FileReference.WriteAllTextIfDifferent(versionScriptFile, SymbolsLinkScriptContent);
-				AdditionalLinkerArguments = " -Wl,--version-script=\"" + versionScriptFile.FullName + "\"";
+				AdditionalLinkerArguments = " -Wl,--version-script=\"" + versionScriptFile + "\"";
 			}
 			else if (target.Platform == UnrealTargetPlatform.IOS)
 			{
