[UBT/BuildSettings]
* Added so virtual->local path lookup is available in BuildSettings

diff --git a/Engine/Source/Programs/UnrealBuildTool/Configuration/UEBuildModuleCPP.cs b/Engine/Source/Programs/UnrealBuildTool/Configuration/UEBuildModuleCPP.cs
--- a/Engine/Source/Programs/UnrealBuildTool/Configuration/UEBuildModuleCPP.cs
+++ b/Engine/Source/Programs/UnrealBuildTool/Configuration/UEBuildModuleCPP.cs
@@ -2464,6 +2464,14 @@ namespace UnrealBuildTool
 			if (Rules.Name == "BuildSettings")
 			{
 				Result.Definitions.Add(String.Format("UE_WITH_DEBUG_INFO={0}", BaseCompileEnvironment.bCreateDebugInfo ? "1" : "0"));
+
+				StringBuilder sb = new();
+				foreach (var (id, vfs, local) in BaseCompileEnvironment.RootPaths)
+				{
+					sb.Append(vfs.FullName.Replace("\\", "\\\\")).Append(';');
+					sb.Append(local.FullName.Replace("\\", "\\\\")).Append(';');
+				}
+				Result.Definitions.Add($"UE_VFS_PATHS=\"{sb}\"");
 			}
 
 			// Add the project definitions
diff --git a/Engine/Source/Runtime/BuildSettings/Private/BuildSettings.cpp b/Engine/Source/Runtime/BuildSettings/Private/BuildSettings.cpp
--- a/Engine/Source/Runtime/BuildSettings/Private/BuildSettings.cpp
+++ b/Engine/Source/Runtime/BuildSettings/Private/BuildSettings.cpp
@@ -111,4 +111,9 @@ namespace BuildSettings
 	{
 		return UE_PERSISTENT_ALLOCATOR_RESERVE_SIZE;
 	}
+
+	const char* GetVfsPaths()
+	{
+		return UE_VFS_PATHS;
+	}
 }
diff --git a/Engine/Source/Runtime/BuildSettings/Public/BuildSettings.h b/Engine/Source/Runtime/BuildSettings/Public/BuildSettings.h
--- a/Engine/Source/Runtime/BuildSettings/Public/BuildSettings.h
+++ b/Engine/Source/Runtime/BuildSettings/Public/BuildSettings.h
@@ -138,4 +138,8 @@ namespace BuildSettings
 	 */
 	BUILDSETTINGS_API uint64 GetPersistentAllocatorReserveSize();
 
+	/**
+	 * Semicolon separated string containing virtual;local;virtual;local; paths. Ends with semicolon
+	 */
+	BUILDSETTINGS_API const char* GetVfsPaths();
 }
