UnrealBuildTool: Replace expensive list search with a dictionary key lookup

#rnx

diff --git a/Engine/Source/Programs/UnrealBuildTool/ProjectFiles/ProjectFile.cs b/Engine/Source/Programs/UnrealBuildTool/ProjectFiles/ProjectFile.cs
--- a/Engine/Source/Programs/UnrealBuildTool/ProjectFiles/ProjectFile.cs
+++ b/Engine/Source/Programs/UnrealBuildTool/ProjectFiles/ProjectFile.cs
@@ -156,6 +156,15 @@ namespace UnrealBuildTool
 		/// All of the targets in this project.  All non-stub projects must have at least one target.
 		public readonly List<Project> ProjectTargets = new List<Project>();
 
+		/// <summary>
+		/// Checks if a file is already contained
+		/// </summary>
+		/// <param name="File">Files to check</param>
+		public bool Contains(FileReference File)
+		{
+			return SourceFileMap.ContainsKey(File);
+		}
+
 		/// <summary>
 		/// Adds a list of files to this project, ignoring dupes
 		/// </summary>
