UnrealBuildTool: Split up the steps of projectfile intellisense gen so we can reintroduce parallelization

#rnx

diff --git a/Engine/Source/Programs/UnrealBuildTool/ProjectFiles/ProjectFileGenerator.cs b/Engine/Source/Programs/UnrealBuildTool/ProjectFiles/ProjectFileGenerator.cs
--- a/Engine/Source/Programs/UnrealBuildTool/ProjectFiles/ProjectFileGenerator.cs
+++ b/Engine/Source/Programs/UnrealBuildTool/ProjectFiles/ProjectFileGenerator.cs
@@ -1,12 +1,14 @@
 // Copyright Epic Games, Inc. All Rights Reserved.
 
 using System;
+using System.Collections.Concurrent;
 using System.Collections.Generic;
 using System.Diagnostics;
 using System.IO;
 using System.Linq;
 using System.Text;
 using System.Threading;
+using System.Threading.Tasks;
 using System.Xml.Linq;
 using EpicGames.Core;
 using Microsoft.Extensions.Logging;
@@ -2139,17 +2141,14 @@ namespace UnrealBuildTool
 		{
 			if (ShouldGenerateIntelliSenseData() && Targets.Any())
 			{
-				string ProgressInfoText = OperatingSystem.IsWindows() ? "Binding IntelliSense data..." : "Generating data for project indexing...";
-				using (ProgressWriter Progress = new ProgressWriter(ProgressInfoText, true, Logger))
+				// Ignore projects for platforms we can't build for
+				List<Tuple<ProjectFile, ProjectTarget>> SupportedTargets = Targets.Where(x => x.Item2.SupportedPlatforms.Any(x => GetIntelliSensePlatforms().Contains(x))).ToList();
+
+				using (ProgressWriter Progress = new ProgressWriter("Compiling Rules Assemblies", true, Logger))
 				{
-					// Ignore projects for platforms we can't build for
-					List<Tuple<ProjectFile, ProjectTarget>> SupportedTargets = Targets.Where(x => x.Item2.SupportedPlatforms.Any(x => GetIntelliSensePlatforms().Contains(x))).ToList();
-					int NumTargets = SupportedTargets.Count;
-					int NumTasks = NumTargets;
 					int TasksFinished = 0;
 					foreach (Tuple<ProjectFile, ProjectTarget> Item in SupportedTargets)
 					{
-						ProjectFile TargetProjectFile = Item.Item1;
 						ProjectTarget CurTarget = Item.Item2;
 
 						Logger.LogDebug("Found target: {Target}", CurTarget.Name);
@@ -2161,6 +2160,42 @@ namespace UnrealBuildTool
 						}
 						NewArguments.AddRange(Arguments);
 
+						// Get the architecture from the target platform
+						UnrealArchitectures DefaultArchitecture = UnrealArchitectureConfig.ForPlatform(BuildHostPlatform.Current.Platform).ActiveArchitectures(CurTarget.UnrealProjectFilePath, CurTarget.Name);
+
+						// Create the target descriptor
+						TargetDescriptor TargetDesc = new TargetDescriptor(CurTarget.UnrealProjectFilePath, CurTarget.Name, BuildHostPlatform.Current.Platform, UnrealTargetConfiguration.Development, DefaultArchitecture, new CommandLineArguments(NewArguments.ToArray()));
+						TargetDesc.IntermediateEnvironment = UnrealIntermediateEnvironment.GenerateProjectFiles;
+
+						// Create the target rules assembly 
+						RulesCompiler.CreateTargetRulesAssembly(CurTarget.UnrealProjectFilePath, CurTarget.Name, false, false, bUsePrecompiled, TargetDesc.ForeignPlugin, TargetDesc.bBuildPluginAsLocal, Logger);
+
+						lock (Progress)
+						{
+							Interlocked.Increment(ref TasksFinished);
+							Progress.Write(TasksFinished, SupportedTargets.Count);
+						}
+					}
+				}
+
+				ConcurrentBag<UEBuildTarget> BuildTargets = [];
+				using (ProgressWriter Progress = new ProgressWriter("Creating Build Targets", true, Logger))
+				{
+					int NumTargets = SupportedTargets.Count;
+					int NumTasks = NumTargets;
+					int TasksFinished = 0;
+
+					Parallel.ForEach(SupportedTargets, Item =>
+					{
+						ProjectTarget CurTarget = Item.Item2;
+
+						List<string> NewArguments = new List<string>(Arguments.Length + 4);
+						if (CurTarget.TargetRules!.Type != TargetType.Program)
+						{
+							NewArguments.Add("-precompile");
+						}
+						NewArguments.AddRange(Arguments);
+
 						try
 						{
 							foreach (UnrealTargetPlatform IntellisensePlatform in GetIntelliSensePlatforms().Where(x => CurTarget.SupportedPlatforms.Contains(x)))
@@ -2173,54 +2208,100 @@ namespace UnrealBuildTool
 								TargetDesc.IntermediateEnvironment = UnrealIntermediateEnvironment.GenerateProjectFiles;
 
 								// Create the target
-								UEBuildTarget Target = UEBuildTarget.Create(TargetDesc, false, false, bUsePrecompiled, UnrealIntermediateEnvironment.GenerateProjectFiles, Logger);
+								UEBuildTarget Target = UEBuildTarget.Create(TargetDesc, true, false, bUsePrecompiled, UnrealIntermediateEnvironment.GenerateProjectFiles, Logger);
+								BuildTargets.Add(Target);
+							}
+						}
+						catch (Exception Ex)
+						{
+							Logger.LogWarning("Exception while creating build target for {Target}: {Ex}", CurTarget.Name, Ex.ToString());
+						}
+						
+						lock (Progress)
+						{
+							Interlocked.Increment(ref TasksFinished);
+							Progress.Write(TasksFinished, SupportedTargets.Count);
+						}
+					});
+				}
 
-								AddTargetForIntellisense(Target, Logger);
+				string ProgressInfoText = OperatingSystem.IsWindows() ? "Binding IntelliSense data..." : "Generating data for project indexing...";
+				using (ProgressWriter Progress = new ProgressWriter(ProgressInfoText, true, Logger))
+				{
+					int TasksFinished = 0;
+					int TotalTasks = BuildTargets.Count;
+					if (ShouldGenerateIntelliSenseCompileEnvironments())
+					{
+						TotalTasks += BuildTargets.SelectMany(x => x.Binaries).SelectMany(x => x.Modules.OfType<UEBuildModuleCPP>()).Count();
+					}
+
+					void Gen(UEBuildTarget Target)
+					{
+						AddTargetForIntellisense(Target, Logger);
 
-								// If the project generator just cares about the result of UEBuildTarget.Create, skip generating the compile environments.
-								if (ShouldGenerateIntelliSenseCompileEnvironments())
+						// If the project generator just cares about the result of UEBuildTarget.Create, skip generating the compile environments.
+						if (ShouldGenerateIntelliSenseCompileEnvironments())
+						{
+							// Generate a compile environment for each module in the binary
+							CppCompileEnvironment GlobalCompileEnvironment = Target.CreateCompileEnvironmentForProjectFiles(Logger);
+							foreach (UEBuildBinary Binary in Target.Binaries)
+							{
+								CppCompileEnvironment BinaryCompileEnvironment = Binary.CreateBinaryCompileEnvironment(GlobalCompileEnvironment);
+								foreach (UEBuildModuleCPP Module in Binary.Modules.OfType<UEBuildModuleCPP>())
 								{
-									// Generate a compile environment for each module in the binary
-									CppCompileEnvironment GlobalCompileEnvironment = Target.CreateCompileEnvironmentForProjectFiles(Logger);
-									foreach (UEBuildBinary Binary in Target.Binaries)
+									if (ModuleToEditorProjectFileMap.TryGetValue(Module.RulesFile, out ProjectFile? ProjectFileForIDE) && ProjectFileForIDE != null && ProjectFileForIDE.ProjectFilePath == Target.ProjectFile)
 									{
-										CppCompileEnvironment BinaryCompileEnvironment = Binary.CreateBinaryCompileEnvironment(GlobalCompileEnvironment);
-										foreach (UEBuildModuleCPP Module in Binary.Modules.OfType<UEBuildModuleCPP>())
+										CppCompileEnvironment ModuleCompileEnvironment = Module.CreateCompileEnvironmentForIntellisense(Target.Rules, BinaryCompileEnvironment, Logger);
+										lock (ProjectFileForIDE)
 										{
-											ProjectFile? ProjectFileForIDE;
-											if (ModuleToEditorProjectFileMap.TryGetValue(Module.RulesFile, out ProjectFileForIDE) && ProjectFileForIDE == TargetProjectFile)
-											{
-												CppCompileEnvironment ModuleCompileEnvironment = Module.CreateCompileEnvironmentForIntellisense(Target.Rules, BinaryCompileEnvironment, Logger);
-												lock (ProjectFileForIDE)
-												{
-													ProjectFileForIDE.AddModuleForIntelliSense(Module, ModuleCompileEnvironment);
-												}
-											}
+											ProjectFileForIDE.AddModuleForIntelliSense(Module, ModuleCompileEnvironment);
 										}
 									}
 
-									// If we're generating project files, then go ahead and wipe out the existing UBTMakefile for every target, to make sure that
-									// it gets a full dependency scan next time.
-									// NOTE: This is just a safeguard and doesn't have to be perfect.  We also check for newer project file timestamps in LoadUBTMakefile()
-									FileReference MakefileLocation = TargetMakefile.GetLocation(TargetDesc.ProjectFile, TargetDesc.Name, TargetDesc.Platform, TargetDesc.Architectures, TargetDesc.Configuration, TargetDesc.IntermediateEnvironment);
-									if (FileReference.Exists(MakefileLocation))
+									lock (Progress)
 									{
-										FileReference.Delete(MakefileLocation);
+										Interlocked.Increment(ref TasksFinished);
+										Progress.Write(TasksFinished, TotalTasks);
 									}
 								}
 							}
 						}
-						catch (Exception Ex)
-						{
-							Logger.LogWarning("Exception while generating include data for {Target}: {Ex}", CurTarget.Name, Ex.ToString());
-						}
 
 						lock (Progress)
 						{
 							Interlocked.Increment(ref TasksFinished);
-							Progress.Write(TasksFinished, NumTasks);
+							Progress.Write(TasksFinished, TotalTasks);
 						}
 					}
+
+					IEnumerable<UEBuildTarget> Shared = BuildTargets.Where(x => x.bUseSharedBuildEnvironment);
+					IEnumerable<UEBuildTarget> Unique = BuildTargets.Where(x => !x.bUseSharedBuildEnvironment);
+
+					IEnumerable<IGrouping<(UnrealTargetPlatform Platform, UnrealTargetConfiguration Configuration, TargetType TargetType), UEBuildTarget>> ParallelGroups = Shared.GroupBy(x => (x.Platform, x.Configuration, x.TargetType));
+
+					// Unique environments can be processed in parallel
+					// Shared environments can be processed in parallel, but each unique combination of Platform\Configuration\TargetType must be processed serially
+					// DebugGame can not be run in parallel with Development
+					// This could be made more async but it's complicated
+
+					Parallel.ForEach(Unique.Where(x => x.Configuration != UnrealTargetConfiguration.DebugGame), Target => Gen(Target));
+					Parallel.ForEach(Unique.Where(x => x.Configuration == UnrealTargetConfiguration.DebugGame), Target => Gen(Target));
+					
+					Parallel.ForEach(ParallelGroups.Where(x => x.Key.Configuration != UnrealTargetConfiguration.DebugGame), Group =>
+					{
+						foreach (UEBuildTarget Target in Group)
+						{
+							Gen(Target);
+						}
+					});
+
+					Parallel.ForEach(ParallelGroups.Where(x => x.Key.Configuration == UnrealTargetConfiguration.DebugGame), Group =>
+					{
+						foreach (UEBuildTarget Target in Group)
+						{
+							Gen(Target);
+						}
+					});
 				}
 			}
 		}
@@ -2477,7 +2558,7 @@ namespace UnrealBuildTool
 
 					foreach (ProjectFile aProjectFile in ProjectFiles)
 					{
-						if (aProjectFile.SourceFiles.Find(x => x.Reference == CurModuleFile) == null)
+						if (!aProjectFile.Contains(CurModuleFile))
 						{
 							// Skip if already added
 							aProjectFile.AddFilesToProject(FoundFiles, BaseFolder);
