[UBT]
* Added cache bucket property to Action in order to be able to have more control over which buckets are used for different builds

diff --git a/Engine/Source/Programs/UnrealBuildTool/Executors/UnrealBuildAccelerator/UBAExecutor.cs b/Engine/Source/Programs/UnrealBuildTool/Executors/UnrealBuildAccelerator/UBAExecutor.cs
--- a/Engine/Source/Programs/UnrealBuildTool/Executors/UnrealBuildAccelerator/UBAExecutor.cs
+++ b/Engine/Source/Programs/UnrealBuildTool/Executors/UnrealBuildAccelerator/UBAExecutor.cs
@@ -678,6 +678,11 @@ namespace UnrealBuildTool
 
 		static uint GetActionCacheBucket(LinkedAction action)
 		{
+			if (action.CacheBucket != 0)
+			{
+				return action.CacheBucket;
+			}
+
 			if (action.Target == null)
 			{
 				return 0;
diff --git a/Engine/Source/Programs/UnrealBuildTool/Platform/Windows/VCCompileAction.cs b/Engine/Source/Programs/UnrealBuildTool/Platform/Windows/VCCompileAction.cs
--- a/Engine/Source/Programs/UnrealBuildTool/Platform/Windows/VCCompileAction.cs
+++ b/Engine/Source/Programs/UnrealBuildTool/Platform/Windows/VCCompileAction.cs
@@ -218,6 +218,9 @@ namespace UnrealBuildTool
 		/// <inheritdoc/>
 		public double Weight { get; set; } = 1.0;
 
+		/// <inheritdoc/>
+		public uint CacheBucket { get; set; }
+
 		/// <inheritdoc/>
 		public bool bShouldOutputLog { get; set; } = true;
 		#endregion
@@ -422,6 +425,7 @@ namespace UnrealBuildTool
 			bIsAnalyzing = InAction.bIsAnalyzing;
 			bWarningsAsErrors = InAction.bWarningsAsErrors;
 			Weight = InAction.Weight;
+			CacheBucket = InAction.CacheBucket;
 
 			AdditionalPrerequisiteItems = new List<FileItem>(InAction.AdditionalPrerequisiteItems);
 			AdditionalProducedItems = new List<FileItem>(InAction.AdditionalProducedItems);
@@ -470,6 +474,7 @@ namespace UnrealBuildTool
 			bIsAnalyzing = Reader.ReadBool();
 			bWarningsAsErrors = Reader.ReadBool();
 			Weight = Reader.ReadDouble();
+			CacheBucket = Reader.ReadUnsignedInt();
 
 			AdditionalPrerequisiteItems = Reader.ReadList(() => Reader.ReadFileItem())!;
 			AdditionalProducedItems = Reader.ReadList(() => Reader.ReadFileItem())!;
@@ -515,6 +520,7 @@ namespace UnrealBuildTool
 			Writer.WriteBool(bIsAnalyzing);
 			Writer.WriteBool(bWarningsAsErrors);
 			Writer.WriteDouble(Weight);
+			Writer.WriteUnsignedInt(CacheBucket);
 
 			Writer.WriteList(AdditionalPrerequisiteItems, Item => Writer.WriteFileItem(Item));
 			Writer.WriteList(AdditionalProducedItems, Item => Writer.WriteFileItem(Item));
diff --git a/Engine/Source/Programs/UnrealBuildTool/Platform/Windows/VCToolChain.cs b/Engine/Source/Programs/UnrealBuildTool/Platform/Windows/VCToolChain.cs
--- a/Engine/Source/Programs/UnrealBuildTool/Platform/Windows/VCToolChain.cs
+++ b/Engine/Source/Programs/UnrealBuildTool/Platform/Windows/VCToolChain.cs
@@ -1852,6 +1852,10 @@ namespace UnrealBuildTool
 				BaseCompileAction.ToolChainVersion = $"{BaseCompileAction.ToolChainVersion} Compressed";
 			}
 
+			// Calculate cache bucket. We want to use as few buckets as possible per build but we also
+			// want to make sure that the buckets are not getting too big.
+			BaseCompileAction.CacheBucket = GetCacheBucket(Target, (!Target.bUseVFS && CompileEnvironment.PrecompiledHeaderAction != PrecompiledHeaderAction.None) ? CompileEnvironment.RootPaths : null);
+
 			return BaseCompileAction;
 		}
 
@@ -2406,6 +2410,7 @@ namespace UnrealBuildTool
 				CompileAction.PrerequisiteItems.UnionWith(CompileEnvironment.ForceIncludeFiles);
 				CompileAction.PrerequisiteItems.UnionWith(CompileEnvironment.AdditionalPrerequisites);
 				CompileAction.ArtifactMode = ArtifactMode.Enabled;
+				CompileAction.CacheBucket = GetCacheBucket(Target, null);
 
 				// Resource tool can run remotely if possible
 				CompileAction.bCanExecuteRemotely = true;
@@ -2615,6 +2620,7 @@ namespace UnrealBuildTool
 			CompileAction.bCanExecuteRemotely = false; // Incompatible with remote distribution
 			CompileAction.RootPaths = CompileEnvironment.RootPaths;
 			CompileAction.ArtifactMode = ArtifactMode.Enabled;
+			CompileAction.CacheBucket = GetCacheBucket(Target, null);
 		}
 
 		public override IEnumerable<string> GetGlobalCommandLineArgs(CppCompileEnvironment CompileEnvironment)
@@ -3056,6 +3062,8 @@ namespace UnrealBuildTool
 				LinkAction.CommandVersion = $"{LinkAction.CommandVersion} Compressed";
 			}
 
+			LinkAction.CacheBucket = GetCacheBucket(Target, null);
+
 			Logger.LogDebug("     Linking: {StatusDescription}", LinkAction.StatusDescription);
 			Logger.LogDebug("     Command: {CommandArguments}", LinkAction.CommandArguments);
 
diff --git a/Engine/Source/Programs/UnrealBuildTool/System/Action.cs b/Engine/Source/Programs/UnrealBuildTool/System/Action.cs
--- a/Engine/Source/Programs/UnrealBuildTool/System/Action.cs
+++ b/Engine/Source/Programs/UnrealBuildTool/System/Action.cs
@@ -4,6 +4,7 @@ using System;
 using System.Collections.Generic;
 using System.Diagnostics;
 using System.Linq;
+using System.Runtime.InteropServices.Marshalling;
 using System.Threading;
 using EpicGames.Core;
 using Microsoft.Extensions.Logging;
@@ -210,6 +211,13 @@ namespace UnrealBuildTool
 		/// Used to determine how much weight(CPU and Memory work) this action is.
 		/// </summary>
 		double Weight { get; }
+
+		/// <summary>
+		/// Used to determine which cache bucket should be used for this action
+		/// Using different cache buckets is an optimization where we want one build to use as few buckets as possible while 
+		/// two different builds should not use the same bucket if there is no chance their cache entries will be the same
+		/// </summary>
+		public uint CacheBucket { get; }
 	}
 
 	/// <summary>
@@ -353,6 +361,9 @@ namespace UnrealBuildTool
 		/// <inheritdoc/>
 		public double Weight { get; set; } = 1.0;
 
+		/// <inheritdoc/>
+		public uint CacheBucket { get; set; }
+
 		IEnumerable<FileItem> IExternalAction.PrerequisiteItems => PrerequisiteItems;
 		IEnumerable<FileItem> IExternalAction.ProducedItems => ProducedItems;
 		IEnumerable<FileItem> IExternalAction.DeleteItems => DeleteItems;
@@ -396,6 +407,7 @@ namespace UnrealBuildTool
 			bUseActionHistory = InOther.bUseActionHistory;
 			bIsHighPriority = InOther.bIsHighPriority;
 			Weight = InOther.Weight;
+			CacheBucket = InOther.CacheBucket;
 		}
 
 		public Action(BinaryArchiveReader Reader)
@@ -427,7 +439,7 @@ namespace UnrealBuildTool
 			bUseActionHistory = Reader.ReadBool();
 			bIsHighPriority = Reader.ReadBool();
 			Weight = Reader.ReadDouble();
-
+			CacheBucket = Reader.ReadUnsignedInt();
 		}
 
 		/// <summary>
@@ -462,6 +474,7 @@ namespace UnrealBuildTool
 			Writer.WriteBool(bUseActionHistory);
 			Writer.WriteBool(bIsHighPriority);
 			Writer.WriteDouble(Weight);
+			Writer.WriteUnsignedInt(CacheBucket);
 		}
 
 		/// <summary>
@@ -592,6 +605,11 @@ namespace UnrealBuildTool
 				Action.Weight = Weight;
 			}
 
+			if (Object.TryGetUnsignedIntegerField("CacheBucket", out uint CacheBucket))
+			{
+				Action.CacheBucket = CacheBucket;
+			}
+
 			return Action;
 		}
 
@@ -661,6 +679,7 @@ namespace UnrealBuildTool
 			Writer.WriteValue("bShouldOutputLog", Action.bShouldOutputLog);
 			Writer.WriteValue("bProducesImportLibrary", Action.bProducesImportLibrary);
 			Writer.WriteValue("Weight", Action.Weight);
+			Writer.WriteValue("CacheBucket", Action.CacheBucket);
 
 			Writer.WriteArrayStart("PrerequisiteActions");
 			foreach (LinkedAction PrerequisiteAction in Action.PrerequisiteActions)
@@ -794,6 +813,7 @@ namespace UnrealBuildTool
 		public bool bUseActionHistory => Inner.bUseActionHistory;
 		public bool bIsHighPriority => IsHighPriority != 0;
 		public double Weight => Inner.Weight;
+		public uint CacheBucket => Inner.CacheBucket;
 
 		#endregion
 
diff --git a/Engine/Source/Programs/UnrealBuildTool/ToolChain/ISPCToolChain.cs b/Engine/Source/Programs/UnrealBuildTool/ToolChain/ISPCToolChain.cs
--- a/Engine/Source/Programs/UnrealBuildTool/ToolChain/ISPCToolChain.cs
+++ b/Engine/Source/Programs/UnrealBuildTool/ToolChain/ISPCToolChain.cs
@@ -378,6 +378,60 @@ namespace UnrealBuildTool
 			return NormalizeCommandLinePath(Item.Location, RootPaths);
 		}
 
+		/// <summary>
+		/// Simple hash function for strings (can be chained)
+		/// </summary>
+		/// <param name="str">The string to hash</param>
+		/// <param name="hash">The hash input</param>
+		/// <returns>hash of input hash + string</returns>
+		protected static uint GetStringHash(string str, uint hash = 5381)
+		{
+			for (int idx = 0; idx < str.Length; idx++)
+			{
+				hash += (hash << 5) + str[idx];
+			}
+			return hash;
+		}
+
+		/// <summary>
+		/// Calculates hash bucket based on target and root paths
+		/// </summary>
+		/// <param name="target">Target rules</param>
+		/// <param name="rootPaths">Optional root paths. Should be null except in special scenario (using pch and vfs is disabled)</param>
+		/// <returns>bucket hash</returns>
+		public static uint GetCacheBucket(ReadOnlyTargetRules target, CppRootPaths? rootPaths)
+		{
+			uint cacheBucket = GetStringHash(target.Platform.ToString()); // Platform
+
+			cacheBucket = GetStringHash(target.Configuration.ToString(), cacheBucket); // Configuration
+
+			cacheBucket = GetStringHash(target.LinkType.ToString(), cacheBucket); // Modular or monolithic
+
+			cacheBucket = GetStringHash(target.Architecture.ToString(), cacheBucket); // arm64 or x64
+
+			if (target.StaticAnalyzer != StaticAnalyzer.None)
+			{
+				cacheBucket = GetStringHash(target.StaticAnalyzer.ToString(), cacheBucket); // If static analyzer
+			}
+			if (!target.bUseUnityBuild)
+			{
+				cacheBucket = GetStringHash("NU", cacheBucket); // Non-unity. There might be matching cache entries between unity and non unity but this will keep bucket size down
+			}
+			if (target.bMergeModules)
+			{
+				cacheBucket = GetStringHash("MM", cacheBucket); // Merged modules.. will modify all .obj files
+			}
+
+			if (rootPaths != null) // Usually if vfs is disabled and pch is enabled
+			{
+				foreach (var (id, vfs, local) in rootPaths)
+				{
+					cacheBucket = GetStringHash(local.FullName, cacheBucket);
+				}
+			}
+			return cacheBucket;
+		}
+
 		protected override CPPOutput GenerateISPCHeaders(CppCompileEnvironment CompileEnvironment, IEnumerable<FileItem> InputFiles, DirectoryReference OutputDir, IActionGraphBuilder Graph)
 		{
 			CPPOutput Result = new CPPOutput();
