[UBA]
* Added hint to cache client connection to be able to back track from cache server log to horde job/step (or computer name if not running in horde)

diff --git a/Engine/Source/Programs/UnrealBuildTool/Executors/UnrealBuildAccelerator/UBAExecutor.cs b/Engine/Source/Programs/UnrealBuildTool/Executors/UnrealBuildAccelerator/UBAExecutor.cs
--- a/Engine/Source/Programs/UnrealBuildTool/Executors/UnrealBuildAccelerator/UBAExecutor.cs
+++ b/Engine/Source/Programs/UnrealBuildTool/Executors/UnrealBuildAccelerator/UBAExecutor.cs
@@ -264,7 +264,7 @@ namespace UnrealBuildTool
 				return Task.Factory.StartNew(() =>
 				{
 					ProcessStartInfo startInfo = _executor.GetActionStartInfo(action);
-					uint bucket = UBAExecutor.GetActionCacheBucket(action);
+					uint bucket = _executor.GetActionCacheBucket(action);
 					ulong rootsHandle = _executor.GetActionRootsHandle(action);
 
 					FetchFromCacheResult result = _executor._cacheClient!.FetchFromCache(rootsHandle, bucket, startInfo);
@@ -409,10 +409,13 @@ namespace UnrealBuildTool
 
 							OnSessionCreated?.Invoke(_session);
 
-							_cacheClient = ICacheClient.CreateCacheClient(_session, UBAConfig.bReportCacheMissReason);
-							if (!String.IsNullOrEmpty(UBAConfig.CacheServer))
+							bool useCacheServer = !String.IsNullOrEmpty(UBAConfig.CacheServer);
+
+							_cacheClient = ICacheClient.CreateCacheClient(_session, UBAConfig.bReportCacheMissReason, "", useCacheServer ? GetCacheClientHint() : "");
+
+							if (useCacheServer)
 							{
-								string[] nameAndPort = UBAConfig.CacheServer.Split(':');
+								string[] nameAndPort = UBAConfig.CacheServer!.Split(':');
 								int port = 1347;
 								if (nameAndPort.Length > 1)
 								{
@@ -676,7 +679,36 @@ namespace UnrealBuildTool
 			return startInfo;
 		}
 
-		static uint GetActionCacheBucket(LinkedAction action)
+
+		string GetCacheClientHint()
+		{
+			string? jobId = Environment.GetEnvironmentVariable("UE_HORDE_JOBID");
+			if (String.IsNullOrEmpty(jobId))
+			{
+				return Environment.MachineName;
+			}
+			string? hordeUrl = Environment.GetEnvironmentVariable("UE_HORDE_URL");
+			string? stepId = Environment.GetEnvironmentVariable("UE_HORDE_STEPID");
+			return $"{hordeUrl}/job/{jobId}?step={stepId}";
+		}
+
+		HashSet<uint> _usedBuckets = new();
+
+		uint GetActionCacheBucket(LinkedAction action)
+		{
+			uint bucket = GetActionCacheBucket2(action);
+
+			lock (_usedBuckets)
+			{
+				if (_usedBuckets.Add(bucket))
+				{
+					_threadedLogger.LogInformation($"First time bucket {bucket} is used. Action: {action.CommandDescription} {action.StatusDescription}");
+				}
+			}
+			
+			return bucket;
+		}
+		uint GetActionCacheBucket2(LinkedAction action)
 		{
 			if (action.CacheBucket != 0)
 			{
