[UBA]
* Removed slash added after hordeurl since it already ends with slash

diff --git a/Engine/Source/Programs/UnrealBuildTool/Executors/UnrealBuildAccelerator/UBAExecutor.cs b/Engine/Source/Programs/UnrealBuildTool/Executors/UnrealBuildAccelerator/UBAExecutor.cs
--- a/Engine/Source/Programs/UnrealBuildTool/Executors/UnrealBuildAccelerator/UBAExecutor.cs
+++ b/Engine/Source/Programs/UnrealBuildTool/Executors/UnrealBuildAccelerator/UBAExecutor.cs
@@ -689,7 +689,7 @@ namespace UnrealBuildTool
 			}
 			string? hordeUrl = Environment.GetEnvironmentVariable("UE_HORDE_URL");
 			string? stepId = Environment.GetEnvironmentVariable("UE_HORDE_STEPID");
-			return $"{hordeUrl}/job/{jobId}?step={stepId}";
+			return $"{hordeUrl}job/{jobId}?step={stepId}";
 		}
 
 		HashSet<uint> _usedBuckets = new();
