[UBA]
* Fixed so executor does not connect to cache server if no actions are cacheable

diff --git a/Engine/Source/Programs/UnrealBuildTool/Executors/UnrealBuildAccelerator/UBAExecutor.cs b/Engine/Source/Programs/UnrealBuildTool/Executors/UnrealBuildAccelerator/UBAExecutor.cs
--- a/Engine/Source/Programs/UnrealBuildTool/Executors/UnrealBuildAccelerator/UBAExecutor.cs
+++ b/Engine/Source/Programs/UnrealBuildTool/Executors/UnrealBuildAccelerator/UBAExecutor.cs
@@ -56,6 +56,8 @@ namespace UnrealBuildTool
 		// Tracking for all actions processed remotely
 		int _remoteProcessedActions = 0;
 
+		int _cacheableActionsLeft = 0;
+
 		// Tracking for remote connection mode
 		string _remoteConnectionMode = "Local";
 
@@ -264,7 +266,7 @@ namespace UnrealBuildTool
 				return Task.Factory.StartNew(() =>
 				{
 					ProcessStartInfo startInfo = _executor.GetActionStartInfo(action);
-					uint bucket = _executor.GetActionCacheBucket(action);
+					uint bucket = _executor.GetCacheBucket(action);
 					ulong rootsHandle = _executor.GetActionRootsHandle(action);
 
 					FetchFromCacheResult result = _executor._cacheClient!.FetchFromCache(rootsHandle, bucket, startInfo);
@@ -276,6 +278,10 @@ namespace UnrealBuildTool
 						exitCode = 1;
 						++_executor._errorCount;
 					}
+					if (result.Success)
+					{
+						_executor.DecrementCacheableAction();
+					}
 					return new ActionArtifactResult(result.Success, result.LogLines, exitCode);
 				}, cancellationToken, TaskCreationOptions.LongRunning | TaskCreationOptions.PreferFairness, TaskScheduler.Default);
 			}
@@ -314,6 +320,17 @@ namespace UnrealBuildTool
 				UBAConfig.Zone = "local";
 			}
 
+			if (!String.IsNullOrEmpty(UBAConfig.CacheServer))
+			{
+				foreach (LinkedAction action in inputActions)
+				{
+					if (action.ArtifactMode.HasFlag(ArtifactMode.Enabled))
+					{
+						++_cacheableActionsLeft;
+					}
+				}
+			}
+
 			PrintConfiguration();
 			await WriteActionOutputFileAsync(inputActions);
 
@@ -409,11 +426,9 @@ namespace UnrealBuildTool
 
 							OnSessionCreated?.Invoke(_session);
 
-							bool useCacheServer = !String.IsNullOrEmpty(UBAConfig.CacheServer);
-
-							_cacheClient = ICacheClient.CreateCacheClient(_session, UBAConfig.bReportCacheMissReason, "", useCacheServer ? GetCacheClientHint() : "");
+							_cacheClient = ICacheClient.CreateCacheClient(_session, UBAConfig.bReportCacheMissReason, "", _cacheableActionsLeft != 0 ? GetCacheClientHint() : "");
 
-							if (useCacheServer)
+							if (_cacheableActionsLeft != 0)
 							{
 								string[] nameAndPort = UBAConfig.CacheServer!.Split(':');
 								int port = 1347;
@@ -679,6 +694,14 @@ namespace UnrealBuildTool
 			return startInfo;
 		}
 
+		void DecrementCacheableAction()
+		{
+			if (Interlocked.Decrement(ref _cacheableActionsLeft) == 0)
+			{
+				_cacheClient?.Dispose();
+				_cacheClient = null;
+			}
+		}
 
 		string GetCacheClientHint()
 		{
@@ -689,55 +712,54 @@ namespace UnrealBuildTool
 			}
 			string? hordeUrl = Environment.GetEnvironmentVariable("UE_HORDE_URL");
 			string? stepId = Environment.GetEnvironmentVariable("UE_HORDE_STEPID");
-			return $"{hordeUrl}job/{jobId}?step={stepId}";
+			string streamId = Environment.GetEnvironmentVariable("UE_HORDE_STREAMID")??"unknown-stream";
+			return $"{hordeUrl}job/{jobId}?step={stepId} ({streamId})";
 		}
 
 		HashSet<uint> _usedBuckets = new();
 
-		uint GetActionCacheBucket(LinkedAction action)
-		{
-			uint bucket = GetActionCacheBucket2(action);
-
-			lock (_usedBuckets)
-			{
-				if (_usedBuckets.Add(bucket))
-				{
-					_threadedLogger.LogInformation($"First time bucket {bucket} is used. Action: {action.CommandDescription} {action.StatusDescription}");
-				}
-			}
-			
-			return bucket;
-		}
-		uint GetActionCacheBucket2(LinkedAction action)
+		uint GetCacheBucket(LinkedAction action)
 		{
+			uint bucket = 0;
 			if (action.CacheBucket != 0)
 			{
-				return action.CacheBucket;
+				bucket = action.CacheBucket;
 			}
-
-			if (action.Target == null)
+			else if (action.Target == null)
 			{
-				return 0;
+				bucket = 0;
 			}
-
-			using (Blake3.Hasher hasher = Blake3.Hasher.New())
+			else
 			{
-				// Use platform and config to chose bucket since there will never be any cache hits between platforms or configs
-				hasher.Update(System.Text.Encoding.UTF8.GetBytes(action.Target.Platform.ToString()));
-				hasher.Update(new byte[] { (byte)action.Target.Configuration });
-
-				// Absolute path is set for actions that uses pch.
-				// And since pch contains absolute paths we unfortunately can't share cache data between machines that have different paths
-				if (action.ArtifactMode.HasFlag(ArtifactMode.AbsolutePath))
+				using (Blake3.Hasher hasher = Blake3.Hasher.New())
 				{
-					foreach (var (id, vfs, local) in action.RootPaths)
+					// Use platform and config to chose bucket since there will never be any cache hits between platforms or configs
+					hasher.Update(System.Text.Encoding.UTF8.GetBytes(action.Target.Platform.ToString()));
+					hasher.Update(new byte[] { (byte)action.Target.Configuration });
+
+					// Absolute path is set for actions that uses pch.
+					// And since pch contains absolute paths we unfortunately can't share cache data between machines that have different paths
+					if (action.ArtifactMode.HasFlag(ArtifactMode.AbsolutePath))
 					{
-						hasher.Update(System.Text.Encoding.UTF8.GetBytes(local.FullName));
+						foreach (var (id, vfs, local) in action.RootPaths)
+						{
+							hasher.Update(System.Text.Encoding.UTF8.GetBytes(local.FullName));
+						}
 					}
+
+					bucket = (uint)IoHash.FromBlake3(hasher).GetHashCode();
 				}
+			}
 
-				return (uint)IoHash.FromBlake3(hasher).GetHashCode();
+			lock (_usedBuckets)
+			{
+				if (_usedBuckets.Add(bucket))
+				{
+					_threadedLogger.LogInformation($"First time bucket {bucket} is used. Action: {action.CommandDescription} {action.StatusDescription}");
+				}
 			}
+
+			return bucket;
 		}
 
 		ulong GetActionRootsHandle(LinkedAction action)
@@ -827,9 +849,13 @@ namespace UnrealBuildTool
 				}
 			}
 
-			uint bucket = GetActionCacheBucket(action);
+			uint bucket = GetCacheBucket(action);
+
+			bool success = _cacheClient!.WriteToCache(bucket, process, inputsMemory.GetBuffer(), (uint)inputsMemory.Position, outputsMemory.GetBuffer(), (uint)outputsMemory.Position);
+
+			DecrementCacheableAction();
 
-			return _cacheClient!.WriteToCache(bucket, process, inputsMemory.GetBuffer(), (uint)inputsMemory.Position, outputsMemory.GetBuffer(), (uint)outputsMemory.Position);
+			return success;
 		}
 
 		Func<Task>? RunActionLocal(ImmediateActionQueue queue, LinkedAction action)
