[UBA]
* Minor tweaks on logging in UBT

diff --git a/Engine/Source/Programs/UnrealBuildTool/Executors/UnrealBuildAccelerator/UBAExecutor.cs b/Engine/Source/Programs/UnrealBuildTool/Executors/UnrealBuildAccelerator/UBAExecutor.cs
--- a/Engine/Source/Programs/UnrealBuildTool/Executors/UnrealBuildAccelerator/UBAExecutor.cs
+++ b/Engine/Source/Programs/UnrealBuildTool/Executors/UnrealBuildAccelerator/UBAExecutor.cs
@@ -706,12 +706,12 @@ namespace UnrealBuildTool
 		string GetCacheClientHint()
 		{
 			string? jobId = Environment.GetEnvironmentVariable("UE_HORDE_JOBID");
-			if (String.IsNullOrEmpty(jobId))
+			if (jobId == null)
 			{
 				return Environment.MachineName;
 			}
-			string? hordeUrl = Environment.GetEnvironmentVariable("UE_HORDE_URL");
-			string? stepId = Environment.GetEnvironmentVariable("UE_HORDE_STEPID");
+			string hordeUrl = Environment.GetEnvironmentVariable("UE_HORDE_URL")??"";
+			string stepId = Environment.GetEnvironmentVariable("UE_HORDE_STEPID")??"unknown";
 			string streamId = Environment.GetEnvironmentVariable("UE_HORDE_STREAMID")??"unknown-stream";
 			return $"{hordeUrl}job/{jobId}?step={stepId} ({streamId})";
 		}
@@ -755,7 +755,7 @@ namespace UnrealBuildTool
 			{
 				if (_usedBuckets.Add(bucket))
 				{
-					_threadedLogger.LogInformation($"First time bucket {bucket} is used. Action: {action.CommandDescription} {action.StatusDescription}");
+					_threadedLogger.LogInformation($"First time cache bucket {bucket} is used. Action: {action.CommandDescription} {action.StatusDescription}");
 				}
 			}
 
