[UBT]
* Fixed so cache bucket id is precalculated for clang/ispc toolchains

diff --git a/Engine/Source/Programs/UnrealBuildTool/ToolChain/ClangToolChain.cs b/Engine/Source/Programs/UnrealBuildTool/ToolChain/ClangToolChain.cs
--- a/Engine/Source/Programs/UnrealBuildTool/ToolChain/ClangToolChain.cs
+++ b/Engine/Source/Programs/UnrealBuildTool/ToolChain/ClangToolChain.cs
@@ -1309,6 +1309,7 @@ namespace UnrealBuildTool
 				CompileEnvironment.bAllowRemotelyCompiledPCHs;
 
 			CompileAction.ArtifactMode = ArtifactMode.Enabled;
+			CompileAction.CacheBucket = GetCacheBucket(TargetRules, null);
 
 			if (CompileEnvironment.PrecompiledHeaderAction != PrecompiledHeaderAction.None)
 			{
diff --git a/Engine/Source/Programs/UnrealBuildTool/ToolChain/ISPCToolChain.cs b/Engine/Source/Programs/UnrealBuildTool/ToolChain/ISPCToolChain.cs
--- a/Engine/Source/Programs/UnrealBuildTool/ToolChain/ISPCToolChain.cs
+++ b/Engine/Source/Programs/UnrealBuildTool/ToolChain/ISPCToolChain.cs
@@ -21,6 +21,7 @@ namespace UnrealBuildTool
 		protected FileReference? ProjectFile = null;
 		protected bool bMergeModules = false;
 		protected bool bAllowUbaCompression = false;
+		protected ReadOnlyTargetRules? TargetRules = null;
 
 		public override void SetUpGlobalEnvironment(ReadOnlyTargetRules Target, CppCompileEnvironment GlobalCompileEnvironment, LinkEnvironment GlobalLinkEnvironment)
 		{
@@ -28,6 +29,7 @@ namespace UnrealBuildTool
 			ProjectFile = Target.ProjectFile;
 			bMergeModules = Target.bMergeModules;
 			bAllowUbaCompression = Target.bAllowUbaCompression;
+			TargetRules = Target;
 		}
 
 		/// <summary>
@@ -399,8 +401,13 @@ namespace UnrealBuildTool
 		/// <param name="target">Target rules</param>
 		/// <param name="rootPaths">Optional root paths. Should be null except in special scenario (using pch and vfs is disabled)</param>
 		/// <returns>bucket hash</returns>
-		public static uint GetCacheBucket(ReadOnlyTargetRules target, CppRootPaths? rootPaths)
+		public static uint GetCacheBucket(ReadOnlyTargetRules? target, CppRootPaths? rootPaths)
 		{
+			if (target == null)
+			{
+				return 0;
+			}
+
 			uint cacheBucket = GetStringHash(target.Platform.ToString()); // Platform
 
 			cacheBucket = GetStringHash(target.Configuration.ToString(), cacheBucket); // Configuration
@@ -507,6 +514,7 @@ namespace UnrealBuildTool
 				CompileAction.StatusDescription = Path.GetFileName(ISPCFile.AbsolutePath);
 				CompileAction.CommandVersion = GetISPCHostCompilerVersion(BuildHostPlatform.Current.Platform).ToString();
 				CompileAction.ArtifactMode = ArtifactMode.Enabled;
+				CompileAction.CacheBucket = GetCacheBucket(TargetRules, null);
 
 				CompileAction.bCanExecuteRemotely = true;
 				CompileAction.bCanExecuteInUBACrossArchitecture = false;
@@ -732,6 +740,7 @@ namespace UnrealBuildTool
 				CompileAction.bCanExecuteRemotelyWithSNDBS = false;
 
 				CompileAction.ArtifactMode = ArtifactMode.Enabled;
+				CompileAction.CacheBucket = GetCacheBucket(TargetRules, null);
 
 				List<string> Arguments = new List<string>();
 
@@ -883,6 +892,7 @@ namespace UnrealBuildTool
 
 							PostCompileAction.RootPaths = RootPaths;
 							PostCompileAction.ArtifactMode = ArtifactMode.Enabled;
+							PostCompileAction.CacheBucket = GetCacheBucket(TargetRules, null);
 
 							// Disable remote execution to workaround mismatched case on XGE
 							PostCompileAction.bCanExecuteRemotelyWithXGE = false;
