[UBA]
* Downgraded cache bucket logging from information to debug

diff --git a/Engine/Source/Programs/UnrealBuildTool/Executors/UnrealBuildAccelerator/UBAExecutor.cs b/Engine/Source/Programs/UnrealBuildTool/Executors/UnrealBuildAccelerator/UBAExecutor.cs
--- a/Engine/Source/Programs/UnrealBuildTool/Executors/UnrealBuildAccelerator/UBAExecutor.cs
+++ b/Engine/Source/Programs/UnrealBuildTool/Executors/UnrealBuildAccelerator/UBAExecutor.cs
@@ -755,7 +755,7 @@ namespace UnrealBuildTool
 			{
 				if (_usedBuckets.Add(bucket))
 				{
-					_threadedLogger.LogInformation($"First time cache bucket {bucket} is used. Action: {action.CommandDescription} {action.StatusDescription}");
+					_threadedLogger.LogDebug($"Cache bucket {bucket} used. Action: {action.CommandDescription} {action.StatusDescription} ({String.Join(" + ", action.GroupNames)})");
 				}
 			}
 
