[UBT]
* Fixed so multi-architecture targets don't use architecture as part of cache bucket id

diff --git a/Engine/Source/Programs/UnrealBuildTool/ToolChain/ISPCToolChain.cs b/Engine/Source/Programs/UnrealBuildTool/ToolChain/ISPCToolChain.cs
--- a/Engine/Source/Programs/UnrealBuildTool/ToolChain/ISPCToolChain.cs
+++ b/Engine/Source/Programs/UnrealBuildTool/ToolChain/ISPCToolChain.cs
@@ -414,7 +414,10 @@ namespace UnrealBuildTool
 
 			cacheBucket = GetStringHash(target.LinkType.ToString(), cacheBucket); // Modular or monolithic
 
-			cacheBucket = GetStringHash(target.Architecture.ToString(), cacheBucket); // arm64 or x64
+			if (!target.Architectures.bIsMultiArch)
+			{
+				cacheBucket = GetStringHash(target.Architecture.ToString(), cacheBucket); // arm64 or x64
+			}
 
 			if (target.StaticAnalyzer != StaticAnalyzer.None)
 			{
