VC Project file generation - generate correct prject type configurations for low level tests

diff --git a/Engine/Source/Programs/UnrealBuildTool/ProjectFiles/VisualStudio/VCProjectFileGenerator.cs b/Engine/Source/Programs/UnrealBuildTool/ProjectFiles/VisualStudio/VCProjectFileGenerator.cs
--- a/Engine/Source/Programs/UnrealBuildTool/ProjectFiles/VisualStudio/VCProjectFileGenerator.cs
+++ b/Engine/Source/Programs/UnrealBuildTool/ProjectFiles/VisualStudio/VCProjectFileGenerator.cs
@@ -1057,7 +1057,10 @@ namespace UnrealBuildTool
 											TargetType TargetType = TargetType.Game;
 											if (ProjectTarget.TargetRules != null)
 											{
-												TargetType = ProjectTarget.TargetRules.Type;
+												if (!ProjectTarget.TargetRules.IsTestTarget)
+												{
+													TargetType = ProjectTarget.TargetRules.Type;
+												}
 											}
 
 											string SolutionConfigName =
