[UBT]
* Fixed so clang compile actions that is not using vfs but use pch use its own cache bucket

diff --git a/Engine/Source/Programs/UnrealBuildTool/ToolChain/ClangToolChain.cs b/Engine/Source/Programs/UnrealBuildTool/ToolChain/ClangToolChain.cs
--- a/Engine/Source/Programs/UnrealBuildTool/ToolChain/ClangToolChain.cs
+++ b/Engine/Source/Programs/UnrealBuildTool/ToolChain/ClangToolChain.cs
@@ -1309,7 +1309,7 @@ namespace UnrealBuildTool
 				CompileEnvironment.bAllowRemotelyCompiledPCHs;
 
 			CompileAction.ArtifactMode = ArtifactMode.Enabled;
-			CompileAction.CacheBucket = GetCacheBucket(TargetRules, null);
+			CompileAction.CacheBucket = GetCacheBucket(TargetRules, (!CompileEnvironment.RootPaths.bUseVfs && CompileEnvironment.PrecompiledHeaderAction != PrecompiledHeaderAction.None) ? CompileEnvironment.RootPaths : null);
 
 			if (CompileEnvironment.PrecompiledHeaderAction != PrecompiledHeaderAction.None)
 			{
