[UBT]
* Disable vfs for link action that uses PGO for now..

diff --git a/Engine/Source/Programs/UnrealBuildTool/Platform/Windows/VCToolChain.cs b/Engine/Source/Programs/UnrealBuildTool/Platform/Windows/VCToolChain.cs
--- a/Engine/Source/Programs/UnrealBuildTool/Platform/Windows/VCToolChain.cs
+++ b/Engine/Source/Programs/UnrealBuildTool/Platform/Windows/VCToolChain.cs
@@ -2727,7 +2727,14 @@ namespace UnrealBuildTool
 		public override FileItem LinkFiles(LinkEnvironment LinkEnvironment, bool bBuildImportLibraryOnly, IActionGraphBuilder Graph)
 		{
 			CppRootPaths RootPaths = LinkEnvironment.RootPaths;
-			
+
+			// Remove when we fix so pgo linking run in uba
+			if (LinkEnvironment.bPGOOptimize || LinkEnvironment.bPGOProfile)
+			{
+				RootPaths = new(RootPaths);
+				RootPaths.bUseVfs = false;
+			}
+
 			if (LinkEnvironment.bIsBuildingDotNetAssembly)
 			{
 				return FileItem.GetItemByFileReference(LinkEnvironment.OutputFilePath);
