[UBT]
* Another temporary fix to remove vfs paths in rsp files when linking with pgo

diff --git a/Engine/Source/Programs/UnrealBuildTool/Platform/Windows/VCToolChain.cs b/Engine/Source/Programs/UnrealBuildTool/Platform/Windows/VCToolChain.cs
--- a/Engine/Source/Programs/UnrealBuildTool/Platform/Windows/VCToolChain.cs
+++ b/Engine/Source/Programs/UnrealBuildTool/Platform/Windows/VCToolChain.cs
@@ -1447,9 +1447,17 @@ namespace UnrealBuildTool
 			}
 			else
 			{
+				// Remove when we fix so pgo linking run in uba
+				CppRootPaths RootPaths = LinkEnvironment.RootPaths;
+				if (LinkEnvironment.bPGOOptimize || LinkEnvironment.bPGOProfile)
+				{
+					RootPaths = new(RootPaths);
+					RootPaths.bUseVfs = false;
+				}
+
 				Arguments.Add("/MANIFEST:EMBED");
 				FileItem ManifestFile = FileItem.GetItemByPath(Target.WindowsPlatform.ManifestFile);
-				Arguments.Add($"/MANIFESTINPUT:\"{NormalizeCommandLinePath(ManifestFile, LinkEnvironment.RootPaths)}\"");
+				Arguments.Add($"/MANIFESTINPUT:\"{NormalizeCommandLinePath(ManifestFile, RootPaths)}\"");
 			}
 
 			// Prevents the linker from displaying its logo for each invocation.
