Fix for action graph appx manifest generation when there are multiple UBT configurations specified.

diff --git a/Engine/Source/Programs/UnrealBuildTool/Platform/Windows/AppXManifestGeneratorBase.cs b/Engine/Source/Programs/UnrealBuildTool/Platform/Windows/AppXManifestGeneratorBase.cs
--- a/Engine/Source/Programs/UnrealBuildTool/Platform/Windows/AppXManifestGeneratorBase.cs
+++ b/Engine/Source/Programs/UnrealBuildTool/Platform/Windows/AppXManifestGeneratorBase.cs
@@ -505,13 +505,6 @@ namespace UnrealBuildTool
 		/// </summary>
 		public List<string>? CreateManifest(string InManifestName, DirectoryReference OutputDirectory, string? InTargetName, FileReference? InProjectFile, Dictionary<UnrealTargetConfiguration, string> InExecutablePairs)
 		{
-			// Check parameter values are valid.
-			if (InExecutablePairs.Count < 1)
-			{
-				Logger.LogError("The number of target configurations is zero, so we cannot generate a manifest.");
-				return null;
-			}
-
 			FileUtils.CreateDirectoryTree(OutputDirectory);
 
 			DirectoryReference ProjectRoot = (InProjectFile != null) ? InProjectFile.Directory : Unreal.EngineDirectory;
