Fix clang warnings for NDK r28

#rnx

diff --git a/Engine/Source/Programs/UnrealBuildTool/ToolChain/ClangWarnings.cs b/Engine/Source/Programs/UnrealBuildTool/ToolChain/ClangWarnings.cs
--- a/Engine/Source/Programs/UnrealBuildTool/ToolChain/ClangWarnings.cs
+++ b/Engine/Source/Programs/UnrealBuildTool/ToolChain/ClangWarnings.cs
@@ -59,9 +59,13 @@ namespace UnrealBuildTool
 			}
 			if (ClangVersion >= new VersionNumber(19))
 			{
+				bool bIsAndroidClang19 = ClangVersion == new VersionNumber(19, 0, 0) && CompileEnvironment.Platform == UnrealTargetPlatform.Android;
 				Arguments.Add("-Wno-extra-qualification");              // https://clang.llvm.org/docs/DiagnosticsReference.html#wextra-qualification					// new warning for clang 19
-				Arguments.Add("-Wno-cast-function-type-mismatch");
-				Arguments.Add("-Wno-missing-template-arg-list-after-template-kw");
+				if (!bIsAndroidClang19) // android NDK r28.0.12433566 and r28.0.12674087 are missing these warnings 
+				{
+					Arguments.Add("-Wno-cast-function-type-mismatch");
+					Arguments.Add("-Wno-missing-template-arg-list-after-template-kw");
+				}
 			}
 
 			Arguments.Add("-Wno-gnu-string-literal-operator-template"); // https://clang.llvm.org/docs/DiagnosticsReference.html#wgnu-string-literal-operator-template	// We use this feature to allow static FNames.
