[UnrealBuildTool]
- Allow usage of concurrent collections and parallel utilities in project sources

#rnx

diff --git a/Engine/Source/Programs/UnrealBuildTool/System/DynamicCompilation.cs b/Engine/Source/Programs/UnrealBuildTool/System/DynamicCompilation.cs
--- a/Engine/Source/Programs/UnrealBuildTool/System/DynamicCompilation.cs
+++ b/Engine/Source/Programs/UnrealBuildTool/System/DynamicCompilation.cs
@@ -222,6 +222,9 @@ namespace UnrealBuildTool
 			MetadataReferences.Add(MetadataReference.CreateFromFile(Assembly.Load("System.IO").Location));
 			MetadataReferences.Add(MetadataReference.CreateFromFile(Assembly.Load("System.IO.FileSystem").Location));
 			MetadataReferences.Add(MetadataReference.CreateFromFile(Assembly.Load("System.Linq").Location));
+			MetadataReferences.Add(MetadataReference.CreateFromFile(Assembly.Load("System.Threading.Tasks").Location));
+			MetadataReferences.Add(MetadataReference.CreateFromFile(Assembly.Load("System.Threading.Tasks.Parallel").Location));
+			MetadataReferences.Add(MetadataReference.CreateFromFile(Assembly.Load("System.Collections.Concurrent").Location));
 			MetadataReferences.Add(MetadataReference.CreateFromFile(Assembly.Load("System.Private.Xml").Location));
 			MetadataReferences.Add(MetadataReference.CreateFromFile(Assembly.Load("System.Private.Xml.Linq").Location));
 			MetadataReferences.Add(MetadataReference.CreateFromFile(Assembly.Load("System.Text.RegularExpressions").Location));
