UnrealBuildTool: Process shared engine targets before shared project targets

#rnx

diff --git a/Engine/Source/Programs/UnrealBuildTool/ProjectFiles/ProjectFileGenerator.cs b/Engine/Source/Programs/UnrealBuildTool/ProjectFiles/ProjectFileGenerator.cs
--- a/Engine/Source/Programs/UnrealBuildTool/ProjectFiles/ProjectFileGenerator.cs
+++ b/Engine/Source/Programs/UnrealBuildTool/ProjectFiles/ProjectFileGenerator.cs
@@ -2280,17 +2280,35 @@ namespace UnrealBuildTool
 					IEnumerable<UEBuildTarget> Shared = BuildTargets.Where(x => x.bUseSharedBuildEnvironment);
 					IEnumerable<UEBuildTarget> Unique = BuildTargets.Where(x => !x.bUseSharedBuildEnvironment);
 
-					IEnumerable<IGrouping<(UnrealTargetPlatform Platform, UnrealTargetConfiguration Configuration, TargetType TargetType), UEBuildTarget>> ParallelGroups = Shared.GroupBy(x => (x.Platform, x.Configuration, x.TargetType));
+					IEnumerable<UEBuildTarget> SharedEngine = Shared.Where(x => x.ProjectFile == null);
+					IEnumerable<UEBuildTarget> SharedProject = Shared.Except(SharedEngine);
 
+					IEnumerable<IGrouping<(UnrealTargetPlatform Platform, UnrealTargetConfiguration Configuration, TargetType TargetType), UEBuildTarget>> ParallelGroupsEngine = SharedEngine.GroupBy(x => (x.Platform, x.Configuration, x.TargetType));
+					IEnumerable<IGrouping<(UnrealTargetPlatform Platform, UnrealTargetConfiguration Configuration, TargetType TargetType), UEBuildTarget>> ParallelGroupsProject = SharedProject.GroupBy(x => (x.Platform, x.Configuration, x.TargetType));
+
+					// Shared engine targets should be processed before shared project targets
 					// Unique environments can be processed in parallel
 					// Shared environments can be processed in parallel, but each unique combination of Platform\Configuration\TargetType must be processed serially
 					// DebugGame can not be run in parallel with Development
 					// This could be made more async but it's complicated
 
-					Parallel.ForEach(Unique.Where(x => x.Configuration != UnrealTargetConfiguration.DebugGame), Target => Gen(Target));
-					Parallel.ForEach(Unique.Where(x => x.Configuration == UnrealTargetConfiguration.DebugGame), Target => Gen(Target));
+					Parallel.ForEach(ParallelGroupsEngine.Where(x => x.Key.Configuration != UnrealTargetConfiguration.DebugGame), Group =>
+					{
+						foreach (UEBuildTarget Target in Group)
+						{
+							Gen(Target);
+						}
+					});
+
+					Parallel.ForEach(ParallelGroupsEngine.Where(x => x.Key.Configuration == UnrealTargetConfiguration.DebugGame), Group =>
+					{
+						foreach (UEBuildTarget Target in Group)
+						{
+							Gen(Target);
+						}
+					});
 					
-					Parallel.ForEach(ParallelGroups.Where(x => x.Key.Configuration != UnrealTargetConfiguration.DebugGame), Group =>
+					Parallel.ForEach(ParallelGroupsProject.Where(x => x.Key.Configuration != UnrealTargetConfiguration.DebugGame), Group =>
 					{
 						foreach (UEBuildTarget Target in Group)
 						{
@@ -2298,13 +2316,16 @@ namespace UnrealBuildTool
 						}
 					});
 
-					Parallel.ForEach(ParallelGroups.Where(x => x.Key.Configuration == UnrealTargetConfiguration.DebugGame), Group =>
+					Parallel.ForEach(ParallelGroupsProject.Where(x => x.Key.Configuration == UnrealTargetConfiguration.DebugGame), Group =>
 					{
 						foreach (UEBuildTarget Target in Group)
 						{
 							Gen(Target);
 						}
 					});
+
+					Parallel.ForEach(Unique.Where(x => x.Configuration != UnrealTargetConfiguration.DebugGame), Target => Gen(Target));
+					Parallel.ForEach(Unique.Where(x => x.Configuration == UnrealTargetConfiguration.DebugGame), Target => Gen(Target));
 				}
 			}
 		}
diff --git a/Engine/Source/Programs/UnrealBuildTool/ProjectFiles/VisualStudio/VCProjectFileGenerator.cs b/Engine/Source/Programs/UnrealBuildTool/ProjectFiles/VisualStudio/VCProjectFileGenerator.cs
--- a/Engine/Source/Programs/UnrealBuildTool/ProjectFiles/VisualStudio/VCProjectFileGenerator.cs
+++ b/Engine/Source/Programs/UnrealBuildTool/ProjectFiles/VisualStudio/VCProjectFileGenerator.cs
@@ -647,7 +647,7 @@ namespace UnrealBuildTool
 				);
 
 				StringBuilder Output = new StringBuilder();
-				Output.Append("<?xml version=\"1.0\" encoding=\"utf-8\"?>\n");
+				Output.AppendLine("<?xml version=\"1.0\" encoding=\"utf-8\"?>");
 
 				XmlWriterSettings XmlSettings = new XmlWriterSettings();
 				XmlSettings.Encoding = new UTF8Encoding(false);
