Build Configuration Project scope QoL fixes: Implied project from Target & -Target="..." support.

diff --git a/Engine/Source/Programs/UnrealBuildTool/System/Utils.cs b/Engine/Source/Programs/UnrealBuildTool/System/Utils.cs
--- a/Engine/Source/Programs/UnrealBuildTool/System/Utils.cs
+++ b/Engine/Source/Programs/UnrealBuildTool/System/Utils.cs
@@ -1531,12 +1531,15 @@ namespace UnrealBuildTool
 			string? CandidateProjectPath = null;
 
 			// look for -project=<path>, if it does not exist check arguments for anything that has .uproject in it
-			if (!Arguments.TryGetValue("-Project=", out CandidateProjectPath))
+			if (!Arguments.TryGetValue("-Project=", out CandidateProjectPath) )
 			{
+				Logger.LogDebug("No project arugment provided while trying to parse Project File. Attempting fallback parsing to attempt to find a Project File.");
+
 				// Go through the argument list and try to match poorly (or well..) formed arguments like 
 				// EngineTest, EngineTest.uproject
 				// Collaboration/FooProject
-				// by checking for those in the native project list
+				// by checking for those in the native project list, and checking for common -Target="...... "SomeProjectFile.uproject\"" strings in the multi-target scenario.
+				List<FileReference> TargetProjectFiles = new List<FileReference>();
 				for (int Idx = 0; Idx < Arguments.Count; Idx++)
 				{
 					if (!Arguments[Idx].StartsWith("-") && Arguments[Idx].EndsWith(".uproject", StringComparison.OrdinalIgnoreCase))
@@ -1549,6 +1552,52 @@ namespace UnrealBuildTool
 						}
 						break;
 					}
+					else if (Arguments[Idx].StartsWith("-Target", StringComparison.OrdinalIgnoreCase)) // Potentially a -Target="...." aggregate argument
+					{
+						Logger.LogDebug("-Target argument detected in original command line args. Will attempt to subparse in search of ProjectFile.");
+
+						FileReference? targetProjectFile = null;
+						CommandLineArguments subCommand = new CommandLineArguments(ParseArgumentList(Arguments[Idx]).ToArray());
+
+						if (Arguments[Idx].Contains("-Project", StringComparison.OrdinalIgnoreCase)) //If there's a -Project argument, just recurse and let the rest capture. Ignore any other sources within this -Target="..."
+						{
+							TryParseProjectFileArgument(subCommand, Logger, out targetProjectFile, MarkArgumentAsUsed);
+						}
+						else // Just extract the Target, and attempt to get a project for it.
+						{
+							string? potentialTarget;
+							if (subCommand.TryGetValue("-Target=", out potentialTarget))
+							{
+								NativeProjects.TryGetProjectForTarget(potentialTarget, Logger, out targetProjectFile);
+							}
+						}
+
+						if(targetProjectFile != null)
+						{
+							TargetProjectFiles.Add(targetProjectFile);
+						}
+					}
+					else if (!Arguments[Idx].StartsWith("-Target", StringComparison.OrdinalIgnoreCase)) // Assume it *could* just be a raw target; this amortizes to a cheap dictionary check now.
+					{
+						FileReference? targetProjectFile = null;
+						if (NativeProjects.TryGetProjectForTarget(Arguments[Idx], Logger, out targetProjectFile) && targetProjectFile != null)
+						{
+							TargetProjectFiles.Add(targetProjectFile);
+						}
+					}
+				}
+				
+				if (TargetProjectFiles.Count > 0 && TargetProjectFiles.GroupBy(x => x.FullName.ToUpperInvariant()).Count() == 1)
+				{
+					// We can return early as we have already resolved a file reference.
+					ProjectFile = TargetProjectFiles[0];
+					Logger.LogDebug("Resolved project argument {CandidateProjectPath} to {CandidateProjectFile} through -Target resolution.", ProjectFile, ProjectFile.ToNormalizedPath());
+
+					return true;
+				}
+				else if (TargetProjectFiles.Count > 0 && TargetProjectFiles.GroupBy(x => x.FullName.ToUpperInvariant()).Count() > 1)
+				{
+					Logger.LogInformation("Multiple unique uprojects provided. Cannot determine a single project file to use from command line: will not apply a project scope BuildConfiguration.xml.");
 				}
 			}
 
diff --git a/Engine/Source/Programs/UnrealBuildTool/System/XmlConfig.cs b/Engine/Source/Programs/UnrealBuildTool/System/XmlConfig.cs
--- a/Engine/Source/Programs/UnrealBuildTool/System/XmlConfig.cs
+++ b/Engine/Source/Programs/UnrealBuildTool/System/XmlConfig.cs
@@ -76,6 +76,11 @@ namespace UnrealBuildTool
 			// Find all the configurable types
 			List<Type> configTypes = FindConfigurableTypes();
 
+			if (projectRootDirectory == null)
+			{
+				logger.LogDebug("No project directory provided for project scope BuildConfiguration.xml resolution.");
+			}
+
 			// Update the cache if necessary
 			if (overrideCacheFile != null)
 			{
@@ -117,6 +122,7 @@ namespace UnrealBuildTool
 					FileReference projectRootConfigLocation = FileReference.Combine(projectRootDirectory, "Saved", "UnrealBuildTool", "BuildConfiguration.xml");
 					if (!FileReference.Exists(projectRootConfigLocation))
 					{
+						logger.LogDebug("No project scope config file at {ProjectRootConfigLocation}. Creating default.", projectRootConfigLocation.FullName);
 						CreateDefaultConfigFile(projectRootConfigLocation);
 					}
 
@@ -162,6 +168,7 @@ namespace UnrealBuildTool
 						new Dictionary<Type, Dictionary<XmlConfigData.TargetMember, XmlConfigData.ValueInfo>>();
 					foreach (FileReference inputFile in inputFileLocations)
 					{
+						logger.LogDebug("Reading configuration file from: {Location}", inputFile.FullName);
 						if (!TryReadFile(inputFile, categoryToFields, typeToValues, schema, logger))
 						{
 							throw new BuildException("Failed to properly read XML file : {0}", inputFile.FullName);
