Fix some warnings in generated gradle files:
- remove assignment symbol (=)
- use compileSdk instead of compileSdkVersion

#rnx

diff --git a/Engine/Source/Programs/UnrealBuildTool/Platform/Android/UEDeployAndroid.cs b/Engine/Source/Programs/UnrealBuildTool/Platform/Android/UEDeployAndroid.cs
--- a/Engine/Source/Programs/UnrealBuildTool/Platform/Android/UEDeployAndroid.cs
+++ b/Engine/Source/Programs/UnrealBuildTool/Platform/Android/UEDeployAndroid.cs
@@ -6439,10 +6439,10 @@ import java.util.Collection;
 				StringBuilder BuildGradleContent = new StringBuilder();
 				BuildGradleContent.AppendLine("apply plugin: 'com.android.library'");
 				BuildGradleContent.AppendLine("android {");
-				BuildGradleContent.AppendLine("\tndkVersion = NDK_VERSION");
-				BuildGradleContent.AppendLine("\tndkPath = System.getenv(\"NDKROOT\")");
-				BuildGradleContent.AppendLine("\tcompileSdkVersion = COMPILE_SDK_VERSION.toInteger()");
-				BuildGradleContent.AppendLine("\tbuildToolsVersion = BUILD_TOOLS_VERSION");
+				BuildGradleContent.AppendLine("\tndkVersion NDK_VERSION");
+				BuildGradleContent.AppendLine("\tndkPath System.getenv(\"NDKROOT\")");
+				BuildGradleContent.AppendLine("\tcompileSdk COMPILE_SDK_VERSION.toInteger()");
+				BuildGradleContent.AppendLine("\tbuildToolsVersion BUILD_TOOLS_VERSION");
 				BuildGradleContent.AppendLine("\tdefaultConfig {");
 
 				// Try to get the SDK target from the AndroidManifest.xml
@@ -6517,11 +6517,11 @@ import java.util.Collection;
 				}
 				if (!String.IsNullOrEmpty(MinSdkVersion))
 				{
-					BuildGradleContent.AppendLine(String.Format("\t\tminSdkVersion = {0}", MinSdkVersion));
+					BuildGradleContent.AppendLine(String.Format("\t\tminSdkVersion {0}", MinSdkVersion));
 				}
 				if (!String.IsNullOrEmpty(TargetSdkVersion))
 				{
-					BuildGradleContent.AppendLine(String.Format("\t\ttargetSdkVersion = {0}", TargetSdkVersion));
+					BuildGradleContent.AppendLine(String.Format("\t\ttargetSdkVersion {0}", TargetSdkVersion));
 				}
 				BuildGradleContent.AppendLine("\t}");
 				BuildGradleContent.AppendLine("}");
