UnrealBuildTool: Make QueryTargets data generation parallel

#rnx

diff --git a/Engine/Source/Programs/UnrealBuildTool/Modes/QueryTargetsMode.cs b/Engine/Source/Programs/UnrealBuildTool/Modes/QueryTargetsMode.cs
--- a/Engine/Source/Programs/UnrealBuildTool/Modes/QueryTargetsMode.cs
+++ b/Engine/Source/Programs/UnrealBuildTool/Modes/QueryTargetsMode.cs
@@ -112,9 +112,9 @@ namespace UnrealBuildTool
 			List<string> TargetNames = new List<string>();
 			Assembly.GetAllTargetNames(TargetNames, bIncludeParentAssembly);
 
-			// Write the output file
-			DirectoryReference.CreateDirectory(OutputFile.Directory);
-			using (JsonWriter Writer = new JsonWriter(OutputFile))
+			// Generate json
+			using System.IO.StringWriter StringWriter = new();
+			using (JsonWriter Writer = new(StringWriter))
 			{
 				Writer.WriteObjectStart();
 				Writer.WriteArrayStart("Targets");
@@ -181,6 +181,10 @@ namespace UnrealBuildTool
 				Writer.WriteArrayEnd();
 				Writer.WriteObjectEnd();
 			}
+
+			// Write the output file
+			DirectoryReference.CreateDirectory(OutputFile.Directory);
+			FileReference.WriteAllTextIfDifferent(OutputFile, StringWriter.ToString());
 		}
 	}
 }
diff --git a/Engine/Source/Programs/UnrealBuildTool/ProjectFiles/ProjectFileGenerator.cs b/Engine/Source/Programs/UnrealBuildTool/ProjectFiles/ProjectFileGenerator.cs
--- a/Engine/Source/Programs/UnrealBuildTool/ProjectFiles/ProjectFileGenerator.cs
+++ b/Engine/Source/Programs/UnrealBuildTool/ProjectFiles/ProjectFileGenerator.cs
@@ -1360,11 +1360,38 @@ namespace UnrealBuildTool
 				Logger.LogInformation("");
 				Logger.LogInformation("Generating QueryTargets data for editor...");
 
+				ConcurrentBag<KeyValuePair<FileReference?, RulesAssembly>> ProjectFileToRules = new();
+				using (ProgressWriter Progress = new ProgressWriter("Compiling Rules Assemblies", true, Logger))
+				{
+					int TasksFinished = 0;
+
+					// Engine targets
+					ProjectFileToRules.Add(new(null, RulesCompiler.CreateEngineRulesAssembly(false, false, false, Logger)));
+					Progress.Write(++TasksFinished, AllGameProjects.Count + 1);
+
+					// Project targets
+					foreach (FileReference ProjectFile in AllGameProjects)
+					{
+						RulesAssembly RulesAssembly = RulesCompiler.CreateProjectRulesAssembly(ProjectFile, false, false, false, Logger);
+						ProjectFileToRules.Add(new(ProjectFile, RulesAssembly));
+						Progress.Write(++TasksFinished, AllGameProjects.Count + 1);
+					}
+				}
+
 				// Generate all the target info files for the editor
-				foreach (FileReference ProjectFile in AllGameProjects)
+				using (ProgressWriter Progress = new ProgressWriter("Write Query Target Info", true, Logger))
 				{
-					RulesAssembly RulesAssembly = RulesCompiler.CreateProjectRulesAssembly(ProjectFile, false, false, false, Logger);
-					QueryTargetsMode.WriteTargetInfo(ProjectFile, RulesAssembly, QueryTargetsMode.GetDefaultOutputFile(ProjectFile), new CommandLineArguments(GetTargetArguments(Arguments)), Logger);
+					int TasksFinished = 0;
+					Parallel.ForEach(ProjectFileToRules, Item =>
+					{
+						QueryTargetsMode.WriteTargetInfo(Item.Key, Item.Value, QueryTargetsMode.GetDefaultOutputFile(Item.Key), new CommandLineArguments(GetTargetArguments(Arguments)), Logger);
+
+						lock (Progress)
+						{
+							Interlocked.Increment(ref TasksFinished);
+							Progress.Write(TasksFinished, ProjectFileToRules.Count);
+						}
+					});
 				}
 			}
 
