UnrealBuildTool: Lazy load GetPlatformsInClass return value

#rnx

diff --git a/Engine/Source/Programs/UnrealBuildTool/System/Utils.cs b/Engine/Source/Programs/UnrealBuildTool/System/Utils.cs
--- a/Engine/Source/Programs/UnrealBuildTool/System/Utils.cs
+++ b/Engine/Source/Programs/UnrealBuildTool/System/Utils.cs
@@ -504,19 +504,18 @@ namespace UnrealBuildTool
 		/// <returns>Array of platforms in the given class</returns>
 		public static UnrealTargetPlatform[] GetPlatformsInClass(UnrealPlatformClass Class)
 		{
-			switch (Class)
-			{
-				case UnrealPlatformClass.All:
-					return UnrealTargetPlatform.GetValidPlatforms();
-				case UnrealPlatformClass.Desktop:
-					return UEBuildPlatform.GetPlatformsInGroup(UnrealPlatformGroup.Desktop).ToArray();
-				case UnrealPlatformClass.Editor:
-					return new UnrealTargetPlatform[] { UnrealTargetPlatform.Win64, UnrealTargetPlatform.Linux, UnrealTargetPlatform.Mac };
-				case UnrealPlatformClass.Server:
-					return new UnrealTargetPlatform[] { UnrealTargetPlatform.Win64, UnrealTargetPlatform.Linux, UnrealTargetPlatform.LinuxArm64, UnrealTargetPlatform.Mac };
-			}
-			throw new ArgumentException(String.Format("'{0}' is not a valid value for UnrealPlatformClass", (int)Class));
+			return s_PlatformsInClass.TryGetValue(Class, out Lazy<UnrealTargetPlatform[]>? LazyArray)
+				? LazyArray.Value
+				: throw new ArgumentException($"'{(int)Class}' is not a valid value for {typeof(UnrealPlatformClass).Name}");
+
 		}
+		static Dictionary<UnrealPlatformClass, Lazy<UnrealTargetPlatform[]>> s_PlatformsInClass = new()
+		{
+			{ UnrealPlatformClass.All, new(UnrealTargetPlatform.GetValidPlatforms) },
+			{ UnrealPlatformClass.Desktop, new(() => [.. UEBuildPlatform.GetPlatformsInGroup(UnrealPlatformGroup.Desktop)]) },
+			{ UnrealPlatformClass.Editor, new(() => [UnrealTargetPlatform.Win64, UnrealTargetPlatform.Linux, UnrealTargetPlatform.Mac]) },
+			{ UnrealPlatformClass.Server, new(() => [UnrealTargetPlatform.Win64, UnrealTargetPlatform.Linux, UnrealTargetPlatform.LinuxArm64, UnrealTargetPlatform.Mac]) },
+		};
 
 		/// <summary>
 		/// Retrieve the list of platforms in this group (if any)
