UnrealBuildTool: Fix Default SupportedPlatforms for TargetType.Server

#rnx

diff --git a/Engine/Source/Programs/UnrealBuildTool/Configuration/TargetRules.cs b/Engine/Source/Programs/UnrealBuildTool/Configuration/TargetRules.cs
--- a/Engine/Source/Programs/UnrealBuildTool/Configuration/TargetRules.cs
+++ b/Engine/Source/Programs/UnrealBuildTool/Configuration/TargetRules.cs
@@ -3369,6 +3369,10 @@ namespace UnrealBuildTool
 			{
 				return Utils.GetPlatformsInClass(UnrealPlatformClass.Editor);
 			}
+			else if (Type == TargetType.Server)
+			{
+				return Utils.GetPlatformsInClass(UnrealPlatformClass.Server);
+			}
 			else
 			{
 				return Utils.GetPlatformsInClass(UnrealPlatformClass.All);
diff --git a/Engine/Source/Programs/UnrealBuildTool/ProjectFiles/Rider/RiderProjectFile.cs b/Engine/Source/Programs/UnrealBuildTool/ProjectFiles/Rider/RiderProjectFile.cs
--- a/Engine/Source/Programs/UnrealBuildTool/ProjectFiles/Rider/RiderProjectFile.cs
+++ b/Engine/Source/Programs/UnrealBuildTool/ProjectFiles/Rider/RiderProjectFile.cs
@@ -69,8 +69,6 @@ namespace UnrealBuildTool
 			DirectoryReference ProjectRootFolder = RootPath;
 			List<TargetEntry> FileToTarget = new List<TargetEntry>();
 
-			HashSet<UnrealTargetPlatform> ServerPlatforms = Utils.GetPlatformsInClass(UnrealPlatformClass.Server).ToHashSet();
-
 			foreach (UnrealTargetPlatform Platform in InPlatforms)
 			{
 				foreach (UnrealTargetConfiguration Configuration in InConfigurations)
@@ -96,8 +94,8 @@ namespace UnrealBuildTool
 							continue;
 						}
 
-						// Skip Server for all invalid platforms
-						if (ProjectTarget.TargetRules.Type == TargetType.Server && !ServerPlatforms.Contains(Platform))
+						// Skip invalid platforms for the current TargetType
+						if (!ProjectTarget.SupportedPlatforms.Contains(Platform))
 						{
 							continue;
 						}
diff --git a/Engine/Source/Programs/UnrealBuildTool/ProjectFiles/VisualStudioWorkspace/VSWorkspaceProjectFile.cs b/Engine/Source/Programs/UnrealBuildTool/ProjectFiles/VisualStudioWorkspace/VSWorkspaceProjectFile.cs
--- a/Engine/Source/Programs/UnrealBuildTool/ProjectFiles/VisualStudioWorkspace/VSWorkspaceProjectFile.cs
+++ b/Engine/Source/Programs/UnrealBuildTool/ProjectFiles/VisualStudioWorkspace/VSWorkspaceProjectFile.cs
@@ -104,7 +104,6 @@ namespace UnrealBuildTool
 		{
 			DirectoryReference ProjectRootFolder = RootPath;
 
-			HashSet<UnrealTargetPlatform> ServerPlatforms = Utils.GetPlatformsInClass(UnrealPlatformClass.Server).ToHashSet();
 			ProcessProjectFiles(
 				InPlatforms,
 				InConfigurations,
@@ -187,7 +186,6 @@ namespace UnrealBuildTool
 			Dictionary<string, AvailableConfiguration>? ProjectGenerationDescriptionData,
 			Func<UnrealTargetConfiguration, UnrealTargetPlatform, ProjectTarget, Task> Handler)
 		{
-			HashSet<UnrealTargetPlatform> ServerPlatforms = Utils.GetPlatformsInClass(UnrealPlatformClass.Server).ToHashSet();
 			List<Task> Tasks = new List<Task>();
 			List<(UnrealTargetConfiguration, UnrealTargetPlatform, ProjectTarget)> ret = new List<(UnrealTargetConfiguration, UnrealTargetPlatform, ProjectTarget)>();
 			foreach (UnrealTargetConfiguration Configuration in InConfigurations)
@@ -217,8 +215,8 @@ namespace UnrealBuildTool
 							continue;
 						}
 
-						// Skip Server for all invalid platforms
-						if (ProjectTarget.TargetRules.Type == TargetType.Server && !ServerPlatforms.Contains(Platform))
+						// Skip invalid platforms for the current TargetType
+						if (!ProjectTarget.SupportedPlatforms.Contains(Platform))
 						{
 							continue;
 						}
diff --git a/Engine/Source/Programs/UnrealBuildTool/ProjectFiles/Xcode/XcodeProjectFileGenerator.cs b/Engine/Source/Programs/UnrealBuildTool/ProjectFiles/Xcode/XcodeProjectFileGenerator.cs
--- a/Engine/Source/Programs/UnrealBuildTool/ProjectFiles/Xcode/XcodeProjectFileGenerator.cs
+++ b/Engine/Source/Programs/UnrealBuildTool/ProjectFiles/Xcode/XcodeProjectFileGenerator.cs
@@ -543,17 +543,14 @@ namespace UnrealBuildTool
 
 				ProjectTarget CurTarget = TargetPair.Item2;
 
-				UnrealTargetPlatform[] PlatformsToGenerate = { UnrealTargetPlatform.Mac, UnrealTargetPlatform.IOS };
+				SortedSet<UnrealTargetPlatform> PlatformsToGenerate = [UnrealTargetPlatform.Mac, UnrealTargetPlatform.IOS];
+				PlatformsToGenerate.IntersectWith(CurTarget.SupportedPlatforms);
+
 				foreach (UnrealTargetPlatform Platform in PlatformsToGenerate)
 				{
 					UnrealArch Arch = UnrealArch.Arm64;
-
-					if (!CurTarget.SupportedPlatforms.Any(x => x == Platform))
-					{
-						continue;
-					}
 					TargetDescriptor TargetDesc = new TargetDescriptor(CurTarget.UnrealProjectFilePath, CurTarget.Name, Platform, UnrealTargetConfiguration.Development,
-						new UnrealArchitectures(Arch), new CommandLineArguments(new string[] { "-skipclangvalidation" }));
+						new UnrealArchitectures(Arch), new CommandLineArguments(["-skipclangvalidation"]));
 					TargetDesc.IntermediateEnvironment = UnrealIntermediateEnvironment.GenerateProjectFiles;
 					DateTime Start = DateTime.UtcNow;
 
