UnrealBuildTool: Fix error with SortedSet in AddAdditionalNativeTargetInformation

#rnx

diff --git a/Engine/Source/Programs/UnrealBuildTool/ProjectFiles/Xcode/XcodeProjectFileGenerator.cs b/Engine/Source/Programs/UnrealBuildTool/ProjectFiles/Xcode/XcodeProjectFileGenerator.cs
--- a/Engine/Source/Programs/UnrealBuildTool/ProjectFiles/Xcode/XcodeProjectFileGenerator.cs
+++ b/Engine/Source/Programs/UnrealBuildTool/ProjectFiles/Xcode/XcodeProjectFileGenerator.cs
@@ -543,7 +543,7 @@ namespace UnrealBuildTool
 
 				ProjectTarget CurTarget = TargetPair.Item2;
 
-				SortedSet<UnrealTargetPlatform> PlatformsToGenerate = [UnrealTargetPlatform.Mac, UnrealTargetPlatform.IOS];
+				HashSet<UnrealTargetPlatform> PlatformsToGenerate = [UnrealTargetPlatform.Mac, UnrealTargetPlatform.IOS];
 				PlatformsToGenerate.IntersectWith(CurTarget.SupportedPlatforms);
 
 				foreach (UnrealTargetPlatform Platform in PlatformsToGenerate)
