Fix that default static analysis should not force clang on windows.

diff --git a/Engine/Source/Programs/UnrealBuildTool/Platform/Windows/UEBuildWindows.cs b/Engine/Source/Programs/UnrealBuildTool/Platform/Windows/UEBuildWindows.cs
--- a/Engine/Source/Programs/UnrealBuildTool/Platform/Windows/UEBuildWindows.cs
+++ b/Engine/Source/Programs/UnrealBuildTool/Platform/Windows/UEBuildWindows.cs
@@ -1259,7 +1259,6 @@ namespace UnrealBuildTool
 						break;
 					case StaticAnalyzer.None:
 					case StaticAnalyzer.Clang:
-					case StaticAnalyzer.Default:
 						Target.WindowsPlatform.Compiler = WindowsCompiler.ClangRTFM;
 
 						// We always use clang linker with AutoRTFM.
