[UnrealBuildTool]
- Parallelize WriteConfiguration for targets
-61s to 50s

#rnx

diff --git a/Engine/Source/Programs/UnrealBuildTool/ProjectFiles/VisualStudio/VCProject.cs b/Engine/Source/Programs/UnrealBuildTool/ProjectFiles/VisualStudio/VCProject.cs
--- a/Engine/Source/Programs/UnrealBuildTool/ProjectFiles/VisualStudio/VCProject.cs
+++ b/Engine/Source/Programs/UnrealBuildTool/ProjectFiles/VisualStudio/VCProject.cs
@@ -1405,9 +1405,30 @@ namespace UnrealBuildTool
 			}
 
 			// Write each configuration
-			foreach (ProjectConfigAndTargetCombination Combination in ProjectConfigAndTargetCombinations)
 			{
-				WriteConfiguration(ProjectName, Combination, VCProjectFileContent, PlatformProjectGenerators, bGenerateUserFileContent ? VCUserFileContent : null, bGenerateUserFileContent ? VCUserFileSettings : null);
+				StringBuilder[] TmpVCProjectFileContent = new StringBuilder[ProjectConfigAndTargetCombinations.Count];
+				StringBuilder[] TmpVCUserFileContent = new StringBuilder[ProjectConfigAndTargetCombinations.Count];
+
+				Parallel.For(0, ProjectConfigAndTargetCombinations.Count, Index =>
+				{
+					TmpVCProjectFileContent[Index] = new StringBuilder();
+					TmpVCUserFileContent[Index] = new StringBuilder();
+
+					WriteConfiguration(ProjectName, ProjectConfigAndTargetCombinations[Index], TmpVCProjectFileContent[Index], PlatformProjectGenerators, bGenerateUserFileContent ? TmpVCUserFileContent[Index] : null, bGenerateUserFileContent ? VCUserFileSettings : null);
+				});
+
+				for (int Index = 0; Index < ProjectConfigAndTargetCombinations.Count; Index++)
+				{
+					if (TmpVCProjectFileContent[Index].Length > 0) 
+					{
+						VCProjectFileContent.Append(TmpVCProjectFileContent[Index]);
+					}
+
+					if (TmpVCUserFileContent[Index].Length > 0) 
+					{
+						VCUserFileContent.Append(TmpVCUserFileContent[Index]);
+					}
+				}
 			}
 
 			{
