UBT: Add a flag to force disable entitlements from the commandline.

diff --git a/Engine/Source/Programs/UnrealBuildTool/Platform/Apple/AppleExports.cs b/Engine/Source/Programs/UnrealBuildTool/Platform/Apple/AppleExports.cs
--- a/Engine/Source/Programs/UnrealBuildTool/Platform/Apple/AppleExports.cs
+++ b/Engine/Source/Programs/UnrealBuildTool/Platform/Apple/AppleExports.cs
@@ -17,6 +17,16 @@ namespace UnrealBuildTool
 	{
 		private static bool bForceModernXcode = Environment.CommandLine.Contains("-modernxcode", StringComparison.OrdinalIgnoreCase);
 		private static bool bForceLegacyXcode = Environment.CommandLine.Contains("-legacyxcode", StringComparison.OrdinalIgnoreCase);
+		private static bool bNoEntitlements = Environment.CommandLine.Contains("-noEntitlements", StringComparison.OrdinalIgnoreCase);
+
+		/// <summary>
+		///  True if we should force disable entitlements based on "-noEntitlements"
+		/// </summary>
+		/// <returns>bool</returns>
+		public static bool ForceNoEntitlements()
+		{
+			return bNoEntitlements;
+		}
 
 		/// <summary>
 		/// Is the current project using modern xcode?
@@ -217,7 +227,7 @@ namespace UnrealBuildTool
 		{
 			DirectoryReference? GeneratedProjectFile;
 			// we don't use distro flag when making a modern project
-			AppleExports.GenerateRunOnlyXcodeProject(ProjectFile, Platform, TargetName, bForDistribution: false, Logger, out GeneratedProjectFile);
+			AppleExports.GenerateRunOnlyXcodeProject(ProjectFile, Platform, TargetName, bForDistribution: false, bNoEntitlements: bNoEntitlements, Logger, out GeneratedProjectFile);
 
 			if (GeneratedProjectFile == null)
 			{
@@ -262,11 +272,12 @@ namespace UnrealBuildTool
 		/// <param name="Platform">The platform to generate a project for</param>
 		/// <param name="TargetName">The name of the target being built, so we can generate a more minimal project</param>
 		/// <param name="bForDistribution">True if this is making a bild for uploading to app store</param>
+		/// <param name="bNoEntitlements">True if we should disable entitlements</param>
 		/// <param name="Logger">Logging object</param>
 		/// <param name="GeneratedProjectFile">Returns the .xcworkspace that was made</param>
-		public static void GenerateRunOnlyXcodeProject(FileReference? UProjectFile, UnrealTargetPlatform Platform, string TargetName, bool bForDistribution, ILogger Logger, out DirectoryReference? GeneratedProjectFile)
+		public static void GenerateRunOnlyXcodeProject(FileReference? UProjectFile, UnrealTargetPlatform Platform, string TargetName, bool bForDistribution, bool bNoEntitlements, ILogger Logger, out DirectoryReference? GeneratedProjectFile)
 		{
-			AppleToolChain.GenerateRunOnlyXcodeProject(UProjectFile, Platform, TargetName, bForDistribution, Logger, out GeneratedProjectFile);
+			AppleToolChain.GenerateRunOnlyXcodeProject(UProjectFile, Platform, TargetName, bForDistribution, bNoEntitlements, Logger, out GeneratedProjectFile);
 		}
 
 		/// <summary>
diff --git a/Engine/Source/Programs/UnrealBuildTool/ProjectFiles/Xcode/XcodeProject.cs b/Engine/Source/Programs/UnrealBuildTool/ProjectFiles/Xcode/XcodeProject.cs
--- a/Engine/Source/Programs/UnrealBuildTool/ProjectFiles/Xcode/XcodeProject.cs
+++ b/Engine/Source/Programs/UnrealBuildTool/ProjectFiles/Xcode/XcodeProject.cs
@@ -372,6 +372,11 @@ namespace UnrealBuildTool.XcodeProjectXcconfig
 		/// </summary>
 		public bool bForDistribution = false;
 
+		/// <summary>
+		///  Used to force not using the entitlements
+		/// </summary>
+		public bool bNoEntitlements = false;
+
 		/// <summary>
 		/// Override for bundle identifier
 		/// </summary>
@@ -414,7 +419,7 @@ namespace UnrealBuildTool.XcodeProjectXcconfig
 			return bIsAppBundle;
 		}
 
-		public UnrealData(FileReference XcodeProjectFileLocation, bool bIsForDistribution, string BundleID, string AppName, bool bMakeProjectPerTarget)
+		public UnrealData(FileReference XcodeProjectFileLocation, bool bIsForDistribution, bool bForceNoEntitlements, string BundleID, string AppName, bool bMakeProjectPerTarget)
 		{
 			// the .xcodeproj is actually a directory
 			this.XcodeProjectFileLocation = new DirectoryReference(XcodeProjectFileLocation.FullName);
@@ -428,6 +433,7 @@ namespace UnrealBuildTool.XcodeProjectXcconfig
 
 			this.bMakeProjectPerTarget = bMakeProjectPerTarget;
 			bForDistribution = bIsForDistribution;
+			bNoEntitlements = bForceNoEntitlements;
 			BundleIdentifier = BundleID;
 			DisplayName = String.IsNullOrEmpty(AppName) ? "$(UE_PRODUCT_NAME)" : AppName;
 		}
@@ -2109,7 +2115,7 @@ namespace UnrealBuildTool.XcodeProjectXcconfig
 				}
 
 				// programs almost never want entitlements, so handle them specially
-				if (UnrealData.TargetRules.Type == TargetType.Program && !bUseEntitlementsForPrograms)
+				if ((UnrealData.bNoEntitlements) || (UnrealData.TargetRules.Type == TargetType.Program && !bUseEntitlementsForPrograms))
 				{
 					EntitlementsMetadata = null;
 				}
@@ -2348,14 +2354,15 @@ namespace UnrealBuildTool.XcodeProjectXcconfig
 		/// <param name="InitFilePath">The path to the project file on disk</param>
 		/// <param name="BaseDir">The base directory for files within this project</param>
 		/// <param name="bIsForDistribution">True for distribution builds</param>
+		/// <param name="bForceNoEntitlements">True to force disable entitlements</param>
 		/// <param name="BundleID">Override option for bundle identifier</param>
 		/// <param name="AppName"></param>
 		/// <param name="bMakeProjectPerTarget"></param>
 		/// <param name="SingleTargetName"></param>
-		public XcodeProjectFile(FileReference InitFilePath, DirectoryReference BaseDir, bool bIsForDistribution, string BundleID, string AppName, bool bMakeProjectPerTarget, string? SingleTargetName)
+		public XcodeProjectFile(FileReference InitFilePath, DirectoryReference BaseDir, bool bIsForDistribution, bool bForceNoEntitlements, string BundleID, string AppName, bool bMakeProjectPerTarget, string? SingleTargetName)
 			: base(InitFilePath, BaseDir)
 		{
-			UnrealData = new UnrealData(InitFilePath, bIsForDistribution, BundleID, AppName, bMakeProjectPerTarget);
+			UnrealData = new UnrealData(InitFilePath, bIsForDistribution, bForceNoEntitlements, BundleID, AppName, bMakeProjectPerTarget);
 
 			// create the container for all the files that will 
 			SharedFileCollection = new XcodeFileCollection(this);
diff --git a/Engine/Source/Programs/UnrealBuildTool/ProjectFiles/Xcode/XcodeProjectFileGenerator.cs b/Engine/Source/Programs/UnrealBuildTool/ProjectFiles/Xcode/XcodeProjectFileGenerator.cs
--- a/Engine/Source/Programs/UnrealBuildTool/ProjectFiles/Xcode/XcodeProjectFileGenerator.cs
+++ b/Engine/Source/Programs/UnrealBuildTool/ProjectFiles/Xcode/XcodeProjectFileGenerator.cs
@@ -43,6 +43,11 @@ namespace UnrealBuildTool
 		/// </summary>
 		bool bForDistribution = false;
 
+		/// <summary>
+		/// Don't use entitlements
+		/// </summary>
+		bool bNoEntitlements = false;
+
 		/// <summary>
 		/// Override BundleID
 		/// </summary>
@@ -102,11 +107,14 @@ namespace UnrealBuildTool
 			: base(InOnlyGameProject)
 		{
 			SingleGameProject = InOnlyGameProject;
-
 			if (CommandLine.HasOption("-distribution"))
 			{
 				bForDistribution = true;
 			}
+			if (CommandLine.HasOption("-noEntitlements"))
+			{
+				bNoEntitlements = true;
+			}
 			if (CommandLine.HasValue("-bundleID="))
 			{
 				BundleIdentifier = CommandLine.GetString("-bundleID=");
@@ -204,7 +212,7 @@ namespace UnrealBuildTool
 			// unfortunately, we can't read the project configs now because we don't have enough information to 
 			// find the .uproject file for that would make this project (we could change the high level to pass it 
 			// down but it would touch all project generators - not worth it if we end up removing the legacy)
-			return new XcodeProjectXcconfig.XcodeProjectFile(InitFilePath, BaseDir, bForDistribution, BundleIdentifier, AppName, bMakeProjectPerTarget, SingleTargetName);
+			return new XcodeProjectXcconfig.XcodeProjectFile(InitFilePath, BaseDir, bForDistribution, bNoEntitlements, BundleIdentifier, AppName, bMakeProjectPerTarget, SingleTargetName);
 		}
 
 		private bool WriteWorkspaceSettingsFile(string Path, ILogger Logger)
diff --git a/Engine/Source/Programs/UnrealBuildTool/ToolChain/AppleToolChain.cs b/Engine/Source/Programs/UnrealBuildTool/ToolChain/AppleToolChain.cs
--- a/Engine/Source/Programs/UnrealBuildTool/ToolChain/AppleToolChain.cs
+++ b/Engine/Source/Programs/UnrealBuildTool/ToolChain/AppleToolChain.cs
@@ -799,7 +799,7 @@ namespace UnrealBuildTool
 			bool bIsBuildingAppBundle = !BinaryLinkEnvironment.bIsBuildingDLL && !BinaryLinkEnvironment.bIsBuildingLibrary && !BinaryLinkEnvironment.bIsBuildingConsoleApplication;
 			if (AppleExports.UseModernXcode(Target.ProjectFile) && bIsBuildingAppBundle)
 			{
-				Action PostBuildAction = ApplePostBuildSyncMode.CreatePostBuildSyncAction(Target, Executable, BinaryLinkEnvironment.IntermediateDirectory!, Graph);
+				Action PostBuildAction = ApplePostBuildSyncMode.CreatePostBuildSyncAction(Target, Executable, BinaryLinkEnvironment.IntermediateDirectory!, Graph, AppleExports.ForceNoEntitlements());
 
 				PostBuildAction.PrerequisiteItems.UnionWith(OutputFiles);
 				OutputFiles.AddRange(PostBuildAction.ProducedItems);
@@ -871,9 +871,10 @@ namespace UnrealBuildTool
 		/// <param name="Platform">The platform to generate a project for</param>
 		/// <param name="TargetName">The name of the target being built, so we can generate a more minimal project</param>
 		/// <param name="bForDistribution">True if this is making a bild for uploading to app store</param>
+		/// <param name="bNoEntitlements">True if we shouldn't use entitlements</param>
 		/// <param name="Logger">Logging object</param>
 		/// <param name="GeneratedProjectFile">Returns the .xcworkspace that was made</param>
-		internal static void GenerateRunOnlyXcodeProject(FileReference? UProjectFile, UnrealTargetPlatform Platform, string TargetName, bool bForDistribution, ILogger Logger, out DirectoryReference? GeneratedProjectFile)
+		internal static void GenerateRunOnlyXcodeProject(FileReference? UProjectFile, UnrealTargetPlatform Platform, string TargetName, bool bForDistribution, bool bNoEntitlements, ILogger Logger, out DirectoryReference? GeneratedProjectFile)
 		{
 			List<string> Options = new()
 			{
@@ -882,6 +883,7 @@ namespace UnrealBuildTool
 				"-NoIntellisense",
 				"-NoDotNet",
 				"-IgnoreJunk",
+				bNoEntitlements ? "-noEntitlements" : "",
 				bForDistribution ? "-distribution" : "-development",
 				"-IncludeTempTargets",
 				"-projectfileformat=XCode",
@@ -1051,6 +1053,9 @@ namespace UnrealBuildTool
 		[CommandLine("-ModernXcode")]
 		public bool bModernXcode;
 
+		[CommandLine("-noEntitlements")]
+		public bool bNoEntitlements;
+
 		public override Task<int> ExecuteAsync(CommandLineArguments Arguments, ILogger Logger)
 		{
 			Arguments.ApplyTo(this);
@@ -1214,13 +1219,12 @@ namespace UnrealBuildTool
 			return FileItem.GetItemByFileReference(StagedExe);
 		}
 
-		public static Action CreatePostBuildSyncAction(ReadOnlyTargetRules Target, FileItem Executable, DirectoryReference IntermediateDir, IActionGraphBuilder Graph)
+		public static Action CreatePostBuildSyncAction(ReadOnlyTargetRules Target, FileItem Executable, DirectoryReference IntermediateDir, IActionGraphBuilder Graph, bool bNoEntitlements = false)
 		{
 			ApplePostBuildSyncTarget PostBuildSync = new(Target, Executable, IntermediateDir);
 			FileReference PostBuildSyncFile = FileReference.Combine(IntermediateDir!, "PostBuildSync.dat");
-			BinaryFormatterUtils.Save(PostBuildSyncFile, PostBuildSync);
-
-			string PostBuildSyncArguments = String.Format("-modernxcode -Input=\"{0}\" -XmlConfigCache=\"{1}\" -remoteini=\"{2}\"", PostBuildSyncFile, XmlConfig.CacheFile, UnrealBuildTool.GetRemoteIniPath());
+			BinaryFormatterUtils.Save(PostBuildSyncFile, PostBuildSync);			
+			string PostBuildSyncArguments = String.Format("-modernxcode -Input=\"{0}\" -XmlConfigCache=\"{1}\" -remoteini=\"{2}\" {3}", PostBuildSyncFile, XmlConfig.CacheFile, UnrealBuildTool.GetRemoteIniPath(), bNoEntitlements ? "-noEntitlements" : "");
 			Action PostBuildSyncAction = Graph.CreateRecursiveAction<ApplePostBuildSyncMode>(ActionType.CreateAppBundle, PostBuildSyncArguments);
 			PostBuildSyncAction.PrerequisiteItems.Add(Executable);
 			PostBuildSyncAction.ProducedItems.Add(GetPostBuildOutputFile(Executable.Location, Target.Name, Target.Platform));
