[UBT]
* Fixed so link event matcher also handles ld.lld and lld-link for clang targets

diff --git a/Engine/Source/Programs/UnrealBuildTool/Matchers/LinkEventMatcher.cs b/Engine/Source/Programs/UnrealBuildTool/Matchers/LinkEventMatcher.cs
--- a/Engine/Source/Programs/UnrealBuildTool/Matchers/LinkEventMatcher.cs
+++ b/Engine/Source/Programs/UnrealBuildTool/Matchers/LinkEventMatcher.cs
@@ -19,7 +19,7 @@ namespace UnrealBuildTool.Matchers
 			@"^(\s*)Undefined symbols for architecture");
 
 		static readonly Regex s_ldClangPattern = new Regex(
-			@"^\s*(ld|clang):");
+			@"^\s*(ld|ld.lld|lld-link|clang):");
 
 		static readonly Regex s_ldMultiplyDefinedPattern = new Regex(
 			@"(?<severity>error): L\d+: symbol `(?<symbol>[^`]+)' multiply defined$");
