UBT: Split out GenerateQueryTargetsDataForEditor() from GenerateProjectFiles() and call it explicitly at the end instead of nested inside the first generator to clean up the code that was recently optimized in cl 39266718.
Remove the bCacheDataForEditor parameter from the ProjectFileGenerator.GenerateProjectFiles signature and mark the old member function obsolete.

#rnx

diff --git a/Engine/Source/Programs/UnrealBuildTool/Modes/GenerateProjectFilesMode.cs b/Engine/Source/Programs/UnrealBuildTool/Modes/GenerateProjectFilesMode.cs
--- a/Engine/Source/Programs/UnrealBuildTool/Modes/GenerateProjectFilesMode.cs
+++ b/Engine/Source/Programs/UnrealBuildTool/Modes/GenerateProjectFilesMode.cs
@@ -248,8 +248,7 @@ namespace UnrealBuildTool
 
 			// Now generate project files
 			ProjectFileGenerator.bGenerateProjectFiles = true;
-			// perform anything that only needs to happen one time, in the first project genereator
-			bool bPerformOneTimeOperations = true;
+			bool bGenerateSuccess = true;
 			foreach (KeyValuePair<ProjectFileFormat, ProjectFileGenerator> Pair in Generators)
 			{
 				Logger.LogInformation("");
@@ -257,18 +256,21 @@ namespace UnrealBuildTool
 
 				ProjectFileGenerator.Current = Pair.Value;
 				Arguments.ApplyTo(Pair.Value);
-				bool bGenerateSuccess = Pair.Value.GenerateProjectFiles(PlatformProjectGenerators, Arguments.GetRawArray(), bCacheDataForEditor: bPerformOneTimeOperations, Logger);
+#pragma warning disable 0618 // Type or member is obsolete
+				bGenerateSuccess = Pair.Value.GenerateProjectFiles(PlatformProjectGenerators, Arguments.GetRawArray(), false, Logger);
+#pragma warning restore 0618 // Type or member is obsolete
 				ProjectFileGenerator.Current = null;
 
 				if (!bGenerateSuccess)
 				{
-					return Task.FromResult((int)CompilationResult.OtherCompilationError);
+					break;
 				}
-
-				// any further generators can skip one-time operations
-				bPerformOneTimeOperations = false;
 			}
-			return Task.FromResult((int)CompilationResult.Succeeded);
+			if (bGenerateSuccess && Generators.Count > 0)
+			{
+				Generators.First().Value.GenerateQueryTargetsDataForEditor(Arguments.GetRawArray(), Logger);
+			}
+			return bGenerateSuccess ? Task.FromResult((int)CompilationResult.Succeeded) : Task.FromResult((int)CompilationResult.OtherCompilationError);
 		}
 	}
 }
diff --git a/Engine/Source/Programs/UnrealBuildTool/ProjectFiles/ProjectFileGenerator.cs b/Engine/Source/Programs/UnrealBuildTool/ProjectFiles/ProjectFileGenerator.cs
--- a/Engine/Source/Programs/UnrealBuildTool/ProjectFiles/ProjectFileGenerator.cs
+++ b/Engine/Source/Programs/UnrealBuildTool/ProjectFiles/ProjectFileGenerator.cs
@@ -880,7 +880,20 @@ namespace UnrealBuildTool
 		/// <param name="Arguments">Command-line arguments</param>
 		/// <param name="bCacheDataForEditor">If true, write out target data for the editor</param>
 		/// <param name="Logger">Logger for output</param>
+		[Obsolete("Deprecated in UE5.6 - bCacheDataForEditor has been removed from the GenerateProjectFiles signature.")]
 		public virtual bool GenerateProjectFiles(PlatformProjectGeneratorCollection PlatformProjectGenerators, string[] Arguments, bool bCacheDataForEditor, ILogger Logger)
+		{
+			return GenerateProjectFiles(PlatformProjectGenerators, Arguments, Logger);
+		}
+
+		/// <summary>
+		/// Generates a Visual Studio solution file and Visual C++ project files for all known engine and game targets.
+		/// Does not actually build anything.
+		/// </summary>
+		/// <param name="PlatformProjectGenerators">The registered platform project generators</param>
+		/// <param name="Arguments">Command-line arguments</param>
+		/// <param name="Logger">Logger for output</param>
+		public virtual bool GenerateProjectFiles(PlatformProjectGeneratorCollection PlatformProjectGenerators, string[] Arguments, ILogger Logger)
 		{
 			bool bSuccess = true;
 
@@ -1354,48 +1367,50 @@ namespace UnrealBuildTool
 				WriteProjectFiles(PlatformProjectGenerators, Logger);
 				Logger.LogDebug("Project generation complete ({GeneratedProjectFilesCount} generated, {OtherProjectFilesCount} imported)", GeneratedProjectFiles.Count, OtherProjectFiles.Count);
 			}
+			return bSuccess;
+		}
+
+		public void GenerateQueryTargetsDataForEditor(string[] Arguments, ILogger Logger)
+		{
+			Logger.LogInformation("");
+			Logger.LogInformation("Generating QueryTargets data for editor...");
+
+			// Build the list of games to generate projects for
+			List<FileReference> AllGameProjects = FindGameProjects(Logger);
 
-			if (bCacheDataForEditor)
+			ConcurrentBag<KeyValuePair<FileReference?, RulesAssembly>> ProjectFileToRules = new();
+			using (ProgressWriter Progress = new ProgressWriter("Compiling Rules Assemblies", true, Logger))
 			{
-				Logger.LogInformation("");
-				Logger.LogInformation("Generating QueryTargets data for editor...");
+				int TasksFinished = 0;
 
-				ConcurrentBag<KeyValuePair<FileReference?, RulesAssembly>> ProjectFileToRules = new();
-				using (ProgressWriter Progress = new ProgressWriter("Compiling Rules Assemblies", true, Logger))
-				{
-					int TasksFinished = 0;
+				// Engine targets
+				ProjectFileToRules.Add(new(null, RulesCompiler.CreateEngineRulesAssembly(false, false, false, Logger)));
+				Progress.Write(++TasksFinished, AllGameProjects.Count + 1);
 
-					// Engine targets
-					ProjectFileToRules.Add(new(null, RulesCompiler.CreateEngineRulesAssembly(false, false, false, Logger)));
+				// Project targets
+				foreach (FileReference ProjectFile in AllGameProjects)
+				{
+					RulesAssembly RulesAssembly = RulesCompiler.CreateProjectRulesAssembly(ProjectFile, false, false, false, Logger);
+					ProjectFileToRules.Add(new(ProjectFile, RulesAssembly));
 					Progress.Write(++TasksFinished, AllGameProjects.Count + 1);
-
-					// Project targets
-					foreach (FileReference ProjectFile in AllGameProjects)
-					{
-						RulesAssembly RulesAssembly = RulesCompiler.CreateProjectRulesAssembly(ProjectFile, false, false, false, Logger);
-						ProjectFileToRules.Add(new(ProjectFile, RulesAssembly));
-						Progress.Write(++TasksFinished, AllGameProjects.Count + 1);
-					}
 				}
+			}
 
-				// Generate all the target info files for the editor
-				using (ProgressWriter Progress = new ProgressWriter("Write Query Target Info", true, Logger))
+			// Generate all the target info files for the editor
+			using (ProgressWriter Progress = new ProgressWriter("Writing Query Target Info", true, Logger))
+			{
+				int TasksFinished = 0;
+				Parallel.ForEach(ProjectFileToRules, Item =>
 				{
-					int TasksFinished = 0;
-					Parallel.ForEach(ProjectFileToRules, Item =>
-					{
-						QueryTargetsMode.WriteTargetInfo(Item.Key, Item.Value, QueryTargetsMode.GetDefaultOutputFile(Item.Key), new CommandLineArguments(GetTargetArguments(Arguments)), Logger);
+					QueryTargetsMode.WriteTargetInfo(Item.Key, Item.Value, QueryTargetsMode.GetDefaultOutputFile(Item.Key), new CommandLineArguments(GetTargetArguments(Arguments)), Logger);
 
-						lock (Progress)
-						{
-							Interlocked.Increment(ref TasksFinished);
-							Progress.Write(TasksFinished, ProjectFileToRules.Count);
-						}
-					});
-				}
+					lock (Progress)
+					{
+						Interlocked.Increment(ref TasksFinished);
+						Progress.Write(TasksFinished, ProjectFileToRules.Count);
+					}
+				});
 			}
-
-			return bSuccess;
 		}
 
 		private bool DoesProgramMatchOnlyGameProject(FileReference TargetFile)
diff --git a/Engine/Source/Programs/UnrealBuildTool/ProjectFiles/Rider/RiderProjectFileGenerator.cs b/Engine/Source/Programs/UnrealBuildTool/ProjectFiles/Rider/RiderProjectFileGenerator.cs
--- a/Engine/Source/Programs/UnrealBuildTool/ProjectFiles/Rider/RiderProjectFileGenerator.cs
+++ b/Engine/Source/Programs/UnrealBuildTool/ProjectFiles/Rider/RiderProjectFileGenerator.cs
@@ -432,8 +432,7 @@ namespace UnrealBuildTool
 			SetupSupportedPlatformsAndConfigurations(true, Logger, out SupportedPlatformNames);
 		}
 
-		public override bool GenerateProjectFiles(PlatformProjectGeneratorCollection PlatformProjectGenerators,
-			string[] arguments, bool bCacheDataForEditor, ILogger Logger)
+		public override bool GenerateProjectFiles(PlatformProjectGeneratorCollection PlatformProjectGenerators, string[] arguments, ILogger Logger)
 		{
 			ConfigureProjectFileGeneration();
 
diff --git a/Engine/Source/Programs/UnrealBuildTool/ProjectFiles/VisualStudioWorkspace/VSWorkspaceProjectFileGenerator.cs b/Engine/Source/Programs/UnrealBuildTool/ProjectFiles/VisualStudioWorkspace/VSWorkspaceProjectFileGenerator.cs
--- a/Engine/Source/Programs/UnrealBuildTool/ProjectFiles/VisualStudioWorkspace/VSWorkspaceProjectFileGenerator.cs
+++ b/Engine/Source/Programs/UnrealBuildTool/ProjectFiles/VisualStudioWorkspace/VSWorkspaceProjectFileGenerator.cs
@@ -210,8 +210,7 @@ namespace UnrealBuildTool
 			return true;
 		}
 
-		public override bool GenerateProjectFiles(PlatformProjectGeneratorCollection PlatformProjectGenerators,
-			string[] Arguments, bool bCacheDataForEditor, ILogger Logger)
+		public override bool GenerateProjectFiles(PlatformProjectGeneratorCollection PlatformProjectGenerators, string[] Arguments, ILogger Logger)
 		{
 			if (ShowVersionInfo)
 			{
diff --git a/Engine/Source/Programs/UnrealBuildTool/ToolChain/AppleToolChain.cs b/Engine/Source/Programs/UnrealBuildTool/ToolChain/AppleToolChain.cs
--- a/Engine/Source/Programs/UnrealBuildTool/ToolChain/AppleToolChain.cs
+++ b/Engine/Source/Programs/UnrealBuildTool/ToolChain/AppleToolChain.cs
@@ -847,7 +847,7 @@ namespace UnrealBuildTool
 				Generator.SingleTargetName = SingleTargetName;
 				// don't need the editor data since these are stub projects
 				ProjectFileGenerator.Current = Generator;
-				bool bSucces = Generator.GenerateProjectFiles(PlatformProjectGenerators, Arguments, bCacheDataForEditor: false, Logger);
+				bool bSucces = Generator.GenerateProjectFiles(PlatformProjectGenerators, Arguments, Logger);
 				ProjectFileGenerator.Current = null;
 				XcodeProjectFile = Generator.XCWorkspace;
 				return bSucces;
