Implement support for android:minAspectRatio, generally useful for foldables with square aspect ratios and apps that may want to cut off the aspect ratio to 4:3 for example. Additionally, adjust Max Aspect Ratio default to 2.4 to not letterbox on many newer devices (with aspect rations above 2.1) by default.

diff --git a/Engine/Source/Programs/UnrealBuildTool/Platform/Android/UEDeployAndroid.cs b/Engine/Source/Programs/UnrealBuildTool/Platform/Android/UEDeployAndroid.cs
--- a/Engine/Source/Programs/UnrealBuildTool/Platform/Android/UEDeployAndroid.cs
+++ b/Engine/Source/Programs/UnrealBuildTool/Platform/Android/UEDeployAndroid.cs
@@ -2668,10 +2668,15 @@ namespace UnrealBuildTool
 			string VersionDisplayName = GetVersionDisplayName(bIsEmbedded);
 			string DepthBufferPreference;
 			Ini.GetString("/Script/AndroidRuntimeSettings.AndroidRuntimeSettings", "DepthBufferPreference", out DepthBufferPreference);
+			float MinAspectRatioValue;
+			if (!Ini.TryGetValue("/Script/AndroidRuntimeSettings.AndroidRuntimeSettings", "MinAspectRatio", out MinAspectRatioValue))
+			{
+				MinAspectRatioValue = 0.0f;
+			}
 			float MaxAspectRatioValue;
 			if (!Ini.TryGetValue("/Script/AndroidRuntimeSettings.AndroidRuntimeSettings", "MaxAspectRatio", out MaxAspectRatioValue))
 			{
-				MaxAspectRatioValue = 2.1f;
+				MaxAspectRatioValue = 2.4f;
 			}
 			string Orientation = ConvertOrientationIniValue(GetOrientation(Arch));
 			bool EnableFullScreen;
@@ -2880,6 +2885,18 @@ namespace UnrealBuildTool
 				// work around scoped storage for non-distribution for SDK 29; add to ExtraApplicationNodeTags if you need it for distribution
 				Text.AppendLine("\t				android:requestLegacyExternalStorage=\"true\"");
 			}
+
+			// Max supported aspect ratio
+			string MaxAspectRatioString = MaxAspectRatioValue.ToString("f", System.Globalization.CultureInfo.InvariantCulture);
+			if (TargetSDKVersion >= 29)
+			{
+				if(MinAspectRatioValue > 0.0)
+				{
+					string MinAspectRatioString = MinAspectRatioValue.ToString("f", System.Globalization.CultureInfo.InvariantCulture);
+					Text.AppendLine(String.Format("\t\t\tandroid:minAspectRatio=\"{0}\"", MinAspectRatioString));
+				}
+				Text.AppendLine(String.Format("\t\t\tandroid:maxAspectRatio=\"{0}\"", MaxAspectRatioString));
+			}
 			Text.AppendLine("\t             android:hasCode=\"true\">");
 			if (bShowLaunchImage)
 			{
@@ -3076,8 +3093,6 @@ namespace UnrealBuildTool
 			Text.AppendLine("\t\t\t</intent-filter>");
 			Text.AppendLine("\t\t</receiver>");
 
-			// Max supported aspect ratio
-			string MaxAspectRatioString = MaxAspectRatioValue.ToString("f", System.Globalization.CultureInfo.InvariantCulture);
 			Text.AppendLine(String.Format("\t\t<meta-data android:name=\"android.max_aspect\" android:value=\"{0}\" />", MaxAspectRatioString));
 
 			Text.AppendLine("\t</application>");
diff --git a/Engine/Source/Runtime/Android/AndroidRuntimeSettings/Classes/AndroidRuntimeSettings.h b/Engine/Source/Runtime/Android/AndroidRuntimeSettings/Classes/AndroidRuntimeSettings.h
--- a/Engine/Source/Runtime/Android/AndroidRuntimeSettings/Classes/AndroidRuntimeSettings.h
+++ b/Engine/Source/Runtime/Android/AndroidRuntimeSettings/Classes/AndroidRuntimeSettings.h
@@ -279,6 +279,10 @@ public:
 	UPROPERTY(GlobalConfig, EditAnywhere, Category = "APK Packaging")
 	TEnumAsByte<EAndroidScreenOrientation::Type> Orientation;
 
+	// Minimum supported aspect ratio (width / height). Android will automatically letterbox application on devices with smaller aspect ratio
+	UPROPERTY(GlobalConfig, EditAnywhere, Category = "APK Packaging", Meta = (DisplayName = "Minimum supported aspect ratio."))
+	float MinAspectRatio;
+
 	// Maximum supported aspect ratio (width / height). Android will automatically letterbox application on devices with bigger aspect ratio
 	UPROPERTY(GlobalConfig, EditAnywhere, Category = "APK Packaging", Meta = (DisplayName = "Maximum supported aspect ratio."))
 	float MaxAspectRatio;
diff --git a/Engine/Source/Runtime/Android/AndroidRuntimeSettings/Private/AndroidRuntimeSettings.cpp b/Engine/Source/Runtime/Android/AndroidRuntimeSettings/Private/AndroidRuntimeSettings.cpp
--- a/Engine/Source/Runtime/Android/AndroidRuntimeSettings/Private/AndroidRuntimeSettings.cpp
+++ b/Engine/Source/Runtime/Android/AndroidRuntimeSettings/Private/AndroidRuntimeSettings.cpp
@@ -21,7 +21,8 @@ DEFINE_LOG_CATEGORY(LogAndroidRuntimeSettings);
 UAndroidRuntimeSettings::UAndroidRuntimeSettings(const FObjectInitializer& ObjectInitializer)
 	: Super(ObjectInitializer)
 	, Orientation(EAndroidScreenOrientation::Landscape)
-	, MaxAspectRatio(2.1f)
+	, MinAspectRatio(0.0f)
+	, MaxAspectRatio(2.4f)
 	, bAndroidVoiceEnabled(false)
 	, bPackageForMetaQuest(false)
 	, bEnableGooglePlaySupport(false)
