Fix for Server AutoRTFM clang build with static analysis

diff --git a/Engine/Source/Programs/UnrealBuildTool/Platform/Windows/UEBuildWindows.cs b/Engine/Source/Programs/UnrealBuildTool/Platform/Windows/UEBuildWindows.cs
--- a/Engine/Source/Programs/UnrealBuildTool/Platform/Windows/UEBuildWindows.cs
+++ b/Engine/Source/Programs/UnrealBuildTool/Platform/Windows/UEBuildWindows.cs
@@ -1187,7 +1187,7 @@ namespace UnrealBuildTool
 						Target.GlobalDefinitions.Add("UE_EXTERNAL_PROFILING_ENABLED=0");
 					}
 
-					// disbling some plugins until we get arm64 libs
+					// disabling some plugins until we get arm64 libs
 					if (Target.WindowsPlatform.Architecture == UnrealArch.Arm64)
 					{
 						Target.DisablePlugins.AddRange(new string[]
@@ -1224,6 +1224,25 @@ namespace UnrealBuildTool
 				Target.bCompileCEF3 = false;
 			}
 
+			if (Target.bUseAutoRTFMCompiler)
+			{
+				// We check the static analyzer, and only allow using the Auto-RTFM compiler if
+				// we are not doing static analysis, or are explicitly wanting to use clang.
+				switch (Target.StaticAnalyzer)
+				{
+					case StaticAnalyzer.PVSStudio:
+						break;
+
+					// All other analysis cases should use the ClangRTFM compiler as requested
+					default:
+						Target.WindowsPlatform.Compiler = WindowsCompiler.ClangRTFM;
+
+						// We always use clang linker with AutoRTFM.
+						Target.WindowsPlatform.bAllowClangLinker = true;
+						break;
+				}
+			}
+
 			// If clang is selected for static analysis, switch compiler to clang and proceed
 			// as normal.
 			if (Target.StaticAnalyzer == StaticAnalyzer.Clang)
@@ -1243,25 +1262,6 @@ namespace UnrealBuildTool
 				Logger.LogInformation("Defaulting static analyzer output type to text");
 			}
 
-			if (Target.bUseAutoRTFMCompiler)
-			{
-				// We check the static analyzer, and only allow using the AutoRTFM compiler if
-				// we are not doing static analysis, or are explicitly wanting to use clang.
-				switch (Target.StaticAnalyzer)
-				{
-					default:
-						break;
-					case StaticAnalyzer.None:
-					case StaticAnalyzer.Clang:
-					case StaticAnalyzer.Default:
-						Target.WindowsPlatform.Compiler = WindowsCompiler.ClangRTFM;
-
-						// We always use clang linker with AutoRTFM.
-						Target.WindowsPlatform.bAllowClangLinker = true;
-						break;
-				}
-			}
-
 			// Set the compiler version if necessary
 			if (Target.WindowsPlatform.Compiler == WindowsCompiler.Default)
 			{
