UBT: Temporarily reintroduce usage of Timeline scoped events to output some high level hierarchical scoped timing information especially for the GPF and Build modes.
Note: OpenTracing has been deprecated and needs to be replaced, and that will probably be a good time to reoncile all profiling timers.

#rnx

diff --git a/Engine/Source/Programs/UnrealBuildTool/Modes/BuildMode.cs b/Engine/Source/Programs/UnrealBuildTool/Modes/BuildMode.cs
--- a/Engine/Source/Programs/UnrealBuildTool/Modes/BuildMode.cs
+++ b/Engine/Source/Programs/UnrealBuildTool/Modes/BuildMode.cs
@@ -264,8 +264,11 @@ namespace UnrealBuildTool
 				ProcessCoreDumps(SaveCrashDumpDirectory, Logger);
 
 				// Save all the caches
-				SourceFileMetadataCache.SaveAll();
-				CppDependencyCache.SaveAll();
+				using (Timeline.ScopeEvent("Saving caches"))
+				{
+					SourceFileMetadataCache.SaveAll();
+					CppDependencyCache.SaveAll();
+				}
 			}
 			return 0;
 		}
@@ -392,6 +395,7 @@ namespace UnrealBuildTool
 
 			for (int Idx = 0; Idx < TargetDescriptors.Count; ++Idx)
 			{
+				using ITimelineEvent SetupScope = Timeline.ScopeEvent($"Creating MakeFile for {TargetDescriptors[Idx].Name}");
 				// Create and execute the init scripts
 				FileReference[] InitScripts = CreateInitScripts(TargetDescriptors[Idx]);
 				if (InitScripts.Length > 0)
@@ -435,6 +439,7 @@ namespace UnrealBuildTool
 			// Execute the build
 			if ((Options & BuildOptions.SkipBuild) == 0)
 			{
+				using ITimelineEvent ActionGraphTimer = Timeline.ScopeEvent($"Preparing ActionGraph");
 				// Make sure that none of the actions conflict with any other (producing output files differently, etc...)
 				using (GlobalTracer.Instance.BuildSpan("ActionGraph.CheckForConflicts").StartActive())
 				{
@@ -762,6 +767,8 @@ namespace UnrealBuildTool
 				// Create directories for the outdated produced items.
 				ActionGraph.CreateDirectoriesForProducedItems(MergedActionsToExecute);
 
+				ActionGraphTimer.Finish();
+
 				// Execute the actions
 				if ((Options & BuildOptions.XGEExport) != 0)
 				{
@@ -770,6 +777,7 @@ namespace UnrealBuildTool
 					// Just export to an XML file
 					using (GlobalTracer.Instance.BuildSpan("XGE.ExportActions()").StartActive())
 					{
+						using ITimelineEvent XGETimer = Timeline.ScopeEvent("XGEExport");
 						XGE.ExportActions(MergedActionsToExecute, Logger);
 					}
 				}
@@ -780,6 +788,7 @@ namespace UnrealBuildTool
 					// Write actions to an output file
 					using (GlobalTracer.Instance.BuildSpan("ActionGraph.WriteActions").StartActive())
 					{
+						using ITimelineEvent WriteActionsTimer = Timeline.ScopeEvent("Writing actions");
 						ActionGraph.ExportJson(MergedActionsToExecute, WriteOutdatedActionsFile);
 						Logger.LogInformation("Exported build actions to {WriteOutdatedActionsFile}", WriteOutdatedActionsFile);
 					}
@@ -812,6 +821,7 @@ namespace UnrealBuildTool
 
 						ActionExecutor.SetMemoryPerActionOverride(Makefiles.Select(x => x.MemoryPerActionGB).Max());
 
+						using ITimelineEvent ExcuteTimer = Timeline.ScopeEvent("Executing ActionGraph");
 						using (GlobalTracer.Instance.BuildSpan("ActionGraph.ExecuteActions()").StartActive())
 						{
 
@@ -847,6 +857,7 @@ namespace UnrealBuildTool
 					}
 
 					// Run the deployment steps
+					using ITimelineEvent DeployTimer = Timeline.ScopeEvent("Deploying targets");
 					foreach (TargetMakefile Makefile in Makefiles)
 					{
 						// Receipt file may not exist when compiling specific files
@@ -862,6 +873,7 @@ namespace UnrealBuildTool
 			}
 			else
 			{
+				using ITimelineEvent MetadataTimer = Timeline.ScopeEvent("Writing Target Metadata");
 				foreach (TargetMakefile Makefile in Makefiles)
 				{
 					FileReference TargetInfoFile = FileReference.Combine(Makefile.ProjectIntermediateDirectory, "TargetMetadata.dat");
diff --git a/Engine/Source/Programs/UnrealBuildTool/Modes/GenerateProjectFilesMode.cs b/Engine/Source/Programs/UnrealBuildTool/Modes/GenerateProjectFilesMode.cs
--- a/Engine/Source/Programs/UnrealBuildTool/Modes/GenerateProjectFilesMode.cs
+++ b/Engine/Source/Programs/UnrealBuildTool/Modes/GenerateProjectFilesMode.cs
@@ -251,6 +251,7 @@ namespace UnrealBuildTool
 			bool bGenerateSuccess = true;
 			foreach (KeyValuePair<ProjectFileFormat, ProjectFileGenerator> Pair in Generators)
 			{
+				using ITimelineEvent ExecuteScope = Timeline.ScopeEvent($"Generating {Pair.Key} project files");
 				Logger.LogInformation("");
 				Logger.LogInformation($"Generating {Pair.Key} project files:");
 
diff --git a/Engine/Source/Programs/UnrealBuildTool/ProjectFiles/ProjectFileGenerator.cs b/Engine/Source/Programs/UnrealBuildTool/ProjectFiles/ProjectFileGenerator.cs
--- a/Engine/Source/Programs/UnrealBuildTool/ProjectFiles/ProjectFileGenerator.cs
+++ b/Engine/Source/Programs/UnrealBuildTool/ProjectFiles/ProjectFileGenerator.cs
@@ -984,6 +984,7 @@ namespace UnrealBuildTool
 			bool bCleanProjectFiles = Arguments.Any(x => x.Equals("-CleanProjects", StringComparison.InvariantCultureIgnoreCase));
 			if (bCleanProjectFiles)
 			{
+				using ITimelineEvent CleanTimer = Timeline.ScopeEvent("CleanProjectFiles");
 				CleanProjectFiles(PrimaryProjectPath, PrimaryProjectName, IntermediateProjectFilesPath, Logger);
 			}
 
@@ -1060,13 +1061,20 @@ namespace UnrealBuildTool
 			if (IncludeCppSource)
 			{
 				// Setup buildable projects for all targets
-				AddProjectsForAllTargets(PlatformProjectGenerators, AllGameProjects, AllTargetFiles, Arguments, EngineProjects, GameProjects, ProjectFileToUProjectFile, ProgramProjects, RulesAssemblies, Logger);
+				Logger.LogInformation("Adding projects for all targets...");
+				using (Timeline.ScopeEvent("Adding projects for all targets", Logger))
+				{
+					AddProjectsForAllTargets(PlatformProjectGenerators, AllGameProjects, AllTargetFiles, Arguments, EngineProjects, GameProjects, ProjectFileToUProjectFile, ProgramProjects, RulesAssemblies, Logger);
+				}
 
 				// Add projects for mods
 				AddProjectsForMods(GameProjects, ModProjects);
 
 				// Add all game projects and game config files
-				AddAllGameProjects(GameProjects);
+				using (Timeline.ScopeEvent("Adding projects for all game projects", Logger))
+				{
+					AddAllGameProjects(GameProjects);
+				}
 
 				// Set the game to be the default project
 				if (ModProjects.Count > 0)
@@ -1084,6 +1092,7 @@ namespace UnrealBuildTool
 				// Place projects into root level solution folders
 				if (bIncludeEngineSource)
 				{
+					using ITimelineEvent EngineTimer = Timeline.ScopeEvent("Adding Engine Source");
 					// If we're still missing an engine project because we don't have any targets for it, make one up.
 					if (EngineProjects.Count == 0)
 					{
@@ -1242,13 +1251,17 @@ namespace UnrealBuildTool
 			}
 
 			// Setup "stub" projects for all modules
-			AddProjectsForAllModules(AllGameProjects, ProjectFileToUProjectFile, ProgramProjects, ModProjects, AllModuleFiles, AdditionalSearchPaths, bGatherThirdPartySource, Logger);
+			using (Timeline.ScopeEvent("Adding projects for all modules", Logger))
+			{
+				AddProjectsForAllModules(AllGameProjects, ProjectFileToUProjectFile, ProgramProjects, ModProjects, AllModuleFiles, AdditionalSearchPaths, bGatherThirdPartySource, Logger);
+			}
 
 			{
 				PrimaryProjectFolder ProgramsFolder = RootFolder.AddSubFolder("Programs");
 
 				if (bIncludeDotNetPrograms)
 				{
+					using ITimelineEvent DotNetTimer = Timeline.ScopeEvent("Adding DotNet Programs");
 					// gather engine directories across extension dirs
 					string[] UnsupportedPlatformNames = Utils.MakeListOfUnsupportedPlatforms(SupportedPlatforms, bIncludeUnbuildablePlatforms: true, Logger).ToArray();
 					List<DirectoryReference> AllEngineDirectories =
@@ -1364,7 +1377,10 @@ namespace UnrealBuildTool
 				GenerateIntelliSenseData(Arguments, IntelliSenseTargetFiles, Logger);
 
 				// Write the project files
-				WriteProjectFiles(PlatformProjectGenerators, Logger);
+				using (Timeline.ScopeEvent("Writing project files", Logger))
+				{
+					WriteProjectFiles(PlatformProjectGenerators, Logger);
+				}
 				Logger.LogDebug("Project generation complete ({GeneratedProjectFilesCount} generated, {OtherProjectFilesCount} imported)", GeneratedProjectFiles.Count, OtherProjectFiles.Count);
 			}
 			return bSuccess;
@@ -1372,6 +1388,7 @@ namespace UnrealBuildTool
 
 		public void GenerateQueryTargetsDataForEditor(string[] Arguments, ILogger Logger)
 		{
+			using ITimelineEvent GenerateTimer = Timeline.ScopeEvent("Generating QueryTargets data for editor", Logger);
 			Logger.LogInformation("");
 			Logger.LogInformation("Generating QueryTargets data for editor...");
 
@@ -1379,6 +1396,7 @@ namespace UnrealBuildTool
 			List<FileReference> AllGameProjects = FindGameProjects(Logger);
 
 			ConcurrentBag<KeyValuePair<FileReference?, RulesAssembly>> ProjectFileToRules = new();
+			using ITimelineEvent CompileTimer = Timeline.ScopeEvent("Compiling Rules Assemblies");
 			using (ProgressWriter Progress = new ProgressWriter("Compiling Rules Assemblies", true, Logger))
 			{
 				int TasksFinished = 0;
@@ -1395,8 +1413,10 @@ namespace UnrealBuildTool
 					Progress.Write(++TasksFinished, AllGameProjects.Count + 1);
 				}
 			}
+			CompileTimer.Finish();
 
 			// Generate all the target info files for the editor
+			using ITimelineEvent WriteTimer = Timeline.ScopeEvent("Writing Query Target Info");
 			using (ProgressWriter Progress = new ProgressWriter("Writing Query Target Info", true, Logger))
 			{
 				int TasksFinished = 0;
@@ -1831,6 +1851,7 @@ namespace UnrealBuildTool
 		/// <returns>Filtered list of target files</returns>
 		public static List<FileReference> DiscoverTargets(List<FileReference> AllGameProjects, ILogger Logger, FileReference? OnlyGameProject, List<UnrealTargetPlatform> SupportedPlatforms, bool bIncludeEngineSource, bool bIncludeTempTargets)
 		{
+			using ITimelineEvent Timer = Timeline.ScopeEvent("DiscoverTargets");
 			List<FileReference> AllTargetFiles = new List<FileReference>();
 
 			// Make a list of all platform name strings that we're *not* including in the project files
@@ -2186,6 +2207,7 @@ namespace UnrealBuildTool
 				// Ignore projects for platforms we can't build for
 				List<Tuple<ProjectFile, ProjectTarget>> SupportedTargets = Targets.Where(x => x.Item2.SupportedPlatforms.Any(x => GetIntelliSensePlatforms().Contains(x))).ToList();
 
+				using ITimelineEvent CompileTimer = Timeline.ScopeEvent("Compiling Rules Assemblies");
 				using (ProgressWriter Progress = new ProgressWriter("Compiling Rules Assemblies", true, Logger))
 				{
 					int TasksFinished = 0;
@@ -2219,9 +2241,11 @@ namespace UnrealBuildTool
 						}
 					}
 				}
+				CompileTimer.Finish();
 
 				ConcurrentBag<UEBuildTarget> BuildTargets = [];
 				ConcurrentDictionary<UEBuildTarget, ProjectFile> BuildTargetToProjectFile = new();
+				using ITimelineEvent BuildTargetsTimer = Timeline.ScopeEvent("Creating Build Targets", Logger);
 				using (ProgressWriter Progress = new ProgressWriter("Creating Build Targets", true, Logger))
 				{
 					int NumTargets = SupportedTargets.Count;
@@ -2269,9 +2293,11 @@ namespace UnrealBuildTool
 						}
 					});
 				}
+				BuildTargetsTimer.Finish();
 
-				string ProgressInfoText = OperatingSystem.IsWindows() ? "Binding IntelliSense data..." : "Generating data for project indexing...";
-				using (ProgressWriter Progress = new ProgressWriter(ProgressInfoText, true, Logger))
+				string ProgressInfoText = OperatingSystem.IsWindows() ? "Binding IntelliSense data" : "Generating data for project indexing";
+				using ITimelineEvent IntellisenseTimer = Timeline.ScopeEvent(ProgressInfoText, Logger);
+				using (ProgressWriter Progress = new ProgressWriter(ProgressInfoText + "...", true, Logger))
 				{
 					int TasksFinished = 0;
 					int TotalTasks = BuildTargets.Count;
diff --git a/Engine/Source/Programs/UnrealBuildTool/System/ScopedTimer.cs b/Engine/Source/Programs/UnrealBuildTool/System/ScopedTimer.cs
--- a/Engine/Source/Programs/UnrealBuildTool/System/ScopedTimer.cs
+++ b/Engine/Source/Programs/UnrealBuildTool/System/ScopedTimer.cs
@@ -59,7 +59,7 @@ namespace UnrealBuildTool
 			StringBuilder IndentText = new StringBuilder(LogIndent * 2);
 			IndentText.Append(' ', LogIndent * 2);
 
-			Logger.Log(Level, "{Indent}{Name} took {TimeSeconds}s", IndentText.ToString(), Name, TotalSeconds);
+			Logger.Log(Level, "{Indent}{Name} took {TimeSeconds:0.000}s", IndentText.ToString(), Name, TotalSeconds);
 		}
 	}
 }
diff --git a/Engine/Source/Programs/UnrealBuildTool/System/Timeline.cs b/Engine/Source/Programs/UnrealBuildTool/System/Timeline.cs
--- a/Engine/Source/Programs/UnrealBuildTool/System/Timeline.cs
+++ b/Engine/Source/Programs/UnrealBuildTool/System/Timeline.cs
@@ -49,17 +49,31 @@ namespace UnrealBuildTool
 			/// </summary>
 			public readonly ITraceSpan Span;
 
+			/// <summary>
+			/// Optional output logger
+			/// </summary>
+			ILogger? Logger;
+
+			/// <summary>
+			/// Optional output verbosity
+			/// </summary>
+			LogLevel? Verbosity;
+
 			/// <summary>
 			/// Constructor
 			/// </summary>
 			/// <param name="Name">Event name</param>
 			/// <param name="StartTime">Time of the event</param>
 			/// <param name="FinishTime">Finish time for the event. May be null.</param>
-			public Event(string Name, TimeSpan StartTime, TimeSpan? FinishTime)
+			/// <param name="Logger">Logger. May be null.</param>
+			/// <param name="Verbosity">Verbosity. May be null.</param>
+			public Event(string Name, TimeSpan StartTime, TimeSpan? FinishTime = null, ILogger? Logger = null, LogLevel? Verbosity = null)
 			{
 				this.Name = Name;
 				this.StartTime = StartTime;
 				this.FinishTime = FinishTime;
+				this.Logger = Logger;
+				this.Verbosity = Verbosity;
 				Span = TraceSpan.Create(Name);
 			}
 
@@ -71,6 +85,11 @@ namespace UnrealBuildTool
 				if (!FinishTime.HasValue)
 				{
 					FinishTime = Stopwatch.Elapsed;
+					if (Logger != null && Verbosity != null)
+					{
+						double Duration = ((TimeSpan)FinishTime - StartTime).TotalSeconds;
+						Logger.Log((LogLevel)Verbosity, "{Name} took {Duration:0.00}s", Name, Duration);
+					}
 				}
 				Span.Dispose();
 			}
@@ -107,6 +126,14 @@ namespace UnrealBuildTool
 			Stopwatch.Restart();
 		}
 
+		/// <summary>
+		/// Stop the stopwatch
+		/// </summary>
+		public static void Stop()
+		{
+			Stopwatch.Stop();
+		}
+
 		/// <summary>
 		/// Records a timeline marker with the given name
 		/// </summary>
@@ -135,15 +162,33 @@ namespace UnrealBuildTool
 			return Event;
 		}
 
+		/// <summary>
+		/// Enters a scope event with the given name. Should be disposed to terminate the scope.
+		/// Also logs the finish of the event at the provided verbosity.
+		/// </summary>
+		/// <param name="Name">Name of the event</param>
+		/// <param name="Logger">Logger.</param>
+		/// <param name="Verbosity">LogLevel of the event.</param>
+		/// <returns>Event to track the length of the event</returns>
+		public static ITimelineEvent ScopeEvent(string Name, ILogger Logger, LogLevel Verbosity = LogLevel.Information)
+		{
+			Event Event = new Event(Name, Stopwatch.Elapsed, null, Logger, Verbosity);
+			lock (Events)
+			{
+				Events.Add(Event);
+			}
+			return Event;
+		}
+
 		/// <summary>
 		/// Prints this information to the log
 		/// </summary>
 		public static void Print(TimeSpan MaxUnknownTime, LogLevel Verbosity, ILogger Logger)
 		{
 			// Print the start time
-			Logger.Log(Verbosity, "Timeline:");
 			Logger.Log(Verbosity, "");
-			Logger.Log(Verbosity, "[{Time,6}]", FormatTime(TimeSpan.Zero));
+			Logger.Log(Verbosity, "Timeline:");
+			Logger.Log(Verbosity, "[{Time,7}]", FormatTime(TimeSpan.Zero));
 
 			// Create the root event
 			TimeSpan FinishTime = Stopwatch.Elapsed;
@@ -200,7 +245,7 @@ namespace UnrealBuildTool
 			}
 
 			// Print the finish time
-			Logger.Log(Verbosity, "[{Time,6}]", FormatTime(FinishTime));
+			Logger.Log(Verbosity, "[{Time,7}]", FormatTime(FinishTime));
 		}
 
 		/// <summary>
@@ -237,14 +282,14 @@ namespace UnrealBuildTool
 
 			for (int Idx = 0; Idx < OuterEvents.Count - 1; Idx++)
 			{
-				Prefix.AppendFormat(" {0,6}          ", FormatTime(StartTime - OuterEvents[Idx].StartTime));
+				Prefix.AppendFormat(" {0,7}          ", FormatTime(StartTime - OuterEvents[Idx].StartTime));
 			}
 
-			Prefix.AppendFormat("[{0,6}]", FormatTime(StartTime - OuterEvents[^1].StartTime));
+			Prefix.AppendFormat("[{0,7}]", FormatTime(StartTime - OuterEvents[^1].StartTime));
 
 			if (!FinishTime.HasValue)
 			{
-				Prefix.AppendFormat("({0,6})", "???");
+				Prefix.AppendFormat("{0,8}", "???");
 			}
 			else if (FinishTime.Value == StartTime)
 			{
@@ -252,7 +297,7 @@ namespace UnrealBuildTool
 			}
 			else
 			{
-				Prefix.AppendFormat("({0,6})", "+" + FormatTime(FinishTime.Value - StartTime));
+				Prefix.AppendFormat("{0,8}", "+" + FormatTime(FinishTime.Value - StartTime));
 			}
 
 			Logger.Log(Verbosity, "{Prefix} {Label}", Prefix.ToString(), Label);
diff --git a/Engine/Source/Programs/UnrealBuildTool/UnrealBuildTool.cs b/Engine/Source/Programs/UnrealBuildTool/UnrealBuildTool.cs
--- a/Engine/Source/Programs/UnrealBuildTool/UnrealBuildTool.cs
+++ b/Engine/Source/Programs/UnrealBuildTool/UnrealBuildTool.cs
@@ -387,6 +387,7 @@ namespace UnrealBuildTool
         /// <returns>Zero on success, non-zero on error</returns>
         private static int Main(string[] ArgumentsArray)
 		{
+			ToolModeOptions ModeOptions = ToolModeOptions.None;
 			FileReference? RunFile = null;
 			DirectoryReference? TempDirectory = null;
 			GlobalSingleInstanceMutex? Mutex = null;
@@ -406,6 +407,7 @@ namespace UnrealBuildTool
 			{
 				// Start capturing performance info
 				Timeline.Start();
+				using ITimelineEvent MainScope = Timeline.ScopeEvent("Main");
 				Tracer = JsonTracer.TryRegisterAsGlobalTracer();
 				ArgumentsArray = ArgumentsArray.Concat(GetExtraArgsFromEnvVar()).ToArray();
 
@@ -481,7 +483,7 @@ namespace UnrealBuildTool
 				BuildMode = ModeType.Name;
 
 				// Get the options for which systems have to be initialized for this mode
-				ToolModeOptions ModeOptions = ModeType.GetCustomAttribute<ToolModeAttribute>()!.Options;
+				ModeOptions = ModeType.GetCustomAttribute<ToolModeAttribute>()!.Options;
 
 				// if we don't care about the trace listener, toss it now
 				if ((ModeOptions & ToolModeOptions.UseStartupTraceListener) == 0)
@@ -501,6 +503,7 @@ namespace UnrealBuildTool
 				// Read the XML configuration files
 				if ((ModeOptions & ToolModeOptions.XmlConfig) != 0)
 				{
+					using (Timeline.ScopeEvent("Apply XmlConfig"))
 					using (GlobalTracer.Instance.BuildSpan("XmlConfig.ReadConfigFiles()").StartActive())
 					{
 						string XmlConfigMutexName = GlobalSingleInstanceMutex.GetUniqueMutexForPath("UnrealBuildTool_Mutex_XmlConfig", FileReference.FromString(Assembly.GetExecutingAssembly().Location));
@@ -606,26 +609,29 @@ namespace UnrealBuildTool
 					}
 				}
 
-				// Register all the build platforms
-				if ((ModeOptions & ToolModeOptions.BuildPlatforms) != 0)
+				using (Timeline.ScopeEvent("RegisterPlatforms"))
 				{
-					using (GlobalTracer.Instance.BuildSpan("UEBuildPlatform.RegisterPlatforms()").StartActive())
+					// Register all the build platforms
+					if ((ModeOptions & ToolModeOptions.BuildPlatforms) != 0)
 					{
-						UEBuildPlatform.RegisterPlatforms(false, false, ModeType, ArgumentsArray, Logger);
+						using (GlobalTracer.Instance.BuildSpan("UEBuildPlatform.RegisterPlatforms()").StartActive())
+						{
+							UEBuildPlatform.RegisterPlatforms(false, false, ModeType, ArgumentsArray, Logger);
+						}
 					}
-				}
-				if ((ModeOptions & ToolModeOptions.BuildPlatformsHostOnly) != 0)
-				{
-					using (GlobalTracer.Instance.BuildSpan("UEBuildPlatform.RegisterPlatforms()").StartActive())
+					if ((ModeOptions & ToolModeOptions.BuildPlatformsHostOnly) != 0)
 					{
-						UEBuildPlatform.RegisterPlatforms(false, true, ModeType, ArgumentsArray, Logger);
+						using (GlobalTracer.Instance.BuildSpan("UEBuildPlatform.RegisterPlatforms()").StartActive())
+						{
+							UEBuildPlatform.RegisterPlatforms(false, true, ModeType, ArgumentsArray, Logger);
+						}
 					}
-				}
-				if ((ModeOptions & ToolModeOptions.BuildPlatformsForValidation) != 0)
-				{
-					using (GlobalTracer.Instance.BuildSpan("UEBuildPlatform.RegisterPlatforms()").StartActive())
+					if ((ModeOptions & ToolModeOptions.BuildPlatformsForValidation) != 0)
 					{
-						UEBuildPlatform.RegisterPlatforms(true, false, ModeType, ArgumentsArray, Logger);
+						using (GlobalTracer.Instance.BuildSpan("UEBuildPlatform.RegisterPlatforms()").StartActive())
+						{
+							UEBuildPlatform.RegisterPlatforms(true, false, ModeType, ArgumentsArray, Logger);
+						}
 					}
 				}
 
@@ -633,19 +639,9 @@ namespace UnrealBuildTool
 				ToolMode Mode = (ToolMode)Activator.CreateInstance(ModeType)!;
 
 				// Execute the mode
-				int Result;
-				try
-				{
-					Result = Mode.ExecuteAsync(Arguments, Logger).GetAwaiter().GetResult();
-				}
-				finally
-				{
-					if ((ModeOptions & ToolModeOptions.ShowExecutionTime) != 0)
-					{
-						Logger.LogInformation("Total execution time: {Time:0.00} seconds", Timeline.Elapsed.TotalSeconds);
-					}
-				}
-
+				MainScope.Finish();
+				using ITimelineEvent ExecuteScope = Timeline.ScopeEvent(ModeType.Name);
+				int Result = Mode.ExecuteAsync(Arguments, Logger).GetAwaiter().GetResult();
 				ApplicationResult = (CompilationResult)Result;
 				return Result;
 			}
@@ -674,18 +670,24 @@ namespace UnrealBuildTool
 				// Useful when investigating why UBT takes time.
 				//DirectoryItem.WriteDebugFileWithAllEnumeratedFiles(@"c:\temp\AllFiles.txt");
 
-				if (!IsRecursive)
+				using (Timeline.ScopeEvent("TelemetryService.FlushEvents"))
 				{
-					TelemetryService.Get().RecordEvent(new TelemetryCompletedEvent(ArgumentsArray, StartTimeUtc, ApplicationResult, DateTime.UtcNow));
+					if (!IsRecursive)
+					{
+						TelemetryService.Get().RecordEvent(new TelemetryCompletedEvent(ArgumentsArray, StartTimeUtc, ApplicationResult, DateTime.UtcNow));
+					}
+					// Flush any remaining telemetry events
+					TelemetryService.Get().FlushEvents();
 				}
 
-				// Flush any remaining telemetry events
-				TelemetryService.Get().FlushEvents();
-
 				Utils.LogWriteFileIfChangedActivity(Logger);
 
 				// Print out all the performance info
-				Timeline.Print(TimeSpan.FromMilliseconds(20.0), LogLevel.Debug, Logger);
+				Timeline.Stop();
+				LogLevel TimelineLogLevel = ModeOptions.HasFlag(ToolModeOptions.ShowExecutionTime) ? LogLevel.Information : LogLevel.Debug;
+				Timeline.Print(TimeSpan.FromMilliseconds(200.0), TimelineLogLevel, Logger);
+				Logger.Log(TimelineLogLevel, (""));
+				Logger.Log(TimelineLogLevel, "Total execution time: {Time:0.00} seconds", Timeline.Elapsed.TotalSeconds);
 
 				// Make sure we flush the logs however we exit
 				Trace.Close();
