UnrealBuildTool: Load syntax trees for dynamically compiled rules assembly async
Reduces parse time locally for UE5Rules 8.5s -> 5s
Clean up reference preparation

#rnx

diff --git a/Engine/Source/Programs/UnrealBuildTool/System/DynamicCompilation.cs b/Engine/Source/Programs/UnrealBuildTool/System/DynamicCompilation.cs
--- a/Engine/Source/Programs/UnrealBuildTool/System/DynamicCompilation.cs
+++ b/Engine/Source/Programs/UnrealBuildTool/System/DynamicCompilation.cs
@@ -6,6 +6,8 @@ using System.IO;
 using System.Linq;
 using System.Reflection;
 using System.Runtime.InteropServices;
+using System.Threading;
+using System.Threading.Tasks;
 using EpicGames.Core;
 using Microsoft.CodeAnalysis;
 using Microsoft.CodeAnalysis.CSharp;
@@ -166,30 +168,37 @@ namespace UnrealBuildTool
 			}
 		}
 
-		private static Assembly? CompileAssembly(FileReference OutputAssemblyPath, IEnumerable<FileReference> SourceFileNames, ILogger Logger, IEnumerable<string>? ReferencedAssembies, IEnumerable<string>? PreprocessorDefines = null, bool TreatWarningsAsErrors = false)
+		private static async Task<SyntaxTree> ParseSyntaxTreeAsync(FileReference sourceFileName, CSharpParseOptions parseOptions, CancellationToken cancellationToken)
 		{
-			CSharpParseOptions ParseOptions = new CSharpParseOptions(
+			string content = await FileReference.ReadAllTextAsync(sourceFileName, cancellationToken);
+			SourceText source = SourceText.From(content, System.Text.Encoding.UTF8);
+			return CSharpSyntaxTree.ParseText(source, parseOptions, sourceFileName.FullName, cancellationToken);
+		}
+
+		private static async Task<IEnumerable<SyntaxTree>> ParseSyntaxTreesAsync(IEnumerable<FileReference> sourceFileNames, IEnumerable<string>? preprocessorDefines, CancellationToken cancellationToken)
+		{
+			CSharpParseOptions parseOptions = new(
 				languageVersion: LanguageVersion.Latest,
 				kind: SourceCodeKind.Regular,
-				preprocessorSymbols: PreprocessorDefines
+				preprocessorSymbols: preprocessorDefines
 			);
 
-			List<SyntaxTree> SyntaxTrees = new List<SyntaxTree>();
+			IEnumerable<Task<SyntaxTree>> tasks = sourceFileNames.Select(sourceFileName => ParseSyntaxTreeAsync(sourceFileName, parseOptions, cancellationToken));
+			return (await Task.WhenAll(tasks)).OrderBy(x => x.FilePath);
+		}
 
-			foreach (FileReference SourceFileName in SourceFileNames)
-			{
-				SourceText Source = SourceText.From(File.ReadAllText(SourceFileName.FullName), System.Text.Encoding.UTF8);
-				SyntaxTree Tree = CSharpSyntaxTree.ParseText(Source, ParseOptions, SourceFileName.FullName);
+		private static async Task<Assembly?> CompileAssemblyAsync(FileReference OutputAssemblyPath, IEnumerable<FileReference> SourceFileNames, ILogger Logger, IEnumerable<string>? ReferencedAssembies, IEnumerable<string>? PreprocessorDefines = null, bool TreatWarningsAsErrors = false)
+		{
+			IEnumerable<SyntaxTree> SyntaxTrees = await ParseSyntaxTreesAsync(SourceFileNames, PreprocessorDefines, default);
 
-				IEnumerable<Diagnostic> Diagnostics = Tree.GetDiagnostics();
-				if (Diagnostics.Any())
+			// Check for errors
+			{
+				IEnumerable<Diagnostic> syntaxTreeErrors = SyntaxTrees.SelectMany(x => x.GetDiagnostics());
+				if (syntaxTreeErrors.Where(x => x.Severity == DiagnosticSeverity.Error).Any())
 				{
-					Logger.LogWarning("Errors generated while parsing '{SourceFileName}'", SourceFileName);
-					LogDiagnostics(Tree.GetDiagnostics(), Logger);
+					LogDiagnostics(syntaxTreeErrors, Logger);
 					return null;
 				}
-
-				SyntaxTrees.Add(Tree);
 			}
 
 			// Create the output directory if it doesn't exist already
@@ -206,56 +215,51 @@ namespace UnrealBuildTool
 				}
 			}
 
-			List<MetadataReference> MetadataReferences = new List<MetadataReference>();
-			if (ReferencedAssembies != null)
-			{
-				foreach (string Reference in ReferencedAssembies)
-				{
-					MetadataReferences.Add(MetadataReference.CreateFromFile(Reference));
-				}
-			}
-
-			MetadataReferences.Add(MetadataReference.CreateFromFile(typeof(object).Assembly.Location));
-			MetadataReferences.Add(MetadataReference.CreateFromFile(Assembly.Load("System.Runtime").Location));
-			MetadataReferences.Add(MetadataReference.CreateFromFile(Assembly.Load("System.CodeDom").Location));
-			MetadataReferences.Add(MetadataReference.CreateFromFile(Assembly.Load("System.Collections").Location));
-			MetadataReferences.Add(MetadataReference.CreateFromFile(Assembly.Load("System.IO").Location));
-			MetadataReferences.Add(MetadataReference.CreateFromFile(Assembly.Load("System.IO.FileSystem").Location));
-			MetadataReferences.Add(MetadataReference.CreateFromFile(Assembly.Load("System.Linq").Location));
-			MetadataReferences.Add(MetadataReference.CreateFromFile(Assembly.Load("System.Threading.Tasks").Location));
-			MetadataReferences.Add(MetadataReference.CreateFromFile(Assembly.Load("System.Threading.Tasks.Parallel").Location));
-			MetadataReferences.Add(MetadataReference.CreateFromFile(Assembly.Load("System.Collections.Concurrent").Location));
-			MetadataReferences.Add(MetadataReference.CreateFromFile(Assembly.Load("System.Private.Xml").Location));
-			MetadataReferences.Add(MetadataReference.CreateFromFile(Assembly.Load("System.Private.Xml.Linq").Location));
-			MetadataReferences.Add(MetadataReference.CreateFromFile(Assembly.Load("System.Text.RegularExpressions").Location));
-			MetadataReferences.Add(MetadataReference.CreateFromFile(Assembly.Load("Microsoft.CodeAnalysis.CSharp").Location));
-			MetadataReferences.Add(MetadataReference.CreateFromFile(Assembly.Load("System.Console").Location));
-			MetadataReferences.Add(MetadataReference.CreateFromFile(Assembly.Load("System.Runtime.Extensions").Location));
-			MetadataReferences.Add(MetadataReference.CreateFromFile(Assembly.Load("Microsoft.Extensions.Logging.Abstractions").Location));
-			MetadataReferences.Add(MetadataReference.CreateFromFile(Assembly.Load("netstandard").Location));
-
-			// process start dependencies
-			MetadataReferences.Add(MetadataReference.CreateFromFile(Assembly.Load("System.ComponentModel.Primitives").Location));
-			MetadataReferences.Add(MetadataReference.CreateFromFile(Assembly.Load("System.Diagnostics.Process").Location));
-
-			// registry access
-			MetadataReferences.Add(MetadataReference.CreateFromFile(Assembly.Load("Microsoft.Win32.Registry").Location));
-
-			// RNGCryptoServiceProvider, used to generate random hex bytes
-			try
-			{
-				// Does not exist for .NET6
-				MetadataReferences.Add(MetadataReference.CreateFromFile(Assembly.Load("System.Security.Cryptography").Location));
-			}
-			catch (FileNotFoundException)
-			{
-			}
-			MetadataReferences.Add(MetadataReference.CreateFromFile(Assembly.Load("System.Security.Cryptography.Algorithms").Location));
-			MetadataReferences.Add(MetadataReference.CreateFromFile(Assembly.Load("System.Security.Cryptography.Csp").Location));
-
-			MetadataReferences.Add(MetadataReference.CreateFromFile(typeof(UnrealBuildTool).Assembly.Location));
-			MetadataReferences.Add(MetadataReference.CreateFromFile(typeof(FileReference).Assembly.Location));
-			MetadataReferences.Add(MetadataReference.CreateFromFile(typeof(UEBuildPlatformSDK).Assembly.Location));
+			List<Type> ReferenceAssemblyTypes = [
+				typeof(object),
+				typeof(UnrealBuildTool),
+				typeof(FileReference),
+				typeof(UEBuildPlatformSDK),
+			];
+
+			List<string> ReferenceAssemblyStrings = [
+				// system references,
+				"System.Runtime",
+				"System.CodeDom",
+				"System.Collections",
+				"System.IO",
+				"System.IO.FileSystem",
+				"System.Linq",
+				"System.Threading.Tasks",
+				"System.Threading.Tasks.Parallel",
+				"System.Collections.Concurrent",
+				"System.Private.Xml",
+				"System.Private.Xml.Linq",
+				"System.Text.RegularExpressions",
+				"Microsoft.CodeAnalysis.CSharp",
+				"System.Console",
+				"System.Runtime.Extensions",
+				"Microsoft.Extensions.Logging.Abstractions",
+				"netstandard",
+
+				// process start dependencies
+				"System.ComponentModel.Primitives",
+				"System.Diagnostics.Process",
+
+				// registry access
+				"Microsoft.Win32.Registry",
+
+				// RNGCryptoServiceProvider, used to generate random hex bytes
+				"System.Security.Cryptography",
+				"System.Security.Cryptography.Algorithms",
+				"System.Security.Cryptography.Csp",
+			];
+
+			List<MetadataReference> MetadataReferences = [
+				.. ReferencedAssembies?.Select(x => MetadataReference.CreateFromFile(x)) ?? [],
+				.. ReferenceAssemblyTypes.Select(x => MetadataReference.CreateFromFile(x.Assembly.Location)),
+				.. ReferenceAssemblyStrings.Select(x => MetadataReference.CreateFromFile(Assembly.Load(x).Location))
+			];
 
 			CSharpCompilationOptions CompilationOptions = new CSharpCompilationOptions(
 				outputKind: OutputKind.DynamicallyLinkedLibrary,
@@ -277,24 +281,23 @@ namespace UnrealBuildTool
 				options: CompilationOptions
 				);
 
-			using (FileStream AssemblyStream = FileReference.Open(OutputAssemblyPath, FileMode.Create))
 			{
-				using (FileStream? PdbStream = RuntimeInformation.IsOSPlatform(OSPlatform.Windows) ? FileReference.Open(OutputAssemblyPath.ChangeExtension(".pdb"), FileMode.Create) : null)
+				using FileStream AssemblyStream = FileReference.Open(OutputAssemblyPath, FileMode.Create);
+				using FileStream? PdbStream = RuntimeInformation.IsOSPlatform(OSPlatform.Windows) ? FileReference.Open(OutputAssemblyPath.ChangeExtension(".pdb"), FileMode.Create) : null;
+
+				EmitOptions EmitOptions = new(
+					includePrivateMembers: true
+				);
+
+				EmitResult Result = Compilation.Emit(
+					peStream: AssemblyStream,
+					pdbStream: PdbStream,
+					options: EmitOptions);
+				LogDiagnostics(Result.Diagnostics, Logger);
+
+				if (!Result.Success)
 				{
-					EmitOptions EmitOptions = new EmitOptions(
-						includePrivateMembers: true
-					);
-
-					EmitResult Result = Compilation.Emit(
-						peStream: AssemblyStream,
-						pdbStream: PdbStream,
-						options: EmitOptions);
-					LogDiagnostics(Result.Diagnostics, Logger);
-
-					if (!Result.Success)
-					{
-						return null;
-					}
+					return null;
 				}
 			}
 
@@ -302,7 +305,7 @@ namespace UnrealBuildTool
 		}
 
 		/// <summary>
-		/// Dynamically compiles an assembly for the specified source file and loads that assembly into the application's
+		/// Dynamically compiles an assembly async for the specified source file and loads that assembly into the application's
 		/// current domain.  If an assembly has already been compiled and is not out of date, then it will be loaded and
 		/// no compilation is necessary.
 		/// </summary>
@@ -315,7 +318,7 @@ namespace UnrealBuildTool
 		/// <param name="ForceCompile"></param>
 		/// <param name="TreatWarningsAsErrors"></param>
 		/// <returns>The assembly that was loaded</returns>
-		public static Assembly? CompileAndLoadAssembly(FileReference OutputAssemblyPath, IEnumerable<FileReference> SourceFileNames, ILogger Logger, IEnumerable<string>? ReferencedAssembies = null, IEnumerable<string>? PreprocessorDefines = null, bool DoNotCompile = false, bool ForceCompile = false, bool TreatWarningsAsErrors = false)
+		public static async Task<Assembly?> CompileAndLoadAssemblyAsync(FileReference OutputAssemblyPath, IEnumerable<FileReference> SourceFileNames, ILogger Logger, IEnumerable<string>? ReferencedAssembies = null, IEnumerable<string>? PreprocessorDefines = null, bool DoNotCompile = false, bool ForceCompile = false, bool TreatWarningsAsErrors = false)
 		{
 			// Check to see if the resulting assembly is compiled and up to date
 			FileReference AssemblyManifestFilePath = FileReference.Combine(OutputAssemblyPath.Directory, Path.GetFileNameWithoutExtension(OutputAssemblyPath.FullName) + "Manifest.json");
@@ -360,7 +363,7 @@ namespace UnrealBuildTool
 			{
 				using (GlobalTracer.Instance.BuildSpan(String.Format("Compiling rules assembly ({0})", OutputAssemblyPath.GetFileName())).StartActive())
 				{
-					CompiledAssembly = CompileAssembly(OutputAssemblyPath, SourceFileNames, Logger, ReferencedAssembies, PreprocessorDefines, TreatWarningsAsErrors);
+					CompiledAssembly = await CompileAssemblyAsync(OutputAssemblyPath, SourceFileNames, Logger, ReferencedAssembies, PreprocessorDefines, TreatWarningsAsErrors);
 				}
 
 				using (JsonWriter Writer = new JsonWriter(AssemblyManifestFilePath))
@@ -379,9 +382,25 @@ namespace UnrealBuildTool
 			return CompiledAssembly;
 		}
 
-		private static string FormatVersionNumber(ReadOnlyBuildVersion Version)
+		/// <summary>
+		/// Dynamically compiles an assembly for the specified source file and loads that assembly into the application's
+		/// current domain.  If an assembly has already been compiled and is not out of date, then it will be loaded and
+		/// no compilation is necessary.
+		/// </summary>
+		/// <param name="OutputAssemblyPath">Full path to the assembly to be created</param>
+		/// <param name="SourceFileNames">List of source file name</param>
+		/// <param name="Logger">Logger for output</param>
+		/// <param name="ReferencedAssembies"></param>
+		/// <param name="PreprocessorDefines"></param>
+		/// <param name="DoNotCompile"></param>
+		/// <param name="ForceCompile"></param>
+		/// <param name="TreatWarningsAsErrors"></param>
+		/// <returns>The assembly that was loaded</returns>
+		public static Assembly? CompileAndLoadAssembly(FileReference OutputAssemblyPath, IEnumerable<FileReference> SourceFileNames, ILogger Logger, IEnumerable<string>? ReferencedAssembies = null, IEnumerable<string>? PreprocessorDefines = null, bool DoNotCompile = false, bool ForceCompile = false, bool TreatWarningsAsErrors = false)
 		{
-			return String.Format("{0}.{1}.{2}", Version.MajorVersion, Version.MinorVersion, Version.PatchVersion);
+			return CompileAndLoadAssemblyAsync(OutputAssemblyPath, SourceFileNames, Logger, ReferencedAssembies, PreprocessorDefines, DoNotCompile, ForceCompile, TreatWarningsAsErrors).Result;
 		}
+
+		private static string FormatVersionNumber(ReadOnlyBuildVersion Version) => $"{Version.MajorVersion}.{Version.MinorVersion}.{Version.PatchVersion}";
 	}
 }
