Apply Warnings refactor - attribute derivation based

diff --git a/Engine/Source/Programs/UnrealBuildTool/Configuration/BuildSystemContext.cs b/Engine/Source/Programs/UnrealBuildTool/Configuration/BuildSystemContext.cs
--- a/Engine/Source/Programs/UnrealBuildTool/Configuration/BuildSystemContext.cs
+++ b/Engine/Source/Programs/UnrealBuildTool/Configuration/BuildSystemContext.cs
@@ -79,6 +79,15 @@ namespace UnrealBuildTool
 			_buildContext = context;
 		}
 
+		/// <summary>
+		/// Helper method to obtain a readonly target rules from the given context, if at all possible.
+		/// </summary>
+		/// <returns>The resolved ReadOnlyTargetRules from the current build context.</returns>
+		internal ReadOnlyTargetRules? GetReadOnlyTargetRules()
+		{
+			return _targetRulesPrivate != null ? new ReadOnlyTargetRules(_targetRulesPrivate) : _moduleRulesPrivate?.Target;
+		}
+
 		internal TargetRules? _targetRulesPrivate;
 		internal ModuleRules? _moduleRulesPrivate;
 		internal IBuildContextProvider _buildContext;
diff --git a/Engine/Source/Programs/UnrealBuildTool/Configuration/CompileWarnings/ApplyWarningsAttribute.cs b/Engine/Source/Programs/UnrealBuildTool/Configuration/CompileWarnings/ApplyWarningsAttribute.cs
new file mode 100644
--- /dev/null
+++ b/Engine/Source/Programs/UnrealBuildTool/Configuration/CompileWarnings/ApplyWarningsAttribute.cs
@@ -0,0 +1,223 @@
+// Copyright Epic Games, Inc. All Rights Reserved.
+
+using System;
+using System.Collections.Generic;
+using EpicGames.Core;
+
+namespace UnrealBuildTool
+{
+	/// <summary>
+	/// Attribute used to mark fields which describes a set of <see cref="WarningLevel"/> arguments.
+	/// </summary>
+	[AttributeUsage(AttributeTargets.Field | AttributeTargets.Property, AllowMultiple = true)]
+	internal abstract class ApplyWarningsAttribute : Attribute
+	{
+		private readonly string[]? _offArgs;
+		private readonly string[]? _warningArgs;
+		private readonly string[]? _errorArgs;
+
+		/// <summary>
+		/// Defines the argument metadata around a <see cref="CppCompileWarnings"/> <see cref="WarningLevel"/> property.
+		/// </summary>
+		/// <param name="offArgs">The arguments to use when the property is defined as <see cref="WarningLevel.Off"/>.</param>
+		/// <param name="warningArgs">The arguments to use when the property is defined as <see cref="WarningLevel.Warning"/>.</param>
+		/// <param name="errorArgs">The arguments to use when the property is defined as <see cref="WarningLevel.Error"/>.</param>
+		internal ApplyWarningsAttribute(string[]? offArgs, string[]? warningArgs, string[]? errorArgs)
+		{
+			_offArgs = offArgs;
+			_warningArgs = warningArgs;
+			_errorArgs = errorArgs;
+		}
+
+		/// <summary>
+		/// Applies the warnings to the provided arguments list.
+		/// </summary>
+		/// <param name="level">The warning level to use when applying.</param>
+		/// <param name="toolChainContext">The toolchain context to use in application.</param>
+		/// <param name="arguments">The arugments list to add values to.</param>
+		internal void ApplyWarningsToArguments(WarningLevel level, CompilerWarningsToolChainContext toolChainContext, List<string> arguments)
+		{
+			ApplyWarningsToArgumentsInternal(level, toolChainContext, arguments);
+		}
+
+		/// <summary>
+		/// Applies the warnings to the provided arguments list.
+		/// </summary>
+		/// <param name="level">The warning level to use when applying.</param>
+		/// <param name="toolChainContext">The toolchain context to use in application.</param>
+		/// <param name="arguments">The arugments list to add values to.</param>
+		/// <remarks>Used to provide custom overriding logic.</remarks>
+		protected virtual void ApplyWarningsToArgumentsInternal(WarningLevel level, CompilerWarningsToolChainContext toolChainContext, List<string> arguments)
+		{
+			switch (level)
+			{
+				case WarningLevel.Off:
+					if (_offArgs != null)
+					{
+						arguments.AddRange(_offArgs);
+					}
+					break;
+				case WarningLevel.Warning:
+					if (_warningArgs != null)
+					{
+						arguments.AddRange(_warningArgs);
+					}
+					break;
+				case WarningLevel.Error:
+					if (_errorArgs != null)
+					{
+						arguments.AddRange(_errorArgs);
+					}
+					break;
+			}
+		}
+
+		/// <summary>
+		/// Whether the given <see cref="ApplyWarningsAttribute"/> can be applied to the provided context.
+		/// </summary>
+		/// <param name="toolChainContext">The toolchain context to be applied to.</param>
+		/// <returns>True if applicable, false otherwise.</returns>
+		internal abstract bool CanApplyToContext(CompilerWarningsToolChainContext toolChainContext);
+
+		/// <summary>
+		/// The list of off arugments corresponding to <see cref="WarningLevel.Off"/>
+		/// </summary>
+		internal string[]? OffArgs => _offArgs;
+
+		/// <summary>
+		/// The list of off arugments corresponding to <see cref="WarningLevel.Warning"/>
+		/// </summary>
+		internal string[]? WarningArgs => _warningArgs;
+
+		/// <summary>
+		/// The list of off arugments corresponding to <see cref="WarningLevel.Error"/>
+		/// </summary>
+		internal string[]? ErrorArgs => _errorArgs;
+	}
+
+	#region -- ToolChains --
+
+#pragma warning disable CA1813
+	internal class WarningsVCToolChainAttribute : ApplyWarningsAttribute
+#pragma warning restore CA1813
+	{
+		/// <inheritdoc/>
+		internal WarningsVCToolChainAttribute(string[]? offArgs, string[]? warningArgs, string[]? errorArgs) : base(offArgs, warningArgs, errorArgs) {}
+
+		/// <inheritdoc/>
+		internal override bool CanApplyToContext(CompilerWarningsToolChainContext toolChainContext)
+		{
+			return typeof(VCToolChain) == toolChainContext._toolChainType;
+		}
+	}
+
+#pragma warning disable CA1813
+	internal class WarningsClangToolChainAttribute : ApplyWarningsAttribute
+	{
+#pragma warning restore CA1813
+		/// <inheritdoc/>
+		internal WarningsClangToolChainAttribute(string[]? offArgs, string[]? warningArgs, string[]? errorArgs) : base(offArgs, warningArgs, errorArgs) { }
+
+		/// <inheritdoc/>
+		internal override bool CanApplyToContext(CompilerWarningsToolChainContext toolChainContext)
+		{
+			return typeof(ClangToolChain) == toolChainContext._toolChainType;
+		}
+	}
+
+	#endregion -- ToolChains --
+
+	#region -- Compiler Specializations -- 
+
+	/// <summary>
+	/// Specialized <see cref="WarningsVCToolChainAttribute"/> for MSVC compiler within the <see cref="VCToolChain"/>.
+	/// </summary>
+	internal sealed class WarningsMSCVAttribute : WarningsVCToolChainAttribute
+	{
+		/// <inheritdoc/>
+		internal WarningsMSCVAttribute(string[]? offArgs, string[]? warningArgs, string[]? errorArgs) : base(offArgs, warningArgs, errorArgs) { }
+
+		/// <inheritdoc/>
+		internal override bool CanApplyToContext(CompilerWarningsToolChainContext toolChainContext)
+		{
+			ReadOnlyTargetRules? readOnlyTargetRules = toolChainContext._buildSystemContext.GetReadOnlyTargetRules();
+			return readOnlyTargetRules != null && readOnlyTargetRules.WindowsPlatform != null && readOnlyTargetRules.WindowsPlatform.Compiler.IsMSVC();
+		}
+	}
+
+	#endregion -- Compiler Specializations --
+
+	#region -- ToolChain Specializations --
+
+	/// <summary>
+	/// Specialized <see cref="WarningsClangToolChainAttribute"/> for Shadow Variable Warnings.
+	/// </summary>
+	internal sealed class ShadowVariableWarningsClangToolChainAttribute : WarningsClangToolChainAttribute
+	{
+		/// <inheritdoc/>
+		internal ShadowVariableWarningsClangToolChainAttribute(string[]? offArgs, string[]? warningArgs, string[]? errorArgs) : base(offArgs, warningArgs, errorArgs) {}
+
+		/// <inheritdoc/>
+		protected override void ApplyWarningsToArgumentsInternal(WarningLevel level, CompilerWarningsToolChainContext toolChainContext, List<string> arguments)
+		{
+			if (toolChainContext._toolChainVersion == null || (toolChainContext._toolChainVersion >= new EpicGames.Core.VersionNumber(17) && toolChainContext._toolChainVersion < new VersionNumber(18, 1, 3)))
+			{
+				if (OffArgs != null)
+				{
+					arguments.AddRange(OffArgs);
+				}
+			}
+			else
+			{
+				// Explicitly remove the OffArgs for consideration
+				switch (level)
+				{
+					case WarningLevel.Warning:
+						if (WarningArgs != null)
+							arguments.AddRange(WarningArgs);
+						break;
+					case WarningLevel.Error:
+						if (ErrorArgs != null)
+							arguments.AddRange(ErrorArgs);
+						break;
+				}
+			}
+		}
+	}
+
+	/// <summary>
+	/// Specialized <see cref="WarningsVCToolChainAttribute"/> for Undefined identifier.
+	/// </summary>
+	internal sealed class UndefinedIdentifierWarningsVCToolChainAttribute : WarningsVCToolChainAttribute
+	{
+		/// <inheritdoc/>
+		internal UndefinedIdentifierWarningsVCToolChainAttribute(string[]? offArgs, string[]? warningArgs, string[]? errorArgs) : base(offArgs, warningArgs, errorArgs) { }
+
+		/// <inheritdoc/>
+		protected override void ApplyWarningsToArgumentsInternal(WarningLevel level, CompilerWarningsToolChainContext toolChainContext, List<string> arguments)
+		{
+			if (!toolChainContext._compileEnvironment.bPreprocessOnly)
+			{
+				base.ApplyWarningsToArgumentsInternal(level, toolChainContext, arguments);
+			}
+		}
+	}
+
+	/// <summary>
+	/// Specialized <see cref="WarningsVCToolChainAttribute"/> for Unsafe Type Casts.
+	/// </summary>
+	internal sealed class UnsafeTypeCastWarningsVCToolChainAttribute : WarningsVCToolChainAttribute
+	{
+		/// <inheritdoc/>
+		internal UnsafeTypeCastWarningsVCToolChainAttribute(string[] offArgs, string[] warningArgs, string[] errorArgs) : base(offArgs, warningArgs, errorArgs) { }
+
+		/// <inheritdoc/>
+		protected override void ApplyWarningsToArgumentsInternal(WarningLevel level, CompilerWarningsToolChainContext toolChainContext, List<string> arguments)
+		{
+			WarningLevel effectiveCastWarningLevel = (toolChainContext._buildSystemContext.GetReadOnlyTargetRules()?.Platform == UnrealTargetPlatform.Win64) ? level : WarningLevel.Off;
+			base.ApplyWarningsToArgumentsInternal(effectiveCastWarningLevel, toolChainContext, arguments);
+		}
+	}
+
+	#endregion -- ToolChain Specializations --
+}
\ No newline at end of file
diff --git a/Engine/Source/Programs/UnrealBuildTool/Configuration/CompileWarnings/CompilerWarningsToolChainContext.cs b/Engine/Source/Programs/UnrealBuildTool/Configuration/CompileWarnings/CompilerWarningsToolChainContext.cs
new file mode 100644
--- /dev/null
+++ b/Engine/Source/Programs/UnrealBuildTool/Configuration/CompileWarnings/CompilerWarningsToolChainContext.cs
@@ -0,0 +1,34 @@
+// Copyright Epic Games, Inc. All Rights Reserved.
+
+using System;
+using EpicGames.Core;
+
+namespace UnrealBuildTool
+{
+	/// <summary>
+	/// Context object used to pass data to <see cref="CppCompileWarnings"/> in apply warnings phase.
+	/// </summary>
+	/// <remarks>This class should be used as the common place to add new context as necessary for <see cref="ApplyWarningsAttribute.ApplyWarningsToArguments"/> extensibility.</remarks>
+	internal class CompilerWarningsToolChainContext
+	{
+		/// <summary>
+		/// Constructor for toolchain context.
+		/// </summary>
+		/// <param name="compileEnvironment">Compile environment for the tool chain.</param>
+		/// <param name="buildSystemContext">Build context of the <see cref="CppCompileWarnings"/>.</param>
+		/// <param name="toolChainType">The type of the toolchain.</param>
+		/// <param name="toolChainVersion">The version of the toolchain.</param>
+		public CompilerWarningsToolChainContext(CppCompileEnvironment compileEnvironment, BuildSystemContext buildSystemContext, Type toolChainType, VersionNumber? toolChainVersion)
+		{
+			_compileEnvironment = compileEnvironment;
+			_buildSystemContext = buildSystemContext;
+			_toolChainType = toolChainType;
+			_toolChainVersion = toolChainVersion;
+		}
+
+		internal CppCompileEnvironment _compileEnvironment;
+		internal BuildSystemContext _buildSystemContext;
+		internal Type _toolChainType;
+		internal VersionNumber? _toolChainVersion;
+	}
+}
\ No newline at end of file
diff --git a/Engine/Source/Programs/UnrealBuildTool/Configuration/CppCompileWarnings.cs b/Engine/Source/Programs/UnrealBuildTool/Configuration/CppCompileWarnings.cs
--- a/Engine/Source/Programs/UnrealBuildTool/Configuration/CppCompileWarnings.cs
+++ b/Engine/Source/Programs/UnrealBuildTool/Configuration/CppCompileWarnings.cs
@@ -3,6 +3,9 @@
 using EpicGames.Core;
 using Microsoft.Extensions.Logging;
 using System;
+using System.Collections.Generic;
+using System.Linq;
+using System.Reflection;
 
 namespace UnrealBuildTool
 {
@@ -17,7 +20,7 @@ namespace UnrealBuildTool
 		/// <remarks>If <paramref name="appliedWarningLevel"/> with value of null or <see cref="WarningLevel.Default"/> will be ignored as an override.</remarks>
 		internal static void ApplyWarning(this ref WarningLevel inSourceWarningLevel, WarningLevel? appliedWarningLevel, bool overwriteSourceValue = true)
 		{
-			if ((appliedWarningLevel.HasValue && appliedWarningLevel.Value != WarningLevel.Default) && (overwriteSourceValue ||  (!overwriteSourceValue && inSourceWarningLevel == WarningLevel.Default)))
+			if ((appliedWarningLevel.HasValue && appliedWarningLevel.Value != WarningLevel.Default) && (overwriteSourceValue || (!overwriteSourceValue && inSourceWarningLevel == WarningLevel.Default)))
 			{
 				inSourceWarningLevel = appliedWarningLevel.Value;
 			}
@@ -27,7 +30,7 @@ namespace UnrealBuildTool
 	/// <summary>
 	/// Container class used for C++ compiler warning settings.
 	/// </summary>
-	public class CppCompileWarnings
+	public partial class CppCompileWarnings
 	{
 		/*
 		 * To add a new compiler warning, you need to do the following:
@@ -44,14 +47,26 @@ namespace UnrealBuildTool
 		 *		For Target defaults				- internal void ApplyTargetDefaults()
 		 *			Note: This is the most appropriatel place for BuildSettings version specific defaults
 		 * 4. Access the new property through the (TargetRules|ModuleRules|CppCompileEnvironment).CppCompileWarnings as needed
+		 * 5. (Optional) For compiler line support, utilize the ApplyWarningsAttribute system to participate in invocations to GenerateWarningCommandLineAgs.
 		 */
 		#region -- Compiler Warning Settings --
 
 		/// <summary>
 		/// Forces shadow variable warnings to be treated as errors on platforms that support it.
+		/// MSVC - 
+		///		https://learn.microsoft.com/en-us/cpp/error-messages/compiler-warnings/compiler-warning-level-4-c4456
+		///		4456 - declaration of 'LocalVariable' hides previous local declaration
+		///		https://learn.microsoft.com/en-us/cpp/error-messages/compiler-warnings/compiler-warning-level-4-c4458
+		///		4458 - declaration of 'parameter' hides class member
+		///		https://learn.microsoft.com/en-us/cpp/error-messages/compiler-warnings/compiler-warning-level-4-c4459
+		///		4459 - declaration of 'LocalVariable' hides global declaration
+		///	Clang - 
+		///		 https://clang.llvm.org/docs/DiagnosticsReference.html#wshadow
 		/// </summary>
 		[CommandLine("-ShadowVariableErrors", Value = nameof(WarningLevel.Error))]
 		[RequiresUniqueBuildEnvironment]
+		[WarningsMSCV(["/wd4456", "/wd4458", "/wd4459"], [], ["/we4456", "/we4458", "/we4459"])]
+		[ShadowVariableWarningsClangToolChain(["-Wno-shadow"], ["-Wshadow", "-Wno-error=shadow"], ["-Wshadow"])]
 		public WarningLevel ShadowVariableWarningLevel
 		{
 			get => ResolveWarning(ParentCppCompileWarnings?.ShadowVariableWarningLevel, _shadowVariableWarningLevel);
@@ -61,9 +76,32 @@ namespace UnrealBuildTool
 
 		/// <summary>
 		/// Indicates what warning/error level to treat unsafe type casts as on platforms that support it (e.g., double->float or int64->int32)
+		/// MSVC - 
+		///		https://learn.microsoft.com/en-us/cpp/error-messages/compiler-warnings/compiler-warning-levels-3-and-4-c4244
+		///		4244 - conversion from 'type1' to 'type2', possible loss of data
+		///		44244 - Note: The extra 4 is not a typo, /wLXXXX sets warning XXXX to level L
+		///		https://learn.microsoft.com/en-us/cpp/error-messages/compiler-warnings/compiler-warning-level-1-c4838
+		///		4838 - conversion from 'type1' to 'type2' requires a narrowing conversion
+		///		44838 - Note: The extra 4 is not a typo, /wLXXXX sets warning XXXX to level L
+		///	Clang - 
+		///		https://clang.llvm.org/docs/DiagnosticsReference.html#wfloat-conversion
+		///		https://clang.llvm.org/docs/DiagnosticsReference.html#wimplicit-float-conversion
+		///		https://clang.llvm.org/docs/DiagnosticsReference.html#wimplicit-int-conversion
+		///		https://clang.llvm.org/docs/DiagnosticsReference.html#wc-11-narrowing
+		///		
+		///		To enable (too many hits right now)
+		///		https://clang.llvm.org/docs/DiagnosticsReference.html#wshorten-64-to-32
+		///		https://clang.llvm.org/docs/DiagnosticsReference.html#wsign-conversion
 		/// </summary>
+		/// <remarks>This should be kept in sync with PRAGMA_DISABLE_UNSAFE_TYPECAST_WARNINGS in ClangPlatformCompilerPreSetup.h</remarks>
+		/// <remarks>Clang: "shorten-64-to-32"; too many hits right now, probably want it *soon*; "sign-conversion"; too many hits right now, probably want it eventually</remarks>
 		[XmlConfigFile(Category = "BuildConfiguration")]
 		[RequiresUniqueBuildEnvironment]
+		[UnsafeTypeCastWarningsVCToolChain(["/wd4244", "/wd4838"], ["/w44244", "/w44838"], ["/we4244", "/we4838"])]
+		[WarningsClangToolChain(
+			["-Wno-float-conversion", "-Wno-implicit-float-conversion", "-Wno-implicit-int-conversion", "-Wno-c++11-narrowing"],
+			["-Wfloat-conversion", "-Wno-error=float-conversion", "-Wimplicit-float-conversion", "-Wno-error=implicit-float-conversion", "-Wimplicit-int-conversion", "-Wno-error=implicit-int-conversion", "-Wc++11-narrowing", "-Wno-error=c++11-narrowing"],
+			["-Wfloat-conversion", "-Wimplicit-float-conversion", "-Wimplicit-int-conversion", "-Wc++11-narrowing"])]
 		public WarningLevel UnsafeTypeCastWarningLevel
 		{
 			get => ResolveWarning(ParentCppCompileWarnings?.UnsafeTypeCastWarningLevel, _unsafeTypeCastWarningLevel);
@@ -73,10 +111,18 @@ namespace UnrealBuildTool
 
 		/// <summary>
 		/// Indicates what warning/error level to treat undefined identifiers in conditional expressions.
+		/// MSVC - 
+		///		https://learn.microsoft.com/en-us/cpp/error-messages/compiler-warnings/compiler-warning-level-4-c4668 
+		///		4668 - 'symbol' is not defined as a preprocessor macro, replacing with '0' for 'directives'
+		///		44668 - Note: The extra 4 is not a typo, /wLXXXX sets warning XXXX to level L
+		/// Clang - 
+		///		https://clang.llvm.org/docs/DiagnosticsReference.html#wundef
 		/// </summary>
 		[XmlConfigFile(Category = "BuildConfiguration")]
 		[CommandLine("-UndefinedIdentifierWarningLevel=")]
 		[RequiresUniqueBuildEnvironment]
+		[UndefinedIdentifierWarningsVCToolChain(null, ["/w44668"], ["/we4668"])]
+		[WarningsClangToolChain(null, ["-Wundef", "-Wno-error=undef"], ["-Wundef"])]
 		public WarningLevel UndefinedIdentifierWarningLevel
 		{
 			get => ResolveWarning(ParentCppCompileWarnings?.UndefinedIdentifierWarningLevel, _undefinedIdentifierWarningLevel);
@@ -86,10 +132,17 @@ namespace UnrealBuildTool
 
 		/// <summary>
 		/// Indicates what warning/error level to treat unhandled enumerators in switches on enumeration-typed values.
+		/// MSVC - 
+		///		4061 - https://learn.microsoft.com/en-us/cpp/error-messages/compiler-warnings/compiler-warning-level-4-c4061 
+		///		44061 - Note: The extra 4 is not a typo, /wLXXXX sets warning XXXX to level L
+		///	Clang - 
+		///		https://clang.llvm.org/docs/DiagnosticsReference.html#wswitch-enum
 		/// </summary>
 		[XmlConfigFile(Category = "BuildConfiguration")]
 		[CommandLine("-SwitchUnhandledEnumeratorWarningLevel=")]
 		[RequiresUniqueBuildEnvironment]
+		[WarningsVCToolChain(null, ["/w44061"], ["/we4061"])]
+		[WarningsClangToolChain(null, ["-Wswitch-enum", "-Wno-error=switch-enum"], ["-Wswitch-enum"])]
 		public WarningLevel SwitchUnhandledEnumeratorWarningLevel
 		{
 			get => ResolveWarning(ParentCppCompileWarnings?.SwitchUnhandledEnumeratorWarningLevel, _switchUnhandledEnumeratorWarningLevel);
@@ -274,7 +327,7 @@ namespace UnrealBuildTool
 				ParentCppCompileWarnings = new ReadOnlyCppCompileWarnings(newParentCppCompilerWarnings);
 				return;
 			}
-			
+
 			_logger?.LogTrace("Promoting existing parent CppCompileWarnings, and inserting new immediate parent.");
 			newParentCppCompilerWarnings.ParentCppCompileWarnings = ParentCppCompileWarnings;
 
@@ -298,6 +351,116 @@ namespace UnrealBuildTool
 
 		#endregion -- Public API --
 
+		private static IEnumerable<PropertyInfo>? s_propertiesWithAttribute = null;
+		private static IEnumerable<FieldInfo>? s_fieldsWithAttribute = null;
+
+		/// <summary>
+		/// Generates the compiler warnings command line arguments for the requested <see cref="UEToolChain"/> type, and version.
+		/// </summary>
+		/// <param name="compileEnvironment">The compile environment of the invoking tool chain.</param>
+		/// <param name="toolChainType">The type of the invoking toolchain.</param>
+		/// <param name="toolChainVersion">The version of the tool chain, if applicable.</param>
+		/// <returns>A list of the compiler command line arguments for compile warnings.</returns>
+		/// <remarks>The type of the tool chain must be the specific class within the hierarchy due to how <see cref="ClangToolChain.GetCompileArguments_WarningsAndErrors"/> applies arguments.</remarks>
+		internal IEnumerable<string> GenerateWarningCommandLineArgs(CppCompileEnvironment compileEnvironment, System.Type toolChainType, VersionNumber? toolChainVersion = null)
+		{
+			if (!typeof(UEToolChain).IsAssignableFrom(toolChainType))
+			{
+				_logger?.LogWarning("Requested command line arugments of a toolchain type ({TypeName}) that is not a part of the UEToolChain hierarchy.", toolChainType.Name);
+
+				return new List<string>(0);
+			}
+
+			CompilerWarningsToolChainContext toolChainContext = new CompilerWarningsToolChainContext(compileEnvironment, BuildContextProvider, toolChainType, toolChainVersion); //pull out into request?
+
+			List<string> warningArguments = new List<string>();
+
+			// Process properties
+			{
+				s_propertiesWithAttribute ??= typeof(CppCompileWarnings)
+					.GetProperties(BindingFlags.Public | BindingFlags.NonPublic | BindingFlags.Instance | BindingFlags.Static)
+					.Where(p => p.IsDefined(typeof(ApplyWarningsAttribute), false));
+
+				foreach (PropertyInfo property in s_propertiesWithAttribute)
+				{
+					if (property.PropertyType != typeof(WarningLevel))
+					{
+						_logger?.LogWarning("Attempting to resolve warnings, and encountered invalid property ({PropertyName}). Skipping this property.", property.Name);
+						continue;
+					}
+
+					ApplyWarningsAttribute? applyWarningsAttribute = null;
+					IEnumerable<ApplyWarningsAttribute?> applyWarningsAttributes = property.GetCustomAttributes<ApplyWarningsAttribute>().Where(x => x.CanApplyToContext(toolChainContext));
+
+					if (applyWarningsAttributes.Count() > 1)
+					{
+						_logger?.LogWarning("Ambiguous ApplyWarningsAttributes on the property ({PropertyName}), for the provided tool-chain ({ToolChainName}). Skipping this property.", property.Name, toolChainContext._toolChainType.Name);
+						continue;
+					}
+
+					applyWarningsAttribute = applyWarningsAttributes.FirstOrDefault();
+
+					if (applyWarningsAttribute != null)
+					{
+						object? instanceBoxedValue = property.GetValue(this);
+
+						if (instanceBoxedValue == null)
+						{
+							_logger?.LogWarning("Invalid boxed value returned from attribute. Skipping this property.");
+							continue;
+						}
+
+						WarningLevel instanceWarningLevel = (WarningLevel)instanceBoxedValue;
+						applyWarningsAttribute.ApplyWarningsToArguments(instanceWarningLevel, toolChainContext, warningArguments);
+					}
+				}
+			}
+
+			// Process fields
+			{
+				s_fieldsWithAttribute ??= typeof(CppCompileWarnings)
+					.GetFields(BindingFlags.Public | BindingFlags.NonPublic | BindingFlags.Instance | BindingFlags.Static)
+					.Where(p => p.IsDefined(typeof(ApplyWarningsAttribute), false));
+
+				foreach (FieldInfo field in s_fieldsWithAttribute)
+				{
+					if (field.FieldType != typeof(WarningLevel))
+					{
+						_logger?.LogWarning("Attempting to resolve warnings, and encountered invalid field ({FieldName}). Skipping this field.", field.Name);
+						continue;
+					}
+
+					ApplyWarningsAttribute? applyWarningsAttribute = null;
+					IEnumerable<ApplyWarningsAttribute?> applyWarningsAttributes = field.GetCustomAttributes<ApplyWarningsAttribute>().Where(x => x.CanApplyToContext(toolChainContext));
+
+					if (applyWarningsAttributes.Count() > 1)
+					{
+						_logger?.LogWarning("Ambiguous ApplyWarningsAttributes on the field ({FieldName}), for the provided tool-chain ({ToolChainName}). Skipping this field.", field.Name, toolChainContext._toolChainType.Name);
+						continue;
+					}
+
+					applyWarningsAttribute = applyWarningsAttributes.FirstOrDefault();
+
+					if (applyWarningsAttribute != null)
+					{
+						object? instanceBoxedValue = field.GetValue(this);
+
+						if (instanceBoxedValue == null)
+						{
+							_logger?.LogWarning("Invalid boxed value returned from attribute. Skipping this field.");
+							continue;
+						}
+
+						WarningLevel instanceWarningLevel = (WarningLevel)instanceBoxedValue;
+						applyWarningsAttribute.ApplyWarningsToArguments(instanceWarningLevel, toolChainContext, warningArguments);
+					}
+				}
+			}
+
+			return warningArguments;
+		}
+
+
 		#region -- Internal API --
 
 		/// <summary>
@@ -326,7 +489,7 @@ namespace UnrealBuildTool
 			cppCompileWarnings._moduleIncludePrivateWarningLevel.ApplyWarning(WarningLevel.Off, overwriteSourceValue);
 			cppCompileWarnings._moduleIncludeSubdirectoryWarningLevel.ApplyWarning(WarningLevel.Off, overwriteSourceValue);
 
-			cppCompileWarnings._disablePluginsConflictWarningLevel.ApplyWarning(WarningLevel.Default,overwriteSourceValue);
+			cppCompileWarnings._disablePluginsConflictWarningLevel.ApplyWarning(WarningLevel.Default, overwriteSourceValue);
 			cppCompileWarnings._nonInlinedGenCppWarningLevel.ApplyWarning(WarningLevel.Off, false);
 		}
 
diff --git a/Engine/Source/Programs/UnrealBuildTool/Platform/Windows/VCToolChain.cs b/Engine/Source/Programs/UnrealBuildTool/Platform/Windows/VCToolChain.cs
--- a/Engine/Source/Programs/UnrealBuildTool/Platform/Windows/VCToolChain.cs
+++ b/Engine/Source/Programs/UnrealBuildTool/Platform/Windows/VCToolChain.cs
@@ -1110,77 +1110,13 @@ namespace UnrealBuildTool
 			if (CompileEnvironment.DeprecationWarningLevel == WarningLevel.Off)
 			{
 				Arguments.Add("/wd4996");
-			}
+			} 
 			else if (CompileEnvironment.DeprecationWarningLevel == WarningLevel.Error)
 			{
 				Arguments.Add("/we4996");
 			}
 
-			//@todo: Disable warnings for VS2017. These should be reenabled as we clear the reasons for them out of the engine source and the VS2015 toolchain evolves.
-			if (Target.WindowsPlatform.Compiler.IsMSVC())
-			{
-				// Disable shadow variable warnings
-				if (CompileEnvironment.CppCompileWarnings.ShadowVariableWarningLevel == WarningLevel.Off)
-				{
-					Arguments.Add("/wd4456"); // 4456 - declaration of 'LocalVariable' hides previous local declaration
-					Arguments.Add("/wd4458"); // 4458 - declaration of 'parameter' hides class member
-					Arguments.Add("/wd4459"); // 4459 - declaration of 'LocalVariable' hides global declaration
-				}
-				else if (CompileEnvironment.CppCompileWarnings.ShadowVariableWarningLevel == WarningLevel.Error)
-				{
-					Arguments.Add("/we4456"); // 4456 - declaration of 'LocalVariable' hides previous local declaration
-					Arguments.Add("/we4458"); // 4458 - declaration of 'parameter' hides class member
-					Arguments.Add("/we4459"); // 4459 - declaration of 'LocalVariable' hides global declaration
-				}
-			}
-
-			if (CompileEnvironment.CppCompileWarnings.UndefinedIdentifierWarningLevel != WarningLevel.Off && !CompileEnvironment.bPreprocessOnly)
-			{
-				if (CompileEnvironment.CppCompileWarnings.UndefinedIdentifierWarningLevel == WarningLevel.Error)
-				{
-					Arguments.Add("/we4668");
-				}
-				else
-				{
-					Arguments.Add("/w44668");
-				}
-			}
-
-			if (CompileEnvironment.CppCompileWarnings.SwitchUnhandledEnumeratorWarningLevel != WarningLevel.Off)
-			{
-				if (CompileEnvironment.CppCompileWarnings.SwitchUnhandledEnumeratorWarningLevel == WarningLevel.Error)
-				{
-					Arguments.Add("/we4061");
-				}
-				else
-				{
-					Arguments.Add("/w44061");
-				}
-			}
-
-			// The unsafe type cast warnings setting controls the following warnings currently:
-			//   4244: conversion from 'type1' to 'type2', possible loss of data
-			//   4838: conversion from 'type1' to 'type2' requires a narrowing conversion
-			//@TODO: FLOATPRECISION: Consider doing the following as well:
-			//   4267: 'var' : conversion from 'size_t' to 'type', possible loss of data
-			//   4305: 'identifier' : truncation from 'type1' to 'type2'
-			WarningLevel EffectiveCastWarningLevel = (Target.Platform == UnrealTargetPlatform.Win64) ? CompileEnvironment.CppCompileWarnings.UnsafeTypeCastWarningLevel : WarningLevel.Off;
-			if (EffectiveCastWarningLevel == WarningLevel.Error)
-			{
-				Arguments.Add("/we4244");
-				Arguments.Add("/we4838");
-			}
-			else if (EffectiveCastWarningLevel == WarningLevel.Warning)
-			{
-				// Note: The extra 4 is not a typo, /wLXXXX sets warning XXXX to level L
-				Arguments.Add("/w44244");
-				Arguments.Add("/w44838");
-			}
-			else
-			{
-				Arguments.Add("/wd4244");
-				Arguments.Add("/wd4838");
-			}
+			Arguments.AddRange(CompileEnvironment.CppCompileWarnings.GenerateWarningCommandLineArgs(CompileEnvironment, typeof(VCToolChain)));
 
 			if (CompileEnvironment.bUseHeaderUnitsForPch)
 			{
diff --git a/Engine/Source/Programs/UnrealBuildTool/ToolChain/ClangWarnings.cs b/Engine/Source/Programs/UnrealBuildTool/ToolChain/ClangWarnings.cs
--- a/Engine/Source/Programs/UnrealBuildTool/ToolChain/ClangWarnings.cs
+++ b/Engine/Source/Programs/UnrealBuildTool/ToolChain/ClangWarnings.cs
@@ -125,61 +125,7 @@ namespace UnrealBuildTool
 				}
 			}
 
-			if (ClangVersion >= new VersionNumber(17) && ClangVersion < new VersionNumber(18, 1, 3))
-			{
-				Arguments.Add("-Wno-shadow");
-			}
-			else
-			{
-				// https://clang.llvm.org/docs/DiagnosticsReference.html#wshadow
-				if (CompileEnvironment.CppCompileWarnings.ShadowVariableWarningLevel != WarningLevel.Off)
-				{
-					Arguments.Add("-Wshadow" + ((CompileEnvironment.CppCompileWarnings.ShadowVariableWarningLevel == WarningLevel.Error) ? "" : " -Wno-error=shadow"));
-				}
-			}
-
-			// https://clang.llvm.org/docs/DiagnosticsReference.html#wundef
-			if (CompileEnvironment.CppCompileWarnings.UndefinedIdentifierWarningLevel != WarningLevel.Off)
-			{
-				Arguments.Add("-Wundef" + (CompileEnvironment.CppCompileWarnings.UndefinedIdentifierWarningLevel == WarningLevel.Error ? "" : " -Wno-error=undef"));
-			}
-
-			if (CompileEnvironment.CppCompileWarnings.SwitchUnhandledEnumeratorWarningLevel != WarningLevel.Off)
-			{
-				Arguments.Add("-Wswitch-enum" + (CompileEnvironment.CppCompileWarnings.SwitchUnhandledEnumeratorWarningLevel == WarningLevel.Error ? "" : " -Wno-error=switch-enum"));
-			}
-
-			// Note: This should be kept in sync with PRAGMA_DISABLE_UNSAFE_TYPECAST_WARNINGS in ClangPlatformCompilerPreSetup.h
-			string[] UnsafeTypeCastWarningList = {
-				"float-conversion",
-				"implicit-float-conversion",
-				"implicit-int-conversion",
-				"c++11-narrowing"
-				//"shorten-64-to-32",	<-- too many hits right now, probably want it *soon*
-				//"sign-conversion",	<-- too many hits right now, probably want it eventually
-			};
-
-			if (CompileEnvironment.CppCompileWarnings.UnsafeTypeCastWarningLevel == WarningLevel.Error)
-			{
-				foreach (string Warning in UnsafeTypeCastWarningList)
-				{
-					Arguments.Add("-W" + Warning);
-				}
-			}
-			else if (CompileEnvironment.CppCompileWarnings.UnsafeTypeCastWarningLevel == WarningLevel.Warning)
-			{
-				foreach (string Warning in UnsafeTypeCastWarningList)
-				{
-					Arguments.Add("-W" + Warning + " -Wno-error=" + Warning);
-				}
-			}
-			else
-			{
-				foreach (string Warning in UnsafeTypeCastWarningList)
-				{
-					Arguments.Add("-Wno-" + Warning);
-				}
-			}
+			Arguments.AddRange(CompileEnvironment.CppCompileWarnings.GenerateWarningCommandLineArgs(CompileEnvironment, typeof(ClangToolChain), ClangVersion));
 
 			// -Wimplicit-int-conversion also controls -Wshorten-64-to-32 as of clang 19
 			if (ClangVersion >= new VersionNumber(19))
