UnrealBuildTool: Fix minor typo in UBA connection output

#rnx

diff --git a/Engine/Source/Programs/UnrealBuildTool/Executors/UnrealBuildAccelerator/UBAAgentCoordinatorHorde.cs b/Engine/Source/Programs/UnrealBuildTool/Executors/UnrealBuildAccelerator/UBAAgentCoordinatorHorde.cs
--- a/Engine/Source/Programs/UnrealBuildTool/Executors/UnrealBuildAccelerator/UBAAgentCoordinatorHorde.cs
+++ b/Engine/Source/Programs/UnrealBuildTool/Executors/UnrealBuildAccelerator/UBAAgentCoordinatorHorde.cs
@@ -433,7 +433,7 @@ namespace UnrealBuildTool
 				encryption = Encryption.Ssl;
 			}
 
-			logger.LogInformation("Horde URL: {Server}, Pool: {Pool}, Cluster {Cluster}, Condition: {Condition}, Connection: {Connection} HordeEncryption: {Encryption}",
+			logger.LogInformation("Horde URL: {Server}, Pool: {Pool}, Cluster: {Cluster}, Condition: {Condition}, Connection: {Connection} HordeEncryption: {Encryption}",
 				server, hordeConfig.HordePool ?? "(none)", hordeConfig.HordeCluster ?? "(none)", hordeConfig.HordeCondition ?? "(none)", connectionMode?.ToString() ?? "(none)", encryption?.ToString() ?? "(none)");
 			try
 			{
