[UnrealBuildTool]
- Parallelize WriteConfiguration for targets
- 41s to 37s

#rnx

diff --git a/Engine/Source/Programs/UnrealBuildTool/Platform/Mac/UEBuildMac.cs b/Engine/Source/Programs/UnrealBuildTool/Platform/Mac/UEBuildMac.cs
--- a/Engine/Source/Programs/UnrealBuildTool/Platform/Mac/UEBuildMac.cs
+++ b/Engine/Source/Programs/UnrealBuildTool/Platform/Mac/UEBuildMac.cs
@@ -286,7 +286,8 @@ namespace UnrealBuildTool
 			UnrealArchitectures Result = new UnrealArchitectures(Architectures);
 			lock (ProjectArchitectureCache)
 			{
-				ProjectArchitectureCache.Add(Key, Result);
+				// Another thread could have beaten us to the add
+				ProjectArchitectureCache.TryAdd(Key, Result);
 			}
 			return Result;
 		}
diff --git a/Engine/Source/Programs/UnrealBuildTool/ProjectFiles/VisualStudio/VCProject.cs b/Engine/Source/Programs/UnrealBuildTool/ProjectFiles/VisualStudio/VCProject.cs
--- a/Engine/Source/Programs/UnrealBuildTool/ProjectFiles/VisualStudio/VCProject.cs
+++ b/Engine/Source/Programs/UnrealBuildTool/ProjectFiles/VisualStudio/VCProject.cs
@@ -1405,9 +1405,30 @@ namespace UnrealBuildTool
 			}
 
 			// Write each configuration
-			foreach (ProjectConfigAndTargetCombination Combination in ProjectConfigAndTargetCombinations)
 			{
-				WriteConfiguration(ProjectName, Combination, VCProjectFileContent, PlatformProjectGenerators, bGenerateUserFileContent ? VCUserFileContent : null, bGenerateUserFileContent ? VCUserFileSettings : null);
+				StringBuilder[] TmpVCProjectFileContent = new StringBuilder[ProjectConfigAndTargetCombinations.Count];
+				StringBuilder[] TmpVCUserFileContent = new StringBuilder[ProjectConfigAndTargetCombinations.Count];
+
+				Parallel.For(0, ProjectConfigAndTargetCombinations.Count, Index =>
+				{
+					TmpVCProjectFileContent[Index] = new StringBuilder();
+					TmpVCUserFileContent[Index] = new StringBuilder();
+
+					WriteConfiguration(ProjectName, ProjectConfigAndTargetCombinations[Index], TmpVCProjectFileContent[Index], PlatformProjectGenerators, bGenerateUserFileContent ? TmpVCUserFileContent[Index] : null, bGenerateUserFileContent ? VCUserFileSettings : null);
+				});
+
+				for (int Index = 0; Index < ProjectConfigAndTargetCombinations.Count; Index++)
+				{
+					if (TmpVCProjectFileContent[Index].Length > 0) 
+					{
+						VCProjectFileContent.Append(TmpVCProjectFileContent[Index]);
+					}
+
+					if (TmpVCUserFileContent[Index].Length > 0) 
+					{
+						VCUserFileContent.Append(TmpVCUserFileContent[Index]);
+					}
+				}
 			}
 
 			{
