[UBT/Instrumentation]
- Fix missing includes when compiling monolithic builds with instrumentation active.

diff --git a/Engine/Source/Programs/UnrealBuildTool/Configuration/UEBuildModuleCPP.cs b/Engine/Source/Programs/UnrealBuildTool/Configuration/UEBuildModuleCPP.cs
--- a/Engine/Source/Programs/UnrealBuildTool/Configuration/UEBuildModuleCPP.cs
+++ b/Engine/Source/Programs/UnrealBuildTool/Configuration/UEBuildModuleCPP.cs
@@ -680,6 +680,14 @@ namespace UnrealBuildTool
 				}
 			}
 
+			if (Target.WindowsPlatform.bEnableInstrumentation)
+			{
+				CppCompileEnvironment.UserIncludePaths.Add(DirectoryReference.Combine(Unreal.EngineSourceDirectory, "Runtime", "Instrumentation", "Public"));
+				CppCompileEnvironment.UserIncludePaths.Add(DirectoryReference.Combine(Unreal.EngineSourceDirectory, "Runtime", "Instrumentation", "Internal"));
+				CppCompileEnvironment.UserIncludePaths.Add(DirectoryReference.Combine(Unreal.EngineSourceDirectory, "Runtime", "Core", "Public"));
+				CppCompileEnvironment.UserIncludePaths.Add(DirectoryReference.Combine(Unreal.EngineSourceDirectory, "Runtime", "Core", "Internal"));
+			}
+
 			// The operator new/delete overrides need to exist in all modules.
 			// We put these overrides in a separate file in order to be able to filter out so only one file is added when modules are merged
 			if (Target.LinkType == TargetLinkType.Modular)
@@ -696,11 +704,6 @@ namespace UnrealBuildTool
 
 				if (Target.WindowsPlatform.bEnableInstrumentation)
 				{
-					CppCompileEnvironment.UserIncludePaths.Add(DirectoryReference.Combine(Unreal.EngineSourceDirectory, "Runtime", "Instrumentation", "Public"));
-					CppCompileEnvironment.UserIncludePaths.Add(DirectoryReference.Combine(Unreal.EngineSourceDirectory, "Runtime", "Instrumentation", "Internal"));
-					CppCompileEnvironment.UserIncludePaths.Add(DirectoryReference.Combine(Unreal.EngineSourceDirectory, "Runtime", "Core", "Public"));
-					CppCompileEnvironment.UserIncludePaths.Add(DirectoryReference.Combine(Unreal.EngineSourceDirectory, "Runtime", "Core", "Internal"));
-
 					// Give the ability for core to hotpatch itself for the per module functions to save another set of JMPs.
 					if (Name == "Core")
 					{ 
