UnrealBuildTool: When chaining PCH, if there is a valid parent pch template it needs to be created if it doesn't exist

#rnx

diff --git a/Engine/Source/Programs/UnrealBuildTool/Configuration/UEBuildModuleCPP.cs b/Engine/Source/Programs/UnrealBuildTool/Configuration/UEBuildModuleCPP.cs
--- a/Engine/Source/Programs/UnrealBuildTool/Configuration/UEBuildModuleCPP.cs
+++ b/Engine/Source/Programs/UnrealBuildTool/Configuration/UEBuildModuleCPP.cs
@@ -1474,47 +1474,19 @@ namespace UnrealBuildTool
 
 				// Setup PCH chaining
 				PrecompiledHeaderInstance? ParentPCHInstance = null;
-				Dictionary<string, string>? DefinitionsDictionary = null;
 				if (Rules.Target.bChainPCHs)
 				{
-					// Create a lookup table for the definitions
-					DefinitionsDictionary = Definitions.ToHashSet().ToDictionary(x => x.Split('=')[0], x => (x.Split("=").Length >= 2) ? x.Split("=")[1] : String.Empty);
-
 					// Find all the dependencies of this module
 					HashSet<UEBuildModule> ReferencedModules = new HashSet<UEBuildModule>();
 					Template.Module.GetAllDependencyModules(new List<UEBuildModule>(), ReferencedModules, bIncludeDynamicallyLoaded: false, bForceCircular: false, bOnlyDirectDependencies: false);
 
-					foreach (PrecompiledHeaderTemplate ParentTemplate in CompileEnvironment.SharedPCHs.Where(x => ReferencedModules.Contains(x.Module)))
+					PrecompiledHeaderTemplate? ParentTemplate = CompileEnvironment.SharedPCHs
+						.Skip(CompileEnvironment.SharedPCHs.IndexOf(Template) + 1)
+						.FirstOrDefault(x => ReferencedModules.Contains(x.Module) && x.IsValidFor(CompileEnvironment));
+					if (ParentTemplate != null) 
 					{
-						if (ParentTemplate.IsValidFor(CompileEnvironment))
-						{
-							bool AreSharedPCHDefinitionsCompatible(PrecompiledHeaderInstance ParentInstance)
-							{
-								if (ParentInstance.DefinitionsDictionary == null)
-								{
-									return false;
-								}
-
-								foreach (KeyValuePair<string, string> Definition in DefinitionsDictionary)
-								{
-									if (ParentInstance.DefinitionsDictionary.TryGetValue(Definition.Key, out string? ParentValue))
-									{
-										if (ParentValue != Definition.Value)
-										{
-											return false;
-										}
-									}
-								}
-								return true;
-							}
-
-							ParentPCHInstance = ParentTemplate.Instances.Find(x => IsCompatibleForSharedPCH(x.CompileEnvironment, ModuleCompileEnvironment) && AreSharedPCHDefinitionsCompatible(x));
-							if (ParentPCHInstance != null)
-							{
-								CompileEnvironment.ParentPCHInstance = ParentPCHInstance;
-								break;
-							}
-						}
+						ParentPCHInstance = FindOrCreateSharedPCH(ToolChain, ParentTemplate, ModuleCompileEnvironment, Graph);
+						CompileEnvironment.ParentPCHInstance = ParentPCHInstance;
 					}
 				}
 
@@ -1530,7 +1502,6 @@ namespace UnrealBuildTool
 				}
 				Instance = new PrecompiledHeaderInstance(WrapperFile, SharedDefinitionsFileItem, CompileEnvironment, Output, GetImmutableDefinitions(Template.BaseCompileEnvironment.Definitions))
 				{
-					DefinitionsDictionary = DefinitionsDictionary,
 					ParentPCHInstance = ParentPCHInstance
 				};
 				Template.Instances.Add(Instance);
diff --git a/Engine/Source/Programs/UnrealBuildTool/System/PrecompiledHeaderInstance.cs b/Engine/Source/Programs/UnrealBuildTool/System/PrecompiledHeaderInstance.cs
--- a/Engine/Source/Programs/UnrealBuildTool/System/PrecompiledHeaderInstance.cs
+++ b/Engine/Source/Programs/UnrealBuildTool/System/PrecompiledHeaderInstance.cs
@@ -45,11 +45,6 @@ namespace UnrealBuildTool
 		/// </summary>
 		public PrecompiledHeaderInstance? ParentPCHInstance;
 
-		/// <summary>
-		/// Dictionary of definitions use in the CppCompileEnvironment
-		/// </summary>
-		public Dictionary<string, string>? DefinitionsDictionary;
-
 		/// <summary>
 		/// Constructor
 		/// </summary>
