Experimental remote debugging support for Visual Studio on Windows.

Refer to https://learn.microsoft.com/en-us/visualstudio/debugger/remote-debugging-cpp for details on setting up Visual Studio and the remote device

Add this to BuildConfiguration.xml and re-run GenerateProjectFiles.bat :

```
<WindowsPlatform>
<!-- main toggle -->
<bEnableExperimentalRemoteDebugging>true</bEnableExperimentalRemoteDebugging>

<!-- x64 remote machine configuration -->
<RemoteDebugMachineX64>MY_REMOTE_X64_PC</RemoteDebugMachineX64>
<RemoteDebugDirectoryX64>C:\Shared</RemoteDebugDirectoryX64>

<!-- arm64 remote machine configuration -->
<RemoteDebugMachineARM64>MY_REMOTE_ARM64_PC</RemoteDebugMachineARM64>
<RemoteDebugDirectoryARM64>C:\Shared</RemoteDebugDirectoryARM64>

<!-- by default the system expects games to be deployed to C:\Shared\MyGame\Windows\ to allow multiple projects to be debugged. Set this to false to assume C:\Shared is the root of a staged build -->
<bPerProjectRemoteDebugPaths>false</bPerProjectRemoteDebugPaths>
</WindowsPlatform>
```

#

diff --git a/Engine/Source/Programs/UnrealBuildTool/Platform/Windows/UEBuildWindows.cs b/Engine/Source/Programs/UnrealBuildTool/Platform/Windows/UEBuildWindows.cs
--- a/Engine/Source/Programs/UnrealBuildTool/Platform/Windows/UEBuildWindows.cs
+++ b/Engine/Source/Programs/UnrealBuildTool/Platform/Windows/UEBuildWindows.cs
@@ -733,6 +733,12 @@ namespace UnrealBuildTool
 		[CommandLine("-UseXCurl")]
 		public bool bUseXCurl = false;
 
+		/// <summary>
+		/// Whether to enable the experimental remote debugging configuration.
+		/// </summary>
+		[XmlConfigFile(Category = "WindowsPlatform")]
+		public bool bEnableExperimentalRemoteDebugging = false;
+
 		/// <summary>
 		/// Settings for PVS studio
 		/// </summary>
@@ -1007,6 +1013,8 @@ namespace UnrealBuildTool
 
 		public bool bUseXCurl => Inner.bUseXCurl;
 
+		public bool bEnableExperimentalRemoteDebugging => Inner.bEnableExperimentalRemoteDebugging;
+
 		public ReadOnlyPVSTargetSettings PVS
 		{
 			get; private set;
diff --git a/Engine/Source/Programs/UnrealBuildTool/Platform/Windows/VCToolChain.cs b/Engine/Source/Programs/UnrealBuildTool/Platform/Windows/VCToolChain.cs
--- a/Engine/Source/Programs/UnrealBuildTool/Platform/Windows/VCToolChain.cs
+++ b/Engine/Source/Programs/UnrealBuildTool/Platform/Windows/VCToolChain.cs
@@ -3289,6 +3289,16 @@ namespace UnrealBuildTool
 
 		protected virtual void ModifyFinalLinkArguments(LinkEnvironment LinkEnvironment, List<string> Arguments, bool bBuildImportLibraryOnly)
 		{
+			// do not strip the PDB path from the executable if we want to remotely debug it (otherwise VS can't find the symbols)
+			if (LinkEnvironment.bCreateDebugInfo &&
+				Target.WindowsPlatform.bEnableExperimentalRemoteDebugging &&
+				Target.Type != TargetType.Editor &&
+				Target.Configuration != UnrealTargetConfiguration.Shipping &&
+				Target.Platform.IsInGroup(UnrealPlatformGroup.Windows))
+			{
+				Arguments.Remove("/PDBALTPATH:%_PDB%");
+			}
+
 			// IMPLEMENT_MODULE_ is not required - it only exists to ensure developers add an IMPLEMENT_MODULE() declaration in code. These are always removed for PGO so that adding/removing a module won't invalidate PGC data.
 			Arguments.RemoveAll(Argument => Argument.StartsWith("/INCLUDE:IMPLEMENT_MODULE_"));
 
diff --git a/Engine/Source/Programs/UnrealBuildTool/Platform/Windows/WindowsProjectGenerator.cs b/Engine/Source/Programs/UnrealBuildTool/Platform/Windows/WindowsProjectGenerator.cs
--- a/Engine/Source/Programs/UnrealBuildTool/Platform/Windows/WindowsProjectGenerator.cs
+++ b/Engine/Source/Programs/UnrealBuildTool/Platform/Windows/WindowsProjectGenerator.cs
@@ -1,6 +1,9 @@
 // Copyright Epic Games, Inc. All Rights Reserved.
 
+using System;
 using System.Collections.Generic;
+using System.Diagnostics.CodeAnalysis;
+using System.IO;
 using System.Text;
 using EpicGames.Core;
 using Microsoft.Extensions.Logging;
@@ -12,6 +15,61 @@ namespace UnrealBuildTool
 	/// </summary>
 	class WindowsProjectGenerator : PlatformProjectGenerator
 	{
+		/// <summary>
+		/// Whether to enable the experimental remote debugging support. 
+		/// This will configure the project settings for remote debugging based on the fields below.
+		/// Refer to https://learn.microsoft.com/en-us/visualstudio/debugger/remote-debugging-cpp for details on setting up Visual Studio and the remote device for remote debugging.
+		/// </summary>
+		protected bool bEnableExperimentalRemoteDebugging = false;
+
+		/// <summary>
+		/// When disabled, the remote debug directory is expected to be the root of a staged build.
+		/// When enabled, the project name and target will be appended to the remote debug directory, allowing multiple projects to be more easily debugged.
+		/// ...this can be used alongside RunUAT BuildCookRun [...] -archive -archivedirectory=\\remotepc\Share\MyProject\
+		/// </summary>
+		[XmlConfigFile(Category = "WindowsPlatform")]
+		protected bool bPerProjectRemoteDebugPaths = true;
+
+
+		/// <summary>
+		/// Remote machine name for Windows X64 remote debugging
+		/// </summary>
+		[XmlConfigFile(Category = "WindowsPlatform")]
+		protected string? RemoteDebugMachineX64;
+
+		/// <summary>
+		/// Local shared folder path on the remote Windows X64 remote machine to deploy to. 
+		/// </summary>
+		[XmlConfigFile(Category = "WindowsPlatform")]
+		protected string? RemoteDebugDirectoryX64;
+
+
+		/// <summary>
+		/// Remote machine name for Windows Arm64 remote debugging
+		/// </summary>
+		[XmlConfigFile(Category="WindowsPlatform")]
+		protected string? RemoteDebugMachineArm64;
+
+		/// <summary>
+		/// Local shared folder path on the remote Windows Arm64 remote machine to deploy to. 
+		/// </summary>
+		[XmlConfigFile(Category = "WindowsPlatform")]
+		protected string? RemoteDebugDirectoryArm64;
+
+
+		/// <summary>
+		/// Remote machine name for Windows Arm64ec remote debugging
+		/// </summary>
+		[XmlConfigFile(Category = "WindowsPlatform")]
+		protected string? RemoteDebugMachineArm64ec;
+
+		/// <summary>
+		/// Local shared folder path on the remote Windows Arm64ec remote machine to deploy to. 
+		/// </summary>
+		[XmlConfigFile(Category = "WindowsPlatform")]
+		protected string? RemoteDebugDirectoryArm64ec;
+
+
 		/// <summary>
 		/// Constructor
 		/// </summary>
@@ -20,6 +78,13 @@ namespace UnrealBuildTool
 		public WindowsProjectGenerator(CommandLineArguments Arguments, ILogger Logger)
 			: base(Arguments, Logger)
 		{
+			XmlConfig.ApplyTo(this);
+
+			// must read this value directly because it exists in the windows target rules too
+			if (XmlConfig.TryGetValue(typeof(WindowsTargetRules), "bEnableExperimentalRemoteDebugging", out object? ConfigObject))
+			{
+				bEnableExperimentalRemoteDebugging = (bool)ConfigObject;
+			}
 		}
 
 		/// <inheritdoc/>
@@ -78,6 +143,66 @@ namespace UnrealBuildTool
 			return VCUserFileContent.ToString();
 		}
 
+		/// <inheritdoc/>
+		public override void GetVisualStudioPathsEntries(VSSettings InVSSettings, TargetType TargetType, FileReference TargetRulesPath, FileReference ProjectFilePath, FileReference NMakeOutputPath, StringBuilder ProjectFileBuilder)
+		{
+			base.GetVisualStudioPathsEntries(InVSSettings, TargetType, TargetRulesPath, ProjectFilePath, NMakeOutputPath, ProjectFileBuilder);
+
+			if (bEnableExperimentalRemoteDebugging && 
+				TargetType != TargetType.Editor &&
+				GetRemoteDebugProperties(InVSSettings.Architecture, out string? RemoteDebugMachine, out string? RemoteDebugDirectory))
+			{			
+				string RemoteDebugWorkingDir = RemoteDebugDirectory;
+
+				// can optionally use a subdirectory within the shared folder to make it simpler to debug multiple projects
+				if (bPerProjectRemoteDebugPaths)
+				{
+					string ProjectName = (TargetType == TargetType.Program)
+						? NMakeOutputPath.Directory.GetDirectoryName()
+						: NMakeOutputPath.Directory!.ParentDirectory!.ParentDirectory!.GetDirectoryName();
+
+					string PlatformTargetName = (TargetType == TargetType.Program) ? "" : "Windows";
+					if (TargetType != TargetType.Game && TargetType != TargetType.Program)
+					{
+						PlatformTargetName += TargetType.ToString();
+					}
+
+					DirectoryReference BaseDir = NMakeOutputPath.Directory!.ParentDirectory!.ParentDirectory!.ParentDirectory!;
+					string RelativeWorkingDir = NMakeOutputPath.Directory.MakeRelativeTo(BaseDir);
+					RemoteDebugWorkingDir = Path.Combine(RemoteDebugDirectory, ProjectName, PlatformTargetName, RelativeWorkingDir);
+				}
+
+				string OutDir = ProjectFile.NormalizeProjectPath(NMakeOutputPath.Directory.FullName); // need to override OutDir for remote debugging so that VS deploys the executable correctly
+				string RemoteDebugCommand = Path.Combine(RemoteDebugWorkingDir, NMakeOutputPath.GetFileName());
+
+				ProjectFileBuilder.AppendLine($"    <OutDir>{OutDir}</OutDir>");
+				ProjectFileBuilder.AppendLine($"    <DeploymentDirectory>{RemoteDebugWorkingDir}</DeploymentDirectory>");
+				ProjectFileBuilder.AppendLine($"    <RemoteDebuggerCommand>{RemoteDebugCommand}</RemoteDebuggerCommand>");
+				ProjectFileBuilder.AppendLine($"    <RemoteDebuggerWorkingDirectory>{RemoteDebugWorkingDir}</RemoteDebuggerWorkingDirectory>");
+				ProjectFileBuilder.AppendLine($"    <RemoteDebuggerServerName>{RemoteDebugMachine}</RemoteDebuggerServerName>");
+				ProjectFileBuilder.AppendLine($"    <RemoteDebuggerDebuggerType>NativeOnly</RemoteDebuggerDebuggerType>");
+				ProjectFileBuilder.AppendLine($"    <RemoteDebuggerDeployCppRuntime>false</RemoteDebuggerDeployCppRuntime>");
+
+				// default to remote debugging for other architectures
+				if (InVSSettings.Architecture != null && InVSSettings.Architecture.Value != UnrealArch.Host.Value)
+				{
+					ProjectFileBuilder.AppendLine($"    <DebuggerFlavor>WindowsRemoteDebugger</DebuggerFlavor>");
+				}
+			}
+		}
+
+		/// <inheritdoc/>
+		public override bool GetVisualStudioDeploymentEnabled(VSSettings InVSSettings)
+		{
+			// remote debugging requires deployment
+			if (bEnableExperimentalRemoteDebugging && GetRemoteDebugProperties(InVSSettings.Architecture, out _, out _))
+			{
+				return true;
+			}
+
+			return false;
+		}
+
 		/// <inheritdoc/>
 		public override bool RequiresVSUserFileGeneration()
 		{
@@ -95,5 +220,32 @@ namespace UnrealBuildTool
 
 			return Result;
 		}
+
+		protected virtual bool GetRemoteDebugProperties(UnrealArch? Architecture, [NotNullWhen(true)] out string? RemoteDebugMachine, [NotNullWhen(true)] out string? RemoteDebugDirectory)
+		{
+			RemoteDebugMachine = null;
+			RemoteDebugDirectory = null;
+
+			if (Architecture != null)
+			{
+				if (Architecture == UnrealArch.X64)
+				{
+					RemoteDebugMachine = RemoteDebugMachineX64;
+					RemoteDebugDirectory = RemoteDebugDirectoryX64;
+				}
+				else if (Architecture == UnrealArch.Arm64)
+				{
+					RemoteDebugMachine = RemoteDebugMachineArm64;
+					RemoteDebugDirectory = RemoteDebugDirectoryArm64;
+				}
+				else if (Architecture == UnrealArch.Arm64ec)
+				{
+					RemoteDebugMachine = RemoteDebugMachineArm64ec;
+					RemoteDebugDirectory = RemoteDebugDirectoryArm64ec;
+				}
+			}
+
+			return !string.IsNullOrEmpty(RemoteDebugMachine) && !string.IsNullOrEmpty(RemoteDebugDirectory);
+		}
 	}
 }
diff --git a/Engine/Source/Programs/UnrealBuildTool/ProjectFiles/VisualStudio/VCProject.cs b/Engine/Source/Programs/UnrealBuildTool/ProjectFiles/VisualStudio/VCProject.cs
--- a/Engine/Source/Programs/UnrealBuildTool/ProjectFiles/VisualStudio/VCProject.cs
+++ b/Engine/Source/Programs/UnrealBuildTool/ProjectFiles/VisualStudio/VCProject.cs
@@ -2375,7 +2375,7 @@ namespace UnrealBuildTool
 					string UBTPlatformName = Platform.ToString();
 					string UBTConfigurationName = Configuration.ToString();
 					VSSettings VSSettings = new(Platform, Configuration, ProjectFileFormat, null);
-					if (Platform == UnrealTargetPlatform.Android && Combination.Architecture != null)
+					if (Combination.Architecture != null)
 					{
 						VSSettings.Architecture = Combination.Architecture;
 					}
