Don't use VFS for generating clangd compile_commands.json

diff --git a/Engine/Source/Programs/UnrealBuildTool/Modes/GenerateClangDatabase.cs b/Engine/Source/Programs/UnrealBuildTool/Modes/GenerateClangDatabase.cs
--- a/Engine/Source/Programs/UnrealBuildTool/Modes/GenerateClangDatabase.cs
+++ b/Engine/Source/Programs/UnrealBuildTool/Modes/GenerateClangDatabase.cs
@@ -94,7 +94,7 @@ namespace UnrealBuildTool
 					// Disable PCHs and unity builds for the target
 					TargetDescriptor.bUseUnityBuild = false;
 					TargetDescriptor.IntermediateEnvironment = UnrealIntermediateEnvironment.GenerateClangDatabase;
-					TargetDescriptor.AdditionalArguments = TargetDescriptor.AdditionalArguments.Append(new string[] { "-NoPCH" });
+					TargetDescriptor.AdditionalArguments = TargetDescriptor.AdditionalArguments.Append(new string[] { "-NoPCH", "-NoVFS" });
 					// Default the compiler to clang
 					if (!TargetDescriptor.AdditionalArguments.Any(x => x.StartsWith("-Compiler=", StringComparison.OrdinalIgnoreCase)))
 					{
