UBT: Add a log for the failed clang version check.

diff --git a/Engine/Source/Programs/UnrealBuildTool/ToolChain/ClangToolChain.cs b/Engine/Source/Programs/UnrealBuildTool/ToolChain/ClangToolChain.cs
--- a/Engine/Source/Programs/UnrealBuildTool/ToolChain/ClangToolChain.cs
+++ b/Engine/Source/Programs/UnrealBuildTool/ToolChain/ClangToolChain.cs
@@ -174,7 +174,7 @@ namespace UnrealBuildTool
 				{
 					return ClangVersion;
 				}
-				throw new BuildException("Failed to query the Clang version number!");
+				throw new BuildException($"Failed to query the Clang version number! Got: {ClangVersionString}");
 			}
 
 			/// <summary>
