Fix AutoRTFM Compiler failing to run on Rocky OS due to the compiler glibc version being to high

diff --git a/Engine/Source/Programs/UnrealBuildTool/Platform/Linux/LinuxToolChain.cs b/Engine/Source/Programs/UnrealBuildTool/Platform/Linux/LinuxToolChain.cs
--- a/Engine/Source/Programs/UnrealBuildTool/Platform/Linux/LinuxToolChain.cs
+++ b/Engine/Source/Programs/UnrealBuildTool/Platform/Linux/LinuxToolChain.cs
@@ -128,7 +128,8 @@ namespace UnrealBuildTool
 					// set up the path to our toolchain
 					if (BuildHostPlatform.Current.Platform == UnrealTargetPlatform.Linux)
 					{
-						ClangPath = FileReference.Combine(AutoRTFMDir, $"verse-clang++{BuildHostPlatform.Current.BinarySuffix}");
+						// FIXME This compiler is broken to actually run on HostLinux, disable setting this and fail to compile later
+						// ClangPath = FileReference.Combine(AutoRTFMDir, $"verse-clang++{BuildHostPlatform.Current.BinarySuffix}");
 					}
 					else
 					{
