UBT: Lower verbosity of some Timeline log lines to make the stdout information even more dense and relevant.
Tweaks of the behavior added in 39439379.

#rnx

diff --git a/Engine/Source/Programs/UnrealBuildTool/System/Timeline.cs b/Engine/Source/Programs/UnrealBuildTool/System/Timeline.cs
--- a/Engine/Source/Programs/UnrealBuildTool/System/Timeline.cs
+++ b/Engine/Source/Programs/UnrealBuildTool/System/Timeline.cs
@@ -183,12 +183,15 @@ namespace UnrealBuildTool
 		/// <summary>
 		/// Prints this information to the log
 		/// </summary>
-		public static void Print(TimeSpan MaxUnknownTime, LogLevel Verbosity, ILogger Logger)
+		/// <param name="MaxUnknownTime">Non-instrumented time limit for inserting an "unknown" event</param>
+		/// <param name="MinKnownTime">Instrumented time limit for printing the labelled event at given verbosity. May be null.</param>
+		/// <param name="Verbosity">LogLevel of the event.</param>
+		/// <param name="Logger">Logger.</param>
+		public static void Print(TimeSpan? MinKnownTime, TimeSpan MaxUnknownTime, LogLevel Verbosity, ILogger Logger)
 		{
 			// Print the start time
 			Logger.Log(Verbosity, "");
 			Logger.Log(Verbosity, "Timeline:");
-			Logger.Log(Verbosity, "[{Time,7}]", FormatTime(TimeSpan.Zero));
 
 			// Create the root event
 			TimeSpan FinishTime = Stopwatch.Elapsed;
@@ -221,7 +224,7 @@ namespace UnrealBuildTool
 					UpdateLastEventTime(ref LastTime, Event.StartTime, MaxUnknownTime, OuterEvents, Verbosity, Logger);
 
 					// Print this event
-					Print(Event.StartTime, Event.FinishTime, Event.Name, OuterEvents, Verbosity, Logger);
+					Print(Event.StartTime, Event.FinishTime, MinKnownTime, Event.Name, OuterEvents, Verbosity, Logger);
 
 					// Push it onto the stack
 					if (Event.FinishTime.HasValue)
@@ -262,7 +265,7 @@ namespace UnrealBuildTool
 			const string UnknownEvent = "<unknown>";
 			if (NewTime - LastTime > MaxUnknownTime)
 			{
-				Print(LastTime, NewTime, UnknownEvent, OuterEvents, Verbosity, Logger);
+				Print(LastTime, NewTime, null, UnknownEvent, OuterEvents, Verbosity, Logger);
 			}
 			LastTime = NewTime;
 		}
@@ -272,11 +275,12 @@ namespace UnrealBuildTool
 		/// </summary>
 		/// <param name="StartTime">Start time for the event</param>
 		/// <param name="FinishTime">Finish time for the event. May be null.</param>
+		/// <param name="MinKnownTime">MinKnownTime for printing at the provided verbosity. May be null.</param>
 		/// <param name="Label">Event name</param>
 		/// <param name="OuterEvents">List of all the start times for parent events</param>
 		/// <param name="Verbosity">Verbosity for the output</param>
 		/// <param name="Logger">Logger for output</param>
-		static void Print(TimeSpan StartTime, TimeSpan? FinishTime, string Label, List<Event> OuterEvents, LogLevel Verbosity, ILogger Logger)
+		static void Print(TimeSpan StartTime, TimeSpan? FinishTime, TimeSpan? MinKnownTime, string Label, List<Event> OuterEvents, LogLevel Verbosity, ILogger Logger)
 		{
 			StringBuilder Prefix = new StringBuilder();
 
@@ -287,6 +291,7 @@ namespace UnrealBuildTool
 
 			Prefix.AppendFormat("[{0,7}]", FormatTime(StartTime - OuterEvents[^1].StartTime));
 
+			bool UseDebugVerbosity = MinKnownTime != null;
 			if (!FinishTime.HasValue)
 			{
 				Prefix.AppendFormat("{0,8}", "???");
@@ -298,9 +303,11 @@ namespace UnrealBuildTool
 			else
 			{
 				Prefix.AppendFormat("{0,8}", "+" + FormatTime(FinishTime.Value - StartTime));
+				UseDebugVerbosity &= (FinishTime.Value - StartTime) < MinKnownTime;
 			}
 
-			Logger.Log(Verbosity, "{Prefix} {Label}", Prefix.ToString(), Label);
+			LogLevel PrintVerbosity = UseDebugVerbosity && Verbosity > LogLevel.Debug ? LogLevel.Debug : Verbosity;
+			Logger.Log(PrintVerbosity, "{Prefix} {Label}", Prefix.ToString(), Label);
 		}
 
 		/// <summary>
diff --git a/Engine/Source/Programs/UnrealBuildTool/UnrealBuildTool.cs b/Engine/Source/Programs/UnrealBuildTool/UnrealBuildTool.cs
--- a/Engine/Source/Programs/UnrealBuildTool/UnrealBuildTool.cs
+++ b/Engine/Source/Programs/UnrealBuildTool/UnrealBuildTool.cs
@@ -685,7 +685,7 @@ namespace UnrealBuildTool
 				// Print out all the performance info
 				Timeline.Stop();
 				LogLevel TimelineLogLevel = ModeOptions.HasFlag(ToolModeOptions.ShowExecutionTime) ? LogLevel.Information : LogLevel.Debug;
-				Timeline.Print(TimeSpan.FromMilliseconds(200.0), TimelineLogLevel, Logger);
+				Timeline.Print(TimeSpan.FromMilliseconds(100.0), TimeSpan.FromMilliseconds(200.0), TimelineLogLevel, Logger);
 				Logger.Log(TimelineLogLevel, (""));
 				Logger.Log(TimelineLogLevel, "Total execution time: {Time:0.00} seconds", Timeline.Elapsed.TotalSeconds);
 
