[UnrealBuildTool]
- Remove manifest from compiled DLLs as it was causing a delay for each module loaded to query WinSxS with an out-of-process call
- Saves about 10s of boot time for ~2300 dlls loaded at editor startup

diff --git a/Engine/Source/Programs/UnrealBuildTool/Platform/Windows/VCToolChain.cs b/Engine/Source/Programs/UnrealBuildTool/Platform/Windows/VCToolChain.cs
--- a/Engine/Source/Programs/UnrealBuildTool/Platform/Windows/VCToolChain.cs
+++ b/Engine/Source/Programs/UnrealBuildTool/Platform/Windows/VCToolChain.cs
@@ -1377,7 +1377,7 @@ namespace UnrealBuildTool
 			}
 
 			// Don't create a side-by-side manifest file for the executable.
-			if (Target.WindowsPlatform.ManifestFile == null)
+			if (Target.WindowsPlatform.ManifestFile == null || LinkEnvironment.bIsBuildingDLL)
 			{
 				Arguments.Add("/MANIFEST:NO");
 			}
@@ -2843,7 +2843,7 @@ namespace UnrealBuildTool
 						NormalizeCommandLinePath(NatvisFile, RootPaths)));
 				}
 
-				if (Target.WindowsPlatform.ManifestFile != null)
+				if (Target.WindowsPlatform.ManifestFile != null && !LinkEnvironment.bIsBuildingDLL)
 				{
 					PrerequisiteItems.Add(FileItem.GetItemByPath(Target.WindowsPlatform.ManifestFile));
 				}
