UnrealBuildTool: Add WindowsPlatform.PreferredCompilers a list of compilers a target will use, allowing for falling back to a different compiler if not available

diff --git a/Engine/Source/Programs/UnrealBuildTool/Platform/Windows/UEBuildWindows.cs b/Engine/Source/Programs/UnrealBuildTool/Platform/Windows/UEBuildWindows.cs
--- a/Engine/Source/Programs/UnrealBuildTool/Platform/Windows/UEBuildWindows.cs
+++ b/Engine/Source/Programs/UnrealBuildTool/Platform/Windows/UEBuildWindows.cs
@@ -274,7 +274,7 @@ namespace UnrealBuildTool
 		[XmlConfigFile(Category = "WindowsPlatform")]
 		[CommandLine("-2022", Value = nameof(WindowsCompiler.VisualStudio2022))]
 		[CommandLine("-Compiler=")]
-		public WindowsCompiler Compiler = WindowsCompiler.Default;
+		public WindowsCompiler Compiler { get; set; } = WindowsCompiler.Default;
 
 		/// <summary>
 		/// Version of the toolchain to use on Windows platform when a non-msvc Compiler is in use, to locate include paths etc.
@@ -282,12 +282,24 @@ namespace UnrealBuildTool
 		[ConfigFile(ConfigHierarchyType.Engine, "/Script/WindowsTargetPlatform.WindowsTargetSettings", "Toolchain")]
 		[XmlConfigFile(Category = "WindowsPlatform")]
 		[CommandLine("-VCToolchain=")]
-		public WindowsCompiler ToolChain
+		public WindowsCompiler ToolChain { get; set; } = WindowsCompiler.Default;
+
+		/// <summary>
+		/// Order of preference when Compiler is set to Default. If populated and none listed are available this will cause a warning.
+		/// </summary>
+		[XmlConfigFile(Category = "WindowsPlatform")]
+		public string[] PreferredCompilers
 		{
-			get => ToolChainPrivate ?? (Compiler.IsMSVC() ? Compiler : WindowsPlatform.GetDefaultCompiler(Target.ProjectFile, Architecture, Target.Logger, true));
-			set => ToolChainPrivate = value;
+			get => _preferredCompilersPrivate ?? DefaultPreferredCompilers;
+			set => _preferredCompilersPrivate = value;
 		}
-		private WindowsCompiler? ToolChainPrivate = null;
+		private string[]? _preferredCompilersPrivate = null;
+
+		/// <summary>
+		/// Default order of preference when Compiler is set to Default. If populated and none listed are available this will cause a warning.
+		/// </summary>
+		[XmlConfigFile(Category = "WindowsPlatform")]
+		private string[] DefaultPreferredCompilers = [];
 
 		/// <summary>
 		/// Architecture of Target.
@@ -427,6 +439,7 @@ namespace UnrealBuildTool
 		/// Enable building with the Win10 SDK instead of the older Win8.1 SDK 
 		/// </summary>
 		[ConfigFile(ConfigHierarchyType.Engine, "/Script/WindowsTargetPlatform.WindowsTargetSettings", "bUseWindowsSDK10")]
+		[Obsolete]
 		public bool bUseWindowsSDK10 = false;
 
 		/// <summary>
@@ -913,8 +926,6 @@ namespace UnrealBuildTool
 
 		public bool bPixProfilingEnabled => Inner.bPixProfilingEnabled;
 
-		public bool bUseWindowsSDK10 => Inner.bUseWindowsSDK10;
-
 		public bool bUseCPPWinRT => Inner.bUseCPPWinRT;
 
 		public bool bVCFastFail => Inner.bVCFastFail;
@@ -1251,6 +1262,15 @@ namespace UnrealBuildTool
 				}
 			}
 
+			// Disable PCHs for PVS studio analyzer.
+			if (Target.StaticAnalyzer == StaticAnalyzer.PVSStudio)
+			{
+				Target.bUsePCHFiles = false;
+
+				// PVSStudio static analysis requires MSVC
+				Target.WindowsPlatform.Compiler = WindowsCompiler.VisualStudio2022;
+			}
+
 			// If clang is selected for static analysis, switch compiler to clang and proceed
 			// as normal.
 			if (Target.StaticAnalyzer == StaticAnalyzer.Clang)
@@ -1273,7 +1293,13 @@ namespace UnrealBuildTool
 			// Set the compiler version if necessary
 			if (Target.WindowsPlatform.Compiler == WindowsCompiler.Default)
 			{
-				Target.WindowsPlatform.Compiler = GetDefaultCompiler(Target.ProjectFile, Target.WindowsPlatform.Architecture, Logger, !Platform.IsInGroup(UnrealPlatformGroup.Microsoft));
+				Target.WindowsPlatform.Compiler = GetDefaultCompiler(Target, Target.WindowsPlatform.Architecture, Logger, !Platform.IsInGroup(UnrealPlatformGroup.Microsoft));
+			}
+
+			// Set the toolchain version if necessary
+			if (Target.WindowsPlatform.ToolChain == WindowsCompiler.Default)
+			{
+				Target.WindowsPlatform.ToolChain = GetDefaultToolchain(Target, Target.WindowsPlatform.Architecture, Logger, !Platform.IsInGroup(UnrealPlatformGroup.Microsoft));
 			}
 
 			// Disable linking and ignore build outputs if we're using a static analyzer
@@ -1286,12 +1312,6 @@ namespace UnrealBuildTool
 				Target.WindowsPlatform.bVCExtendedWarningInfo = true;
 			}
 
-			// Disable PCHs for PVS studio analyzer.
-			if (Target.StaticAnalyzer == StaticAnalyzer.PVSStudio)
-			{
-				Target.bUsePCHFiles = false;
-			}
-
 			// Disable chaining PCH files if not using clang
 			if (!Target.WindowsPlatform.Compiler.IsClang())
 			{
@@ -1400,8 +1420,61 @@ namespace UnrealBuildTool
 		/// Gets the default compiler which should be used, if it's not set explicitly by the target, command line, or config file.
 		/// </summary>
 		/// <returns>The default compiler version</returns>
-		internal static WindowsCompiler GetDefaultCompiler(FileReference? ProjectFile, UnrealArch Architecture, ILogger Logger, bool bSkipWarning = false)
+		private static IEnumerable<WindowsCompiler> GetPreferredCompilers(TargetRules? TargetRules)
 		{
+			if (TargetRules == null)
+			{
+				yield break;
+			}
+			foreach (string enumString in TargetRules.WindowsPlatform.PreferredCompilers)
+			{
+				if (Enum.TryParse(enumString, out WindowsCompiler result))
+				{
+					yield return result;
+				}
+			}
+		}
+
+		/// <summary>
+		/// Gets the default compiler which should be used, if it's not set explicitly by the target, command line, or config file.
+		/// </summary>
+		/// <returns>The default compiler version</returns>
+		internal static WindowsCompiler GetDefaultCompiler(TargetRules? TargetRules, UnrealArch Architecture, ILogger Logger, bool bSkipWarning = false)
+		{
+			IEnumerable<WindowsCompiler> PreferredCompilers = GetPreferredCompilers(TargetRules);
+			if (PreferredCompilers.Any())
+			{
+				WindowsCompiler? Compiler = PreferredCompilers.FirstOrDefault(x => MicrosoftPlatformSDK.HasValidCompiler(x, Architecture, Logger));
+				if (Compiler != null)
+				{
+					if (Compiler != PreferredCompilers.First())
+					{
+						Logger.Log(bSkipWarning ? LogLevel.Debug : LogLevel.Warning, "{Target} preferred compiler {Preferred} not available, compiler set to {Compiler}", TargetRules?.Name ?? "Default", PreferredCompilers.First(), Compiler);
+					}
+					else
+					{
+						Logger.LogDebug("{Target} Compiler set to preferred compiler {Compiler}", TargetRules?.Name ?? "Default", Compiler);
+					}
+					return (WindowsCompiler)Compiler;
+				}
+				Logger.Log(bSkipWarning ? LogLevel.Debug : LogLevel.Warning, "{Target} No preferred compilers available [{Preferred}]", TargetRules?.Name ?? "Default", String.Join(',', PreferredCompilers));
+			}
+
+			return GetDefaultToolchain(TargetRules, Architecture, Logger, bSkipWarning);
+		}
+
+		/// <summary>
+		/// Gets the default toolchain which should be used, if it's not set explicitly by the target, command line, or config file.
+		/// </summary>
+		/// <returns>The default toolchain version</returns>
+		internal static WindowsCompiler GetDefaultToolchain(TargetRules? TargetRules, UnrealArch Architecture, ILogger Logger, bool bSkipWarning = false)
+		{
+			// If the compiler is set to MSVC, it should also be used as the toolchain
+			if (TargetRules?.WindowsPlatform.Compiler.IsMSVC() == true)
+			{
+				return TargetRules.WindowsPlatform.Compiler;
+			}
+
 			// If there's no specific compiler set, try to pick the matching compiler for the selected IDE
 			if (ProjectFileGeneratorSettings.Format != null)
 			{
@@ -1427,7 +1500,7 @@ namespace UnrealBuildTool
 
 			// Check the editor settings too
 			ProjectFileFormat PreferredAccessor;
-			if (ProjectFileGenerator.GetPreferredSourceCodeAccessor(ProjectFile, out PreferredAccessor))
+			if (ProjectFileGenerator.GetPreferredSourceCodeAccessor(TargetRules?.ProjectFile, out PreferredAccessor))
 			{
 				if (PreferredAccessor == ProjectFileFormat.VisualStudio2022)
 				{
