UnrealBuildTool: Always pass /Brepro when compiling msvc platforms with clang

#rnx

diff --git a/Engine/Source/Programs/UnrealBuildTool/Platform/Windows/VCToolChain.cs b/Engine/Source/Programs/UnrealBuildTool/Platform/Windows/VCToolChain.cs
--- a/Engine/Source/Programs/UnrealBuildTool/Platform/Windows/VCToolChain.cs
+++ b/Engine/Source/Programs/UnrealBuildTool/Platform/Windows/VCToolChain.cs
@@ -667,10 +667,11 @@ namespace UnrealBuildTool
 						Arguments.Add("/wd5048");
 					}
 				}
-				else if (Target.WindowsPlatform.Compiler.IsClang())
-				{
-					Arguments.Add("/Brepro");
-				}
+			}
+
+			if (Target.WindowsPlatform.Compiler.IsClang())
+			{
+				Arguments.Add("/Brepro");
 			}
 
 			if (Target.WindowsPlatform.Compiler.IsMSVC() && Target.WindowsPlatform.bVCFastFail)
@@ -1502,10 +1503,11 @@ namespace UnrealBuildTool
 				{
 					Arguments.Add("/experimental:deterministic");
 				}
-				else if (Target.WindowsPlatform.Compiler.IsClang())
-				{
-					Arguments.Add("/Brepro");
-				}
+			}
+
+			if (Target.WindowsPlatform.Compiler.IsClang())
+			{
+				Arguments.Add("/Brepro");
 			}
 
 			if (Target.WindowsPlatform.Compiler.IsMSVC() && Target.WindowsPlatform.bVCFastFail && !LinkEnvironment.bPGOOptimize && !LinkEnvironment.bPGOProfile)
