Revert change here to disable AutoRTFM compiler

diff --git a/Engine/Source/Programs/UnrealBuildTool/Platform/Linux/LinuxToolChain.cs b/Engine/Source/Programs/UnrealBuildTool/Platform/Linux/LinuxToolChain.cs
--- a/Engine/Source/Programs/UnrealBuildTool/Platform/Linux/LinuxToolChain.cs
+++ b/Engine/Source/Programs/UnrealBuildTool/Platform/Linux/LinuxToolChain.cs
@@ -128,8 +128,7 @@ namespace UnrealBuildTool
 					// set up the path to our toolchain
 					if (BuildHostPlatform.Current.Platform == UnrealTargetPlatform.Linux)
 					{
-						// FIXME This compiler is broken to actually run on HostLinux, disable setting this and fail to compile later
-						// ClangPath = FileReference.Combine(AutoRTFMDir, $"verse-clang++{BuildHostPlatform.Current.BinarySuffix}");
+						ClangPath = FileReference.Combine(AutoRTFMDir, $"verse-clang++{BuildHostPlatform.Current.BinarySuffix}");
 					}
 					else
 					{
