UnrealBuildTool: Clang PCH chaining does not work with windows compiles, disable it as to not create unnecessary dependencies

#rnx

diff --git a/Engine/Source/Programs/UnrealBuildTool/Platform/Windows/UEBuildWindows.cs b/Engine/Source/Programs/UnrealBuildTool/Platform/Windows/UEBuildWindows.cs
--- a/Engine/Source/Programs/UnrealBuildTool/Platform/Windows/UEBuildWindows.cs
+++ b/Engine/Source/Programs/UnrealBuildTool/Platform/Windows/UEBuildWindows.cs
@@ -1312,11 +1312,8 @@ namespace UnrealBuildTool
 				Target.WindowsPlatform.bVCExtendedWarningInfo = true;
 			}
 
-			// Disable chaining PCH files if not using clang
-			if (!Target.WindowsPlatform.Compiler.IsClang())
-			{
-				Target.bChainPCHs = false;
-			}
+			// Disable chaining PCH files
+			Target.bChainPCHs = false;
 
 			// E&C support.
 			if (Target.bSupportEditAndContinue || Target.bAdaptiveUnityEnablesEditAndContinue)
