UnrealBuildTool: Fix VCToolchain clang -ftime-trace

#rnx

diff --git a/Engine/Source/Programs/UnrealBuildTool/Platform/Windows/VCToolChain.cs b/Engine/Source/Programs/UnrealBuildTool/Platform/Windows/VCToolChain.cs
--- a/Engine/Source/Programs/UnrealBuildTool/Platform/Windows/VCToolChain.cs
+++ b/Engine/Source/Programs/UnrealBuildTool/Platform/Windows/VCToolChain.cs
@@ -1924,7 +1924,7 @@ namespace UnrealBuildTool
 					// Experimental: support for JSON output of timing data
 					if (Target.WindowsPlatform.Compiler.IsClang() && (Target.bPrintToolChainTimingInfo || Target.WindowsPlatform.bClangTimeTrace))
 					{
-						CompileAction.Arguments.Add("-Xclang -ftime-trace");
+						CompileAction.Arguments.Add("-ftime-trace");
 						CompileAction.AdditionalProducedItems.Add(FileItem.GetItemByFileReference(ObjectFile.Location.ChangeExtension(".json")));
 					}
 				}
