UnrealBuildTool: Suppress preferred compiler spam when generating projects

#rnx

diff --git a/Engine/Source/Programs/UnrealBuildTool/Platform/Windows/UEBuildWindows.cs b/Engine/Source/Programs/UnrealBuildTool/Platform/Windows/UEBuildWindows.cs
--- a/Engine/Source/Programs/UnrealBuildTool/Platform/Windows/UEBuildWindows.cs
+++ b/Engine/Source/Programs/UnrealBuildTool/Platform/Windows/UEBuildWindows.cs
@@ -1446,7 +1446,7 @@ namespace UnrealBuildTool
 				{
 					if (Compiler != PreferredCompilers.First())
 					{
-						Logger.Log(bSkipWarning ? LogLevel.Debug : LogLevel.Warning, "{Target} preferred compiler {Preferred} not available, compiler set to {Compiler}", TargetRules?.Name ?? "Default", PreferredCompilers.First(), Compiler);
+						Logger.Log(bSkipWarning || MicrosoftPlatformSDK.IgnoreToolchainErrors ? LogLevel.Debug : LogLevel.Warning, "{Target} preferred compiler {Preferred} not available, compiler set to {Compiler}", TargetRules?.Name ?? "Default", PreferredCompilers.First(), Compiler);
 					}
 					else
 					{
@@ -1454,7 +1454,7 @@ namespace UnrealBuildTool
 					}
 					return (WindowsCompiler)Compiler;
 				}
-				Logger.Log(bSkipWarning ? LogLevel.Debug : LogLevel.Warning, "{Target} No preferred compilers available [{Preferred}]", TargetRules?.Name ?? "Default", String.Join(',', PreferredCompilers));
+				Logger.Log(bSkipWarning || MicrosoftPlatformSDK.IgnoreToolchainErrors ? LogLevel.Debug : LogLevel.Warning, "{Target} No preferred compilers available [{Preferred}]", TargetRules?.Name ?? "Default", String.Join(',', PreferredCompilers));
 			}
 
 			return GetDefaultToolchain(TargetRules, Architecture, Logger, bSkipWarning);
