UnrealBuildTool: Skip compiling UBT plugins when generating project files

#rnx

diff --git a/Engine/Source/Programs/UnrealBuildTool/Configuration/UEBuildTarget.cs b/Engine/Source/Programs/UnrealBuildTool/Configuration/UEBuildTarget.cs
--- a/Engine/Source/Programs/UnrealBuildTool/Configuration/UEBuildTarget.cs
+++ b/Engine/Source/Programs/UnrealBuildTool/Configuration/UEBuildTarget.cs
@@ -5054,6 +5054,11 @@ namespace UnrealBuildTool
 				throw new BuildException("Build plugins has not be populated");
 			}
 
+			if (Rules.bGenerateProjectFiles)
+			{
+				return;
+			}
+
 			// Collect the plugins
 			UbtPlugins = PluginsBase.EnumerateUbtPlugins(ProjectFile);
 
