UnrealBuildTool: PVS exclude *.gen.cpp when bStaticAnalyzerIncludeGenerated is set

#rnx

diff --git a/Engine/Source/Programs/UnrealBuildTool/Platform/Windows/PVSToolChain.cs b/Engine/Source/Programs/UnrealBuildTool/Platform/Windows/PVSToolChain.cs
--- a/Engine/Source/Programs/UnrealBuildTool/Platform/Windows/PVSToolChain.cs
+++ b/Engine/Source/Programs/UnrealBuildTool/Platform/Windows/PVSToolChain.cs
@@ -922,9 +922,17 @@ namespace UnrealBuildTool
 					ConfigFileContents.Append("new-output-format=yes\n");
 				}
 
-				if (Target.bStaticAnalyzerProjectOnly && AnalyzerVersion.CompareTo(new Version("7.34")) >= 0)
+				if (AnalyzerVersion.CompareTo(new Version("7.34")) >= 0)
 				{
-					ConfigFileContents.Append($"analysis-paths=skip={Unreal.EngineSourceDirectory}\n");
+					if (Target.bStaticAnalyzerProjectOnly)
+					{
+						ConfigFileContents.Append($"analysis-paths=skip={Unreal.EngineSourceDirectory}\n");
+					}
+					if (!Target.bStaticAnalyzerIncludeGenerated)
+					{
+						ConfigFileContents.Append($"analysis-paths=skip=*.gen.cpp\n");
+						ConfigFileContents.Append($"analysis-paths=skip=*.generated.h\n");
+					}
 				}
 
 				string BaseFileName = PreprocessedFileItem.Location.GetFileName();
