UnrealBuildTool: Fix Adaptive unity logging when PCH is not in use

#rnx

diff --git a/Engine/Source/Programs/UnrealBuildTool/Configuration/UEBuildModuleCPP.cs b/Engine/Source/Programs/UnrealBuildTool/Configuration/UEBuildModuleCPP.cs
--- a/Engine/Source/Programs/UnrealBuildTool/Configuration/UEBuildModuleCPP.cs
+++ b/Engine/Source/Programs/UnrealBuildTool/Configuration/UEBuildModuleCPP.cs
@@ -1742,64 +1742,47 @@ namespace UnrealBuildTool
 			IActionGraphBuilder Graph,
 			ILogger Logger)
 		{
-			bool bAdaptiveUnityDisablesPCH = false;
-			if (CompileEnvironment.PrecompiledHeaderAction == PrecompiledHeaderAction.Include)
+			bool isEngine = Rules.File.IsUnderDirectory(Unreal.EngineDirectory) || Rules.bTreatAsEngineModule || Rules.PrivatePCHHeaderFile == null;
+			bool bAdaptiveUnityDisablesPCH = isEngine ? Target.bAdaptiveUnityDisablesPCH : Target.bAdaptiveUnityDisablesPCHForProject;
+			string adaptiveUnityDisablesPCHName = isEngine ? nameof(Target.bAdaptiveUnityDisablesPCH) : nameof(Target.bAdaptiveUnityDisablesPCHForProject);
+
+			CPPOutput OutputFiles = new CPPOutput();
+
+			if (NormalFiles.Count > 0)
 			{
-				if (Rules.PCHUsage == ModuleRules.PCHUsageMode.UseExplicitOrSharedPCHs)
-				{
-					if (Rules.bTreatAsEngineModule || Rules.PrivatePCHHeaderFile == null)
-					{
-						bAdaptiveUnityDisablesPCH = Target.bAdaptiveUnityDisablesPCH;
-					}
-					else
-					{
-						bAdaptiveUnityDisablesPCH = Target.bAdaptiveUnityDisablesPCHForProject;
-					}
-				}
+				OutputFiles = ToolChain.CompileAllCPPFiles(CompileEnvironment, NormalFiles, IntermediateDirectory, Name, Graph);
+			}
 
-				if (AdaptiveFiles.Count > 0)
+			if (AdaptiveFiles.Count > 0)
+			{
+				if (CompileEnvironment.PrecompiledHeaderAction == PrecompiledHeaderAction.Include)
 				{
 					if (Target.bAdaptiveUnityCreatesDedicatedPCH)
 					{
 						Graph.AddDiagnostic("[Adaptive Build] Creating dedicated PCH for each excluded file. Set bAdaptiveUnityCreatesDedicatedPCH to false in BuildConfiguration.xml to change this behavior.");
 					}
-					else if (Target.bAdaptiveUnityDisablesPCH)
-					{
-						Graph.AddDiagnostic("[Adaptive Build] Disabling PCH for excluded files. Set bAdaptiveUnityDisablesPCH to false in BuildConfiguration.xml to change this behavior.");
-					}
-
-					if (Target.bAdaptiveUnityDisablesOptimizations)
-					{
-						Graph.AddDiagnostic("[Adaptive Build] Disabling optimizations for excluded files. Set bAdaptiveUnityDisablesOptimizations to false in BuildConfiguration.xml to change this behavior.");
-					}
-					if (Target.bAdaptiveUnityEnablesEditAndContinue)
+					else if (bAdaptiveUnityDisablesPCH)
 					{
-						Graph.AddDiagnostic("[Adaptive Build] Enabling Edit & Continue for excluded files. Set bAdaptiveUnityEnablesEditAndContinue to false in BuildConfiguration.xml to change this behavior.");
+						Graph.AddDiagnostic($"[Adaptive Build] Disabling PCH for excluded files. Set {adaptiveUnityDisablesPCHName} to false in BuildConfiguration.xml to change this behavior.");
 					}
-
-					Graph.AddDiagnostic($"[Adaptive Build] Excluded from {Name} unity file: " + String.Join(", ", AdaptiveFiles.Select(File => Path.GetFileName(File.AbsolutePath))));
 				}
 
-				if ((!Target.bAdaptiveUnityDisablesOptimizations && !bAdaptiveUnityDisablesPCH && !Target.bAdaptiveUnityCreatesDedicatedPCH) || Target.bStressTestUnity)
+				if (Target.bAdaptiveUnityDisablesOptimizations)
 				{
-					NormalFiles = NormalFiles.Concat(AdaptiveFiles).ToList();
-					AdaptiveFiles = new List<FileItem>();
+					Graph.AddDiagnostic("[Adaptive Build] Disabling optimizations for excluded files. Set bAdaptiveUnityDisablesOptimizations to false in BuildConfiguration.xml to change this behavior.");
+				}
+				if (Target.bAdaptiveUnityEnablesEditAndContinue)
+				{
+					Graph.AddDiagnostic("[Adaptive Build] Enabling Edit & Continue for excluded files. Set bAdaptiveUnityEnablesEditAndContinue to false in BuildConfiguration.xml to change this behavior.");
 				}
-			}
-
-			CPPOutput OutputFiles = new CPPOutput();
 
-			if (NormalFiles.Count > 0)
-			{
-				OutputFiles = ToolChain.CompileAllCPPFiles(CompileEnvironment, NormalFiles, IntermediateDirectory, Name, Graph);
-			}
+				Graph.AddDiagnostic($"[Adaptive Build] Excluded from {Name} unity file: " + String.Join(", ", AdaptiveFiles.Select(File => Path.GetFileName(File.AbsolutePath))));
 
-			if (AdaptiveFiles.Count > 0)
-			{
 				// Create the new compile environment. Always turn off PCH due to different compiler settings.
 				CppCompileEnvironment AdaptiveUnityEnvironment = new CppCompileEnvironment(ModuleCompileEnvironment);
 				AdaptiveUnityEnvironment.AdditionalPrerequisites = new(AdaptiveUnityEnvironment.AdditionalPrerequisites.Concat(CompileEnvironment.AdditionalPrerequisites).Distinct());
 				AdaptiveUnityEnvironment.UserIncludePaths.UnionWith(CompileEnvironment.SharedUserIncludePaths);
+				AdaptiveUnityEnvironment.bDisableStaticAnalysis = CompileEnvironment.bDisableStaticAnalysis;
 				if (Target.bAdaptiveUnityDisablesOptimizations)
 				{
 					AdaptiveUnityEnvironment.bOptimizeCode = false;
@@ -1809,6 +1792,7 @@ namespace UnrealBuildTool
 					AdaptiveUnityEnvironment.bSupportEditAndContinue = true;
 				}
 
+
 				// Create a per-file PCH
 				CPPOutput AdaptiveOutput;
 				if (CompileEnvironment.PrecompiledHeaderAction == PrecompiledHeaderAction.Include && Target.bAdaptiveUnityCreatesDedicatedPCH)
