If bUseChangeListAsStoreVersion is set, allow setting the Android store version to the changelist number even on non-build machines, so it is no longer necessary to rename builder-built obb's to .2 when installing them with a local executable.

diff --git a/Engine/Source/Programs/UnrealBuildTool/Platform/Android/UEDeployAndroid.cs b/Engine/Source/Programs/UnrealBuildTool/Platform/Android/UEDeployAndroid.cs
--- a/Engine/Source/Programs/UnrealBuildTool/Platform/Android/UEDeployAndroid.cs
+++ b/Engine/Source/Programs/UnrealBuildTool/Platform/Android/UEDeployAndroid.cs
@@ -948,27 +948,22 @@ namespace UnrealBuildTool
 				bool bUseChangeListAsStoreVersion = false;
 				Ini.GetBool("/Script/AndroidRuntimeSettings.AndroidRuntimeSettings", "bUseChangeListAsStoreVersion", out bUseChangeListAsStoreVersion);
 
-				bool IsBuildMachine = Unreal.IsBuildMachine();
-				// override store version with changelist if enabled and is build machine
-				if (bUseChangeListAsStoreVersion && IsBuildMachine)
+				// override store version with changelist if enabled
+				if (bUseChangeListAsStoreVersion)
 				{
 					// make sure changelist is cached (clear unused warning)
 					string EngineVersion = ReadEngineVersion();
-					if (EngineVersion == null)
-					{
-						throw new BuildException("No engine version!");
-					}
 
 					int Changelist = 0;
-					if (Int32.TryParse(EngineChangelist, out Changelist))
+					if (EngineVersion == null || !(Int32.TryParse(EngineChangelist, out Changelist)) || Changelist == 0)
 					{
-						if (Changelist != 0)
-						{
-							StoreVersion = Changelist;
-						}
+						throw new BuildException("AndroidRuntimeSettings.bUseChangeListAsStoreVersion is enabled but no valid Engine version changelist was found!");
 					}
+
+					StoreVersion = Changelist;
 				}
 
+				bool IsBuildMachine = Unreal.IsBuildMachine();
 				Logger.LogInformation("GotStoreVersion found v{StoreVersion}. (bUseChangeListAsStoreVersion={UseChangeListAsStoreVersion} IsBuildMachine={IsBuildMachine} EngineChangeList={EngineChangeList})", StoreVersion, bUseChangeListAsStoreVersion, IsBuildMachine, EngineChangelist);
 
 				CachedStoreVersion = StoreVersion;
