UnrealBuildTool: Add remark that -StaticAnalyzerProjectOnly will reenable PVS-Studio warnings that are disabled in MicrosoftPlatformCodeAnalysis.h

#rnx

diff --git a/Engine/Source/Programs/UnrealBuildTool/Configuration/TargetRules.cs b/Engine/Source/Programs/UnrealBuildTool/Configuration/TargetRules.cs
--- a/Engine/Source/Programs/UnrealBuildTool/Configuration/TargetRules.cs
+++ b/Engine/Source/Programs/UnrealBuildTool/Configuration/TargetRules.cs
@@ -2201,6 +2201,7 @@ namespace UnrealBuildTool
 		/// <summary>
 		/// Only run static analysis against project modules, skipping engine modules
 		/// </summary>
+		/// <remarks>PVS-Studio: warnings that are disabled via Engine/Source/Runtime/Core/Public/Microsoft/MicrosoftPlatformCodeAnalysis.h will be reenabled</remarks>
 		[CommandLine("-StaticAnalyzerProjectOnly")]
 		[XmlConfigFile(Category = "BuildConfiguration")]
 		public bool bStaticAnalyzerProjectOnly { get; set; } = false;
