AutomationTool.references file incorrectly placed in project intermediates

diff --git a/Engine/Source/Programs/UnrealBuildTool/ProjectFiles/VisualStudio/VCProjectFileGenerator.cs b/Engine/Source/Programs/UnrealBuildTool/ProjectFiles/VisualStudio/VCProjectFileGenerator.cs
--- a/Engine/Source/Programs/UnrealBuildTool/ProjectFiles/VisualStudio/VCProjectFileGenerator.cs
+++ b/Engine/Source/Programs/UnrealBuildTool/ProjectFiles/VisualStudio/VCProjectFileGenerator.cs
@@ -658,7 +658,11 @@ namespace UnrealBuildTool
 				{
 					AutomationToolDocument.Save(Writer);
 				}
-				Utils.WriteFileIfChanged(FileReference.Combine(IntermediateProjectFilesPath, "AutomationTool.csproj.References"), Output.ToString(), Logger);
+				
+				// We cannot just use the IntermediateProjectFilesPath, as it can be changed to the project scope under "-Game" solution generation; AutomationTool.csproj always comes from 
+				// Engine/Source/Programs
+				DirectoryReference engineSourceIntermediateProjectFilesPath = DirectoryReference.Combine(Unreal.EngineDirectory, "Intermediate", "ProjectFiles");
+				Utils.WriteFileIfChanged(FileReference.Combine(engineSourceIntermediateProjectFilesPath, "AutomationTool.csproj.References"), Output.ToString(), Logger);
 			}
 
 			return true;
