UnrealBuildTool: Place files larger than NumIncludedBytesPerUnityCPP into their own unity files, to prevent oversized unity files

#rnx

diff --git a/Engine/Source/Programs/UnrealBuildTool/Configuration/ReadOnlyTargetRules.cs b/Engine/Source/Programs/UnrealBuildTool/Configuration/ReadOnlyTargetRules.cs
--- a/Engine/Source/Programs/UnrealBuildTool/Configuration/ReadOnlyTargetRules.cs
+++ b/Engine/Source/Programs/UnrealBuildTool/Configuration/ReadOnlyTargetRules.cs
@@ -398,6 +398,7 @@ namespace UnrealBuildTool
 
 		public bool bStressTestUnity => Inner.bStressTestUnity;
 
+		[Obsolete("Deprecated in UE5.6")]
 		public bool bDetailedUnityFiles => Inner.bDetailedUnityFiles;
 
 		[Obsolete("Deprecated in UE5.4 - Replace with ReadOnlyTargetRules.DebugInfo")]
diff --git a/Engine/Source/Programs/UnrealBuildTool/Configuration/TargetRules.cs b/Engine/Source/Programs/UnrealBuildTool/Configuration/TargetRules.cs
--- a/Engine/Source/Programs/UnrealBuildTool/Configuration/TargetRules.cs
+++ b/Engine/Source/Programs/UnrealBuildTool/Configuration/TargetRules.cs
@@ -2015,9 +2015,8 @@ namespace UnrealBuildTool
 		/// <summary>
 		/// Whether to add additional information to the unity files, such as '_of_X' in the file name. Not recommended.
 		/// </summary>
-		[CommandLine("-DetailedUnityFiles", Value = "true")]
-		[CommandLine("-NoDetailedUnityFiles", Value = "false")]
 		[XmlConfigFile(Category = "BuildConfiguration")]
+		[Obsolete("Deprecated in UE5.6")]
 		public bool bDetailedUnityFiles { get; set; }
 
 		/// <summary>
diff --git a/Engine/Source/Programs/UnrealBuildTool/System/Unity.cs b/Engine/Source/Programs/UnrealBuildTool/System/Unity.cs
--- a/Engine/Source/Programs/UnrealBuildTool/System/Unity.cs
+++ b/Engine/Source/Programs/UnrealBuildTool/System/Unity.cs
@@ -88,6 +88,10 @@ namespace UnrealBuildTool
 			/// <param name="File">The file to add.</param>
 			public void AddFile(FileItem File)
 			{
+				if (SplitLength != -1 && File.Length > SplitLength)
+				{
+					EndCurrentUnityFile();
+				}
 				CurrentCollection.AddFile(File);
 				if (SplitLength != -1 && CurrentCollection.VirtualLength > SplitLength)
 				{
@@ -102,6 +106,10 @@ namespace UnrealBuildTool
 			/// <param name="File">The file to add virtually.  Only the size of the file is tracked.</param>
 			public void AddVirtualFile(FileItem File)
 			{
+				if (SplitLength != -1 && File.Length > SplitLength)
+				{
+					EndCurrentUnityFile();
+				}
 				CurrentCollection.AddVirtualFile(File);
 				if (SplitLength != -1 && CurrentCollection.VirtualLength > SplitLength)
 				{
@@ -193,7 +201,7 @@ namespace UnrealBuildTool
 				// Note that we're relying on this not only sorting files within each directory, but also the directories
 				// themselves, so the whole list of file paths is the same across computers.
 				// Case-insensitive file path compare, because you never know what is going on with local file systems.
-				List<FileItem> SortedSourceFiles = new List<FileItem>(SourceFiles);
+				List<FileItem> SortedSourceFiles = [.. SourceFiles];
 				SortedSourceFiles.Sort((A, B) =>
 				{
 					// Generated files from UHT need to be first in the list because they implement templated functions that aren't
@@ -202,21 +210,27 @@ namespace UnrealBuildTool
 					// which results in 'error C2908: explicit specialization; '*****' has already been instantiated'
 					bool bAIsGenerated = A.AbsolutePath.EndsWith(".gen.cpp") || CompileEnvironment.FileMatchesExtraGeneratedCPPTypes(A.AbsolutePath);
 					bool bBIsGenerated = B.AbsolutePath.EndsWith(".gen.cpp") || CompileEnvironment.FileMatchesExtraGeneratedCPPTypes(B.AbsolutePath);
-					if (bAIsGenerated && !bBIsGenerated)
+					if (bAIsGenerated != bBIsGenerated)
 					{
-						return -1;
+						return bAIsGenerated && !bBIsGenerated ? -1 : 1;
 					}
 
-					if (!bAIsGenerated && bBIsGenerated)
+					// Sort oversized files to the end of the list so they will be placed into their own unity file
+					if (!bForceIntoSingleUnityFile)
 					{
-						return 1;
+						bool bAIsOversized = A.Length > NumIncludedBytesPerUnitySource;
+						bool bBIsOversized = B.Length > NumIncludedBytesPerUnitySource;
+						if (bAIsOversized != bBIsOversized)
+						{
+							return !bAIsOversized && bBIsOversized ? -1 : 1;
+						}
 					}
 
 					return String.Compare(A.AbsolutePath, B.AbsolutePath, StringComparison.OrdinalIgnoreCase);
 				});
 
-				HashSet<FileItem> AdaptiveFileSet = new HashSet<FileItem>(AdaptiveFiles);
-				UnityFileBuilder SourceUnityFileBuilder = new UnityFileBuilder(bForceIntoSingleUnityFile ? -1 : NumIncludedBytesPerUnitySource);
+				HashSet<FileItem> AdaptiveFileSet = [.. AdaptiveFiles];
+				UnityFileBuilder SourceUnityFileBuilder = new(bForceIntoSingleUnityFile ? -1 : NumIncludedBytesPerUnitySource);
 				foreach (FileItem SourceFile in SortedSourceFiles)
 				{
 					if (!bForceIntoSingleUnityFile && SourceFile.AbsolutePath.Contains(".GeneratedWrapper.", StringComparison.InvariantCultureIgnoreCase))
@@ -250,28 +264,16 @@ namespace UnrealBuildTool
 			{
 				++CurrentUnityFileCount;
 
-				string UnityExt = UnityFile.Files.First().Location.GetExtension().ToLowerInvariant();
+				FileItem FirstFile = UnityFile.Files.FirstOrDefault() ?? UnityFile.VirtualFiles.First();
+				string UnityExt = FirstFile.Location.GetExtension().ToLowerInvariant();
 				StringWriter OutputUnitySourceWriter = new();
 
 				OutputUnitySourceWriter.WriteLine($"// This file is automatically generated at compile-time to include some subset of the user-created {UnityExt.Trim('.')} files.");
 
 				// Determine unity file path name
-				string UnitySourceFileName;
-				if (AllUnityFiles.Count > 1)
-				{
-					if (Target.bDetailedUnityFiles)
-					{
-						UnitySourceFileName = $"{ModulePrefix}{BaseName}.{CurrentUnityFileCount}_of_{AllUnityFiles.Count}{UnityExt}";
-					}
-					else
-					{
-						UnitySourceFileName = $"{ModulePrefix}{BaseName}.{CurrentUnityFileCount}{UnityExt}";
-					}
-				}
-				else
-				{
-					UnitySourceFileName = $"{ModulePrefix}{BaseName}{UnityExt}";
-				}
+				string UnitySourceFileName = AllUnityFiles.Count > 1
+					? $"{ModulePrefix}{BaseName}.{CurrentUnityFileCount}{UnityExt}"
+					: $"{ModulePrefix}{BaseName}{UnityExt}";
 				FileReference UnitySourceFilePath = FileReference.Combine(IntermediateDirectory, UnitySourceFileName);
 
 				List<FileItem> InlinedGenCPPFilesInUnity = [];
