UnrealBuildTool: Always use unity files for gen.cpp even if unity is disabled for the current module

#rnx

diff --git a/Engine/Source/Programs/UnrealBuildTool/Configuration/ReadOnlyTargetRules.cs b/Engine/Source/Programs/UnrealBuildTool/Configuration/ReadOnlyTargetRules.cs
--- a/Engine/Source/Programs/UnrealBuildTool/Configuration/ReadOnlyTargetRules.cs
+++ b/Engine/Source/Programs/UnrealBuildTool/Configuration/ReadOnlyTargetRules.cs
@@ -338,6 +338,8 @@ namespace UnrealBuildTool
 
 		public bool bForceUnityBuild => Inner.bForceUnityBuild;
 
+		public bool bAlwaysUseUnityForGeneratedFiles => Inner.bAlwaysUseUnityForGeneratedFiles;
+
 		public IReadOnlyList<string>? DisableMergingModuleAndGeneratedFilesInUnityFiles => Inner.DisableMergingModuleAndGeneratedFilesInUnityFiles;
 
 		public bool bAdaptiveUnityDisablesOptimizations => Inner.bAdaptiveUnityDisablesOptimizations;
diff --git a/Engine/Source/Programs/UnrealBuildTool/Configuration/TargetRules.cs b/Engine/Source/Programs/UnrealBuildTool/Configuration/TargetRules.cs
--- a/Engine/Source/Programs/UnrealBuildTool/Configuration/TargetRules.cs
+++ b/Engine/Source/Programs/UnrealBuildTool/Configuration/TargetRules.cs
@@ -1741,6 +1741,13 @@ namespace UnrealBuildTool
 		[XmlConfigFile(Category = "BuildConfiguration")]
 		public bool bForceUnityBuild { get; set; }
 
+		/// <summary>
+		/// Whether to force generated files to be placed in a unity file even when unity files are otherwise disabled.
+		/// It is not recommended to disable this unless debugging issues with UHT file generation.
+		/// </summary>
+		[XmlConfigFile(Category = "BuildConfiguration")]
+		public bool bAlwaysUseUnityForGeneratedFiles { get; set; } = true;
+
 		/// <summary>
 		/// Whether to merge module and generated unity files for faster compilation.
 		/// </summary>
diff --git a/Engine/Source/Programs/UnrealBuildTool/Configuration/UEBuildModuleCPP.cs b/Engine/Source/Programs/UnrealBuildTool/Configuration/UEBuildModuleCPP.cs
--- a/Engine/Source/Programs/UnrealBuildTool/Configuration/UEBuildModuleCPP.cs
+++ b/Engine/Source/Programs/UnrealBuildTool/Configuration/UEBuildModuleCPP.cs
@@ -806,23 +806,22 @@ namespace UnrealBuildTool
 				}
 			}
 
-			// Compile Generated CPP Files
-			if (bModuleUsesUnityBuild && Rules.UnityFileTypes.HasFlag(ModuleRules.UnityFileType.Cpp))
+			// Compile Generated CPP Files, always use unity for .gen.cpp even if disabled for the current module
+			if (bModuleUsesUnityBuild && bMergeUnityFiles && Rules.UnityFileTypes.HasFlag(ModuleRules.UnityFileType.Cpp))
 			{
-				if (bMergeUnityFiles)
-				{
-					CPPFiles.AddRange(GeneratedFileItems);
-				}
-				else
-				{
-					Unity.GenerateUnitySource(Target, GeneratedFileItems, new List<FileItem>(), CppCompileEnvironment, WorkingSet, (Rules.ShortName ?? Name) + ".gen", IntermediateDirectory, Graph, SourceFileToUnityFile,
-						out List<FileItem> NormalGeneratedFiles, out List<FileItem> AdaptiveGeneratedFiles, NumIncludedBytesPerUnityCPP);
-					LinkInputFiles.AddRange(CompileFilesWithToolChain(Target, ToolChain, GeneratedCPPCompileEnvironment, ModuleCompileEnvironment, NormalGeneratedFiles, AdaptiveGeneratedFiles, Graph, Logger).ObjectFiles);
-				}
+				CPPFiles.AddRange(GeneratedFileItems);
+			}
+			else if (Rules.Target.bAlwaysUseUnityForGeneratedFiles && !Rules.Target.bEnableCppModules)
+			{
+				Unity.GenerateUnitySource(Target, GeneratedFileItems, new List<FileItem>(), CppCompileEnvironment, WorkingSet, (Rules.ShortName ?? Name) + ".gen", IntermediateDirectory, Graph, SourceFileToUnityFile,
+					out List<FileItem> NormalGeneratedFiles, out List<FileItem> AdaptiveGeneratedFiles, NumIncludedBytesPerUnityCPP);
+				LinkInputFiles.AddRange(CompileFilesWithToolChain(Target, ToolChain, GeneratedCPPCompileEnvironment, ModuleCompileEnvironment, NormalGeneratedFiles, AdaptiveGeneratedFiles, Graph, Logger).ObjectFiles);
 			}
 			else
 			{
-				LinkInputFiles.AddRange(ToolChain.CompileAllCPPFiles(GeneratedCPPCompileEnvironment, GeneratedFileItems, IntermediateDirectory, Name, Graph).ObjectFiles);
+				Unity.GetAdaptiveFiles(Target, GeneratedFileItems, InputFiles.HeaderFiles, GeneratedCPPCompileEnvironment, WorkingSet, (Rules.ShortName ?? Name) + ".gen", IntermediateDirectory, Graph,
+					out List<FileItem> NormalGeneratedFiles, out List<FileItem> AdaptiveGeneratedFiles);
+				LinkInputFiles.AddRange(CompileFilesWithToolChain(Target, ToolChain, GeneratedCPPCompileEnvironment, ModuleCompileEnvironment, NormalGeneratedFiles, AdaptiveGeneratedFiles, Graph, Logger).ObjectFiles);
 			}
 
 
diff --git a/Engine/Source/Programs/UnrealBuildTool/Platform/Android/UEBuildAndroid.cs b/Engine/Source/Programs/UnrealBuildTool/Platform/Android/UEBuildAndroid.cs
--- a/Engine/Source/Programs/UnrealBuildTool/Platform/Android/UEBuildAndroid.cs
+++ b/Engine/Source/Programs/UnrealBuildTool/Platform/Android/UEBuildAndroid.cs
@@ -210,6 +210,10 @@ namespace UnrealBuildTool
 
 				// Clang static analysis requires non unity builds
 				Target.bUseUnityBuild = false;
+				if (Target.bStaticAnalyzerIncludeGenerated)
+				{
+					Target.bAlwaysUseUnityForGeneratedFiles = false;
+				}
 			}
 
 			Target.bCompileRecast = true;
diff --git a/Engine/Source/Programs/UnrealBuildTool/Platform/IOS/UEBuildIOS.cs b/Engine/Source/Programs/UnrealBuildTool/Platform/IOS/UEBuildIOS.cs
--- a/Engine/Source/Programs/UnrealBuildTool/Platform/IOS/UEBuildIOS.cs
+++ b/Engine/Source/Programs/UnrealBuildTool/Platform/IOS/UEBuildIOS.cs
@@ -763,6 +763,10 @@ namespace UnrealBuildTool
 
 				// Clang static analysis requires non unity builds
 				Target.bUseUnityBuild = false;
+				if (Target.bStaticAnalyzerIncludeGenerated)
+				{
+					Target.bAlwaysUseUnityForGeneratedFiles = false;
+				}
 			}
 
 			// we assume now we are building with IOS8 or later
diff --git a/Engine/Source/Programs/UnrealBuildTool/Platform/Linux/UEBuildLinux.cs b/Engine/Source/Programs/UnrealBuildTool/Platform/Linux/UEBuildLinux.cs
--- a/Engine/Source/Programs/UnrealBuildTool/Platform/Linux/UEBuildLinux.cs
+++ b/Engine/Source/Programs/UnrealBuildTool/Platform/Linux/UEBuildLinux.cs
@@ -227,6 +227,10 @@ namespace UnrealBuildTool
 
 				// Clang static analysis requires non unity builds
 				Target.bUseUnityBuild = false;
+				if (Target.bStaticAnalyzerIncludeGenerated)
+				{
+					Target.bAlwaysUseUnityForGeneratedFiles = false;
+				}
 			}
 
 			// Editor target types get overwritten in UEBuildTarget.cs so lets avoid adding this here. ResetTarget is called with
diff --git a/Engine/Source/Programs/UnrealBuildTool/Platform/Mac/UEBuildMac.cs b/Engine/Source/Programs/UnrealBuildTool/Platform/Mac/UEBuildMac.cs
--- a/Engine/Source/Programs/UnrealBuildTool/Platform/Mac/UEBuildMac.cs
+++ b/Engine/Source/Programs/UnrealBuildTool/Platform/Mac/UEBuildMac.cs
@@ -357,6 +357,10 @@ namespace UnrealBuildTool
 
 				// Clang static analysis requires non unity builds
 				Target.bUseUnityBuild = false;
+				if (Target.bStaticAnalyzerIncludeGenerated)
+				{
+					Target.bAlwaysUseUnityForGeneratedFiles = false;
+				}
 
 				// Disable chaining PCHs for the moment because it is crashing clang
 				Target.bChainPCHs = false;
diff --git a/Engine/Source/Programs/UnrealBuildTool/Platform/Windows/UEBuildWindows.cs b/Engine/Source/Programs/UnrealBuildTool/Platform/Windows/UEBuildWindows.cs
--- a/Engine/Source/Programs/UnrealBuildTool/Platform/Windows/UEBuildWindows.cs
+++ b/Engine/Source/Programs/UnrealBuildTool/Platform/Windows/UEBuildWindows.cs
@@ -1283,6 +1283,10 @@ namespace UnrealBuildTool
 
 				// Clang static analysis requires non unity builds
 				Target.bUseUnityBuild = false;
+				if (Target.bStaticAnalyzerIncludeGenerated)
+				{
+					Target.bAlwaysUseUnityForGeneratedFiles = false;
+				}
 			}
 			else if (Target.StaticAnalyzer != StaticAnalyzer.None &&
 					 Target.StaticAnalyzerOutputType != StaticAnalyzerOutputType.Text)
