UnrealBuildTool: add bPrimaryProjectNameAppendDriveLetter option to append the drive letter to the generated primary project name.
This is very useful when having the same branch on different drives with the same folder name, where bPrimaryProjectNameFromFolder would generate the same project name, resulting in confusion when opened in multiple Visual Studio instances.

diff --git a/Engine/Source/Programs/UnrealBuildTool/ProjectFiles/ProjectFileGenerator.cs b/Engine/Source/Programs/UnrealBuildTool/ProjectFiles/ProjectFileGenerator.cs
--- a/Engine/Source/Programs/UnrealBuildTool/ProjectFiles/ProjectFileGenerator.cs
+++ b/Engine/Source/Programs/UnrealBuildTool/ProjectFiles/ProjectFileGenerator.cs
@@ -436,6 +436,12 @@ namespace UnrealBuildTool
 		[XmlConfigFile]
 		protected bool bPrimaryProjectNameFromFolder = false;
 
+		/// <summary>
+		/// If true, append the drive letter to the primary project name.
+		/// </summary>
+		[XmlConfigFile]
+		protected bool bPrimaryProjectNameAppendDriveLetter = false;
+
 		/// <summary>
 		/// Maps all module files that were included in generated project files, to actual project file objects.
 		/// </summary>
@@ -950,6 +956,20 @@ namespace UnrealBuildTool
 				}
 			}
 
+			if (bPrimaryProjectNameAppendDriveLetter)
+			{
+				string PrimaryProjectPathString = PrimaryProjectPath.ToString();
+
+				if (PrimaryProjectPathString.Length > 2 && PrimaryProjectPathString.Substring(1, 1) == ":")
+				{
+					string NewPrimaryProjectDrive = PrimaryProjectPathString.Substring(0, 1);
+					if (!String.IsNullOrEmpty(NewPrimaryProjectDrive))
+					{
+						PrimaryProjectName += "-" + NewPrimaryProjectDrive;
+					}
+				}
+			}
+
 			// if making a temp primary project, put it in with the rest of the project files
 			if (bGeneratingTemporaryProjects)
 			{
