UnrealBuildTool: Remove unnecessary recursion to scan for referenced modules in test mode and ensure that the makefile will be invalidated for test builds

#rnx

diff --git a/Engine/Source/Programs/UnrealBuildTool/Modes/TestMode.cs b/Engine/Source/Programs/UnrealBuildTool/Modes/TestMode.cs
--- a/Engine/Source/Programs/UnrealBuildTool/Modes/TestMode.cs
+++ b/Engine/Source/Programs/UnrealBuildTool/Modes/TestMode.cs
@@ -111,8 +111,10 @@ namespace UnrealBuildTool
 			}
 
 			// Analyzing dependency graph before any build commands to automatically set compilation flags.
-			// Extracted from AnalyzeMode for simplicity.
-			AnalyzeDependencyGraph(TargetDescriptors.First(), BuildConfiguration, Logger);
+			AnalyzeDependencyGraph(TestTargetDescriptor, BuildConfiguration, Logger);
+
+			// Reset arguments and add a dummy argument that will invalidate the makefile if a build had been done outside of test mode
+			TestTargetDescriptor.AdditionalArguments = new([.. TestTargetDescriptor.AdditionalArguments.AsEnumerable(), "-TestMode"]);
 
 			using (ISourceFileWorkingSet WorkingSet = SourceFileWorkingSet.Create(Unreal.RootDirectory, new HashSet<DirectoryReference>(), Logger))
 			{
@@ -259,109 +261,24 @@ namespace UnrealBuildTool
 			UEBuildTarget Target = UEBuildTarget.Create(TargetDescriptor, BuildConfiguration, Logger);
 			DirectoryReference.CreateDirectory(Target.ReceiptFileName.Directory);
 
-			HashSet<UEBuildModule> Visited = new();
-			IncludeChainMap Chains = new();
-
-			if (Target.Rules.LaunchModuleName != null)
-			{
-				AnalyzeModuleDependencies(Target.Rules.LaunchModuleName, new List<string>() { "target" }, Target, Chains, Visited, Logger);
-			}
-
-			foreach (string RootModuleName in Target.Rules.ExtraModuleNames)
-			{
-				AnalyzeModuleDependencies(RootModuleName, new List<string>() { "target" }, Target, Chains, Visited, Logger);
-			}
-
-			foreach (UEBuildPlugin Plugin in Target.BuildPlugins!)
-			{
-				foreach (UEBuildModule Module in Plugin.Modules)
-				{
-					if (!Chains.ContainsKey(Module))
-					{
-						AnalyzeModuleDependencies(Module.Name, new List<string>() { "target", Plugin.ReferenceChain, Plugin.File.GetFileName() }, Target, Chains, Visited, Logger);
-					}
-				}
-			}
-
-			if (Target.Rules.bBuildAllModules)
-			{
-				foreach (UEBuildBinary Binary in Target.Binaries)
-				{
-					foreach (UEBuildModule Module in Binary.Modules)
-					{
-						if (!Chains.ContainsKey(Module))
-						{
-							List<string> IncludeChain = Enumerable.Repeat(String.Empty, 1000).ToList();
-							IncludeChain.Add("allmodules option");
-							if (Module.Rules.Plugin != null)
-							{
-								IncludeChain.Add(Module.Rules.Plugin.File.GetFileName());
-								AnalyzeModuleDependencies(Module.Name, IncludeChain, Target, Chains, Visited, Logger);
-							}
-							else
-							{
-								AnalyzeModuleDependencies(Module.Name, IncludeChain, Target, Chains, Visited, Logger);
-							}
-						}
-					}
-				}
-			}
-		}
-
-		private void AnalyzeModuleDependencies(string ModuleName, List<string> ParentChain, UEBuildTarget Target, IncludeChainMap Chains, HashSet<UEBuildModule> Visited, ILogger Logger)
-		{
-			if (ParentChain.Contains(ModuleName))
-			{
-				return;
-			}
+			HashSet<string> Modules = [.. Target.Binaries.SelectMany(x => x.Modules).Select(x => x.Name)];
 
-			UEBuildModule Module = Target.GetModuleByName(ModuleName);
-
-			List<string> CurrentChain = new(ParentChain) { Module.Name };
-
-			if (!Chains.ContainsKey(Module))
-			{
-				Chains.Add(Module, CurrentChain);
-			}
-
-			if (Chains[Module].Count > CurrentChain.Count)
-			{
-				List<UEBuildModule> RecheckModules = new();
-				Module.GetAllDependencyModules(RecheckModules, new(), true, false, true);
-				Visited.ExceptWith(RecheckModules);
-				Chains[Module] = CurrentChain;
-			}
-			else if (Visited.Contains(Module))
-			{
-				return;
-			}
-
-			SetupTestsTargetForModule(Module, Logger);
-			Visited.Add(Module);
-
-			List<UEBuildModule> TargetModules = new();
-			Module.GetAllDependencyModules(TargetModules, new(), true, false, true);
-			TargetModules.ForEach(x => AnalyzeModuleDependencies(x.Name, CurrentChain, Target, Chains, Visited, Logger));
-		}
-
-		private void SetupTestsTargetForModule(UEBuildModule Module, ILogger Logger)
-		{
-			if (Module.Name == ENGINE_MODULE && !TestTargetRules.bTestsRequireEngine)
+			if (Modules.Contains(ENGINE_MODULE) && !TestTargetRules.bTestsRequireEngine)
 			{
 				TestTargetRules.bTestsRequireEngine = true;
 				Logger.LogInformation("{Module} module found in dependency graph, will compile against it.", ENGINE_MODULE);
 			}
-			else if (Module.Name == APPLICATION_CORE_MODULE && !TestTargetRules.bTestsRequireApplicationCore)
+			else if (Modules.Contains(APPLICATION_CORE_MODULE) && !TestTargetRules.bTestsRequireApplicationCore)
 			{
 				TestTargetRules.bTestsRequireApplicationCore = true;
 				Logger.LogInformation("{Module} module found in dependency graph, will compile against it.", APPLICATION_CORE_MODULE);
 			}
-			else if (Module.Name == COREUOBJECT_MODULE && !TestTargetRules.bTestsRequireCoreUObject)
+			else if (Modules.Contains(COREUOBJECT_MODULE) && !TestTargetRules.bTestsRequireCoreUObject)
 			{
 				TestTargetRules.bTestsRequireCoreUObject = true;
 				Logger.LogInformation("{Module} module found in dependency graph, will compile against it.", COREUOBJECT_MODULE);
 			}
-			else if (Module.Name == EDITOR_MODULE && !TestTargetRules.bTestsRequireEditor)
+			else if (Modules.Contains(EDITOR_MODULE) && !TestTargetRules.bTestsRequireEditor)
 			{
 				TestTargetRules.bTestsRequireEditor = true;
 				Logger.LogInformation("{Module} module found in dependency graph, will compile against it.", EDITOR_MODULE);
