UnrealBuildTool: Rescan a target in test mode if a test requires property is changed, as that may add additional modules to the reference list

#rnx

diff --git a/Engine/Source/Programs/UnrealBuildTool/Modes/TestMode.cs b/Engine/Source/Programs/UnrealBuildTool/Modes/TestMode.cs
--- a/Engine/Source/Programs/UnrealBuildTool/Modes/TestMode.cs
+++ b/Engine/Source/Programs/UnrealBuildTool/Modes/TestMode.cs
@@ -258,30 +258,40 @@ namespace UnrealBuildTool
 
 		private void AnalyzeDependencyGraph(TargetDescriptor TargetDescriptor, BuildConfiguration BuildConfiguration, ILogger Logger)
 		{
-			UEBuildTarget Target = UEBuildTarget.Create(TargetDescriptor, BuildConfiguration, Logger);
-			DirectoryReference.CreateDirectory(Target.ReceiptFileName.Directory);
+			UEBuildTarget target = UEBuildTarget.Create(TargetDescriptor, BuildConfiguration, Logger);
+			DirectoryReference.CreateDirectory(target.ReceiptFileName.Directory);
 
-			HashSet<string> Modules = [.. Target.Binaries.SelectMany(x => x.Modules).Select(x => x.Name)];
+			HashSet<string> moduleNames = [.. target.Binaries.SelectMany(x => x.Modules).Select(x => x.Name)];
 
-			if (Modules.Contains(ENGINE_MODULE) && !TestTargetRules.bTestsRequireEngine)
+			bool updated = false;
+			if (moduleNames.Contains(ENGINE_MODULE) && !TestTargetRules.bTestsRequireEngine)
 			{
 				TestTargetRules.bTestsRequireEngine = true;
 				Logger.LogInformation("{Module} module found in dependency graph, will compile against it.", ENGINE_MODULE);
+				updated = true;
 			}
-			else if (Modules.Contains(APPLICATION_CORE_MODULE) && !TestTargetRules.bTestsRequireApplicationCore)
+			else if (moduleNames.Contains(APPLICATION_CORE_MODULE) && !TestTargetRules.bTestsRequireApplicationCore)
 			{
 				TestTargetRules.bTestsRequireApplicationCore = true;
 				Logger.LogInformation("{Module} module found in dependency graph, will compile against it.", APPLICATION_CORE_MODULE);
+				updated = true;
 			}
-			else if (Modules.Contains(COREUOBJECT_MODULE) && !TestTargetRules.bTestsRequireCoreUObject)
+			else if (moduleNames.Contains(COREUOBJECT_MODULE) && !TestTargetRules.bTestsRequireCoreUObject)
 			{
 				TestTargetRules.bTestsRequireCoreUObject = true;
 				Logger.LogInformation("{Module} module found in dependency graph, will compile against it.", COREUOBJECT_MODULE);
+				updated = true;
 			}
-			else if (Modules.Contains(EDITOR_MODULE) && !TestTargetRules.bTestsRequireEditor)
+			else if (moduleNames.Contains(EDITOR_MODULE) && !TestTargetRules.bTestsRequireEditor)
 			{
 				TestTargetRules.bTestsRequireEditor = true;
 				Logger.LogInformation("{Module} module found in dependency graph, will compile against it.", EDITOR_MODULE);
+				updated = true;
+			}
+			// If a property changes scan again for additional dependencies
+			if (updated)
+			{
+				AnalyzeDependencyGraph(TargetDescriptor, BuildConfiguration, Logger);
 			}
 		}
 	}
