[UBA]
* Enabled known inputs by default

diff --git a/Engine/Source/Programs/UnrealBuildTool/Executors/UnrealBuildAccelerator/UnrealBuildAcceleratorConfig.cs b/Engine/Source/Programs/UnrealBuildTool/Executors/UnrealBuildAccelerator/UnrealBuildAcceleratorConfig.cs
--- a/Engine/Source/Programs/UnrealBuildTool/Executors/UnrealBuildAccelerator/UnrealBuildAcceleratorConfig.cs
+++ b/Engine/Source/Programs/UnrealBuildTool/Executors/UnrealBuildAccelerator/UnrealBuildAcceleratorConfig.cs
@@ -206,7 +206,7 @@ namespace UnrealBuildTool
 		/// </summary>
 		[XmlConfigFile(Category = "UnrealBuildAccelerator")]
 		[CommandLine("-UBAUseKnownInputs", Value = "true")]
-		public bool bUseKnownInputs { get; set; } = false;
+		public bool bUseKnownInputs { get; set; } = true;
 
 		/// <summary>
 		/// Write yaml file with all actions that are queued for build. This can be used to replay using "UbaCli.exe local file.yaml"
