[UBT]
* Changed so clang-cl.exe is not creating temp files since uba will not write the files out until it is done anyway
* Set vctoolsdir to undefined for clang builds so it doesn't pick random vs installs from the machine and read files from there
* Explicitly set resource-dir so uba can find it and prefetch builtin headers
* Added RAD_SHARED_THREAD_POOL when linking with radlinker.. this prevents crazy thread contention that brings down the computer

diff --git a/Engine/Source/Programs/UnrealBuildTool/Platform/Windows/VCToolChain.cs b/Engine/Source/Programs/UnrealBuildTool/Platform/Windows/VCToolChain.cs
--- a/Engine/Source/Programs/UnrealBuildTool/Platform/Windows/VCToolChain.cs
+++ b/Engine/Source/Programs/UnrealBuildTool/Platform/Windows/VCToolChain.cs
@@ -369,7 +369,13 @@ namespace UnrealBuildTool
 
 				// We have 'this' vs nullptr comparisons that get optimized away for newer versions of Clang, which is undesirable until we refactor these checks.
 				Arguments.Add("-fno-delete-null-pointer-checks");
-				
+
+				// The temp files are there to prevent corrupt output files but with uba we don't write output files until process has finished successfully so creating temp files just create wasteful IOPS
+				Arguments.Add("-fno-temp-file");
+
+				// Prevents clang-cl from reading a few files from installed visual studio folders (that we are not even using)
+				Arguments.Add("-vctoolsdir undefined");
+
 				if (Target.bWithLiveCoding)
 				{
 					Arguments.Add("-Z7");
@@ -441,9 +447,12 @@ namespace UnrealBuildTool
 					Arguments.Add("-winsysroot=.");
 					Arguments.Add("-Xclang -gno-codeview-command-line");
 					Arguments.Add("-fintegrated-cc1");
-					DirectoryReference ResourceDir = DirectoryReference.Combine(EnvVars.CompilerDir!, "lib", "clang", EnvVars.CompilerVersion.Components[0].ToString());
-					Arguments.Add($"-resource-dir=\"{NormalizeCommandLinePath(ResourceDir, CompileEnvironment.RootPaths)}\"");
 				}
+
+				// Resource dir is for built-in includes. We need it for vfs but also for uba include crawler to find includes to transfer
+				DirectoryReference ResourceDir = DirectoryReference.Combine(EnvVars.CompilerDir!, "lib", "clang", EnvVars.CompilerVersion.Components[0].ToString());
+				Arguments.Add($"-resource-dir=\"{NormalizeCommandLinePath(ResourceDir, CompileEnvironment.RootPaths)}\"");
+
 			}
 			else if (Target.StaticAnalyzer == StaticAnalyzer.Default && !CompileEnvironment.bDisableStaticAnalysis)
 			{
@@ -2708,7 +2717,8 @@ namespace UnrealBuildTool
 			if (Target.WindowsPlatform.bAllowRadLinker && !bBuildImportLibraryOnly)
 			{
 				// Saves space in type table in pdb
-				Arguments.Add("/rad_pdb_hash_type_names:lenient");
+				Arguments.Add("/RAD_PDB_HASH_TYPE_NAMES:lenient");
+				Arguments.Add("/RAD_SHARED_THREAD_POOL");
 
 				// Speeds up linking - TODO: disabled because of windows issues
 				//Arguments.Add("/rad_large_pages");
