[UBA]
* Added lease and instance type to description for agent (so they show up in visualizer0
* Changed agent max connection count to 8
* Changed MaxWorkers to be 64 (instead of core count).. workers most of the time just wait on network so it is ok to have more than core count

diff --git a/Engine/Source/Programs/UnrealBuildTool/Executors/UnrealBuildAccelerator/UBAAgentCoordinatorHorde.cs b/Engine/Source/Programs/UnrealBuildTool/Executors/UnrealBuildAccelerator/UBAAgentCoordinatorHorde.cs
--- a/Engine/Source/Programs/UnrealBuildTool/Executors/UnrealBuildAccelerator/UBAAgentCoordinatorHorde.cs
+++ b/Engine/Source/Programs/UnrealBuildTool/Executors/UnrealBuildAccelerator/UBAAgentCoordinatorHorde.cs
@@ -245,6 +245,8 @@ namespace UnrealBuildTool
 				"PhysicalCores",
 				"LogicalCores",
 				"EC2",
+				"LeaseId",
+				"aws-instance-type",
 			};
 
 		public async Task<bool> AddWorkerAsync(Requirements requirements, UnrealBuildAcceleratorHordeConfig hordeConfig, CancellationToken cancellationToken, int activeCores)
@@ -296,6 +298,8 @@ namespace UnrealBuildTool
 				string agentName = String.Empty;
 				int numLogicalCores = 24; // Assume 24 if something goes wrong here and property is not found
 				string computeIp = String.Empty;
+				string leaseId = string.Empty;
+				string instanceType = string.Empty;
 				foreach (string property in lease.Properties)
 				{
 					int equalsIdx = property.IndexOf('=', StringComparison.OrdinalIgnoreCase);
@@ -316,6 +320,14 @@ namespace UnrealBuildTool
 						{
 							agentName = $"Worker{workerId}";
 						}
+						else if (propertyName == "LeaseId")
+						{
+							leaseId = property[(equalsIdx + 1)..];
+						}
+						else if (propertyName == "aws-instance-type")
+						{
+							instanceType = property[(equalsIdx + 1)..];
+						}
 					}
 				}
 
@@ -342,7 +354,7 @@ namespace UnrealBuildTool
 					Port = ubaPort,
 					ProxyPort = ubaProxyPort,
 				};
-				worker.BackgroundTask = RunWorkerAsync(worker, lease, locator, exeName, workerLogger, hordeConfig, _cancellationTokenSource.Token, agentName);
+				worker.BackgroundTask = RunWorkerAsync(worker, lease, locator, exeName, workerLogger, hordeConfig, _cancellationTokenSource.Token, agentName, leaseId, instanceType);
 				lock (_workers)
 				{
 					_workers.Add(worker);
@@ -481,7 +493,7 @@ namespace UnrealBuildTool
 			return result.AccessToken;
 		}
 
-		async Task RunWorkerAsync(Worker self, IComputeLease lease, BlobLocator tool, string executable, ILogger logger, UnrealBuildAcceleratorHordeConfig hordeConfig, CancellationToken cancellationToken, string agentName)
+		async Task RunWorkerAsync(Worker self, IComputeLease lease, BlobLocator tool, string executable, ILogger logger, UnrealBuildAcceleratorHordeConfig hordeConfig, CancellationToken cancellationToken, string agentName, string leaseId, string instanceType)
 		{
 			logger.LogDebug("Running worker task..");
 			try
@@ -563,7 +575,12 @@ namespace UnrealBuildTool
 						arguments.Add("-Eventfile=%UE_HORDE_TERMINATION_SIGNAL_FILE%");
 						arguments.Add("-MaxCpu=%UE_HORDE_CPU_COUNT%");
 						arguments.Add("-MulCpu=%UE_HORDE_CPU_MULTIPLIER%");
+						arguments.Add("-MaxCon=8");
+						arguments.Add("-MaxWorkers=64");
 						arguments.Add("-MaxIdle=15");
+						//arguments.Add("-ProxyUseLocalStorage"); // If set then agent's cas storage will be used by proxy
+						//arguments.Add("-NoStore"); // This means that no cas files will be written to disk.. so no state is stored between runs
+						arguments.Add($"-Description=\"{instanceType}, {leaseId}\"");
 
 						if (_owner.UBAConfig.bLogEnabled)
 						{
