PR #12787: Fixed typo in JavaReservedKeywords in UEDeployAndroid.cs

diff --git a/Engine/Source/Programs/UnrealBuildTool/Platform/Android/UEDeployAndroid.cs b/Engine/Source/Programs/UnrealBuildTool/Platform/Android/UEDeployAndroid.cs
--- a/Engine/Source/Programs/UnrealBuildTool/Platform/Android/UEDeployAndroid.cs
+++ b/Engine/Source/Programs/UnrealBuildTool/Platform/Android/UEDeployAndroid.cs
@@ -51,7 +51,7 @@ namespace UnrealBuildTool
 			"abstract", "assert", "boolean", "break", "byte", "case", "catch", "char", "class", "const", "continue", "default", "do",
 			"double", "else", "enum", "extends", "final", "finally", "float", "for", "goto", "if", "implements", "import", "instanceof",
 			"int", "interface", "long", "native", "new", "package", "private", "protected", "public", "return", "short", "static",
-			"strictfp", "super", "switch", "sychronized", "this", "throw", "throws", "transient", "try", "void", "volatile", "while",
+			"strictfp", "super", "switch", "synchronized", "this", "throw", "throws", "transient", "try", "void", "volatile", "while",
 			"false", "null", "true"
 		};
 
