UnrealBuildTool: Temporarily disable new -Description flag as submitted UBA binaries do not support it

diff --git a/Engine/Source/Programs/UnrealBuildTool/Executors/UnrealBuildAccelerator/UBAAgentCoordinatorHorde.cs b/Engine/Source/Programs/UnrealBuildTool/Executors/UnrealBuildAccelerator/UBAAgentCoordinatorHorde.cs
--- a/Engine/Source/Programs/UnrealBuildTool/Executors/UnrealBuildAccelerator/UBAAgentCoordinatorHorde.cs
+++ b/Engine/Source/Programs/UnrealBuildTool/Executors/UnrealBuildAccelerator/UBAAgentCoordinatorHorde.cs
@@ -580,7 +580,7 @@ namespace UnrealBuildTool
 						arguments.Add("-MaxIdle=15");
 						//arguments.Add("-ProxyUseLocalStorage"); // If set then agent's cas storage will be used by proxy
 						//arguments.Add("-NoStore"); // This means that no cas files will be written to disk.. so no state is stored between runs
-						arguments.Add($"-Description=\"{instanceType}, {leaseId}\"");
+						//arguments.Add($"-Description=\"{instanceType}, {leaseId}\"");
 
 						if (_owner.UBAConfig.bLogEnabled)
 						{
