UBT: Keep ld_classic around a bit longer in the AppleToolChain.cs.

diff --git a/Engine/Source/Programs/UnrealBuildTool/ToolChain/AppleToolChain.cs b/Engine/Source/Programs/UnrealBuildTool/ToolChain/AppleToolChain.cs
--- a/Engine/Source/Programs/UnrealBuildTool/ToolChain/AppleToolChain.cs
+++ b/Engine/Source/Programs/UnrealBuildTool/ToolChain/AppleToolChain.cs
@@ -819,7 +819,7 @@ namespace UnrealBuildTool
 				AppleClangVersion = ((AppleToolChainInfo)GetToolChainInfo()).AppleClangVersion;
 			}
 			// Temp solution for UE-191350
-			if (AppleClangVersion >= new VersionNumber(15) && AppleClangVersion < new VersionNumber(16))
+			if (AppleClangVersion >= new VersionNumber(15) && AppleClangVersion < new VersionNumber(17))
 			{
 				Arguments.Add(" -ld_classic");
 			}
