Remove DirectXMath support

It's been inoperable since LWC (so, 5.0), so there's no need for a separate deprecation period.
Deprecate the associated build flags in UBT the regular way, though.

diff --git a/Engine/Source/Programs/UnrealBuildTool/Configuration/ReadOnlyTargetRules.cs b/Engine/Source/Programs/UnrealBuildTool/Configuration/ReadOnlyTargetRules.cs
--- a/Engine/Source/Programs/UnrealBuildTool/Configuration/ReadOnlyTargetRules.cs
+++ b/Engine/Source/Programs/UnrealBuildTool/Configuration/ReadOnlyTargetRules.cs
@@ -263,6 +263,7 @@ namespace UnrealBuildTool
 
 		public bool bUseDebugLiveCodingConsole => Inner.bUseDebugLiveCodingConsole;
 
+		[Obsolete("Deprecated in UE 5.6 - DirectXMath path was incompatible with LWC and unused since 5.0")]
 		public bool bWithDirectXMath => Inner.bWithDirectXMath;
 
 		public bool bWithFixedTimeStepSupport => Inner.bWithFixedTimeStepSupport;
diff --git a/Engine/Source/Programs/UnrealBuildTool/Configuration/TargetRules.cs b/Engine/Source/Programs/UnrealBuildTool/Configuration/TargetRules.cs
--- a/Engine/Source/Programs/UnrealBuildTool/Configuration/TargetRules.cs
+++ b/Engine/Source/Programs/UnrealBuildTool/Configuration/TargetRules.cs
@@ -1510,10 +1510,11 @@ namespace UnrealBuildTool
 		public bool bUseDebugLiveCodingConsole { get; set; }
 
 		/// <summary>
-		/// Whether to enable support for DirectX Math
-		/// LWC_TODO: No longer supported. Needs removing.
+		/// Obsolete: whether to enable support for DirectX Math
+		/// Not compatible with LWC so has been disabled since UE 5.
 		/// </summary>
 		[RequiresUniqueBuildEnvironment]
+		[Obsolete("Deprecated in UE 5.6 - DirectXMath path was incompatible with LWC and unused since 5.0")]
 		public bool bWithDirectXMath { get; set; }
 
 		/// <summary>
diff --git a/Engine/Source/Runtime/Core/Core.Build.cs b/Engine/Source/Runtime/Core/Core.Build.cs
--- a/Engine/Source/Runtime/Core/Core.Build.cs
+++ b/Engine/Source/Runtime/Core/Core.Build.cs
@@ -239,15 +239,6 @@ public class Core : ModuleRules
 			PrivateDependencyModuleNames.Add("VSPerfExternalProfiler");
 		}
 
-		if (Target.bWithDirectXMath && Target.Platform == UnrealTargetPlatform.Win64)
-		{
-			PublicDefinitions.Add("WITH_DIRECTXMATH=1");
-		}
-		else
-		{
-			PublicDefinitions.Add("WITH_DIRECTXMATH=0");
-		}
-
 		if (Target.IsInPlatformGroup(UnrealPlatformGroup.Apple) ||
 			Target.Platform == UnrealTargetPlatform.Android)
 		{
diff --git a/Engine/Source/Runtime/Core/Public/Math/UnrealMathDirectX.h b/Engine/Source/Runtime/Core/Public/Math/UnrealMathDirectX.h
deleted file mode 100644
diff --git a/Engine/Source/Runtime/Core/Public/Math/UnrealMathSSE.h b/Engine/Source/Runtime/Core/Public/Math/UnrealMathSSE.h
--- a/Engine/Source/Runtime/Core/Public/Math/UnrealMathSSE.h
+++ b/Engine/Source/Runtime/Core/Public/Math/UnrealMathSSE.h
@@ -9,7 +9,7 @@
 
 // IWYU pragma: begin_exports
 
-#if !WITH_DIRECTXMATH && !PLATFORM_ENABLE_VECTORINTRINSICS_NEON && !(defined(__cplusplus_cli) && !PLATFORM_HOLOLENS) && PLATFORM_ENABLE_VECTORINTRINSICS
+#if !PLATFORM_ENABLE_VECTORINTRINSICS_NEON && !(defined(__cplusplus_cli) && !PLATFORM_HOLOLENS) && PLATFORM_ENABLE_VECTORINTRINSICS
 
 struct VectorRegisterConstInit {};
 
diff --git a/Engine/Source/Runtime/Core/Public/Math/VectorRegister.h b/Engine/Source/Runtime/Core/Public/Math/VectorRegister.h
--- a/Engine/Source/Runtime/Core/Public/Math/VectorRegister.h
+++ b/Engine/Source/Runtime/Core/Public/Math/VectorRegister.h
@@ -9,9 +9,7 @@
 #define SUPPORT_DOUBLE_TO_FLOAT_VECTOR_CONVERSION 1
 
 // Platform specific vector intrinsics include.
-#if WITH_DIRECTXMATH
-#include "Math/UnrealMathDirectX.h"
-#elif PLATFORM_ENABLE_VECTORINTRINSICS_NEON
+#if PLATFORM_ENABLE_VECTORINTRINSICS_NEON
 #include "Math/UnrealMathNeon.h"
 #elif defined(__cplusplus_cli) && !PLATFORM_HOLOLENS
 #include "Math/UnrealMathFPU.h" // there are compile issues with UnrealMathSSE in managed mode, so use the FPU version
