[UBT]
* Fixed so .exp file is produced by msvc linker

diff --git a/Engine/Source/Programs/UnrealBuildTool/Platform/Windows/VCToolChain.cs b/Engine/Source/Programs/UnrealBuildTool/Platform/Windows/VCToolChain.cs
--- a/Engine/Source/Programs/UnrealBuildTool/Platform/Windows/VCToolChain.cs
+++ b/Engine/Source/Programs/UnrealBuildTool/Platform/Windows/VCToolChain.cs
@@ -2800,7 +2800,7 @@ namespace UnrealBuildTool
 			if (LinkEnvironment.bIsCrossReferenced && !bBuildImportLibraryOnly)
 			{
 				Arguments.Add("/NOIMPLIB");
-				if (!Target.WindowsPlatform.bAllowClangLinker)
+				if (!Target.WindowsPlatform.Compiler.IsClang())
 				{
 					Arguments.Add("/NOEXP"); // This compiler flag does not exist on lld-link.exe.. it skips the writing of the .exp file
 				}
