Make it an error for a plugin to specify that it cannot contain Verse, but still enable Verse asset reflection, because it necessitates redundant checking in the compiler toolchain of whether or not to consider a plugin "compilable" from a Verse standpoint by only checking the `bCanContainVerse` flag, especially for plugins that only contain  assets but no actual `.verse` source files. Also exposes the `bEnableVerseAssetReflection`/`bEnableSceneGraph` flags to the C# part of UBT.

diff --git a/Engine/Source/Programs/UnrealBuildTool/System/PluginDescriptor.cs b/Engine/Source/Programs/UnrealBuildTool/System/PluginDescriptor.cs
--- a/Engine/Source/Programs/UnrealBuildTool/System/PluginDescriptor.cs
+++ b/Engine/Source/Programs/UnrealBuildTool/System/PluginDescriptor.cs
@@ -159,6 +159,16 @@ namespace UnrealBuildTool
 		/// </summary>
 		public uint? VerseVersion;
 
+		/// <summary>
+		/// Whether this plugin has SceneGraph enabled, which impacts the generated Verse Asset Digest.
+		/// </summary>
+		public bool? bEnableSceneGraph = false;
+
+		/// <summary>
+		/// If to generate Verse source code definitions from assets contained in this plugin
+		/// </summary>
+		public bool? bEnableVerseAssetReflection = false;
+
 		/// <summary>
 		/// Whether this plugin should be enabled by default for all projects
 		/// </summary>
@@ -337,6 +347,18 @@ namespace UnrealBuildTool
 				VerseVersion = PluginVerseVersion;
 			}
 
+			bool bPluginEnableSceneGraph;
+			if (RawObject.TryGetBoolField("bEnableSceneGraph", out bPluginEnableSceneGraph))
+			{
+				bEnableSceneGraph = bPluginEnableSceneGraph;
+			}
+
+			bool bPluginEnableVerseAssetReflection;
+			if (RawObject.TryGetBoolField("bEnableVerseAssetReflection", out bPluginEnableVerseAssetReflection))
+			{
+				bEnableVerseAssetReflection = bPluginEnableVerseAssetReflection;
+			}
+
 			bool bEnabledByDefaultValue;
 			if (RawObject.TryGetBoolField("EnabledByDefault", out bEnabledByDefaultValue))
 			{
