Adding event matcher to silent ld_classic warnings

#rnx

diff --git a/Engine/Source/Programs/UnrealBuildTool/Matchers/CompileEventMatcher.cs b/Engine/Source/Programs/UnrealBuildTool/Matchers/CompileEventMatcher.cs
--- a/Engine/Source/Programs/UnrealBuildTool/Matchers/CompileEventMatcher.cs
+++ b/Engine/Source/Programs/UnrealBuildTool/Matchers/CompileEventMatcher.cs
@@ -66,6 +66,7 @@ namespace UnrealBuildTool.Matchers
 		static readonly Regex s_scriptCompilePattern = new Regex(@"^\s*[A-Za-z0-9_\.]+ ERROR:.* [A-Za-z_]+ failed to compile\.");
 		static readonly Regex s_cscSummaryPattern = new Regex(@"^\s+\d+ (?:Warning|Error)\(s\)");
 		static readonly Regex s_cscOutputPattern = new Regex(@"^  [^ ]+ -> ");
+		static readonly Regex s_xcodeClassicPattern = new Regex(@"ld_classic is deprecated and will be removed in a future release");
 
 		static readonly string[] s_invalidExtensions =
 		{
@@ -209,6 +210,13 @@ namespace UnrealBuildTool.Matchers
 				return true;
 			}
 
+			// Temporary XCode silencing when minOS is too low
+			if (builder.Current.IsMatch(s_xcodeClassicPattern))
+			{
+				outEvent = builder.ToMatch(LogEventPriority.Highest, LogLevel.Information, KnownLogEvents.Systemic_XCode);
+				return true;
+			}
+
 			outEvent = null;
 			return false;
 		}
