PR #12521: Add System.Text.Json reference so module rules can read JSON files

diff --git a/Engine/Source/Programs/UnrealBuildTool/System/DynamicCompilation.cs b/Engine/Source/Programs/UnrealBuildTool/System/DynamicCompilation.cs
--- a/Engine/Source/Programs/UnrealBuildTool/System/DynamicCompilation.cs
+++ b/Engine/Source/Programs/UnrealBuildTool/System/DynamicCompilation.cs
@@ -230,6 +230,7 @@ namespace UnrealBuildTool
 				"System.IO",
 				"System.IO.FileSystem",
 				"System.Linq",
+				"System.Text.Json",
 				"System.Threading.Tasks",
 				"System.Threading.Tasks.Parallel",
 				"System.Collections.Concurrent",
