UnrealBuildTool: Ensure StaticAnalzyer.Default uses the correct toolchain in all cases

#rnx

diff --git a/Engine/Source/Programs/UnrealBuildTool/Platform/Windows/UEBuildWindows.cs b/Engine/Source/Programs/UnrealBuildTool/Platform/Windows/UEBuildWindows.cs
--- a/Engine/Source/Programs/UnrealBuildTool/Platform/Windows/UEBuildWindows.cs
+++ b/Engine/Source/Programs/UnrealBuildTool/Platform/Windows/UEBuildWindows.cs
@@ -1262,19 +1262,32 @@ namespace UnrealBuildTool
 				}
 			}
 
-			// Disable PCHs for PVS studio analyzer.
-			if (Target.StaticAnalyzer == StaticAnalyzer.PVSStudio)
+			// Set the compiler version if necessary
+			if (Target.WindowsPlatform.Compiler == WindowsCompiler.Default)
 			{
-				Target.bUsePCHFiles = false;
+				Target.WindowsPlatform.Compiler = GetDefaultCompiler(Target, Target.WindowsPlatform.Architecture, Logger, !Platform.IsInGroup(UnrealPlatformGroup.Microsoft));
+			}
 
-				// PVSStudio static analysis requires MSVC
-				Target.WindowsPlatform.Compiler = WindowsCompiler.VisualStudio2022;
+			// Set the toolchain version if necessary
+			if (Target.WindowsPlatform.ToolChain == WindowsCompiler.Default)
+			{
+				Target.WindowsPlatform.ToolChain = GetDefaultToolchain(Target, Target.WindowsPlatform.Architecture, Logger, !Platform.IsInGroup(UnrealPlatformGroup.Microsoft));
 			}
 
-			// If clang is selected for static analysis, switch compiler to clang and proceed
-			// as normal.
-			if (Target.StaticAnalyzer == StaticAnalyzer.Clang)
+			if (Target.StaticAnalyzer == StaticAnalyzer.VisualCpp
+				|| (Target.StaticAnalyzer == StaticAnalyzer.Default && Target.WindowsPlatform.Compiler.IsMSVC()))
 			{
+				// VisualCpp static analysis requires MSVC
+				if (!Target.WindowsPlatform.Compiler.IsMSVC())
+				{
+					Target.WindowsPlatform.Compiler = WindowsCompiler.VisualStudio2022;
+				}
+				Target.StaticAnalyzer = StaticAnalyzer.Default;
+			}
+			else if (Target.StaticAnalyzer == StaticAnalyzer.Clang
+				|| (Target.StaticAnalyzer == StaticAnalyzer.Default && Target.WindowsPlatform.Compiler.IsClang()))
+			{
+				// Clang static analysis requires a clang toolchain
 				if (!Target.WindowsPlatform.Compiler.IsClang())
 				{
 					Target.WindowsPlatform.Compiler = WindowsCompiler.Clang;
@@ -1288,22 +1301,13 @@ namespace UnrealBuildTool
 					Target.bAlwaysUseUnityForGeneratedFiles = false;
 				}
 			}
-			else if (Target.StaticAnalyzer != StaticAnalyzer.None &&
-					 Target.StaticAnalyzerOutputType != StaticAnalyzerOutputType.Text)
-			{
-				Logger.LogInformation("Defaulting static analyzer output type to text");
-			}
-
-			// Set the compiler version if necessary
-			if (Target.WindowsPlatform.Compiler == WindowsCompiler.Default)
+			else if (Target.StaticAnalyzer == StaticAnalyzer.PVSStudio)
 			{
-				Target.WindowsPlatform.Compiler = GetDefaultCompiler(Target, Target.WindowsPlatform.Architecture, Logger, !Platform.IsInGroup(UnrealPlatformGroup.Microsoft));
-			}
+				// Disable PCHs for PVS studio analyzer.
+				Target.bUsePCHFiles = false;
 
-			// Set the toolchain version if necessary
-			if (Target.WindowsPlatform.ToolChain == WindowsCompiler.Default)
-			{
-				Target.WindowsPlatform.ToolChain = GetDefaultToolchain(Target, Target.WindowsPlatform.Architecture, Logger, !Platform.IsInGroup(UnrealPlatformGroup.Microsoft));
+				// PVSStudio static analysis requires MSVC
+				Target.WindowsPlatform.Compiler = WindowsCompiler.VisualStudio2022;
 			}
 
 			// Disable linking and ignore build outputs if we're using a static analyzer
