[UBA]
* Moved horde initialization code from UbaExecutor to horde file
* Added
* Dispose cache client if failing to connect
* Added AllowDetour config to be able to run with uba executor but with uba detouring entirely disabled
* Added so pch actions are set to normal priority instead of below normal in order to speed up compiles
* Added support for running extra horde coordinator (for people running horde at home but still wants to use company horde. Use <Horde2> in buildconfiguration. Very limited config
* Added so helper sessions are clickable in visualizer if using horde and takes  you to the horde lease
* Improved visualizer description of horde coordinator status

diff --git a/Engine/Source/Programs/UnrealBuildTool/Executors/UnrealBuildAccelerator/UBAAgentCoordinatorHorde.cs b/Engine/Source/Programs/UnrealBuildTool/Executors/UnrealBuildAccelerator/UBAAgentCoordinatorHorde.cs
--- a/Engine/Source/Programs/UnrealBuildTool/Executors/UnrealBuildAccelerator/UBAAgentCoordinatorHorde.cs
+++ b/Engine/Source/Programs/UnrealBuildTool/Executors/UnrealBuildAccelerator/UBAAgentCoordinatorHorde.cs
@@ -331,6 +331,12 @@ namespace UnrealBuildTool
 					}
 				}
 
+				string desc = instanceType;
+				if (!String.IsNullOrEmpty(leaseId))
+				{
+					string link = $" {hordeConfig.HordeServer!.TrimEnd('/')}/lease/{leaseId}";
+					desc += link;
+				}
 				// When using relay connection mode, the IP will be relay server's IP
 				string ip = String.IsNullOrEmpty(lease.Ip) ? computeIp : lease.Ip;
 
@@ -354,7 +360,7 @@ namespace UnrealBuildTool
 					Port = ubaPort,
 					ProxyPort = ubaProxyPort,
 				};
-				worker.BackgroundTask = RunWorkerAsync(worker, lease, locator, exeName, workerLogger, hordeConfig, _cancellationTokenSource.Token, agentName, leaseId, instanceType);
+				worker.BackgroundTask = RunWorkerAsync(worker, lease, locator, exeName, workerLogger, hordeConfig, _cancellationTokenSource.Token, agentName, desc);
 				lock (_workers)
 				{
 					_workers.Add(worker);
@@ -419,18 +425,18 @@ namespace UnrealBuildTool
 			}
 		}
 
-		public static async Task<UBAHordeSession?> TryCreateHordeSessionAsync(UnrealBuildAcceleratorHordeConfig hordeConfig, UBAExecutor executor, bool bStrictErrors, ILogger logger, CancellationToken cancellationToken = default)
+		public static async Task<(UBAHordeSession?, string)> TryCreateHordeSessionAsync(UnrealBuildAcceleratorHordeConfig hordeConfig, UBAExecutor executor, bool bStrictErrors, ILogger logger, CancellationToken cancellationToken = default)
 		{
 			if (hordeConfig.bDisableHorde)
 			{
 				logger.LogInformation("Horde disabled via command line option.");
-				return null;
+				return (null, "Disabled via command line");
 			}
 
 			if (String.IsNullOrEmpty(hordeConfig.HordeServer) && HordeOptions.GetServerUrlFromEnvironment() == null && HordeOptions.GetDefaultServerUrl() == null)
 			{
 				logger.LogInformation("Horde disabled. Url not set.");
-				return null;
+				return (null, "No server url set");
 			}
 
 			Uri? server = (hordeConfig.HordeServer == null) ? null : new Uri(hordeConfig.HordeServer);
@@ -453,16 +459,16 @@ namespace UnrealBuildTool
 
 				UBAHordeSession session = new(executor, server, token, hordeConfig.HordePool, allowWine, hordeConfig.HordeMaxCores, bStrictErrors, connectionMode, encryption, logger);
 				await session.InitAsync(useSentry: !String.IsNullOrEmpty(hordeConfig.UBASentryUrl), cancellationToken);
-				return session;
+				return (session, "");
 			}
 			catch (TaskCanceledException)
 			{
-				return null;
+				return (null, "Cancelled");
 			}
 			catch (Exception ex)
 			{
 				logger.Log(bStrictErrors ? LogLevel.Error : LogLevel.Information, ex, "Unable to create Horde session: {Message}", ex.Message);
-				return null;
+				return (null, ex.Message);
 			}
 		}
 
@@ -493,7 +499,7 @@ namespace UnrealBuildTool
 			return result.AccessToken;
 		}
 
-		async Task RunWorkerAsync(Worker self, IComputeLease lease, BlobLocator tool, string executable, ILogger logger, UnrealBuildAcceleratorHordeConfig hordeConfig, CancellationToken cancellationToken, string agentName, string leaseId, string instanceType)
+		async Task RunWorkerAsync(Worker self, IComputeLease lease, BlobLocator tool, string executable, ILogger logger, UnrealBuildAcceleratorHordeConfig hordeConfig, CancellationToken cancellationToken, string agentName, string desc)
 		{
 			logger.LogDebug("Running worker task..");
 			try
@@ -578,9 +584,12 @@ namespace UnrealBuildTool
 						arguments.Add("-MaxCon=8");
 						arguments.Add("-MaxWorkers=64");
 						arguments.Add("-MaxIdle=15");
+						//arguments.Add("-UseCrawler");
 						//arguments.Add("-ProxyUseLocalStorage"); // If set then agent's cas storage will be used by proxy
 						//arguments.Add("-NoStore"); // This means that no cas files will be written to disk.. so no state is stored between runs
-						//arguments.Add($"-Description=\"{instanceType}, {leaseId}\"");
+
+						if (!String.IsNullOrEmpty(desc))
+							arguments.Add($"-Description=\"{desc}\"");
 
 						if (_owner.UBAConfig.bLogEnabled)
 						{
@@ -697,35 +706,84 @@ namespace UnrealBuildTool
 
 	class UBAAgentCoordinatorHorde : IUBAAgentCoordinator, IDisposable
 	{
-		public string? Server => HordeConfig.HordeServer;
-		public bool Enabled => !HordeConfig.bDisableHorde;
-		public string ConnectionModeString => (Enum.TryParse(HordeConfig.HordeConnectionMode, true, out ConnectionMode cm) ? cm : ConnectionMode.Direct).ToString();
+		public string? Server => _hordeConfig.HordeServer;
+		public bool Enabled => !_hordeConfig.bDisableHorde;
+		public string ConnectionModeString => (Enum.TryParse(_hordeConfig.HordeConnectionMode, true, out ConnectionMode cm) ? cm : ConnectionMode.Direct).ToString();
+
+		public static List<IUBAAgentCoordinator> Init(ILogger logger, UBAExecutor executor, IEnumerable<TargetDescriptor> targetDescriptors, UnrealBuildAcceleratorConfig ubaConfig, ref string remoteConnectionMode)
+		{
+			List<IUBAAgentCoordinator> outCoordinators = new();
+			List<UBAAgentCoordinatorHorde> hordeAgentCoordinators = new();
+			foreach (TargetDescriptor targetDescriptor in targetDescriptors)
+			{
+				if (targetDescriptor.HotReloadMode != HotReloadMode.Disabled)
+				{
+					ubaConfig.bStoreObjFilesCompressed = false;
+					ubaConfig.CacheServer = null;
+				}
+				targetDescriptor.AdditionalArguments.ApplyTo(executor);
+				targetDescriptor.AdditionalArguments.ApplyTo(ubaConfig);
+				hordeAgentCoordinators.Add(new UBAAgentCoordinatorHorde(logger, ubaConfig, targetDescriptor.AdditionalArguments, targetDescriptor.ProjectFile?.Directory));
+			}
+			hordeAgentCoordinators.RemoveAll(x => !x.Enabled);
+			remoteConnectionMode = hordeAgentCoordinators.FirstOrDefault()?.ConnectionModeString ?? "Local";
+			outCoordinators.AddRange(hordeAgentCoordinators.DistinctBy(x => x.Server));
+
+			// Some users want their extra local horde. This is the very simplified version with much fewer settings
+			{
+				UnrealBuildAcceleratorHordeConfig2 config2 = new();
+				XmlConfig.ApplyTo(config2);
+				if (!String.IsNullOrEmpty(config2.HordeServer))
+				{
+					UnrealBuildAcceleratorHordeConfig extraHordeConfig = new();
+					extraHordeConfig.HordeServer = config2.HordeServer;
+					extraHordeConfig.bDisableHorde = config2.bDisableHorde;
+					UBAAgentCoordinatorHorde extraHorde = new(logger, ubaConfig, extraHordeConfig);
+					if (extraHorde.Enabled)
+					{
+						outCoordinators.Add(extraHorde);
+					}
+				}
+			}
+			return outCoordinators;
+		}
+
+		public UBAAgentCoordinatorHorde(ILogger logger, UnrealBuildAcceleratorConfig ubaConfig, UnrealBuildAcceleratorHordeConfig hordeConfig)
+		{
+			_logger = logger;
+			_title = "Horde2";
+			_ubaConfig = ubaConfig;
+			_hordeConfig = hordeConfig;
+		}
 
 		public UBAAgentCoordinatorHorde(ILogger logger, UnrealBuildAcceleratorConfig ubaConfig, CommandLineArguments? additionalArguments = null, DirectoryReference? projectDir = null)
 		{
 			_logger = logger;
+			_title = "Horde";
 			_ubaConfig = ubaConfig;
 
-			ConfigCache.ReadSettings(projectDir, BuildHostPlatform.Current.Platform, HordeConfig);
-			if (HordeConfig.HordeEnabled?.Equals("False", StringComparison.OrdinalIgnoreCase) == true || (HordeConfig.HordeEnabled?.Equals("BuildMachineOnly", StringComparison.OrdinalIgnoreCase) == true && !Unreal.IsBuildMachine()))
+			_hordeConfig = new();
+
+			ConfigCache.ReadSettings(projectDir, BuildHostPlatform.Current.Platform, _hordeConfig);
+			if (_hordeConfig.HordeEnabled?.Equals("False", StringComparison.OrdinalIgnoreCase) == true || (_hordeConfig.HordeEnabled?.Equals("BuildMachineOnly", StringComparison.OrdinalIgnoreCase) == true && !Unreal.IsBuildMachine()))
 			{
-				HordeConfig.bDisableHorde = true;
+				_hordeConfig.bDisableHorde = true;
 			}
 			if (Unreal.IsBuildMachine())
 			{
 				string? hordeUrl = System.Environment.GetEnvironmentVariable("UE_HORDE_URL");
 				if (!String.IsNullOrEmpty(hordeUrl))
 				{
-					HordeConfig.HordeServer = hordeUrl;
+					_hordeConfig.HordeServer = hordeUrl;
 				}
 			}	
-			XmlConfig.ApplyTo(HordeConfig);
-			additionalArguments?.ApplyTo(HordeConfig);
+			XmlConfig.ApplyTo(_hordeConfig);
+			additionalArguments?.ApplyTo(_hordeConfig);
 
 			// Sentry is currently unsupported for non-Windows and non-x64
 			if (!OperatingSystem.IsWindows() || RuntimeInformation.ProcessArchitecture != Architecture.X64)
 			{
-				HordeConfig.UBASentryUrl = null;
+				_hordeConfig.UBASentryUrl = null;
 			}
 		}
 
@@ -748,9 +806,9 @@ namespace UnrealBuildTool
 
 			_executor = executor;
 			_cancellationSource = new CancellationTokenSource();
-			_hordeSessionTask = UBAHordeSession.TryCreateHordeSessionAsync(HordeConfig, executor, _ubaConfig.bStrict, _logger, _cancellationSource.Token);
+			_hordeSessionTask = UBAHordeSession.TryCreateHordeSessionAsync(_hordeConfig, executor, _ubaConfig.bStrict, _logger, _cancellationSource.Token);
 			await _hordeSessionTask;
-			executor.AgentCoordinatorInitialized(this, _hordeSessionTask.IsCompletedSuccessfully && _hordeSessionTask.Result != null);
+			executor.AgentCoordinatorInitialized(this, _hordeSessionTask.IsCompletedSuccessfully && _hordeSessionTask.Result.Item1 != null);
 		}
 
 		public void Start(ImmediateActionQueue queue, Func<LinkedAction, bool> canRunRemotely, StatusUpdateAction updateStatus)
@@ -760,22 +818,27 @@ namespace UnrealBuildTool
 			int timerPeriod = 5000;
 			bool shownNoAgentsFoundMessage = false;
 
-			UBAHordeSession? hordeSession = _hordeSessionTask?.Result;
+			UBAHordeSession? hordeSession = null;
+			string reason = "Unknown";
+			if (_hordeSessionTask != null)
+			{
+				(hordeSession, reason) = _hordeSessionTask.Result;
+			}
 
 			string? link = null;
-			if (HordeConfig.HordeServer != null)
+			if (_hordeConfig.HordeServer != null)
 			{
-				link = $"{HordeConfig.HordeServer!.TrimEnd('/')}/agents";
-				if (HordeConfig.HordePool != string.Empty)
+				link = $"{_hordeConfig.HordeServer!.TrimEnd('/')}/agents";
+				if (_hordeConfig.HordePool != string.Empty)
 				{
-					link += $"?agent={HordeConfig.HordePool}";
+					link += $"?agent={_hordeConfig.HordePool}";
 				}
 			}
-			updateStatus(0, 1, "Horde", EpicGames.UBA.LogEntryType.Info, link);
+			updateStatus(0, 1, _title, EpicGames.UBA.LogEntryType.Info, link);
 
 			if (hordeSession == null)
 			{
-				updateStatus(0, 6, "Not created (check log)", EpicGames.UBA.LogEntryType.Info, null);
+				updateStatus(0, 6, reason, EpicGames.UBA.LogEntryType.Info, null);
 				return;
 			}
 
@@ -793,7 +856,7 @@ namespace UnrealBuildTool
 					return;
 				}
 
-				UBAHordeSession? hordeSession = await _hordeSessionTask;
+				var (hordeSession, reason) = await _hordeSessionTask;
 
 				if (hordeSession == null)
 				{
@@ -831,7 +894,7 @@ namespace UnrealBuildTool
 
 						queueWeight -= queuedCores;
 
-						bool isSatisfied = queueWeight <= queueThreshold || (activeCores + queuedCores) >= HordeConfig.HordeMaxCores || _cancellationSource!.IsCancellationRequested;
+						bool isSatisfied = queueWeight <= queueThreshold || (activeCores + queuedCores) >= _hordeConfig.HordeMaxCores || _cancellationSource!.IsCancellationRequested;
 
 						hordeSession.UpdateHordeStatus(isSatisfied ? " - Requests paused" : $" - Requesting agent{("....".Substring(0, (requestCounter++ % 4)))}");
 
@@ -840,8 +903,8 @@ namespace UnrealBuildTool
 							break;
 						}
 
-						Requirements requirements = GetRequirements(HordeConfig);
-						if (!await hordeSession.AddWorkerAsync(requirements, HordeConfig, _cancellationSource.Token, activeCores))
+						Requirements requirements = GetRequirements(_hordeConfig);
+						if (!await hordeSession.AddWorkerAsync(requirements, _hordeConfig, _cancellationSource.Token, activeCores))
 						{
 							await Task.Delay(500); // Sleep a little bit to make sure previous horde status is visible
 							string status = (queuedAgents + activeAgents) == 0 ? " - No agents available" : " - No additional agents available";
@@ -869,7 +932,7 @@ namespace UnrealBuildTool
 				}
 
 				_timer?.Change(timerPeriod, Timeout.Infinite);
-			}, null, HordeConfig.HordeDelay * 1000, timerPeriod);
+			}, null, _hordeConfig.HordeDelay * 1000, timerPeriod);
 		}
 
 		public void Stop()
@@ -886,7 +949,7 @@ namespace UnrealBuildTool
 				return;
 			}
 
-			UBAHordeSession? hordeSession = await _hordeSessionTask;
+			var (hordeSession, reason) = await _hordeSessionTask;
 			_hordeSessionTask = null;
 			if (hordeSession != null)
 			{
@@ -965,12 +1028,13 @@ namespace UnrealBuildTool
 		}
 
 		readonly ILogger _logger;
+		readonly string _title;
 		readonly UnrealBuildAcceleratorConfig _ubaConfig;
-		UnrealBuildAcceleratorHordeConfig HordeConfig { get; init; } = new();
+		readonly UnrealBuildAcceleratorHordeConfig _hordeConfig;
 
 		UBAExecutor? _executor;
 		CancellationTokenSource? _cancellationSource;
-		Task<UBAHordeSession?>? _hordeSessionTask;
+		Task<(UBAHordeSession?, string)>? _hordeSessionTask;
 		StatusUpdateAction? _updateStatus;
 		Timer? _timer;
 	}
diff --git a/Engine/Source/Programs/UnrealBuildTool/Executors/UnrealBuildAccelerator/UBAExecutor.cs b/Engine/Source/Programs/UnrealBuildTool/Executors/UnrealBuildAccelerator/UBAExecutor.cs
--- a/Engine/Source/Programs/UnrealBuildTool/Executors/UnrealBuildAccelerator/UBAExecutor.cs
+++ b/Engine/Source/Programs/UnrealBuildTool/Executors/UnrealBuildAccelerator/UBAExecutor.cs
@@ -124,22 +124,7 @@ namespace UnrealBuildTool
 			XmlConfig.ApplyTo(UBAConfig);
 			CommandLine.ParseArguments(Environment.GetCommandLineArgs(), this, logger);
 
-			List<UBAAgentCoordinatorHorde> hordeAgentCoordinators = new();
-			foreach (TargetDescriptor targetDescriptor in targetDescriptors)
-			{
-				if (targetDescriptor.HotReloadMode != HotReloadMode.Disabled)
-				{
-					UBAConfig.bStoreObjFilesCompressed = false;
-					UBAConfig.CacheServer = null;
-				}
-				targetDescriptor.AdditionalArguments.ApplyTo(this);
-				targetDescriptor.AdditionalArguments.ApplyTo(UBAConfig);
-				hordeAgentCoordinators.Add(new UBAAgentCoordinatorHorde(logger, UBAConfig, targetDescriptor.AdditionalArguments, targetDescriptor.ProjectFile?.Directory));
-			}
-			hordeAgentCoordinators.RemoveAll(x => !x.Enabled);
-			_remoteConnectionMode = hordeAgentCoordinators.FirstOrDefault()?.ConnectionModeString ?? "Local";
-			_agentCoordinators.AddRange(hordeAgentCoordinators.DistinctBy(x => x.Server));
-
+			_agentCoordinators.AddRange(UBAAgentCoordinatorHorde.Init(logger, this, targetDescriptors, UBAConfig, ref _remoteConnectionMode));
 			_threadedLogger = new ThreadedLogger(logger);
 		}
 
@@ -452,6 +437,8 @@ namespace UnrealBuildTool
 								}
 								else
 								{
+									_cacheClient.Dispose();
+									_cacheClient = null;
 									_session.UpdateStatus(1, 6, "Not connected", LogEntryType.Info);
 								}
 							}
@@ -644,8 +631,13 @@ namespace UnrealBuildTool
 		/// </summary>
 		/// <param name="action">The action to check</param>
 		/// <returns>If this action must be local, non-detoured</returns>
-		static bool ForceLocalNoDetour(LinkedAction action)
+		bool ForceLocalNoDetour(LinkedAction action)
 		{
+			if (!UBAConfig.bAllowDetour)
+			{
+				return true;
+			}
+
 			if (!OperatingSystem.IsMacOS()) // Below code is slow, so early out
 			{
 				return false;
@@ -679,12 +671,22 @@ namespace UnrealBuildTool
 			if (!String.IsNullOrEmpty(action.CommandDescription))
 				description = $"{action.StatusDescription} ({action.CommandDescription})";
 
+			// Let pch compile at higher priority
+			System.Diagnostics.ProcessPriorityClass priority = ProcessPriority;
+			if (ProcessPriority == System.Diagnostics.ProcessPriorityClass.BelowNormal)
+			{
+				if (action.StatusDescription.Contains("PCH."))
+				{
+					priority = System.Diagnostics.ProcessPriorityClass.Normal;
+				}
+			}
+				
 			ProcessStartInfo startInfo = new()
 			{
 				Application = action.CommandPath.FullName,
 				WorkingDirectory = action.WorkingDirectory.FullName,
 				Arguments = action.CommandArguments,
-				Priority = ProcessPriority,
+				Priority = priority,
 				UserData = action,
 				Description = description,
 				Configuration = action.bIsGCCCompiler ? EpicGames.UBA.ProcessStartInfo.CommonProcessConfigs.CompileClang : EpicGames.UBA.ProcessStartInfo.CommonProcessConfigs.CompileMsvc,
@@ -705,6 +707,8 @@ namespace UnrealBuildTool
 
 		string GetCacheClientHint()
 		{
+			// This is extra information sent to cache server when fetching stuff. This way it will be possible to backtrack from cache server log to who connected
+			// These environment variables are set when UBT is running on a horde machine.
 			string? jobId = Environment.GetEnvironmentVariable("UE_HORDE_JOBID");
 			if (jobId == null)
 			{
@@ -830,10 +834,18 @@ namespace UnrealBuildTool
 
 				if (action.DependencyListFile != null)
 				{
-					CppDependencyCache.DependencyInfo info = CppDependencyCache.ReadDependencyInfo(action.DependencyListFile);
-					foreach (FileItem f in info.Files)
+					try
+					{
+						CppDependencyCache.DependencyInfo info = CppDependencyCache.ReadDependencyInfo(action.DependencyListFile);
+						foreach (FileItem f in info.Files)
+						{
+							writer.Write(f.FullName);
+						}
+					}
+					catch (Exception)
 					{
-						writer.Write(f.FullName);
+						_threadedLogger.LogInformation("{Description} {StatusDescription}: Failed to read dependency file {DependencyFile} for WriteCache", action.CommandDescription, action.StatusDescription, action.DependencyListFile);
+						return false;
 					}
 				}
 			}
diff --git a/Engine/Source/Programs/UnrealBuildTool/Executors/UnrealBuildAccelerator/UnrealBuildAcceleratorConfig.cs b/Engine/Source/Programs/UnrealBuildTool/Executors/UnrealBuildAccelerator/UnrealBuildAcceleratorConfig.cs
--- a/Engine/Source/Programs/UnrealBuildTool/Executors/UnrealBuildAccelerator/UnrealBuildAcceleratorConfig.cs
+++ b/Engine/Source/Programs/UnrealBuildTool/Executors/UnrealBuildAccelerator/UnrealBuildAcceleratorConfig.cs
@@ -236,6 +236,13 @@ namespace UnrealBuildTool
 		[CommandLine("-UBAWriteCache")]
 		public bool bWriteCache { get; set; }
 
+		/// <summary>
+		/// Disable all detouring making all actions run outside uba
+		/// </summary>
+		[XmlConfigFile(Category = "UnrealBuildAccelerator", Name = "AllowDetour")]
+		[CommandLine("-UBANoDetour", Value = "false")]
+		public bool bAllowDetour { get; set; } = true;
+
 		/// <summary>
 		/// Max number of cache download tasks that can execute in parallel
 		/// </summary>
diff --git a/Engine/Source/Programs/UnrealBuildTool/Executors/UnrealBuildAccelerator/UnrealBuildAcceleratorHordeConfig.cs b/Engine/Source/Programs/UnrealBuildTool/Executors/UnrealBuildAccelerator/UnrealBuildAcceleratorHordeConfig.cs
--- a/Engine/Source/Programs/UnrealBuildTool/Executors/UnrealBuildAccelerator/UnrealBuildAcceleratorHordeConfig.cs
+++ b/Engine/Source/Programs/UnrealBuildTool/Executors/UnrealBuildAccelerator/UnrealBuildAcceleratorHordeConfig.cs
@@ -107,17 +107,17 @@ namespace UnrealBuildTool
 		[XmlConfigFile(Category = "Horde", Name = "Requirements")]
 		[CommandLine("-UBAHordeRequirements=")]
 		public string? HordeCondition { get; set; }
-		
+
 		/// <summary>
 		/// Default compute cluster ID if nothing is set
 		/// </summary>
 		public const string ClusterDefault = "default";
-		
+
 		/// <summary>
 		/// Compute cluster ID for automatically resolving the most suitable cluster
 		/// </summary>
 		public const string ClusterAuto = "_auto";
-		
+
 		/// <summary>
 		/// Compute cluster ID to use in Horde. Set to "_auto" to let Horde server resolve a suitable cluster.
 		/// In multi-region setups this is can simplify configuration of UBT/UBA a lot.
@@ -197,6 +197,20 @@ namespace UnrealBuildTool
 		/// </summary>
 		[ConfigFile(ConfigHierarchyType.Engine, "Horde", "UbaEnabled")]
 		public string? HordeEnabled { get; set; } = null;
+	}
 
+	class UnrealBuildAcceleratorHordeConfig2
+	{
+		/// <summary>
+		/// Uri of the Horde server
+		/// </summary>
+		[ConfigFile(ConfigHierarchyType.Engine, "Horde2", "ServerUrl")]
+		[XmlConfigFile(Category = "Horde2", Name = "Server")]
+		[CommandLine("-UBAHorde2=")]
+		public string? HordeServer { get; set; }
+
+		[XmlConfigFile(Category = "UnrealBuildAccelerator")]
+		[CommandLine("-UBADisableHorde")]
+		public bool bDisableHorde { get; set; } = false;
 	}
 }
\ No newline at end of file
