[UBA]
* Added comment to property field

diff --git a/Engine/Source/Programs/UnrealBuildTool/Executors/UnrealBuildAccelerator/UnrealBuildAcceleratorHordeConfig.cs b/Engine/Source/Programs/UnrealBuildTool/Executors/UnrealBuildAccelerator/UnrealBuildAcceleratorHordeConfig.cs
--- a/Engine/Source/Programs/UnrealBuildTool/Executors/UnrealBuildAccelerator/UnrealBuildAcceleratorHordeConfig.cs
+++ b/Engine/Source/Programs/UnrealBuildTool/Executors/UnrealBuildAccelerator/UnrealBuildAcceleratorHordeConfig.cs
@@ -209,6 +209,9 @@ namespace UnrealBuildTool
 		[CommandLine("-UBAHorde2=")]
 		public string? HordeServer { get; set; }
 
+		/// <summary>
+		/// Disable horde all together
+		/// </summary>
 		[XmlConfigFile(Category = "UnrealBuildAccelerator")]
 		[CommandLine("-UBADisableHorde")]
 		public bool bDisableHorde { get; set; } = false;
