Update compile warnings for Clang to use CppCompileWarnings system

diff --git a/Engine/Source/Programs/Shared/UnrealBuildTool.Tests/CppCompileWarningsTests.cs b/Engine/Source/Programs/Shared/UnrealBuildTool.Tests/CppCompileWarningsTests.cs
--- a/Engine/Source/Programs/Shared/UnrealBuildTool.Tests/CppCompileWarningsTests.cs
+++ b/Engine/Source/Programs/Shared/UnrealBuildTool.Tests/CppCompileWarningsTests.cs
@@ -386,6 +386,218 @@ namespace UnrealBuildTool.Tests
 
 					baseWarnings.ShadowVariableWarningLevel = WarningLevel.Default;
 				}
+
+				// ToolChainVersionConstrainedWarningsClangToolChainAttribute
+				{
+					testTargetRules.WindowsPlatform.Compiler = WindowsCompiler.Clang;
+					baseWarnings.Shorten64To32WarningLevel = WarningLevel.Off;
+
+					VersionNumber newVersionNumber = new VersionNumber(18, 9, 9);
+					args = baseWarnings.GenerateWarningCommandLineArgs(newCompileEnvironment, typeof(ClangToolChain), newVersionNumber);
+					Assert.IsFalse(args.Any());
+
+					//
+					newVersionNumber = new VersionNumber(19, 0, 0);
+					args = baseWarnings.GenerateWarningCommandLineArgs(newCompileEnvironment, typeof(ClangToolChain), newVersionNumber);
+					Assert.IsTrue(args.Count() == 1);
+					Assert.IsTrue(args.Contains("-Wno-shorten-64-to-32"));
+					baseWarnings.Shorten64To32WarningLevel = WarningLevel.Default;
+				}
+
+				// ConfigurationConstrainedClangToolChainAttribute
+				{
+					testTargetRules.WindowsPlatform.Compiler = WindowsCompiler.Clang;
+					baseWarnings.UnusedValueWarningLevel = WarningLevel.Off;
+					VersionNumber newVersionNumber = new VersionNumber(18, 9, 9);
+
+					CppCompileEnvironment invalidCompileEnvironment = new CppCompileEnvironment(UnrealTargetPlatform.Win64, CppConfiguration.Debug, new UnrealArchitectures(UnrealArch.X64), null!);
+					CppCompileEnvironment validCompileEnvironment = new CppCompileEnvironment(UnrealTargetPlatform.Win64, CppConfiguration.Shipping, new UnrealArchitectures(UnrealArch.X64), null!);
+
+					args = baseWarnings.GenerateWarningCommandLineArgs(invalidCompileEnvironment, typeof(ClangToolChain), newVersionNumber);
+					Assert.IsFalse(args.Any());
+
+					args = baseWarnings.GenerateWarningCommandLineArgs(validCompileEnvironment, typeof(ClangToolChain), newVersionNumber);
+					Assert.IsTrue(args.Count() == 1);
+					Assert.IsTrue(args.Contains("-Wno-unused-value"));
+
+					baseWarnings.UnusedValueWarningLevel = WarningLevel.Default;
+				}
+
+				// CppStandardConstrainedClangToolChainAttribute
+				{
+					testTargetRules.WindowsPlatform.Compiler = WindowsCompiler.Clang;
+					baseWarnings.AmbiguousReversedOperatorWarningLevel = WarningLevel.Off;
+					VersionNumber newVersionNumber = new VersionNumber(18, 9, 9);
+
+					CppCompileEnvironment cppStandardCompileEnvironment = new CppCompileEnvironment(UnrealTargetPlatform.Win64, CppConfiguration.Debug, new UnrealArchitectures(UnrealArch.X64), null!);
+					cppStandardCompileEnvironment.CppStandard = CppStandardVersion.Cpp17;
+
+					args = baseWarnings.GenerateWarningCommandLineArgs(cppStandardCompileEnvironment, typeof(ClangToolChain), newVersionNumber);
+					Assert.IsFalse(args.Any());
+
+					cppStandardCompileEnvironment.CppStandard = CppStandardVersion.Cpp20;
+					args = baseWarnings.GenerateWarningCommandLineArgs(cppStandardCompileEnvironment, typeof(ClangToolChain), newVersionNumber);
+					Assert.IsTrue(args.Count() == 1);
+					Assert.IsTrue(args.Contains("-Wno-ambiguous-reversed-operator"));
+
+					cppStandardCompileEnvironment.CppStandard = CppStandardVersion.Latest;
+					args = baseWarnings.GenerateWarningCommandLineArgs(cppStandardCompileEnvironment, typeof(ClangToolChain), newVersionNumber);
+					Assert.IsTrue(args.Count() == 1);
+					Assert.IsTrue(args.Contains("-Wno-ambiguous-reversed-operator"));
+
+					baseWarnings.AmbiguousReversedOperatorWarningLevel = WarningLevel.Default;
+				}
+
+				// DeterministicClangToolChainAttribute
+				{
+					testTargetRules.WindowsPlatform.Compiler = WindowsCompiler.Clang;
+					baseWarnings.DeterministicWarningLevel = WarningLevel.Off;
+					VersionNumber newVersionNumber = new VersionNumber(18, 9, 9);
+
+					CppCompileEnvironment deterministicCompileEnvironment = new CppCompileEnvironment(UnrealTargetPlatform.Win64, CppConfiguration.Debug, new UnrealArchitectures(UnrealArch.X64), null!);
+					deterministicCompileEnvironment.bDeterministic = false;
+
+					args = baseWarnings.GenerateWarningCommandLineArgs(deterministicCompileEnvironment, typeof(ClangToolChain), newVersionNumber);
+					Assert.IsFalse(args.Any());
+
+					deterministicCompileEnvironment.bDeterministic = true;
+
+					args = baseWarnings.GenerateWarningCommandLineArgs(deterministicCompileEnvironment, typeof(ClangToolChain), newVersionNumber);
+					Assert.IsFalse(args.Any());
+
+					baseWarnings.DeterministicWarningLevel = WarningLevel.Warning;
+					args = baseWarnings.GenerateWarningCommandLineArgs(deterministicCompileEnvironment, typeof(ClangToolChain), newVersionNumber);
+					Assert.IsTrue(args.Count() == 2);
+					Assert.IsTrue(args.Contains("-Wdate-time"));
+					Assert.IsTrue(args.Contains("-Wno-error=date-time"));
+
+					baseWarnings.DeterministicWarningLevel = WarningLevel.Error;
+					args = baseWarnings.GenerateWarningCommandLineArgs(deterministicCompileEnvironment, typeof(ClangToolChain), newVersionNumber);
+					Assert.IsTrue(args.Count() == 1);
+					Assert.IsTrue(args.Contains("-Wdate-time"));
+
+					baseWarnings.DeterministicWarningLevel = WarningLevel.Default;
+				}
+
+				// PGOOptimizedClangToolChainAttribute
+				{
+					testTargetRules.WindowsPlatform.Compiler = WindowsCompiler.Clang;
+					baseWarnings.ProfileInstructWarningLevel = WarningLevel.Off;
+					VersionNumber newVersionNumber = new VersionNumber(18, 9, 9);
+
+					CppCompileEnvironment pgoOptimizedCompileEnvironment = new CppCompileEnvironment(UnrealTargetPlatform.Win64, CppConfiguration.Debug, new UnrealArchitectures(UnrealArch.X64), null!);
+					pgoOptimizedCompileEnvironment.bPGOOptimize = false;
+
+					args = baseWarnings.GenerateWarningCommandLineArgs(pgoOptimizedCompileEnvironment, typeof(ClangToolChain), newVersionNumber);
+					Assert.IsFalse(args.Any());
+
+					pgoOptimizedCompileEnvironment.bPGOOptimize = true;
+
+					args = baseWarnings.GenerateWarningCommandLineArgs(pgoOptimizedCompileEnvironment, typeof(ClangToolChain), newVersionNumber);
+
+					Assert.IsTrue(args.Count() == 2);
+					Assert.IsTrue(args.Contains("-Wno-profile-instr-out-of-date"));
+					Assert.IsTrue(args.Contains("-Wno-profile-instr-unprofiled"));
+
+					baseWarnings.ProfileInstructWarningLevel = WarningLevel.Off;
+				}
+
+				// AndroidNDKR28ToolChainVersionExclusionAttribute
+				{
+					testTargetRules.WindowsPlatform.Compiler = WindowsCompiler.Clang;
+					baseWarnings.CastFunctionTypeMistmatchWarningLevel = WarningLevel.Off;
+
+					VersionNumber newVersionNumber = new VersionNumber(19,0,0);
+					CppCompileEnvironment androidCompileEnvironment = new CppCompileEnvironment(UnrealTargetPlatform.Android, CppConfiguration.Development, new UnrealArchitectures(UnrealArch.X64), null!);
+					args = baseWarnings.GenerateWarningCommandLineArgs(androidCompileEnvironment, typeof(ClangToolChain), newVersionNumber);
+
+					Assert.IsFalse(args.Any());
+
+					args = baseWarnings.GenerateWarningCommandLineArgs(newCompileEnvironment, typeof(ClangToolChain), newVersionNumber);
+
+					Assert.IsTrue(args.Count() == 1);
+					Assert.IsTrue(args.Contains("-Wno-cast-function-type-mismatch"));
+
+					newVersionNumber = new VersionNumber(19, 0, 1);
+					args = baseWarnings.GenerateWarningCommandLineArgs(androidCompileEnvironment, typeof(ClangToolChain), newVersionNumber);
+
+					Assert.IsTrue(args.Count() == 1);
+					Assert.IsTrue(args.Contains("-Wno-cast-function-type-mismatch"));
+
+					baseWarnings.CastFunctionTypeMistmatchWarningLevel = WarningLevel.Default;
+				}
+
+				// AndroidNDKR26ToolChainVersionExclusionAttribute
+				{
+					testTargetRules.WindowsPlatform.Compiler = WindowsCompiler.Clang;
+					baseWarnings.InvalidUnevaluatedStringWarningLevel = WarningLevel.Off;
+					CppCompileEnvironment androidCompileEnvironment = new CppCompileEnvironment(UnrealTargetPlatform.Android, CppConfiguration.Development, new UnrealArchitectures(UnrealArch.X64), null!);
+
+					// Exclusion cases
+					{
+						VersionNumber newVersionNumber = new VersionNumber(17, 0, 2);
+						args = baseWarnings.GenerateWarningCommandLineArgs(androidCompileEnvironment, typeof(ClangToolChain), newVersionNumber);
+
+						Assert.IsTrue(args.Count() == 1 && args.Contains("-Wno-shadow"));
+
+						newVersionNumber = new VersionNumber(16);
+						args = baseWarnings.GenerateWarningCommandLineArgs(androidCompileEnvironment, typeof(ClangToolChain), newVersionNumber);
+
+						Assert.IsFalse(args.Any());
+
+						newVersionNumber = new VersionNumber(17, 1);
+						androidCompileEnvironment.CppStandard = CppStandardVersion.Latest;
+						args = baseWarnings.GenerateWarningCommandLineArgs(androidCompileEnvironment, typeof(ClangToolChain), newVersionNumber);
+
+						Assert.IsTrue(args.Count() == 1 && args.Contains("-Wno-shadow"));
+					}
+
+					// Inclusion cases
+					{
+						VersionNumber newVersionNumber = new VersionNumber(17, 1);
+						androidCompileEnvironment.CppStandard = CppStandardVersion.Cpp20;
+
+						args = baseWarnings.GenerateWarningCommandLineArgs(androidCompileEnvironment, typeof(ClangToolChain), newVersionNumber);
+
+						Assert.IsTrue(args.Contains("-Wno-invalid-unevaluated-string"));
+
+						newVersionNumber = new VersionNumber(17, 0, 2);
+
+						args = baseWarnings.GenerateWarningCommandLineArgs(newCompileEnvironment, typeof(ClangToolChain), newVersionNumber);
+
+						Assert.IsTrue(args.Contains("-Wno-invalid-unevaluated-string"));
+					}
+
+					baseWarnings.InvalidUnevaluatedStringWarningLevel = WarningLevel.Default;
+				}
+
+				// UnusedValueClangToolChainAttribute
+				{
+					testTargetRules.WindowsPlatform.Compiler = WindowsCompiler.Clang;
+					baseWarnings.UnusedValueWarningLevel = WarningLevel.Off;
+					VersionNumber newVersionNumber = new VersionNumber(18, 9, 9);
+
+					CppCompileEnvironment invalidCompileEnvironment = new CppCompileEnvironment(UnrealTargetPlatform.Win64, CppConfiguration.Debug, new UnrealArchitectures(UnrealArch.X64), null!);
+					CppCompileEnvironment validCompileEnvironment = new CppCompileEnvironment(UnrealTargetPlatform.Win64, CppConfiguration.Shipping, new UnrealArchitectures(UnrealArch.X64), null!);
+
+					args = baseWarnings.GenerateWarningCommandLineArgs(invalidCompileEnvironment, typeof(ClangToolChain), newVersionNumber, StaticAnalyzer.None);
+					Assert.IsFalse(args.Any());
+
+					args = baseWarnings.GenerateWarningCommandLineArgs(validCompileEnvironment, typeof(ClangToolChain), newVersionNumber, StaticAnalyzer.None);
+					Assert.IsTrue(args.Count() == 1 && args.Contains("-Wno-unused-value"));
+
+					args = baseWarnings.GenerateWarningCommandLineArgs(validCompileEnvironment, typeof(ClangToolChain), newVersionNumber, StaticAnalyzer.Clang);
+					Assert.IsTrue(args.Count() == 1 && args.Contains("-Wno-unused-value"));
+				}
+
+				// VCClangToolChainAttribute
+				{
+					testTargetRules.WindowsPlatform.Compiler = WindowsCompiler.Clang;
+					baseWarnings.MicrosoftGroupWarningLevel = WarningLevel.Off;
+
+					args = baseWarnings.GenerateWarningCommandLineArgs(newCompileEnvironment, typeof(VCToolChain));
+					Assert.IsTrue(args.Any());
+				}
 			}
 		}
 	}
diff --git a/Engine/Source/Programs/UnrealBuildTool/Configuration/CompileWarnings/ApplyWarningsAttribute.cs b/Engine/Source/Programs/UnrealBuildTool/Configuration/CompileWarnings/ApplyWarningsAttribute.cs
--- a/Engine/Source/Programs/UnrealBuildTool/Configuration/CompileWarnings/ApplyWarningsAttribute.cs
+++ b/Engine/Source/Programs/UnrealBuildTool/Configuration/CompileWarnings/ApplyWarningsAttribute.cs
@@ -102,7 +102,7 @@ namespace UnrealBuildTool
 #pragma warning restore CA1813
 	{
 		/// <inheritdoc/>
-		internal WarningsVCToolChainAttribute(string[]? offArgs, string[]? warningArgs, string[]? errorArgs) : base(offArgs, warningArgs, errorArgs) {}
+		internal WarningsVCToolChainAttribute(string[]? offArgs, string[]? warningArgs, string[]? errorArgs) : base(offArgs, warningArgs, errorArgs) { }
 
 		/// <inheritdoc/>
 		internal override bool CanApplyToContext(CompilerWarningsToolChainContext toolChainContext)
@@ -141,7 +141,25 @@ namespace UnrealBuildTool
 		internal override bool CanApplyToContext(CompilerWarningsToolChainContext toolChainContext)
 		{
 			ReadOnlyTargetRules? readOnlyTargetRules = toolChainContext._buildSystemContext.GetReadOnlyTargetRules();
-			return readOnlyTargetRules != null && readOnlyTargetRules.WindowsPlatform != null && readOnlyTargetRules.WindowsPlatform.Compiler.IsMSVC();
+
+			return base.CanApplyToContext(toolChainContext) && readOnlyTargetRules != null && readOnlyTargetRules.WindowsPlatform != null && readOnlyTargetRules.WindowsPlatform.Compiler.IsMSVC();
+		}
+	}
+
+	/// <summary>
+	/// Specialized <see cref="WarningsVCToolChainAttribute"/> for Clang compiler within the <see cref="VCToolChain"/>.
+	/// </summary>
+	internal sealed class WarningsVCClangAttribute : WarningsVCToolChainAttribute
+	{
+		/// <inheritdoc/>
+		internal WarningsVCClangAttribute(string[]? offArgs, string[]? warningArgs, string[]? errorArgs) : base(offArgs, warningArgs, errorArgs) { }
+
+		/// <inheritdoc/>
+		internal override bool CanApplyToContext(CompilerWarningsToolChainContext toolChainContext)
+		{
+			ReadOnlyTargetRules? readOnlyTargetRules = toolChainContext._buildSystemContext.GetReadOnlyTargetRules();
+
+			return base.CanApplyToContext(toolChainContext) && readOnlyTargetRules != null && readOnlyTargetRules.WindowsPlatform != null && readOnlyTargetRules.WindowsPlatform.Compiler.IsClang();
 		}
 	}
 
@@ -149,15 +167,225 @@ namespace UnrealBuildTool
 
 	#region -- ToolChain Specializations --
 
+	/// <summary>
+	/// Specialized <see cref="WarningsClangToolChainAttribute"/> to enforce toolchain version constraints.
+	/// </summary>
+#pragma warning disable CA1813
+	internal class ToolChainVersionConstrainedWarningsClangToolChainAttribute : WarningsClangToolChainAttribute
+#pragma warning restore CA1813
+	{
+		private readonly VersionNumber? _lowerBound;
+		private readonly VersionNumber? _upperBound;
+
+		/// <summary>
+		/// Constructor for ToolChainVersionConstrainedWarningsClangToolChainAttribute that enforces a clang toolchain version constraint.
+		/// </summary>
+		/// <param name="offArgs">The arguments to use when the property is defined as <see cref="WarningLevel.Off"/>.</param>
+		/// <param name="warningArgs">The arguments to use when the property is defined as <see cref="WarningLevel.Warning"/>.</param>
+		/// <param name="errorArgs">The arguments to use when the property is defined as <see cref="WarningLevel.Error"/>.</param>
+		/// <param name="lowerBoundVersionComponents">The lower bound of the <see cref="VersionNumber"/> to constrain this warning attribute to.</param>
+		/// <param name="upperBoundVersionComponents">The upper bound of the <see cref="VersionNumber"/> to constrain this warning attribute to.</param>
+		/// <remarks><see cref="VersionNumber"/> is not a viable parameter to this attribute, as it's a generic (via Comparable interface), and as such is not viable for attributes.</remarks>
+		internal ToolChainVersionConstrainedWarningsClangToolChainAttribute(string[]? offArgs, string[]? warningArgs, string[]? errorArgs, int[]? lowerBoundVersionComponents, int[]? upperBoundVersionComponents) : base(offArgs, warningArgs, errorArgs)
+		{
+			_lowerBound = lowerBoundVersionComponents == null ? null : new VersionNumber(lowerBoundVersionComponents);
+			_upperBound = upperBoundVersionComponents == null ? null : new VersionNumber(upperBoundVersionComponents);
+		}
+
+		/// <inheritdoc/>
+		internal override bool CanApplyToContext(CompilerWarningsToolChainContext toolChainContext)
+		{
+			bool baseCanApplyToContext = base.CanApplyToContext(toolChainContext);
+
+			return baseCanApplyToContext && IsWithinBounds(toolChainContext._toolChainVersion);
+		}
+
+		protected bool IsWithinBounds(VersionNumber? version)
+		{
+			// If our evaluation version is null, we can't possibly make a statement about whether it fits within a range
+			if (version == null)
+			{
+				return false;
+			}
+
+			bool isAboveLowerBound = _lowerBound == null || version >= _lowerBound;
+			bool isBelowUpperBound = _upperBound == null || version <= _upperBound;
+
+			return isAboveLowerBound && isBelowUpperBound;
+		}
+	}
+
+	/// <summary>
+	/// Specialized <see cref="WarningsClangToolChainAttribute"/> to enforce configuration constraints.
+	/// </summary>
+#pragma warning disable CA1813
+	internal class ConfigurationConstrainedClangToolChainAttribute : WarningsClangToolChainAttribute
+#pragma warning restore CA1813
+	{
+		private readonly CppConfiguration _cppConfiguration;
+
+		/// <summary>
+		/// Constructor for ConfigurationConstrainedClangToolChainAttribute that enforces specific configuration constraint.
+		/// </summary>
+		/// <param name="offArgs">The arguments to use when the property is defined as <see cref="WarningLevel.Off"/>.</param>
+		/// <param name="warningArgs">The arguments to use when the property is defined as <see cref="WarningLevel.Warning"/>.</param>
+		/// <param name="errorArgs">The arguments to use when the property is defined as <see cref="WarningLevel.Error"/>.</param>
+		/// <param name="enabledConfiguration">The configuration in which to apply this WarningsAttribute to.</param>
+		internal ConfigurationConstrainedClangToolChainAttribute(string[]? offArgs, string[]? warningArgs, string[]? errorArgs, CppConfiguration enabledConfiguration) : base(offArgs, warningArgs, errorArgs)
+		{
+			_cppConfiguration = enabledConfiguration;
+		}
+
+		/// <inheritdoc/>
+		internal override bool CanApplyToContext(CompilerWarningsToolChainContext toolChainContext)
+		{
+			return base.CanApplyToContext(toolChainContext) && (_cppConfiguration == toolChainContext._compileEnvironment.Configuration);
+		}
+	}
+
+	/// <summary>
+	/// Specialized <see cref="WarningsClangToolChainAttribute"/> to enforce cpp standard constraints.
+	/// </summary>
+	internal sealed class CppStandardConstrainedClangToolChainAttribute : WarningsClangToolChainAttribute
+	{
+		private readonly CppStandardVersion _lowerBound;
+		private readonly CppStandardVersion _upperBound;
+
+		/// <summary>
+		/// Constructor for CppStandardConstrainedClangToolChainAttribute that enforces a cpp standard constraint.
+		/// </summary>
+		/// <param name="offArgs">The arguments to use when the property is defined as <see cref="WarningLevel.Off"/>.</param>
+		/// <param name="warningArgs">The arguments to use when the property is defined as <see cref="WarningLevel.Warning"/>.</param>
+		/// <param name="errorArgs">The arguments to use when the property is defined as <see cref="WarningLevel.Error"/>.</param>
+		/// <param name="lowerBound">The lower bound of the <see cref="CppStandardVersion"/> to constrain this warning attribute to.</param>
+		/// <param name="upperBound">The upper bound of the <see cref="CppStandardVersion"/> to constrain this warning attribute to.</param>
+		internal CppStandardConstrainedClangToolChainAttribute(string[]? offArgs, string[]? warningArgs, string[]? errorArgs, CppStandardVersion lowerBound = CppStandardVersion.Cpp14, CppStandardVersion upperBound = CppStandardVersion.Latest) : base(offArgs, warningArgs, errorArgs)
+		{
+			_lowerBound = lowerBound;
+			_upperBound = upperBound;
+		}
+
+		/// <inheritdoc/>
+		internal override bool CanApplyToContext(CompilerWarningsToolChainContext toolChainContext)
+		{
+			bool baseCanApplyToContext = base.CanApplyToContext(toolChainContext);
+
+			return baseCanApplyToContext && IsWithinBounds(toolChainContext._compileEnvironment.CppStandard);
+		}
+
+		private bool IsWithinBounds(CppStandardVersion version)
+		{
+			bool isAboveLowerBound = version >= _lowerBound;
+			bool isBelowUpperBound = version <= _upperBound;
+
+			return isAboveLowerBound && isBelowUpperBound;
+		}
+	}
+
+	/// <summary>
+	///  Specialized <see cref="WarningsClangToolChainAttribute"/> to enforce deterministic compilation constraints.
+	/// </summary>
+	internal sealed class DeterministicClangToolChainAttribute : WarningsClangToolChainAttribute
+	{
+		/// <inheritdoc/>
+		internal DeterministicClangToolChainAttribute(string[]? offArgs, string[]? warningArgs, string[]? errorArgs) : base(offArgs, warningArgs, errorArgs) { }
+
+		/// <inheritdoc/>
+		internal override bool CanApplyToContext(CompilerWarningsToolChainContext toolChainContext)
+		{
+			bool baseCanApplyToContext = base.CanApplyToContext(toolChainContext);
+
+			return baseCanApplyToContext && toolChainContext._compileEnvironment.bDeterministic;
+		}
+	}
+
+	/// <summary>
+	/// Specialized <see cref="WarningsClangToolChainAttribute"/> profile guided optimization constraints.
+	/// </summary>
+	internal sealed class PGOOptimizedClangToolChainAttribute : WarningsClangToolChainAttribute
+	{
+		/// <inheritdoc/>
+		internal PGOOptimizedClangToolChainAttribute(string[]? offArgs, string[]? warningArgs, string[]? errorArgs) : base(offArgs, warningArgs, errorArgs) { }
+
+		/// <inheritdoc/>
+		internal override bool CanApplyToContext(CompilerWarningsToolChainContext toolChainContext)
+		{
+			return base.CanApplyToContext(toolChainContext) && toolChainContext._compileEnvironment.bPGOOptimize;
+		}
+	}
+
+	/// <summary>
+	/// Specialized <see cref="WarningsClangToolChainAttribute"/> to enforce Android as an excluded platform, at a specific toolchain version.
+	/// </summary>
+	internal sealed class AndroidNDKR28ToolChainVersionExclusionAttribute : WarningsClangToolChainAttribute
+	{
+		/// <inheritdoc/>
+		internal AndroidNDKR28ToolChainVersionExclusionAttribute(string[]? offArgs, string[]? warningArgs, string[]? errorArgs) : base(offArgs, warningArgs, errorArgs) { }
+
+		/// <inheritdoc/>
+		internal override bool CanApplyToContext(CompilerWarningsToolChainContext toolChainContext)
+		{
+			if (toolChainContext._toolChainVersion == null)
+			{
+				return false;
+			}
+
+			bool versionMin = toolChainContext._toolChainVersion >= new VersionNumber(19);
+			bool isExcludedVersion = toolChainContext._toolChainVersion == new VersionNumber(19, 0, 0);
+			bool bIsAndroidClang19 = isExcludedVersion && toolChainContext._compileEnvironment.Platform == UnrealTargetPlatform.Android;
+
+			return base.CanApplyToContext(toolChainContext) && versionMin && !bIsAndroidClang19;
+		}
+	}
+
+	/// <summary>
+	/// Specialized <see cref="WarningsClangToolChainAttribute"/> to enforce Android as an excluded platform, at a specific toolchain version, at a specific C++ standard.
+	/// </summary>
+	internal sealed class AndroidNDKR26ToolChainVersionExclusionAttribute : WarningsClangToolChainAttribute
+	{
+		/// <inheritdoc/>
+		internal AndroidNDKR26ToolChainVersionExclusionAttribute(string[]? offArgs, string[]? warningArgs, string[]? errorArgs) : base(offArgs, warningArgs, errorArgs) { }
+
+		/// <inheritdoc/>
+		internal override bool CanApplyToContext(CompilerWarningsToolChainContext toolChainContext)
+		{
+			if (toolChainContext._toolChainVersion == null)
+			{
+				return false;
+			}
+
+			bool versionMin = toolChainContext._toolChainVersion >= new VersionNumber(17);
+			bool bIsAndroidClang17 = toolChainContext._toolChainVersion == new VersionNumber(17, 0, 2) && toolChainContext._compileEnvironment.Platform == UnrealTargetPlatform.Android;
+
+			return base.CanApplyToContext(toolChainContext) && versionMin && !bIsAndroidClang17 && toolChainContext._compileEnvironment.CppStandard < CppStandardVersion.Latest;
+		}
+	}
+
+	/// <summary>
+	/// / Specialized <see cref="WarningsClangToolChainAttribute"/> for Unused Value.
+	/// </summary>
+	internal sealed class UnusedValueClangToolChainAttribute : ConfigurationConstrainedClangToolChainAttribute
+	{
+		/// <inheritdoc/>
+		internal UnusedValueClangToolChainAttribute(string[]? offArgs, string[]? warningArgs, string[]? errorArgs, CppConfiguration enabledConfiguration) : base(offArgs, warningArgs, errorArgs, enabledConfiguration) { }
+
+		/// <inheritdoc/>
+		internal override bool CanApplyToContext(CompilerWarningsToolChainContext toolChainContext)
+		{
+			return base.CanApplyToContext(toolChainContext) || (toolChainContext._analyzer != StaticAnalyzer.None);
+		}
+	}
+
 	/// <summary>
 	/// Specialized <see cref="WarningsClangToolChainAttribute"/> for Shadow Variable Warnings.
 	/// </summary>
 	internal sealed class ShadowVariableWarningsClangToolChainAttribute : WarningsClangToolChainAttribute
 	{
 		/// <inheritdoc/>
-		internal ShadowVariableWarningsClangToolChainAttribute(string[]? offArgs, string[]? warningArgs, string[]? errorArgs) : base(offArgs, warningArgs, errorArgs) {}
+		internal ShadowVariableWarningsClangToolChainAttribute(string[]? offArgs, string[]? warningArgs, string[]? errorArgs) : base(offArgs, warningArgs, errorArgs) { }
 
 		/// <inheritdoc/>
+		/// <remarks>No matter what our <see cref="CppCompileWarnings.ShadowVariableWarningLevel"/> is, in the clang 17-18.1.3 range we always disable.</remarks>
 		protected override void ApplyWarningsToArgumentsInternal(WarningLevel level, CompilerWarningsToolChainContext toolChainContext, List<string> arguments)
 		{
 			if (toolChainContext._toolChainVersion == null || (toolChainContext._toolChainVersion >= new EpicGames.Core.VersionNumber(17) && toolChainContext._toolChainVersion < new VersionNumber(18, 1, 3)))
diff --git a/Engine/Source/Programs/UnrealBuildTool/Configuration/CompileWarnings/CompilerWarningsToolChainContext.cs b/Engine/Source/Programs/UnrealBuildTool/Configuration/CompileWarnings/CompilerWarningsToolChainContext.cs
--- a/Engine/Source/Programs/UnrealBuildTool/Configuration/CompileWarnings/CompilerWarningsToolChainContext.cs
+++ b/Engine/Source/Programs/UnrealBuildTool/Configuration/CompileWarnings/CompilerWarningsToolChainContext.cs
@@ -18,17 +18,20 @@ namespace UnrealBuildTool
 		/// <param name="buildSystemContext">Build context of the <see cref="CppCompileWarnings"/>.</param>
 		/// <param name="toolChainType">The type of the toolchain.</param>
 		/// <param name="toolChainVersion">The version of the toolchain.</param>
-		public CompilerWarningsToolChainContext(CppCompileEnvironment compileEnvironment, BuildSystemContext buildSystemContext, Type toolChainType, VersionNumber? toolChainVersion)
+		/// <param name="analyzer">The static analyzer utilized within the current toolchain context.</param>
+		public CompilerWarningsToolChainContext(CppCompileEnvironment compileEnvironment, BuildSystemContext buildSystemContext, Type toolChainType, VersionNumber? toolChainVersion, StaticAnalyzer analyzer = StaticAnalyzer.None)
 		{
 			_compileEnvironment = compileEnvironment;
 			_buildSystemContext = buildSystemContext;
 			_toolChainType = toolChainType;
 			_toolChainVersion = toolChainVersion;
+			_analyzer = analyzer;
 		}
 
 		internal CppCompileEnvironment _compileEnvironment;
 		internal BuildSystemContext _buildSystemContext;
 		internal Type _toolChainType;
 		internal VersionNumber? _toolChainVersion;
+		internal StaticAnalyzer _analyzer;
 	}
 }
\ No newline at end of file
diff --git a/Engine/Source/Programs/UnrealBuildTool/Configuration/CppCompileWarnings.cs b/Engine/Source/Programs/UnrealBuildTool/Configuration/CppCompileWarnings.cs
--- a/Engine/Source/Programs/UnrealBuildTool/Configuration/CppCompileWarnings.cs
+++ b/Engine/Source/Programs/UnrealBuildTool/Configuration/CppCompileWarnings.cs
@@ -48,7 +48,10 @@ namespace UnrealBuildTool
 		 *			Note: This is the most appropriatel place for BuildSettings version specific defaults
 		 * 4. Access the new property through the (TargetRules|ModuleRules|CppCompileEnvironment).CppCompileWarnings as needed
 		 * 5. (Optional) For compiler line support, utilize the ApplyWarningsAttribute system to participate in invocations to GenerateWarningCommandLineAgs.
+		 * 
+		 * Note: [REDACTED_PROFILE_URL]/clang/tree/master/test/SemaCXX contains practical examples for various clang warnings.
 		 */
+
 		#region -- Compiler Warning Settings --
 
 		/// <summary>
@@ -150,6 +153,847 @@ namespace UnrealBuildTool
 		}
 		private WarningLevel _switchUnhandledEnumeratorWarningLevel;
 
+		#region -- Clang Compiler Disables --
+
+		/// <summary>
+		/// Indicates what warning/error level to treat conversions between enum types.
+		/// Clang -
+		///		https://clang.llvm.org/docs/DiagnosticsReference.html#wenum-enum-conversion
+		/// </summary>
+		[WarningsClangToolChain(["-Wno-enum-enum-conversion"], null, null)]
+		public WarningLevel EnumEnumConversionWarningLevel
+		{
+			get => ResolveWarning(ParentCppCompileWarnings?.EnumEnumConversionWarningLevel, _enumEnumConversionWarningLevel);
+			set => _enumEnumConversionWarningLevel = value;
+		}
+		private WarningLevel _enumEnumConversionWarningLevel;
+
+		/// <summary>
+		/// Indicates what warning/error level to treat conversions between an enumeration and a floating-point type.
+		/// Clang - 
+		///		https://clang.llvm.org/docs/DiagnosticsReference.html#wenum-float-conversion
+		/// </summary>
+		[WarningsClangToolChain(["-Wno-enum-float-conversion"], null, null)]
+		public WarningLevel EnumFloatConversionWarningLevel
+		{
+			get => ResolveWarning(ParentCppCompileWarnings?.EnumFloatConversionWarningLevel, _enumFloatConversionWarningLevel);
+			set => _enumFloatConversionWarningLevel = value;
+		}
+		private WarningLevel _enumFloatConversionWarningLevel;
+
+		/// <summary>
+		/// Indicates what warning/error level to treat ambiguous use of reversed operator overloads.
+		/// Clang -
+		///		https://clang.llvm.org/docs/DiagnosticsReference.html#wambiguous-reversed-operator
+		/// </summary>
+		[CppStandardConstrainedClangToolChain(["-Wno-ambiguous-reversed-operator"], null, null, CppStandardVersion.Cpp20)]
+		public WarningLevel AmbiguousReversedOperatorWarningLevel
+		{
+			get => ResolveWarning(ParentCppCompileWarnings?.AmbiguousReversedOperatorWarningLevel, _ambiguousReversedOperatorWarningLevel);
+			set => _ambiguousReversedOperatorWarningLevel = value;
+		}
+		private WarningLevel _ambiguousReversedOperatorWarningLevel;
+
+		/// <summary>
+		/// Indicates what warning/error level to treat implicit conversions between two different anonymous enumerations
+		/// Clang -
+		///		https://clang.llvm.org/docs/DiagnosticsReference.html#wdeprecated-anon-enum-enum-conversion
+		/// </summary>
+		[CppStandardConstrainedClangToolChain(["-Wno-deprecated-anon-enum-enum-conversion"], null, null, CppStandardVersion.Cpp20)]
+		public WarningLevel DeprecatedAnonEnumEnumConversionWarningLevel
+		{
+			get => ResolveWarning(ParentCppCompileWarnings?.DeprecatedAnonEnumEnumConversionWarningLevel, _deprecatedAnonEnumEnumConversionWarningLevel);
+			set => _deprecatedAnonEnumEnumConversionWarningLevel = value;
+		}
+		private WarningLevel _deprecatedAnonEnumEnumConversionWarningLevel;
+
+		/// <summary>
+		/// Indicates what warning/error level to treat deprecated usages of volatile (due to better alternatives like atomic operations).
+		/// Clang -
+		///		https://clang.llvm.org/docs/DiagnosticsReference.html#wdeprecated-volatile
+		/// </summary>
+		[CppStandardConstrainedClangToolChain(["-Wno-deprecated-volatile"], null, null, CppStandardVersion.Cpp20)]
+		public WarningLevel DeprecatedVolatileWarningLevel
+		{
+			get => ResolveWarning(ParentCppCompileWarnings?.DeprecatedVolatileWarningLevel, _deprecatedVolatileWarningLevel);
+			set => _deprecatedVolatileWarningLevel = value;
+		}
+		private WarningLevel _deprecatedVolatileWarningLevel;
+
+		/// <summary>
+		/// Indicates what warning/error level to treat compairing function pointers with relational operators (&lt;,&gt;, etc), which is undefined behaviour.
+		/// Clang -
+		///		https://clang.llvm.org/docs/DiagnosticsReference.html#wordered-compare-function-pointers
+		/// </summary>
+		[WarningsClangToolChain(["-Wno-ordered-compare-function-pointers"], null, null)]
+		public WarningLevel OrderedCompareFunctionPointers
+		{
+			get => ResolveWarning(ParentCppCompileWarnings?.OrderedCompareFunctionPointers, _orderedCompareFunctionPointers);
+			set => _orderedCompareFunctionPointers = value;
+		}
+		private WarningLevel _orderedCompareFunctionPointers;
+
+		/// <summary>
+		/// Indicates what warning/error level to treat bitwise operations are used where logical operators are likely intended.
+		/// Clang -
+		///		https://clang.llvm.org/docs/DiagnosticsReference.html#wbitwise-instead-of-logical
+		/// </summary>
+		[ToolChainVersionConstrainedWarningsClangToolChain(["-Wno-bitwise-instead-of-logical"], null, null, [14], null)]
+		public WarningLevel BitwiseInsteadOfLogicalWarningLevel
+		{
+			get => ResolveWarning(ParentCppCompileWarnings?.BitwiseInsteadOfLogicalWarningLevel, _bitwiseInsteadOfLogicalWarningLevel);
+			set => _bitwiseInsteadOfLogicalWarningLevel = value;
+		}
+		private WarningLevel _bitwiseInsteadOfLogicalWarningLevel;
+
+		/// <summary>
+		/// Indicates what warning/error level to treat usages of a class where a deprecated copy constructor or copy assignment operator exists.
+		/// Clang -
+		///		https://clang.llvm.org/docs/DiagnosticsReference.html#wdeprecated-copy
+		/// </summary>
+		[ToolChainVersionConstrainedWarningsClangToolChain(["-Wno-deprecated-copy"], null, null, [16], null)]
+		public WarningLevel DeprecatedCopyWarningLevel
+		{
+			get => ResolveWarning(ParentCppCompileWarnings?.DeprecatedCopyWarningLevel, _deprecatedCopyWarningLevel);
+			set => _deprecatedCopyWarningLevel = value;
+		}
+		private WarningLevel _deprecatedCopyWarningLevel;
+
+		/// <summary>
+		/// Indicates what warning/error level to treat usages of a class where a deprecated implicitly declared copy constructor or copy assignment operator exists.
+		/// Clang -
+		///		https://clang.llvm.org/docs/DiagnosticsReference.html#wdeprecated-copy-with-user-provided-copy
+		/// </summary>
+		[ToolChainVersionConstrainedWarningsClangToolChain(["-Wno-deprecated-copy-with-user-provided-copy"], null, null, [16], null)]
+		public WarningLevel DeprecatedCopyWithUserProvidedCopyWarningLevel
+		{
+			get => ResolveWarning(ParentCppCompileWarnings?.DeprecatedCopyWithUserProvidedCopyWarningLevel, _deprecatedCopyWithUserProvidedCopyWarningLevel);
+			set => _deprecatedCopyWithUserProvidedCopyWarningLevel = value;
+		}
+		private WarningLevel _deprecatedCopyWithUserProvidedCopyWarningLevel;
+
+		/// <summary>
+		/// Indicates what warning/error level to treat invalid, unevaluated strings.
+		/// Clang -
+		///		https://clang.llvm.org/docs/DiagnosticsReference.html#winvalid-unevaluated-string
+		/// </summary>
+		[AndroidNDKR26ToolChainVersionExclusion(["-Wno-invalid-unevaluated-string"], null, null)]
+		public WarningLevel InvalidUnevaluatedStringWarningLevel
+		{
+			get => ResolveWarning(ParentCppCompileWarnings?.InvalidUnevaluatedStringWarningLevel, _invalidUnevaluatedStringWarningLevel);
+			set => _invalidUnevaluatedStringWarningLevel = value;
+		}
+		private WarningLevel _invalidUnevaluatedStringWarningLevel;
+
+		/// <summary>
+		/// Indicates what warning/error level to treat infinity macros, and subsequent undefined behaviour.
+		/// Clang -
+		///		https://clang.llvm.org/docs/DiagnosticsReference.html#wnan-infinity-disabled
+		/// </summary>
+		/// <remarks>
+		/// We use the NAN macro in a few places to initialize floats while we also set -ffast-math, which disables NaN support. 
+		/// It could be easily fixed in the engine, but would create the risk of Win64 building fine and the code failing just on Clang-based platforms. 
+		/// We tend to use NAN just as a bit pattern, and not rely on it in calculations, so it should be reasonably safe to disable.
+		/// </remarks>
+		[ToolChainVersionConstrainedWarningsClangToolChain(["-Wno-nan-infinity-disabled"], null, null, [18], null)]
+		public WarningLevel NaNInfinityDisabledWarningLevel
+		{
+			get => ResolveWarning(ParentCppCompileWarnings?.NaNInfinityDisabledWarningLevel, _nanInfinityDisabledWarningLevel);
+			set => _nanInfinityDisabledWarningLevel = value;
+		}
+		private WarningLevel _nanInfinityDisabledWarningLevel;
+
+		/// <summary>
+		/// Indicates what warning/error level to treat extra qualifications.
+		/// Clang -
+		///		https://clang.llvm.org/docs/DiagnosticsReference.html#wextra-qualification
+		/// </summary>
+		[ToolChainVersionConstrainedWarningsClangToolChain(["-Wno-extra-qualification"], null, null, [19], null)]
+		public WarningLevel LevelExtraQualificationWarningLevel
+		{
+			get => ResolveWarning(ParentCppCompileWarnings?.LevelExtraQualificationWarningLevel, _levelExtraQualificationWarningLevel);
+			set => _levelExtraQualificationWarningLevel = value;
+		}
+		private WarningLevel _levelExtraQualificationWarningLevel;
+
+		/// <summary>
+		/// Indicates what warning/error level to treat function pointer casts to an incompatible function type, potentially leading to undefined behaviour.
+		/// Clang -
+		///		https://clang.llvm.org/docs/DiagnosticsReference.html#wcast-function-type-mismatch
+		/// </summary>
+		[AndroidNDKR28ToolChainVersionExclusion(["-Wno-cast-function-type-mismatch"], null, null)]
+		public WarningLevel CastFunctionTypeMistmatchWarningLevel
+		{
+			get => ResolveWarning(ParentCppCompileWarnings?.CastFunctionTypeMistmatchWarningLevel, _castFunctionTypeMistmatchWarningLevel);
+			set => _castFunctionTypeMistmatchWarningLevel = value;
+		}
+		private WarningLevel _castFunctionTypeMistmatchWarningLevel;
+
+		/// <summary>
+		/// Indicates what warning/error level to treat templates used without specifying the required template arguments, such as missing parameter packs (&lt;...&gt;) after the template keyword.
+		/// Clang -
+		///		https://clang.llvm.org/docs/DiagnosticsReference.html#wmissing-template-arg-list-after-template-kw
+		/// </summary>
+		[AndroidNDKR28ToolChainVersionExclusion(["-Wno-missing-template-arg-list-after-template-kw"], null, null)]
+		public WarningLevel MissingTemplateArgListAfterTemplateWarningLevel
+		{
+			get => ResolveWarning(ParentCppCompileWarnings?.MissingTemplateArgListAfterTemplateWarningLevel, _missingTemplateArgListAfterTemplateWarningLevel);
+			set => _missingTemplateArgListAfterTemplateWarningLevel = value;
+		}
+		private WarningLevel _missingTemplateArgListAfterTemplateWarningLevel;
+
+		/// <summary>
+		/// Indicates what warning/error level to treat gnu string literal operator templates.
+		/// Clang -
+		///		https://clang.llvm.org/docs/DiagnosticsReference.html#wgnu-string-literal-operator-template
+		/// </summary>
+		/// <remarks>We use this feature to allow static FNames.</remarks>
+		[WarningsClangToolChain(["-Wno-gnu-string-literal-operator-template"], null, null)]
+		public WarningLevel GNUStringLiteralOperatorTemplateWarningLevel
+		{
+			get => ResolveWarning(ParentCppCompileWarnings?.GNUStringLiteralOperatorTemplateWarningLevel, _gnuStringLiteralOperatorTemplateWarningLevel);
+			set => _gnuStringLiteralOperatorTemplateWarningLevel = value;
+		}
+		private WarningLevel _gnuStringLiteralOperatorTemplateWarningLevel;
+
+		/// <summary>
+		/// Indicates what warning/error level to treat inconsistent missing overrides.
+		/// Clang -
+		///		https://clang.llvm.org/docs/DiagnosticsReference.html#winconsistent-missing-override
+		/// </summary>
+		[WarningsClangToolChain(["-Wno-inconsistent-missing-override"], null, null)]
+		public WarningLevel InconsistentMissingOverrideWarningLevel
+		{
+			get => ResolveWarning(ParentCppCompileWarnings?.InconsistentMissingOverrideWarningLevel, _inconsistentMissingOverrideWarningLevel);
+			set => _inconsistentMissingOverrideWarningLevel = value;
+		}
+		private WarningLevel _inconsistentMissingOverrideWarningLevel;
+
+		/// <summary>
+		/// Indicates what warning/error level to treat incorrect usages of std::offsetof, such as with non-POD types, where behaviour is undefined.
+		/// Clang -
+		///		https://clang.llvm.org/docs/DiagnosticsReference.html#winvalid-offsetof
+		/// </summary>
+		/// <remarks>Needed to suppress warnings about using offsetof on non-POD types.</remarks>
+		[WarningsClangToolChain(["-Wno-invalid-offsetof"], null, null)]
+		public WarningLevel InvalidOffsetWarningLevel
+		{
+			get => ResolveWarning(ParentCppCompileWarnings?.InvalidOffsetWarningLevel, _invalidOffsetOfWarningLevel);
+			set => _invalidOffsetOfWarningLevel = value;
+		}
+		private WarningLevel _invalidOffsetOfWarningLevel;
+
+		/// <summary>
+		/// Indicates what warning/error level to treat switch related warnings.
+		/// Clang -
+		///		https://clang.llvm.org/docs/DiagnosticsReference.html#wswitch
+		/// </summary>
+		/// <remarks>Fhis hides the "enumeration value 'XXXXX' not handled in switch [-Wswitch]" warnings - we should maybe remove this at some point and add UE_LOG(, Fatal, ) to default cases.</remarks>
+		[WarningsClangToolChain(["-Wno-switch"], null, null)]
+		public WarningLevel SwitchWarningLevel
+		{
+			get => ResolveWarning(ParentCppCompileWarnings?.SwitchWarningLevel, _switchWarningLevel);
+			set => _switchWarningLevel = value;
+		}
+		private WarningLevel _switchWarningLevel;
+
+		/// <summary>
+		/// Indicates what warning/error level to treat tautological comparisons (i.e. obvious comparisons that are always guaranteed).
+		/// Clang -
+		///		https://clang.llvm.org/docs/DiagnosticsReference.html#wtautological-compare
+		/// </summary>
+		/// <remarks>This hides the "warning : comparison of unsigned expression &lt; 0 is always false" type warnings due to constant comparisons, which are possible with template arguments/</remarks>
+		[WarningsClangToolChain(["-Wno-tautological-compare"], null, null)]
+		public WarningLevel TautologicalCompareWarningLevel
+		{
+			get => ResolveWarning(ParentCppCompileWarnings?.TautologicalCompareWarningLevel, _tautologicalCompareWarningLevel);
+			set => _tautologicalCompareWarningLevel = value;
+		}
+		private WarningLevel _tautologicalCompareWarningLevel;
+
+		/// <summary>
+		/// Indicates what warning/error level to treat unknown pragmas with.
+		/// Clang -
+		///		https://clang.llvm.org/docs/DiagnosticsReference.html#wunknown-pragmas
+		/// </summary>
+		/// <remarks>Slate triggers this (with its optimize on/off pragmas).</remarks>
+		[WarningsClangToolChain(["-Wno-unknown-pragmas"], null, null)]
+		public WarningLevel UnknownPragmasWarningLevel
+		{
+			get => ResolveWarning(ParentCppCompileWarnings?.UnknownPragmasWarningLevel, _unknownPragmasWarningLevel);
+			set => _unknownPragmasWarningLevel = value;
+		}
+		private WarningLevel _unknownPragmasWarningLevel;
+
+		/// <summary>
+		/// Indicates what warning/error level to treat unused symbols.
+		/// Clang - 
+		///		https://clang.llvm.org/docs/DiagnosticsReference.html#wunused-but-set-variable
+		///		https://clang.llvm.org/docs/DiagnosticsReference.html#wunused-but-set-parameter
+		///		https://clang.llvm.org/docs/DiagnosticsReference.html#wunused-function
+		///		https://clang.llvm.org/docs/DiagnosticsReference.html#wunused-lambda-capture
+		///		https://clang.llvm.org/docs/DiagnosticsReference.html#wunused-local-typedef
+		///		https://clang.llvm.org/docs/DiagnosticsReference.html#wunused-private-field
+		///		https://clang.llvm.org/docs/DiagnosticsReference.html#wunused-variable
+		/// </summary>
+		/// <remarks>
+		/// unused-function: This will hide the warnings about static functions in headers that aren't used in every single .cpp file.
+		/// unused-lambda-capture: Suppressed because capturing of compile-time constants is seemingly inconsistent. And MSVC doesn't do that.
+		/// unused-local-typedef: Clang is being overly strict here? PhysX headers trigger this.
+		/// unused-private-field: This will prevent the issue of warnings for unused private variables. MultichannelTcpSocket.h triggers this, possibly more.
+		/// </remarks>
+		[WarningsClangToolChain(["-Wno-unused-but-set-variable", "-Wno-unused-but-set-parameter", "-Wno-unused-function", "-Wno-unused-lambda-capture", "-Wno-unused-local-typedef", "-Wno-unused-private-field", "-Wno-unused-variable"], null, null)]
+		public WarningLevel UnusedWarningLevel
+		{
+			get => ResolveWarning(ParentCppCompileWarnings?.UnusedWarningLevel, _unusedWarningLevel);
+			set => _unusedWarningLevel = value;
+		}
+		private WarningLevel _unusedWarningLevel;
+
+		/// <summary>
+		/// Indicates what warning/error level to treat undefined templates for variables.
+		/// Clang -
+		///		https://clang.llvm.org/docs/DiagnosticsReference.html#wundefined-var-template
+		/// </summary>
+		/// <remarks>Not a good warning to be disabling.</remarks>
+		[WarningsClangToolChain(["-Wno-undefined-var-template"], null, null)]
+		public WarningLevel UndefinedVarTemplateWarningLevel
+		{
+			get => ResolveWarning(ParentCppCompileWarnings?.UndefinedVarTemplateWarningLevel, _undefinedVarTemplateWarningLevel);
+			set => _undefinedVarTemplateWarningLevel = value;
+		}
+		private WarningLevel _undefinedVarTemplateWarningLevel;
+
+		/// <summary>
+		/// Indicates what warning/error level to treat profile instructions.
+		/// Clang -
+		///		https://clang.llvm.org/docs/DiagnosticsReference.html#wprofile-instr-out-of-date
+		///		https://clang.llvm.org/docs/DiagnosticsReference.html#wprofile-instr-unprofiled
+		/// </summary>
+		/// <remarks> Clang emits warnings for each compiled object file that doesn't have a matching entry in the profile data. This can happen when the profile data is older than the binaries we're compiling. Disable these warnings.</remarks>
+		[PGOOptimizedClangToolChain(["-Wno-profile-instr-out-of-date", "-Wno-profile-instr-unprofiled"], null, null)]
+		public WarningLevel ProfileInstructWarningLevel
+		{
+			get => ResolveWarning(ParentCppCompileWarnings?.ProfileInstructWarningLevel, _profileInstructWarningLevel);
+			set => _profileInstructWarningLevel = value;
+		}
+		private WarningLevel _profileInstructWarningLevel;
+
+		/// <summary>
+		/// Clang -
+		///		https://clang.llvm.org/docs/DiagnosticsReference.html#wbackend-plugin
+		/// </summary>
+		/// <remarks>Apparently there can be hashing conflicts with PGO which can result in: 'Function control flow change detected (hash mismatch)' warnings. </remarks>
+		[PGOOptimizedClangToolChain(["-Wno-backend-plugin"], null, null)]
+		public WarningLevel BackendPluginWarningLevel
+		{
+			get => ResolveWarning(ParentCppCompileWarnings?.BackendPluginWarningLevel, _backendPluginWarningLevel);
+			set => _backendPluginWarningLevel = value;
+		}
+		private WarningLevel _backendPluginWarningLevel;
+
+		/// <summary>
+		/// Indicates what warning/error level to treat values that are computed, but not used (i.e.unnecessary evaluation of an rvalue or expression).
+		/// Clang -
+		///		https://clang.llvm.org/docs/DiagnosticsReference.html#wunused-value
+		/// </summary>
+		/// <remarks>Shipping builds will cause this warning with "ensure", so disable only in those case.</remarks>
+		[UnusedValueClangToolChain(["-Wno-unused-value"], null, null, CppConfiguration.Shipping)]
+		public WarningLevel UnusedValueWarningLevel
+		{
+			get => ResolveWarning(ParentCppCompileWarnings?.UnusedValueWarningLevel, _unusedValueWarningLevel);
+			set => _unusedValueWarningLevel = value;
+		}
+		private WarningLevel _unusedValueWarningLevel;
+
+		/// <summary>
+		/// Indicates what warning/error level to treat implicit conversaion losses for integer precision.
+		/// Clang -
+		///		https://clang.llvm.org/docs/DiagnosticsReference.html#wshorten-64-to-32
+		/// </summary>
+		/// <remarks>-Wimplicit-int-conversion also controls -Wshorten-64-to-32 as of clang 19.</remarks>
+		[ToolChainVersionConstrainedWarningsClangToolChain(["-Wno-shorten-64-to-32"], null, null, [19], null)]
+		public WarningLevel Shorten64To32WarningLevel
+		{
+			get => ResolveWarning(ParentCppCompileWarnings?.Shorten64To32WarningLevel, _shorten64To32WarningLevel);
+			set => _shorten64To32WarningLevel = value;
+		}
+		private WarningLevel _shorten64To32WarningLevel;
+
+		/// <summary>
+		/// Indicates what warning/error level to treat incorrect extern template declarations.
+		/// Clang -
+		///		https://clang.llvm.org/docs/DiagnosticsReference.html#wdllexport-explicit-instantiation-decl
+		/// </summary>
+		/// <remarks>The code base contains lots of places where we do "extern template X_API class ..." and we want to keep doing that to reduce compile times.</remarks>
+		[WarningsClangToolChain(["-Wno-dllexport-explicit-instantiation-decl"], null, null)]
+		public WarningLevel DllExportExplicitInstantiationDeclWarningLevel
+		{
+			get => ResolveWarning(ParentCppCompileWarnings?.DllExportExplicitInstantiationDeclWarningLevel, _dllExportExplicitInstantiationDeclWarningLevel);
+			set => _dllExportExplicitInstantiationDeclWarningLevel = value;
+		}
+		private WarningLevel _dllExportExplicitInstantiationDeclWarningLevel;
+
+		#endregion
+
+		#region -- Clang VC Compiler Disables --
+
+		/// <summary>
+		/// Indicates what warning/error level to treat warnings associated with Microsoft-specific language extensions and compatibility issues.
+		/// Clang -
+		///		https://clang.llvm.org/docs/DiagnosticsReference.html#wmicrosoft
+		/// </summary>
+		/// <remarks>Allow Microsoft-specific syntax to slide, even though it may be non-standard.  Needed for Windows headers.</remarks>
+		[WarningsVCClang(["-Wno-microsoft"], null, null)]
+		public WarningLevel MicrosoftGroupWarningLevel
+		{
+			get => ResolveWarning(ParentCppCompileWarnings?.MicrosoftGroupWarningLevel, _microsoftGroupWarningLevel);
+			set => _microsoftGroupWarningLevel = value;
+		}
+		private WarningLevel _microsoftGroupWarningLevel;
+
+		/// <summary>
+		/// Indicates what warning/error level to treat compatibility mods for clang, and how it searches for system headers.
+		/// Clang -
+		///		https://clang.llvm.org/docs/DiagnosticsReference.html#wmsvc-include
+		/// </summary>
+		/// <remarks> Hack due to how we have our 'DummyPCH' wrappers setup when using unity builds.  This warning should not be disabled!</remarks>
+		[WarningsVCClang(["-Wno-msvc-include"], null, null)]
+		public WarningLevel MSVCIncludeWarningLevel
+		{
+			get => ResolveWarning(ParentCppCompileWarnings?.MSVCIncludeWarningLevel, _msvcIncludeWarningLevel);
+			set => _msvcIncludeWarningLevel = value;
+		}
+		private WarningLevel _msvcIncludeWarningLevel;
+
+		/// <summary>
+		/// Indicates what warning/error level to treat improper or inconsistent usage of #pragma pack, which can lead to ABI/unexpected struct layouts.
+		/// Clang -
+		///		https://clang.llvm.org/docs/DiagnosticsReference.html#wpragma-pack
+		/// </summary>
+		/// <remarks>
+		/// This is disabled because clang explicitly warns about changing pack alignment in a header and not restoring it afterwards, which is something we do with the Pre/PostWindowsApi.h headers.
+		/// @TODO clang: push/pop this in  Pre/PostWindowsApi.h headers instead?
+		/// </remarks>
+		[WarningsVCClang(["-Wno-pragma-pack"], null, null)]
+		public WarningLevel PragmaPackWarningLevel
+		{
+			get => ResolveWarning(ParentCppCompileWarnings?.PragmaPackWarningLevel, _pragmaPackWarningLevel);
+			set => _pragmaPackWarningLevel = value;
+		}
+		private WarningLevel _pragmaPackWarningLevel;
+
+		/// <summary>
+		/// Indicates what warning/error level to treat cases where operator new or operator delete is declared inline, leading to issues with memory allocdealloc across different translation units.
+		/// Clang -
+		///		https://clang.llvm.org/docs/DiagnosticsReference.html#winline-new-delete
+		/// </summary>
+		/// <remarks>@todo clang: We declare operator new as inline.  Clang doesn't seem to like that.</remarks>
+		[WarningsVCClang(["-Wno-inline-new-delete"], null, null)]
+		public WarningLevel InlineNewDeleteWarningLevel
+		{
+			get => ResolveWarning(ParentCppCompileWarnings?.InlineNewDeleteWarningLevel, _inlineNewDeleteWarningLevel);
+			set => _inlineNewDeleteWarningLevel = value;
+		}
+		private WarningLevel _inlineNewDeleteWarningLevel;
+
+		/// <summary>
+		/// Indicates what warning/error level to treat cases where function declarations and definitions have mismatched implicit exception qualifiers.
+		/// Clang -
+		///		https://clang.llvm.org/docs/DiagnosticsReference.html#wimplicit-exception-spec-mismatch
+		/// </summary>
+
+		[WarningsVCClang(["-Wno-implicit-exception-spec-mismatch"], null, null)]
+		public WarningLevel ImplicitExceptionSpecMismatchWarningLevel
+		{
+			get => ResolveWarning(ParentCppCompileWarnings?.ImplicitExceptionSpecMismatchWarningLevel, _implicitExceptionSpecMismatchWarningLevel);
+			set => _implicitExceptionSpecMismatchWarningLevel = value;
+		}
+		private WarningLevel _implicitExceptionSpecMismatchWarningLevel;
+
+		/// <summary>
+		/// Indicates what warning/error level to treat cases when an expression that isn't explicitly a boolean is used in a boolean context.
+		/// Clang -
+		///		https://clang.llvm.org/docs/DiagnosticsReference.html#wundefined-bool-conversion
+		/// </summary>
+		/// <remarks>Sometimes we compare 'this' pointers against nullptr, which Clang warns about by default.</remarks>
+		[WarningsVCClang(["-Wno-undefined-bool-conversion"], null, null)]
+		public WarningLevel UndefinedBoolConversionWarningLevel
+		{
+			get => ResolveWarning(ParentCppCompileWarnings?.UndefinedBoolConversionWarningLevel, _undefinedBoolConversionWarningLevel);
+			set => _undefinedBoolConversionWarningLevel = value;
+		}
+		private WarningLevel _undefinedBoolConversionWarningLevel;
+
+		/// <summary>
+		/// Indicates what warning/error level to treat usages of string literals as writable char* (allowed in older c++), but has since been deprecated.
+		/// Clang -
+		///		https://clang.llvm.org/docs/DiagnosticsReference.html#wdeprecated-writable-strings
+		/// </summary>
+
+		[WarningsVCClang(["-Wno-deprecated-writable-strings"], null, null)]
+		public WarningLevel DeprecatedWriteableStringsWarningLevel
+		{
+			get => ResolveWarning(ParentCppCompileWarnings?.DeprecatedWriteableStringsWarningLevel, _deprecatedWriteableStringsWarningLevel);
+			set => _deprecatedWriteableStringsWarningLevel = value;
+		}
+		private WarningLevel _deprecatedWriteableStringsWarningLevel;
+
+		/// <summary>
+		/// Indicates what warning/error level to treat deprecated usages of the register storage class specifier.
+		/// Clang -
+		///		https://clang.llvm.org/docs/DiagnosticsReference.html#wdeprecated-register
+		/// </summary>
+		[WarningsVCClang(["-Wno-deprecated-register"], null, null)]
+		public WarningLevel DeprecatedRegisterWarningLevel
+		{
+			get => ResolveWarning(ParentCppCompileWarnings?.DeprecatedRegisterWarningLevel, _deprecatedRegisterWarningLevel);
+			set => _deprecatedRegisterWarningLevel = value;
+		}
+		private WarningLevel _deprecatedRegisterWarningLevel;
+
+		/// <summary>
+		/// Indicates what warning/error level to treat usages of logical operations and related parenthesis use. 
+		///		e.g.   return false &amp;&amp; false || array[f.get()]; // expected-warning {{'&amp;&amp;' within '||'}} expected-note {{parentheses}}
+		/// Clang -
+		///		https://clang.llvm.org/docs/DiagnosticsReference.html#wlogical-op-parentheses
+		/// </summary>
+		/// <remarks>Disable needed for external headers we shan't change.</remarks>
+		[WarningsVCClang(["-Wno-logical-op-parentheses"], null, null)]
+		public WarningLevel LogicalOpParenthesesWarningLevel
+		{
+			get => ResolveWarning(ParentCppCompileWarnings?.LogicalOpParenthesesWarningLevel, _logicalOpParenthesesWarningLevel);
+			set => _logicalOpParenthesesWarningLevel = value;
+		}
+		private WarningLevel _logicalOpParenthesesWarningLevel;
+
+		/// <summary>
+		/// Indicates what warning/error level to treat occurrences where arithmetic oeprations are performed on nullptr.
+		/// Clang -
+		///		https://clang.llvm.org/docs/DiagnosticsReference.html#wnull-arithmetic
+		/// </summary>
+		/// <remarks>Disable needed for external headers we shan't change.</remarks>
+		[WarningsVCClang(["-Wno-null-arithmetic"], null, null)]
+		public WarningLevel NullArithmeticWarningLevel
+		{
+			get => ResolveWarning(ParentCppCompileWarnings?.NullArithmeticWarningLevel, _nullArithmeticWarningLevel);
+			set => _nullArithmeticWarningLevel = value;
+		}
+		private WarningLevel _nullArithmeticWarningLevel;
+
+		/// <summary>
+		/// Indicates what warning/error level to treat cases where a function with C linkage has a return type that isn't compatible with C.
+		/// Clang -
+		///		https://clang.llvm.org/docs/DiagnosticsReference.html#wreturn-type-c-linkage
+		/// </summary>
+		/// <remarks>Disable needed for PhysX.</remarks>
+		[WarningsVCClang(["-Wno-return-type-c-linkage"], null, null)]
+		public WarningLevel ReturnTypeCLinkageWarningLevel
+		{
+			get => ResolveWarning(ParentCppCompileWarnings?.ReturnTypeCLinkageWarningLevel, _returnTypeCLinkageWarningLevel);
+			set => _returnTypeCLinkageWarningLevel = value;
+		}
+		private WarningLevel _returnTypeCLinkageWarningLevel;
+
+		/// <summary>
+		/// Indicates what warning/error level to treat usages of attributes applied to a function or variable that are ignored.
+		/// Clang -
+		///		https://clang.llvm.org/docs/DiagnosticsReference.html#wignored-attributes
+		/// </summary>
+		/// <remarks>Disable needed for nvtesslib.</remarks>
+		[WarningsVCClang(["-Wno-ignored-attributes"], null, null)]
+		public WarningLevel IgnoredAttributesWarningLevel
+		{
+			get => ResolveWarning(ParentCppCompileWarnings?.IgnoredAttributesWarningLevel, _ignoredAttributesWarningLevel);
+			set => _ignoredAttributesWarningLevel = value;
+		}
+		private WarningLevel _ignoredAttributesWarningLevel;
+
+		/// <summary>
+		/// Indicates what warning/error level to treat varaibles that are used before initialization.
+		/// Clang -
+		///		https://clang.llvm.org/docs/DiagnosticsReference.html#wuninitialized
+		/// </summary>
+		[WarningsVCClang(["-Wno-uninitialized"], null, null)]
+		public WarningLevel UninitializedWarningLevel
+		{
+			get => ResolveWarning(ParentCppCompileWarnings?.UninitializedWarningLevel, _uninitializedWarningLevel);
+			set => _uninitializedWarningLevel = value;
+		}
+		private WarningLevel _uninitializedWarningLevel;
+
+		/// <summary>
+		/// Indicates what warning/error level to treat functions with a non-void return type, but doesn't return a value in all code paths.
+		/// Clang -
+		///		https://clang.llvm.org/docs/DiagnosticsReference.html#wreturn-type
+		/// </summary>
+		/// <remarks>Disable needed for external headers we shan't change.</remarks>
+		[WarningsVCClang(["-Wno-return-type"], null, null)]
+		public WarningLevel ReturnTypeWarningLevel
+		{
+			get => ResolveWarning(ParentCppCompileWarnings?.ReturnTypeWarningLevel, _returnTypeWarningLevel);
+			set => _returnTypeWarningLevel = value;
+		}
+		private WarningLevel _returnTypeWarningLevel;
+
+		/// <summary>
+		/// Indicates what warning/error level to treat function parameters that are declared, but not used within the function body.
+		/// Clang -
+		///		https://clang.llvm.org/docs/DiagnosticsReference.html#wunused-parameter
+		/// </summary>
+		/// <remarks>Unused function parameter. A lot are named 'bUnused'...</remarks>
+		[WarningsVCClang(["-Wno-unused-parameter"], null, null)]
+		public WarningLevel UnusedParameterWarningLevel
+		{
+			get => ResolveWarning(ParentCppCompileWarnings?.UnusedParameterWarningLevel, _unusedParameterWarningLevel);
+			set => _unusedParameterWarningLevel = value;
+		}
+		private WarningLevel _unusedParameterWarningLevel;
+
+		/// <summary>
+		/// Indicates what warning/error level to treat type qualifiers (const/volatile) that are applied to a pointer or ref, but are ignored due to not having an effect.
+		/// Clang -
+		///		https://clang.llvm.org/docs/DiagnosticsReference.html#wignored-qualifiers
+		/// </summary>
+		/// <remarks>const ignored when returning by value e.g. 'const int foo() { return 4; }'.</remarks>
+		[WarningsVCClang(["-Wno-ignored-qualifiers"], null, null)]
+		public WarningLevel IgnoredQualifiersWarningLevel
+		{
+			get => ResolveWarning(ParentCppCompileWarnings?.IgnoredQualifiersWarningLevel, _ignoredQualifiersWarningLevel);
+			set => _ignoredQualifiersWarningLevel = value;
+		}
+		private WarningLevel _ignoredQualifiersWarningLevel;
+
+		/// <summary>
+		/// Indicates what warning/error level to treat occurrences of macro expansion resulting in a defined expression being used in an undefined manner.
+		/// Clang -
+		///		https://clang.llvm.org/docs/DiagnosticsReference.html#wexpansion-to-defined
+		/// </summary>
+		/// <remarks>Usage of 'defined(X)' in a macro definition. Gives different results under MSVC.</remarks>
+		[WarningsVCClang(["-Wno-expansion-to-defined"], null, null)]
+		public WarningLevel ExpansionToDefined
+		{
+			get => ResolveWarning(ParentCppCompileWarnings?.ExpansionToDefined, _expansionToDefined);
+			set => _expansionToDefined = value;
+		}
+		private WarningLevel _expansionToDefined;
+
+		/// <summary>
+		/// Indicates what warning/error level to treat comparisons of signed and unsigned integer types.
+		/// Clang -
+		///		https://clang.llvm.org/docs/DiagnosticsReference.html#wsign-compare
+		/// </summary>
+		/// <remarks>Signed/unsigned comparison - millions of these.</remarks>
+		[WarningsVCClang(["-Wno-sign-compare"], null, null)]
+		public WarningLevel SignCompareWarningLevel
+		{
+			get => ResolveWarning(ParentCppCompileWarnings?.SignCompareWarningLevel, _signCompareWarningLevel);
+			set => _signCompareWarningLevel = value;
+		}
+		private WarningLevel _signCompareWarningLevel;
+
+		/// <summary>
+		/// Indicates what warning/error level to treat initializer lists used for structor or array missing one or more fields.
+		/// Clang -
+		///		https://clang.llvm.org/docs/DiagnosticsReference.html#wmissing-field-initializers
+		/// </summary>
+		/// <remarks>Excessive warning, generated when you initialize with MyStruct A = {0};</remarks>
+		[WarningsVCClang(["-Wno-missing-field-initializers"], null, null)]
+		public WarningLevel MissingFieldInitializersWarningLevel
+		{
+			get => ResolveWarning(ParentCppCompileWarnings?.MissingFieldInitializersWarningLevel, _missingFieldInitializersWarningLevel);
+			set => _missingFieldInitializersWarningLevel = value;
+		}
+		private WarningLevel _missingFieldInitializersWarningLevel;
+
+		/// <summary>
+		/// Indicates what warning/error level to treat usages of non-portable include paths.
+		/// Clang -
+		///		https://clang.llvm.org/docs/DiagnosticsReference.html#wnonportable-include-path
+		/// </summary>
+		[WarningsVCClang(["-Wno-nonportable-include-path"], null, null)]
+		public WarningLevel NonPortableIncludePathWarningLevel
+		{
+			get => ResolveWarning(ParentCppCompileWarnings?.NonPortableIncludePathWarningLevel, _nonPortableIncludePathWarningLevel);
+			set => _nonPortableIncludePathWarningLevel = value;
+		}
+		private WarningLevel _nonPortableIncludePathWarningLevel;
+
+		/// <summary>
+		/// Indicates what warning/error level to treat preprocessor's token-pasting operator (##) is used incorrectly, resulting in ill-formed tokens after concatenation. 
+		/// Clang -
+		///		https://clang.llvm.org/docs/DiagnosticsReference.html#winvalid-token-paste
+		///	MASVC -
+		///		https://learn.microsoft.com/en-us/cpp/preprocessor/token-pasting-operator-hash-hash?view=msvc-170
+		/// </summary>
+		[WarningsVCClang(["-Wno-invalid-token-paste"], null, null)]
+		public WarningLevel InvalidTokenPasteWarningLevel
+		{
+			get => ResolveWarning(ParentCppCompileWarnings?.InvalidTokenPasteWarningLevel, _invalidTokenPasteWarningLevel);
+			set => _invalidTokenPasteWarningLevel = value;
+		}
+		private WarningLevel _invalidTokenPasteWarningLevel;
+
+		/// <summary>
+		/// Indicates what warning/error level to treat arithmetic operations that are performed with a null pointer.
+		/// Clang -
+		///		https://clang.llvm.org/docs/DiagnosticsReference.html#wnull-pointer-arithmetic
+		/// </summary>
+		[WarningsVCClang(["-Wno-null-pointer-arithmetic"], null, null)]
+		public WarningLevel NullPointerArithmeticWarningLevel
+		{
+			get => ResolveWarning(ParentCppCompileWarnings?.NullPointerArithmeticWarningLevel, _nullPointerArithmeticWarningLevel);
+			set => _nullPointerArithmeticWarningLevel = value;
+		}
+		private WarningLevel _nullPointerArithmeticWarningLevel;
+
+		/// <summary>
+		/// Indicates what warning/error level to treat usages of constant values being used within a logical operation, which may lead to redundant logic evaluation.
+		/// Clang -
+		///		https://clang.llvm.org/docs/DiagnosticsReference.html#wconstant-logical-operand
+		/// </summary>
+		/// <remarks>Triggered by || of two template-derived values inside a static_assert.</remarks>
+		[WarningsVCClang(["-Wno-constant-logical-operand"], null, null)]
+		public WarningLevel ConstantLogicalOperandWarningLevel
+		{
+			get => ResolveWarning(ParentCppCompileWarnings?.ConstantLogicalOperandWarningLevel, _constantLogicalOperandWarningLevel);
+			set => _constantLogicalOperandWarningLevel = value;
+		}
+		private WarningLevel _constantLogicalOperandWarningLevel;
+
+		/// <summary>
+		/// Indicates what warning/error level to treat assignments of enum types to a bit=field, which can lead to implicit conversations (causing unexpected results).
+		/// Clang -
+		///		https://clang.llvm.org/docs/DiagnosticsReference.html#wbitfield-enum-conversion
+		/// </summary>
+		[WarningsVCClang(["-Wno-bitfield-enum-conversion"], null, null)]
+		public WarningLevel BitfieldEnumConversion
+		{
+			get => ResolveWarning(ParentCppCompileWarnings?.BitfieldEnumConversion, _bitfieldEnumConversion);
+			set => _bitfieldEnumConversion = value;
+		}
+		private WarningLevel _bitfieldEnumConversion;
+
+		/// <summary>
+		/// Indicates what warning/error level to treat instances of pointer subtraction, where a pointer may be null.
+		/// Clang -
+		///		https://clang.llvm.org/docs/DiagnosticsReference.html#wnull-pointer-subtraction
+		/// </summary>
+		[WarningsVCClang(["-Wno-null-pointer-subtraction"], null, null)]
+		public WarningLevel NullPointerSubtractionWarningLevel
+		{
+			get => ResolveWarning(ParentCppCompileWarnings?.NullPointerSubtractionWarningLevel, _nullPointerSubtractionWarningLevel);
+			set => _nullPointerSubtractionWarningLevel = value;
+		}
+		private WarningLevel _nullPointerSubtractionWarningLevel;
+
+		/// <summary>
+		/// Indicates what warning/error level to treat dangling pointer usage.
+		/// Clang -
+		///		https://clang.llvm.org/docs/DiagnosticsReference.html#wdangling
+		/// </summary>	
+		[WarningsVCClang(["-Wno-dangling"], null, null)]
+		public WarningLevel DanglingWarningLevel
+		{
+			get => ResolveWarning(ParentCppCompileWarnings?.DanglingWarningLevel, _danglingWarningLevel);
+			set => _danglingWarningLevel = value;
+		}
+		private WarningLevel _danglingWarningLevel;
+
+		#region -- Duplicated Clang VC Compiler Disables --
+
+		/// <summary>
+		/// Indicates what warning/error level to treat unhandled enumerators in switches on enumeration-typed values.
+		/// MSVC - 
+		///		4061 - https://learn.microsoft.com/en-us/cpp/error-messages/compiler-warnings/compiler-warning-level-4-c4061 
+		///		44061 - Note: The extra 4 is not a typo, /wLXXXX sets warning XXXX to level L
+		///	Clang - 
+		///		https://clang.llvm.org/docs/DiagnosticsReference.html#wswitch-enum
+		///	</summary>
+		/// <remarks>
+		/// This is a special case where we have a duplicate field for clang <see cref="CppCompileWarnings.SwitchUnhandledEnumeratorWarningLevel"/>.
+		/// The reason for duplication is that the clang field will only participate in the shipping configuration, whereas the MSVC variant has been disabled in all.
+		/// In order to ensure no change of behaviour for MSVC toolchain within the clang compilation context, we will keep this in, and deprecate with assignments to the clang setting.
+		/// </remarks>
+		[WarningsVCClang(["-Wno-switch-enum"], null, null)]
+		public WarningLevel MSVCSwitchEnumWarningLevel
+		{
+			get => ResolveWarning(ParentCppCompileWarnings?.MSVCSwitchEnumWarningLevel, _msvcSwitchEnumWarningLevel);
+			set => _msvcSwitchEnumWarningLevel = value;
+		}
+		private WarningLevel _msvcSwitchEnumWarningLevel;
+
+		/// <summary>
+		/// Indicates what warning/error level to treat values that are computed, but not used (i.e.unnecessary evaluation of an rvalue or expression).
+		/// Clang -
+		///		https://clang.llvm.org/docs/DiagnosticsReference.html#wunused-value
+		/// </summary>
+		/// <remarks>
+		/// This is a special case where we have a duplicate field for clang <see cref="CppCompileWarnings.UnusedValueWarningLevel"/>.
+		/// The reason for duplication is that the clang field will only participate in the shipping configuration, whereas the MSVC variant has been disabled in all.
+		/// In order to ensure no change of behaviour for MSVC toolchain within the clang compilation context, we will keep this in, and deprecate with assignments to the clang setting.
+		/// </remarks>
+		[WarningsVCClang(["-Wno-unused-value"], null, null)]
+		public WarningLevel MSVCUnusedValueWarningLevel
+		{
+			get => ResolveWarning(ParentCppCompileWarnings?.MSVCUnusedValueWarningLevel, _msvcUnusedValueWarningLevel);
+			set => _msvcUnusedValueWarningLevel = value;
+		}
+		private WarningLevel _msvcUnusedValueWarningLevel;
+
+		/// <summary>
+		/// Indicates what warning/error level to treat usages of functions, methods, or variables, that are marked as deprecated.
+		/// Clang -
+		///		https://clang.llvm.org/docs/DiagnosticsReference.html#wdeprecated-declarations
+		/// </summary>
+		/// <remarks>
+		/// This is a special case where we have a duplicate field for clang <see cref="CppCompileWarnings.DeprecationWarningLevel"/>.
+		/// The reason for duplication is that the clang field will be defaulted to WARNING in clang, but not for msvc which is disabled.
+		/// In order to ensure no change of behaviour for MSVC toolchain within the clang compilation context, we will keep this in, and deprecate with assignments to the clang setting.
+		/// </remarks>
+		[WarningsVCClang(["-Wno-deprecated-declarations"], null, null)]
+		public WarningLevel MSVCDeprecationWarningLevel
+		{
+			get => ResolveWarning(ParentCppCompileWarnings?.MSVCDeprecationWarningLevel, _msvcDeprecationWarningLevel);
+			set => _msvcDeprecationWarningLevel = value;
+		}
+		private WarningLevel _msvcDeprecationWarningLevel;
+
+		#endregion -- Duplicated Clang VC Compiler Disables --
+
+		#endregion -- Clang VC Compiler Disables --
+
+		#region -- CppCompileEnvironment Settings
+
+		/// <summary>
+		/// Indicates what warning/error level to treat usages of functions, methods, or variables, that are marked as deprecated.
+		/// Clang -
+		///		https://clang.llvm.org/docs/DiagnosticsReference.html#wdeprecated-declarations
+		/// </summary>
+		/// <remarks>Disable needed for wxWidgets.</remarks>
+		[WarningsClangToolChain(["-Wno-deprecated-declarations"], null, ["-Werror=deprecated-declarations"])]
+		public WarningLevel DeprecationWarningLevel
+		{
+			get => ResolveWarning(ParentCppCompileWarnings?.DeprecationWarningLevel, _deprecationWarningLevel);
+			set => _deprecationWarningLevel = value;
+		}
+		private WarningLevel _deprecationWarningLevel;
+
+		/// <summary>
+		/// Indicates what warning/error level to treat usages of __DATE__ or __TIME__ as they prevent reproducible builds.
+		/// Clang -
+		///		https://clang.llvm.org/docs/DiagnosticsReference.html#wdate-time
+		/// </summary>
+		[DeterministicClangToolChain(null, ["-Wdate-time", "-Wno-error=date-time"], ["-Wdate-time"])]
+		public WarningLevel DeterministicWarningLevel
+		{
+			get => ResolveWarning(ParentCppCompileWarnings?.DeterministicWarningLevel, _deterministicWarningLevel);
+			set => _deterministicWarningLevel = value;
+		}
+		private WarningLevel _deterministicWarningLevel;
+
+		#endregion -- CppCompileEnvironment Settings
+
 		#endregion -- Compiler Warning Settings --
 
 		#region -- Performance Warning Settings --
@@ -360,9 +1204,10 @@ namespace UnrealBuildTool
 		/// <param name="compileEnvironment">The compile environment of the invoking tool chain.</param>
 		/// <param name="toolChainType">The type of the invoking toolchain.</param>
 		/// <param name="toolChainVersion">The version of the tool chain, if applicable.</param>
+		/// <param name="analyzer">The analyzer to use when considering arguments, if applicable.</param>
 		/// <returns>A list of the compiler command line arguments for compile warnings.</returns>
 		/// <remarks>The type of the tool chain must be the specific class within the hierarchy due to how <see cref="ClangToolChain.GetCompileArguments_WarningsAndErrors"/> applies arguments.</remarks>
-		internal IEnumerable<string> GenerateWarningCommandLineArgs(CppCompileEnvironment compileEnvironment, System.Type toolChainType, VersionNumber? toolChainVersion = null)
+		internal IEnumerable<string> GenerateWarningCommandLineArgs(CppCompileEnvironment compileEnvironment, System.Type toolChainType, VersionNumber? toolChainVersion = null, StaticAnalyzer analyzer = StaticAnalyzer.None)
 		{
 			if (!typeof(UEToolChain).IsAssignableFrom(toolChainType))
 			{
@@ -371,7 +1216,7 @@ namespace UnrealBuildTool
 				return new List<string>(0);
 			}
 
-			CompilerWarningsToolChainContext toolChainContext = new CompilerWarningsToolChainContext(compileEnvironment, BuildContextProvider, toolChainType, toolChainVersion); //pull out into request?
+			CompilerWarningsToolChainContext toolChainContext = new CompilerWarningsToolChainContext(compileEnvironment, BuildContextProvider, toolChainType, toolChainVersion, analyzer); //pull out into request?
 
 			List<string> warningArguments = new List<string>();
 
@@ -476,7 +1321,73 @@ namespace UnrealBuildTool
 			cppCompileWarnings._shadowVariableWarningLevel.ApplyWarning((cppCompileWarnings.BuildContextProvider?._buildContext.GetBuildSettings() >= BuildSettingsVersion.V2) ? WarningLevel.Error : WarningLevel.Warning, overwriteSourceValue);
 			cppCompileWarnings._unsafeTypeCastWarningLevel.ApplyWarning(WarningLevel.Off, overwriteSourceValue);
 			cppCompileWarnings._undefinedIdentifierWarningLevel.ApplyWarning(WarningLevel.Error, overwriteSourceValue);
-			cppCompileWarnings._switchUnhandledEnumeratorWarningLevel.ApplyWarning(WarningLevel.Off, overwriteSourceValue);
+			cppCompileWarnings._switchUnhandledEnumeratorWarningLevel.ApplyWarning(WarningLevel.Off, overwriteSourceValue); // Also disabled in Clang VC Toolchain
+
+			// Always clang disabled
+			cppCompileWarnings._enumEnumConversionWarningLevel.ApplyWarning(WarningLevel.Off, overwriteSourceValue);
+			cppCompileWarnings._enumFloatConversionWarningLevel.ApplyWarning(WarningLevel.Off, overwriteSourceValue);
+			cppCompileWarnings._ambiguousReversedOperatorWarningLevel.ApplyWarning(WarningLevel.Off, overwriteSourceValue);
+			cppCompileWarnings._deprecatedAnonEnumEnumConversionWarningLevel.ApplyWarning(WarningLevel.Off, overwriteSourceValue);
+			cppCompileWarnings._deprecatedVolatileWarningLevel.ApplyWarning(WarningLevel.Off, overwriteSourceValue);
+			cppCompileWarnings._orderedCompareFunctionPointers.ApplyWarning(WarningLevel.Off, overwriteSourceValue);
+			cppCompileWarnings._unusedWarningLevel.ApplyWarning(WarningLevel.Off, overwriteSourceValue);
+			cppCompileWarnings._unusedValueWarningLevel.ApplyWarning(WarningLevel.Off, overwriteSourceValue); // Also disabled in Clang VC Toolchain, but mirrored with _msvcUnusedValueWarningLevel
+			cppCompileWarnings._unknownPragmasWarningLevel.ApplyWarning(WarningLevel.Off, overwriteSourceValue);
+			cppCompileWarnings._tautologicalCompareWarningLevel.ApplyWarning(WarningLevel.Off, overwriteSourceValue);
+			cppCompileWarnings._switchWarningLevel.ApplyWarning(WarningLevel.Off, overwriteSourceValue);
+			cppCompileWarnings._invalidOffsetOfWarningLevel.ApplyWarning(WarningLevel.Off, overwriteSourceValue);
+			cppCompileWarnings._inconsistentMissingOverrideWarningLevel.ApplyWarning(WarningLevel.Off, overwriteSourceValue);
+			cppCompileWarnings._gnuStringLiteralOperatorTemplateWarningLevel.ApplyWarning(WarningLevel.Off, overwriteSourceValue);
+			cppCompileWarnings._undefinedVarTemplateWarningLevel.ApplyWarning(WarningLevel.Off, overwriteSourceValue);
+			cppCompileWarnings._bitwiseInsteadOfLogicalWarningLevel.ApplyWarning(WarningLevel.Off, overwriteSourceValue);
+			cppCompileWarnings._deprecatedCopyWithUserProvidedCopyWarningLevel.ApplyWarning(WarningLevel.Off, overwriteSourceValue); // Also disabled in Clang VC Toolchain
+			cppCompileWarnings._deprecatedCopyWarningLevel.ApplyWarning(WarningLevel.Off, overwriteSourceValue);
+			cppCompileWarnings._invalidUnevaluatedStringWarningLevel.ApplyWarning(WarningLevel.Off, overwriteSourceValue);
+			cppCompileWarnings._nanInfinityDisabledWarningLevel.ApplyWarning(WarningLevel.Off, overwriteSourceValue);
+			cppCompileWarnings._castFunctionTypeMistmatchWarningLevel.ApplyWarning(WarningLevel.Off, overwriteSourceValue);
+			cppCompileWarnings._missingTemplateArgListAfterTemplateWarningLevel.ApplyWarning(WarningLevel.Off, overwriteSourceValue);
+			cppCompileWarnings._levelExtraQualificationWarningLevel.ApplyWarning(WarningLevel.Off, overwriteSourceValue);
+			cppCompileWarnings._profileInstructWarningLevel.ApplyWarning(WarningLevel.Off, overwriteSourceValue);
+			cppCompileWarnings._backendPluginWarningLevel.ApplyWarning(WarningLevel.Off, overwriteSourceValue);
+			cppCompileWarnings._shorten64To32WarningLevel.ApplyWarning(WarningLevel.Off, overwriteSourceValue);
+			cppCompileWarnings._dllExportExplicitInstantiationDeclWarningLevel.ApplyWarning(WarningLevel.Off, overwriteSourceValue);
+
+			// Always clang vc disabled
+			cppCompileWarnings._microsoftGroupWarningLevel.ApplyWarning(WarningLevel.Off, overwriteSourceValue);
+			cppCompileWarnings._msvcIncludeWarningLevel.ApplyWarning(WarningLevel.Off, overwriteSourceValue);
+			cppCompileWarnings._pragmaPackWarningLevel.ApplyWarning(WarningLevel.Off, overwriteSourceValue);
+			cppCompileWarnings._inlineNewDeleteWarningLevel.ApplyWarning(WarningLevel.Off, overwriteSourceValue);
+			cppCompileWarnings._implicitExceptionSpecMismatchWarningLevel.ApplyWarning(WarningLevel.Off, overwriteSourceValue);
+			cppCompileWarnings._undefinedBoolConversionWarningLevel.ApplyWarning(WarningLevel.Off, overwriteSourceValue);
+			cppCompileWarnings._deprecatedWriteableStringsWarningLevel.ApplyWarning(WarningLevel.Off, overwriteSourceValue);
+			cppCompileWarnings._deprecatedRegisterWarningLevel.ApplyWarning(WarningLevel.Off, overwriteSourceValue);
+			cppCompileWarnings._logicalOpParenthesesWarningLevel.ApplyWarning(WarningLevel.Off, overwriteSourceValue);
+			cppCompileWarnings._nullArithmeticWarningLevel.ApplyWarning(WarningLevel.Off, overwriteSourceValue);
+			cppCompileWarnings._returnTypeCLinkageWarningLevel.ApplyWarning(WarningLevel.Off, overwriteSourceValue);
+			cppCompileWarnings._ignoredAttributesWarningLevel.ApplyWarning(WarningLevel.Off, overwriteSourceValue);
+			cppCompileWarnings._uninitializedWarningLevel.ApplyWarning(WarningLevel.Off, overwriteSourceValue);
+			cppCompileWarnings._returnTypeWarningLevel.ApplyWarning(WarningLevel.Off, overwriteSourceValue);
+			cppCompileWarnings._unusedParameterWarningLevel.ApplyWarning(WarningLevel.Off, overwriteSourceValue);
+			cppCompileWarnings._ignoredQualifiersWarningLevel.ApplyWarning(WarningLevel.Off, overwriteSourceValue);
+			cppCompileWarnings._expansionToDefined.ApplyWarning(WarningLevel.Off, overwriteSourceValue);
+			cppCompileWarnings._signCompareWarningLevel.ApplyWarning(WarningLevel.Off, overwriteSourceValue);
+			cppCompileWarnings._missingFieldInitializersWarningLevel.ApplyWarning(WarningLevel.Off, overwriteSourceValue);
+			cppCompileWarnings._nonPortableIncludePathWarningLevel.ApplyWarning(WarningLevel.Off, overwriteSourceValue);
+			cppCompileWarnings._invalidTokenPasteWarningLevel.ApplyWarning(WarningLevel.Off, overwriteSourceValue);
+			cppCompileWarnings._nullPointerArithmeticWarningLevel.ApplyWarning(WarningLevel.Off, overwriteSourceValue);
+			cppCompileWarnings._constantLogicalOperandWarningLevel.ApplyWarning(WarningLevel.Off, overwriteSourceValue);
+			cppCompileWarnings._bitfieldEnumConversion.ApplyWarning(WarningLevel.Off, overwriteSourceValue);
+			cppCompileWarnings._nullPointerSubtractionWarningLevel.ApplyWarning(WarningLevel.Off, overwriteSourceValue);
+			cppCompileWarnings._danglingWarningLevel.ApplyWarning(WarningLevel.Off, overwriteSourceValue);
+
+			// Special case duplicates, for MSVC guaranteed disable.
+			cppCompileWarnings._msvcUnusedValueWarningLevel.ApplyWarning(WarningLevel.Off, overwriteSourceValue);
+			cppCompileWarnings._msvcDeprecationWarningLevel.ApplyWarning(WarningLevel.Off, overwriteSourceValue);
+			cppCompileWarnings._msvcSwitchEnumWarningLevel.ApplyWarning(WarningLevel.Off, overwriteSourceValue);
+
+			// CppCompileEnvironment refactors
+			cppCompileWarnings._deprecationWarningLevel.ApplyWarning(WarningLevel.Warning, overwriteSourceValue); // Also disabled in Clang VC Toolchain
+			cppCompileWarnings._deterministicWarningLevel.ApplyWarning(WarningLevel.Off, overwriteSourceValue);
 
 			// Performance Warning Settings
 			cppCompileWarnings._pchPerformanceIssueWarningLevel.ApplyWarning(WarningLevel.Off, overwriteSourceValue);
@@ -525,7 +1436,73 @@ namespace UnrealBuildTool
 			_shadowVariableWarningLevel = WarningLevel.Warning;
 			_unsafeTypeCastWarningLevel = WarningLevel.Off;
 			_undefinedIdentifierWarningLevel = WarningLevel.Off;
-			_switchUnhandledEnumeratorWarningLevel = WarningLevel.Off;
+			_switchUnhandledEnumeratorWarningLevel = WarningLevel.Off; // Also disabled in Clang VC Toolchain
+
+			// Always clang disabled
+			_enumEnumConversionWarningLevel = WarningLevel.Off;
+			_enumFloatConversionWarningLevel = WarningLevel.Off;
+			_ambiguousReversedOperatorWarningLevel = WarningLevel.Off;
+			_deprecatedAnonEnumEnumConversionWarningLevel = WarningLevel.Off;
+			_deprecatedVolatileWarningLevel = WarningLevel.Off;
+			_orderedCompareFunctionPointers = WarningLevel.Off;
+			_unusedWarningLevel = WarningLevel.Off;
+			_unusedValueWarningLevel = WarningLevel.Off; // Also disabled in Clang VC Toolchain
+			_unknownPragmasWarningLevel = WarningLevel.Off;
+			_tautologicalCompareWarningLevel = WarningLevel.Off;
+			_switchWarningLevel = WarningLevel.Off;
+			_invalidOffsetOfWarningLevel = WarningLevel.Off;
+			_inconsistentMissingOverrideWarningLevel = WarningLevel.Off;
+			_gnuStringLiteralOperatorTemplateWarningLevel = WarningLevel.Off;
+			_undefinedVarTemplateWarningLevel = WarningLevel.Off;
+			_bitwiseInsteadOfLogicalWarningLevel = WarningLevel.Off;
+			_deprecatedCopyWithUserProvidedCopyWarningLevel = WarningLevel.Off; // Also disabled in Clang VC Toolchain
+			_deprecatedCopyWarningLevel = WarningLevel.Off;
+			_invalidUnevaluatedStringWarningLevel = WarningLevel.Off;
+			_nanInfinityDisabledWarningLevel = WarningLevel.Off;
+			_castFunctionTypeMistmatchWarningLevel = WarningLevel.Off;
+			_missingTemplateArgListAfterTemplateWarningLevel = WarningLevel.Off;
+			_levelExtraQualificationWarningLevel = WarningLevel.Off;
+			_profileInstructWarningLevel = WarningLevel.Off;
+			_backendPluginWarningLevel = WarningLevel.Off;
+			_shorten64To32WarningLevel = WarningLevel.Off;
+			_dllExportExplicitInstantiationDeclWarningLevel = WarningLevel.Off;
+
+			// Always clang vc disables
+			_microsoftGroupWarningLevel = WarningLevel.Off;
+			_msvcIncludeWarningLevel = WarningLevel.Off;
+			_pragmaPackWarningLevel = WarningLevel.Off;
+			_inlineNewDeleteWarningLevel = WarningLevel.Off;
+			_implicitExceptionSpecMismatchWarningLevel = WarningLevel.Off;
+			_undefinedBoolConversionWarningLevel = WarningLevel.Off;
+			_deprecatedWriteableStringsWarningLevel = WarningLevel.Off;
+			_deprecatedRegisterWarningLevel = WarningLevel.Off;
+			_logicalOpParenthesesWarningLevel = WarningLevel.Off;
+			_nullArithmeticWarningLevel = WarningLevel.Off;
+			_returnTypeCLinkageWarningLevel = WarningLevel.Off;
+			_ignoredAttributesWarningLevel = WarningLevel.Off;
+			_uninitializedWarningLevel = WarningLevel.Off;
+			_returnTypeWarningLevel = WarningLevel.Off;
+			_unusedParameterWarningLevel = WarningLevel.Off;
+			_ignoredQualifiersWarningLevel = WarningLevel.Off;
+			_expansionToDefined = WarningLevel.Off;
+			_signCompareWarningLevel = WarningLevel.Off;
+			_missingFieldInitializersWarningLevel = WarningLevel.Off;
+			_nonPortableIncludePathWarningLevel = WarningLevel.Off;
+			_invalidTokenPasteWarningLevel = WarningLevel.Off;
+			_nullPointerArithmeticWarningLevel = WarningLevel.Off;
+			_constantLogicalOperandWarningLevel = WarningLevel.Off;
+			_bitfieldEnumConversion = WarningLevel.Off;
+			_nullPointerSubtractionWarningLevel = WarningLevel.Off;
+			_danglingWarningLevel = WarningLevel.Off;
+			
+			// Special case duplicates, for MSVC guaranteed disable.
+			_msvcUnusedValueWarningLevel = WarningLevel.Off;
+			_msvcDeprecationWarningLevel = WarningLevel.Off;
+			_msvcSwitchEnumWarningLevel = WarningLevel.Off;
+
+			// CppCompileEnvironment refactors
+			_deprecationWarningLevel = WarningLevel.Warning;
+			_deterministicWarningLevel = WarningLevel.Off;
 
 			// Performance Warning Settings
 			_pchPerformanceIssueWarningLevel = WarningLevel.Off;
diff --git a/Engine/Source/Programs/UnrealBuildTool/Configuration/ReadOnlyCppCompileWarnings.cs b/Engine/Source/Programs/UnrealBuildTool/Configuration/ReadOnlyCppCompileWarnings.cs
--- a/Engine/Source/Programs/UnrealBuildTool/Configuration/ReadOnlyCppCompileWarnings.cs
+++ b/Engine/Source/Programs/UnrealBuildTool/Configuration/ReadOnlyCppCompileWarnings.cs
@@ -44,6 +44,79 @@ namespace UnrealBuildTool
 		public WarningLevel UnsafeTypeCastWarningLevel => _inner.UnsafeTypeCastWarningLevel;
 		public WarningLevel SwitchUnhandledEnumeratorWarningLevel => _inner.SwitchUnhandledEnumeratorWarningLevel;
 
+		#region -- Clang Compiler Disables --
+
+		public WarningLevel EnumEnumConversionWarningLevel => _inner.EnumEnumConversionWarningLevel;
+		public WarningLevel EnumFloatConversionWarningLevel => _inner.EnumFloatConversionWarningLevel;
+		public WarningLevel AmbiguousReversedOperatorWarningLevel => _inner.AmbiguousReversedOperatorWarningLevel;
+		public WarningLevel DeprecatedAnonEnumEnumConversionWarningLevel => _inner.DeprecatedAnonEnumEnumConversionWarningLevel;
+		public WarningLevel DeprecatedVolatileWarningLevel => _inner.DeprecatedVolatileWarningLevel;
+		public WarningLevel OrderedCompareFunctionPointers => _inner.OrderedCompareFunctionPointers;
+		public WarningLevel UnusedWarningLevel => _inner.UnusedWarningLevel;
+		public WarningLevel UnusedValueWarningLevel => _inner.UnusedValueWarningLevel;
+		public WarningLevel UnknownPragmasWarningLevel => _inner.UnknownPragmasWarningLevel;
+		public WarningLevel TautologicalCompareWarningLevel => _inner.TautologicalCompareWarningLevel;
+		public WarningLevel SwitchWarningLevel => _inner.SwitchWarningLevel;
+		public WarningLevel InvalidOffsetWarningLevel => _inner.InvalidOffsetWarningLevel;
+		public WarningLevel InconsistentMissingOverrideWarningLevel => _inner.InconsistentMissingOverrideWarningLevel;
+		public WarningLevel GNUStringLiteralOperatorTemplateWarningLevel => _inner.GNUStringLiteralOperatorTemplateWarningLevel;
+		public WarningLevel UndefinedVarTemplateWarningLevel => _inner.UndefinedVarTemplateWarningLevel;
+		public WarningLevel DeprecatedCopyWithUserProvidedCopyWarningLevel => _inner.DeprecatedCopyWithUserProvidedCopyWarningLevel;
+		public WarningLevel BitwiseInsteadOfLogicalWarningLevel => _inner.BitwiseInsteadOfLogicalWarningLevel;
+		public WarningLevel DeprecatedCopyWarningLevel => _inner.DeprecatedCopyWarningLevel;
+		public WarningLevel InvalidUnevaluatedStringWarningLevel => _inner.InvalidUnevaluatedStringWarningLevel;
+		public WarningLevel NaNInfinityDisabledWarningLevel => _inner.NaNInfinityDisabledWarningLevel;
+		public WarningLevel CastFunctionTypeMistmatchWarningLevel => _inner.CastFunctionTypeMistmatchWarningLevel;
+		public WarningLevel MissingTemplateArgListAfterTemplateWarningLevel => _inner.MissingTemplateArgListAfterTemplateWarningLevel;
+		public WarningLevel LevelExtraQualificationWarningLevel => _inner.LevelExtraQualificationWarningLevel;
+		public WarningLevel ProfileInstructWarningLevel => _inner.ProfileInstructWarningLevel;
+		public WarningLevel BackendPluginWarningLevel => _inner.BackendPluginWarningLevel;
+		public WarningLevel Shorten64To32WarningLevel => _inner.Shorten64To32WarningLevel;
+		public WarningLevel DllExportExplicitInstantiationDeclWarningLevel => _inner.DllExportExplicitInstantiationDeclWarningLevel;
+		public WarningLevel DeprecationWarningLevel => _inner.DeprecationWarningLevel;
+		public WarningLevel DeterministicWarningLevel => _inner.DeterministicWarningLevel;
+
+		#endregion -- Clang Compiler Disables --
+
+		#region -- Clang VC Compiler Disables --
+
+		public WarningLevel MicrosoftGroupWarningLevel => _inner.MicrosoftGroupWarningLevel;
+		public WarningLevel MSVCIncludeWarningLevel => _inner.MSVCIncludeWarningLevel;
+		public WarningLevel PragmaPackWarningLevel => _inner.PragmaPackWarningLevel;
+		public WarningLevel InlineNewDeleteWarningLevel => _inner.InlineNewDeleteWarningLevel;
+		public WarningLevel ImplicitExceptionSpecMismatchWarningLevel => _inner.ImplicitExceptionSpecMismatchWarningLevel;
+		public WarningLevel UndefinedBoolConversionWarningLevel => _inner.UndefinedBoolConversionWarningLevel;
+		public WarningLevel DeprecatedWriteableStringsWarningLevel => _inner.DeprecatedWriteableStringsWarningLevel;
+		public WarningLevel DeprecatedRegisterWarningLevel => _inner.DeprecatedRegisterWarningLevel;
+		public WarningLevel LogicalOpParenthesesWarningLevel => _inner.LogicalOpParenthesesWarningLevel;
+		public WarningLevel NullArithmeticWarningLevel => _inner.NullArithmeticWarningLevel;
+		public WarningLevel ReturnTypeCLinkageWarningLevel => _inner.ReturnTypeCLinkageWarningLevel;
+		public WarningLevel IgnoredAttributesWarningLevel => _inner.IgnoredAttributesWarningLevel;
+		public WarningLevel UninitializedWarningLevel => _inner.UninitializedWarningLevel;
+		public WarningLevel ReturnTypeWarningLevel => _inner.ReturnTypeWarningLevel;
+		public WarningLevel UnusedParameterWarningLevel => _inner.UnusedParameterWarningLevel;
+		public WarningLevel IgnoredQualifiersWarningLevel => _inner.IgnoredQualifiersWarningLevel;
+		public WarningLevel ExpansionToDefined => _inner.ExpansionToDefined;
+		public WarningLevel SignCompareWarningLevel => _inner.SignCompareWarningLevel;
+		public WarningLevel MissingFieldInitializersWarningLevel => _inner.MissingFieldInitializersWarningLevel;
+		public WarningLevel NonPortableIncludePathWarningLevel => _inner.NonPortableIncludePathWarningLevel;
+		public WarningLevel InvalidTokenPasteWarningLevel => _inner.InvalidTokenPasteWarningLevel;
+		public WarningLevel NullPointerArithmeticWarningLevel => _inner.NullPointerArithmeticWarningLevel;
+		public WarningLevel ConstantLogicalOperandWarningLevel => _inner.ConstantLogicalOperandWarningLevel;
+		public WarningLevel BitfieldEnumConversion => _inner.BitfieldEnumConversion;
+		public WarningLevel NullPointerSubtractionWarningLevel => _inner.NullPointerSubtractionWarningLevel;
+		public WarningLevel DanglingWarningLevel => _inner.DanglingWarningLevel;
+
+		#region -- Duplicated Clang VC Compiler Disables --
+		
+		public WarningLevel MSVCSwitchEnumWarningLevel => _inner.MSVCSwitchEnumWarningLevel;
+		public WarningLevel MSVCUnusedValueWarningLevel => _inner.MSVCUnusedValueWarningLevel;
+		public WarningLevel MSVCDeprecationWarningLevel => _inner.MSVCDeprecationWarningLevel;
+		
+		#endregion -- Duplicated Clang VC Compiler Disables --
+
+		#endregion -- Clang VC Compiler Disables --
+
 		#endregion -- Compiler Warning Settings --
 
 		#region -- Performance Warning Settings --
@@ -53,7 +126,7 @@ namespace UnrealBuildTool
 		#endregion -- Performance Warning Settings --
 
 		#region -- Rules Validation Warning Settings --
-		
+
 		public WarningLevel ModuleUnsupportedWarningLevel => _inner.ModuleUnsupportedWarningLevel;
 		public WarningLevel PluginModuleUnsupportedWarningLevel => _inner.PluginModuleUnsupportedWarningLevel;
 		public WarningLevel ModuleIncludePathWarningLevel => _inner.ModuleIncludePathWarningLevel;
diff --git a/Engine/Source/Programs/UnrealBuildTool/Platform/Windows/VCToolChain.cs b/Engine/Source/Programs/UnrealBuildTool/Platform/Windows/VCToolChain.cs
--- a/Engine/Source/Programs/UnrealBuildTool/Platform/Windows/VCToolChain.cs
+++ b/Engine/Source/Programs/UnrealBuildTool/Platform/Windows/VCToolChain.cs
@@ -10,6 +10,7 @@ using EpicGames.Core;
 using Microsoft.CodeAnalysis;
 using Microsoft.Extensions.Logging;
 using UnrealBuildBase;
+using UnrealBuildTool.VerseVMBytecode;
 
 namespace UnrealBuildTool
 {
@@ -1126,8 +1127,6 @@ namespace UnrealBuildTool
 				Arguments.Add("/we4996");
 			}
 
-			Arguments.AddRange(CompileEnvironment.CppCompileWarnings.GenerateWarningCommandLineArgs(CompileEnvironment, typeof(VCToolChain)));
-
 			if (CompileEnvironment.bUseHeaderUnitsForPch)
 			{
 				Arguments.Add("/wd4324"); // 'struct_name' : structure was padded due to __declspec(align())
@@ -1138,12 +1137,15 @@ namespace UnrealBuildTool
 				Arguments.Add("/wd4180"); // qualifier applied to function type has no meaning; ignored
 			}
 
+
 			// Downgrade C4702: unreachable code to a warning when running LTCG or PGO
 			if (CompileEnvironment.bPGOOptimize || CompileEnvironment.bPGOProfile || CompileEnvironment.bAllowLTCG)
 			{
 				Arguments.Add("/wd4702 /w44702");
 			}
 
+			AppendCLArguments_CompileWarnings(CompileEnvironment, Arguments);
+
 			if (CompileEnvironment.Architecture == UnrealArch.Arm64ec)
 			{
 				Arguments.Add("/arm64EC");
@@ -1161,7 +1163,7 @@ namespace UnrealBuildTool
 		protected virtual void AppendCLArguments_H(CppCompileEnvironment CompileEnvironment, List<string> Arguments)
 		{
 			AppendCLArguments_CPP(CompileEnvironment, Arguments);
-
+			
 			if (Target.WindowsPlatform.Compiler.IsClang())
 			{
 				ClangWarnings.GetHeaderDisabledWarnings(Arguments);
@@ -1205,16 +1207,7 @@ namespace UnrealBuildTool
 				Arguments.Add("/GR-");
 			}
 
-			// Set warning level.
-			// Restrictive during regular compilation.
-			Arguments.Add("/W4");
-
-			// Treat warnings as errors
-			if (CompileEnvironment.bWarningsAsErrors)
-			{
-				Arguments.Add("/WX");
-			}
-
+		
 			switch (CompileEnvironment.CppStandard)
 			{
 				case CppStandardVersion.Cpp14:
@@ -1272,26 +1265,6 @@ namespace UnrealBuildTool
 					Arguments.Add("-Xclang -gcodeview-ghash");
 				}
 
-				// Disable specific warnings that cause problems with Clang
-				// NOTE: These must appear after we set the MSVC warning level
-
-				// @todo clang: Ideally we want as few warnings disabled as possible
-				//
-
-				// Treat all warnings as errors by default
-				Arguments.Add("-Werror");                                   // https://clang.llvm.org/docs/UsersManual.html#cmdoption-werror
-
-				ClangWarnings.GetEnabledWarnings(Arguments);
-
-				VersionNumber ClangVersion = Target.WindowsPlatform.Compiler == WindowsCompiler.Intel ? MicrosoftPlatformSDK.GetClangVersionForIntelCompiler(EnvVars.CompilerPath) : EnvVars.CompilerVersion;
-				ClangWarnings.GetDisabledWarnings(CompileEnvironment,
-					Target.StaticAnalyzer,
-					ClangVersion,
-					Arguments);
-
-				// Additional disabled warnings for msvc. Everything below should be checked if it is necessary
-				ClangWarnings.GetVCDisabledWarnings(Arguments);
-
 				if (CompileEnvironment.bEnableAutoRTFMInstrumentation)
 				{
 					Arguments.Add("-fautortfm");
@@ -1325,6 +1298,13 @@ namespace UnrealBuildTool
 			Arguments.Add("/TC");
 
 			// Level 0 warnings.  Needed for external C projects that produce warnings at higher warning levels.
+			// Foricbly remove "/W4" in order to avoid excessive warning that we know will happen.
+			// > Command line warning D9025 : overriding '/W4' with '/W0'
+			if (Arguments.Contains("/W4"))
+			{
+				Arguments.Remove("/W4");
+			}
+
 			Arguments.Add("/W0");
 
 			// Select C Standard version available
@@ -1370,6 +1350,37 @@ namespace UnrealBuildTool
 			}
 		}
 
+		void AppendCLArguments_CompileWarnings(CppCompileEnvironment CompileEnvironment, List<string> Arguments)
+		{
+			// Set warning level.
+			// Restrictive during regular compilation.
+			Arguments.Add("/W4");
+
+			// Treat warnings as errors
+			if (CompileEnvironment.bWarningsAsErrors)
+			{
+				Arguments.Add("/WX");
+			}
+
+			if (Target.WindowsPlatform.Compiler.IsClang())
+			{
+				// Disable specific warnings that cause problems with Clang
+				// NOTE: These must appear after we set the MSVC warning level
+
+				// @todo clang: Ideally we want as few warnings disabled as possible
+
+				// Treat all warnings as errors by default
+				Arguments.Add("-Werror");   // https://clang.llvm.org/docs/UsersManual.html#cmdoption-werror
+
+				VersionNumber ClangVersion = Target.WindowsPlatform.Compiler == WindowsCompiler.Intel ? MicrosoftPlatformSDK.GetClangVersionForIntelCompiler(EnvVars.CompilerPath) : EnvVars.CompilerVersion;
+
+				ClangWarnings.GetEnabledWarnings(Arguments);
+				ClangWarnings.GetDisabledWarnings(CompileEnvironment, Target.StaticAnalyzer, ClangVersion, Arguments);
+			}
+
+			Arguments.AddRange(CompileEnvironment.CppCompileWarnings.GenerateWarningCommandLineArgs(CompileEnvironment, typeof(VCToolChain)));
+		}
+
 		protected virtual void AppendLinkArguments(LinkEnvironment LinkEnvironment, List<string> Arguments)
 		{
 			if (Target.WindowsPlatform.Compiler == WindowsCompiler.Clang && Target.WindowsPlatform.bAllowClangLinker)
@@ -2581,6 +2592,7 @@ namespace UnrealBuildTool
 		public override IEnumerable<string> GetGlobalCommandLineArgs(CppCompileEnvironment CompileEnvironment)
 		{
 			List<string> Arguments = new();
+			
 			AppendCLArguments_Global(new(CompileEnvironment), Arguments);
 			return Arguments;
 		}
@@ -2588,13 +2600,16 @@ namespace UnrealBuildTool
 		public override IEnumerable<string> GetCPPCommandLineArgs(CppCompileEnvironment CompileEnvironment)
 		{
 			List<string> Arguments = new();
+
 			AppendCLArguments_CPP(CompileEnvironment, Arguments);
+
 			return Arguments;
 		}
 
 		public override IEnumerable<string> GetCCommandLineArgs(CppCompileEnvironment CompileEnvironment)
 		{
 			List<string> Arguments = new();
+
 			AppendCLArguments_C(CompileEnvironment, Arguments);
 			return Arguments;
 		}
diff --git a/Engine/Source/Programs/UnrealBuildTool/System/CppCompileEnvironment.cs b/Engine/Source/Programs/UnrealBuildTool/System/CppCompileEnvironment.cs
--- a/Engine/Source/Programs/UnrealBuildTool/System/CppCompileEnvironment.cs
+++ b/Engine/Source/Programs/UnrealBuildTool/System/CppCompileEnvironment.cs
@@ -366,7 +366,11 @@ namespace UnrealBuildTool
 		/// <summary>
 		/// Whether to warn about deprecated variables
 		/// </summary>
-		public WarningLevel DeprecationWarningLevel = WarningLevel.Warning;
+		public WarningLevel DeprecationWarningLevel
+		{
+			get => CppCompileWarnings.DeprecationWarningLevel;
+			set => CppCompileWarnings.DeprecationWarningLevel = value;
+		}
 
 		/// <summary>
 		/// The compiler warnings associacted with the compile environment.
@@ -715,7 +719,11 @@ namespace UnrealBuildTool
 		/// <summary>
 		/// Set flags for determinstic compile warnings.
 		/// </summary>
-		public WarningLevel DeterministicWarningLevel = WarningLevel.Off;
+		public WarningLevel DeterministicWarningLevel
+		{
+			get => CppCompileWarnings.DeterministicWarningLevel;
+			set => CppCompileWarnings.DeterministicWarningLevel = value;
+		}
 
 		/// <summary>
 		/// Emits compilation errors for incorrect UE_LOG format strings.
@@ -786,7 +794,6 @@ namespace UnrealBuildTool
 			bEnableObjCExceptions = Other.bEnableObjCExceptions;
 			bEnableObjCAutomaticReferenceCounting = Other.bEnableObjCAutomaticReferenceCounting;
 			DefaultWarningLevel = Other.DefaultWarningLevel;
-			DeprecationWarningLevel = Other.DeprecationWarningLevel;
 			CppCompileWarnings = CppCompileWarnings.CreateShallowCopy(Other.CppCompileWarnings);
 			bWarningsAsErrors = Other.bWarningsAsErrors;
 			bDisableStaticAnalysis = Other.bDisableStaticAnalysis;
@@ -853,7 +860,6 @@ namespace UnrealBuildTool
 			bEnableCoroutines = Other.bEnableCoroutines;
 			IncludeOrderVersion = Other.IncludeOrderVersion;
 			bDeterministic = Other.bDeterministic;
-			DeterministicWarningLevel = Other.DeterministicWarningLevel;
 			CrashDiagnosticDirectory = Other.CrashDiagnosticDirectory;
 			bValidateFormatStrings = Other.bValidateFormatStrings;
 			bValidateInternalApi = Other.bValidateInternalApi;
diff --git a/Engine/Source/Programs/UnrealBuildTool/ToolChain/ClangWarnings.cs b/Engine/Source/Programs/UnrealBuildTool/ToolChain/ClangWarnings.cs
--- a/Engine/Source/Programs/UnrealBuildTool/ToolChain/ClangWarnings.cs
+++ b/Engine/Source/Programs/UnrealBuildTool/ToolChain/ClangWarnings.cs
@@ -17,124 +17,7 @@ namespace UnrealBuildTool
 
 		internal static void GetDisabledWarnings(CppCompileEnvironment CompileEnvironment, StaticAnalyzer Analyzer, VersionNumber ClangVersion, List<string> Arguments)
 		{
-			Arguments.Add("-Wno-enum-enum-conversion");                 // https://clang.llvm.org/docs/DiagnosticsReference.html#wenum-enum-conversion					// ?? no reason given
-			Arguments.Add("-Wno-enum-float-conversion");                // https://clang.llvm.org/docs/DiagnosticsReference.html#wenum-float-conversion					// ?? no reason given
-
-			// C++20 warnings that should be addressed
-			if (CompileEnvironment.CppStandard >= CppStandardVersion.Cpp20)
-			{
-				Arguments.Add("-Wno-ambiguous-reversed-operator");          // https://clang.llvm.org/docs/DiagnosticsReference.html#wambiguous-reversed-operator
-				Arguments.Add("-Wno-deprecated-anon-enum-enum-conversion"); // https://clang.llvm.org/docs/DiagnosticsReference.html#wdeprecated-anon-enum-enum-conversion
-				Arguments.Add("-Wno-deprecated-volatile");                  // https://clang.llvm.org/docs/DiagnosticsReference.html#wdeprecated-volatile
-			}
-
-			if (ClangVersion >= new VersionNumber(13))
-			{
-				Arguments.Add("-Wno-unused-but-set-variable");           // https://clang.llvm.org/docs/DiagnosticsReference.html#wunused-but-set-variable				// new warning for clang 13
-				Arguments.Add("-Wno-unused-but-set-parameter");          // https://clang.llvm.org/docs/DiagnosticsReference.html#wunused-but-set-parameter				// new warning for clang 13
-				Arguments.Add("-Wno-ordered-compare-function-pointers"); // https://clang.llvm.org/docs/DiagnosticsReference.html#wordered-compare-function-pointers	// new warning for clang 13
-			}
-			if (ClangVersion >= new VersionNumber(14))
-			{
-				Arguments.Add("-Wno-bitwise-instead-of-logical");       // https://clang.llvm.org/docs/DiagnosticsReference.html#wbitwise-instead-of-logical			// new warning for clang 14
-			}
-			if (ClangVersion >= new VersionNumber(16))
-			{
-				Arguments.Add("-Wno-deprecated-copy");                  // https://clang.llvm.org/docs/DiagnosticsReference.html#wdeprecated-copy						// new warning for clang 16
-				Arguments.Add("-Wno-deprecated-copy-with-user-provided-copy");
-			}
-			if (ClangVersion >= new VersionNumber(17))
-			{
-				bool bIsAndroidClang17 = ClangVersion == new VersionNumber(17, 0, 2) && CompileEnvironment.Platform == UnrealTargetPlatform.Android;
-				if (CompileEnvironment.CppStandard < CppStandardVersion.Latest && !bIsAndroidClang17) // Android clang 17.0.2 in NDK r26b is missing this warning
-				{
-					Arguments.Add("-Wno-invalid-unevaluated-string");   // https://clang.llvm.org/docs/DiagnosticsReference.html#winvalid-unevaluated-string			// new warning for clang 17
-				}
-			}
-			if (ClangVersion >= new VersionNumber(18))
-			{
-				Arguments.Add("-Wno-nan-infinity-disabled");            // https://clang.llvm.org/docs/DiagnosticsReference.html#wnan-infinity-disabled					// We use the NAN macro in a few places to initialize floats while we also set -ffast-math, which disables NaN support.
-																																				// It could be easily fixed in the engine, but would create the risk of Win64 building fine and the code failing just on Clang-based platforms.
-																																				// We tend to use NAN just as a bit pattern, and not rely on it in calculations, so it should be reasonably safe to disable.
-			}
-			if (ClangVersion >= new VersionNumber(19))
-			{
-				bool bIsAndroidClang19 = ClangVersion == new VersionNumber(19, 0, 0) && CompileEnvironment.Platform == UnrealTargetPlatform.Android;
-				Arguments.Add("-Wno-extra-qualification");              // https://clang.llvm.org/docs/DiagnosticsReference.html#wextra-qualification					// new warning for clang 19
-				if (!bIsAndroidClang19) // android NDK r28.0.12433566 and r28.0.12674087 are missing these warnings 
-				{
-					Arguments.Add("-Wno-cast-function-type-mismatch");
-					Arguments.Add("-Wno-missing-template-arg-list-after-template-kw");
-				}
-			}
-
-			Arguments.Add("-Wno-gnu-string-literal-operator-template"); // https://clang.llvm.org/docs/DiagnosticsReference.html#wgnu-string-literal-operator-template	// We use this feature to allow static FNames.
-			Arguments.Add("-Wno-inconsistent-missing-override");        // https://clang.llvm.org/docs/DiagnosticsReference.html#winconsistent-missing-override			// ?? no reason given
-			Arguments.Add("-Wno-invalid-offsetof");                     // https://clang.llvm.org/docs/DiagnosticsReference.html#winvalid-offsetof						// needed to suppress warnings about using offsetof on non-POD types.
-			Arguments.Add("-Wno-switch");                               // https://clang.llvm.org/docs/DiagnosticsReference.html#wswitch								// this hides the "enumeration value 'XXXXX' not handled in switch [-Wswitch]" warnings - we should maybe remove this at some point and add UE_LOG(, Fatal, ) to default cases
-			Arguments.Add("-Wno-tautological-compare");                 // https://clang.llvm.org/docs/DiagnosticsReference.html#wtautological-compare					// this hides the "warning : comparison of unsigned expression < 0 is always false" type warnings due to constant comparisons, which are possible with template arguments
-			Arguments.Add("-Wno-unknown-pragmas");                      // https://clang.llvm.org/docs/DiagnosticsReference.html#wunknown-pragmas						// Slate triggers this (with its optimize on/off pragmas)
-			Arguments.Add("-Wno-unused-function");                      // https://clang.llvm.org/docs/DiagnosticsReference.html#wunused-function						// this will hide the warnings about static functions in headers that aren't used in every single .cpp file
-			Arguments.Add("-Wno-unused-lambda-capture");                // https://clang.llvm.org/docs/DiagnosticsReference.html#wunused-lambda-capture					// suppressed because capturing of compile-time constants is seemingly inconsistent. And MSVC doesn't do that.
-			Arguments.Add("-Wno-unused-local-typedef");                 // https://clang.llvm.org/docs/DiagnosticsReference.html#wunused-local-typedef					// clang is being overly strict here? PhysX headers trigger this.
-			Arguments.Add("-Wno-unused-private-field");                 // https://clang.llvm.org/docs/DiagnosticsReference.html#wunused-private-field					// this will prevent the issue of warnings for unused private variables. MultichannelTcpSocket.h triggers this, possibly more
-			Arguments.Add("-Wno-unused-variable");                      // https://clang.llvm.org/docs/DiagnosticsReference.html#wunused-variable						// ?? no reason given
-			Arguments.Add("-Wno-undefined-var-template");               // https://clang.llvm.org/docs/DiagnosticsReference.html#wundefined-var-template				// not really a good warning to disable
-
-			// Profile Guided Optimization (PGO) and Link Time Optimization (LTO)
-			if (CompileEnvironment.bPGOOptimize)
-			{
-				//
-				// Clang emits warnings for each compiled object file that doesn't have a matching entry in the profile data.
-				// This can happen when the profile data is older than the binaries we're compiling.
-				//
-				// Disable these warnings. They are far too verbose.
-				//
-				Arguments.Add("-Wno-profile-instr-out-of-date");        // https://clang.llvm.org/docs/DiagnosticsReference.html#wprofile-instr-out-of-date
-				Arguments.Add("-Wno-profile-instr-unprofiled");         // https://clang.llvm.org/docs/DiagnosticsReference.html#wprofile-instr-unprofiled
-
-				// apparently there can be hashing conflicts with PGO which can result in:
-				// 'Function control flow change detected (hash mismatch)' warnings. 
-				Arguments.Add("-Wno-backend-plugin");                   // https://clang.llvm.org/docs/DiagnosticsReference.html#wbackend-plugin
-			}
-
-			// shipping builds will cause this warning with "ensure", so disable only in those case
-			if (CompileEnvironment.Configuration == CppConfiguration.Shipping || Analyzer != StaticAnalyzer.None)
-			{
-				Arguments.Add("-Wno-unused-value");                     // https://clang.llvm.org/docs/DiagnosticsReference.html#wunused-value
-			}
-
-			// https://clang.llvm.org/docs/DiagnosticsReference.html#wdeprecated-declarations
-			if (CompileEnvironment.DeprecationWarningLevel == WarningLevel.Error)
-			{
-				// TODO: This may be unnecessary with -Werror
-				Arguments.Add("-Werror=deprecated-declarations");
-			}
-
-			// Warn if __DATE__ or __TIME__ are used as they prevent reproducible builds
-			if (CompileEnvironment.bDeterministic)
-			{
-				// https://clang.llvm.org/docs/DiagnosticsReference.html#wdate-time
-				if (CompileEnvironment.DeterministicWarningLevel == WarningLevel.Error)
-				{
-					Arguments.Add("-Wdate-time");
-				}
-				else if (CompileEnvironment.DeterministicWarningLevel == WarningLevel.Warning)
-				{
-					Arguments.Add("-Wdate-time -Wno-error=date-time");
-				}
-			}
-
 			Arguments.AddRange(CompileEnvironment.CppCompileWarnings.GenerateWarningCommandLineArgs(CompileEnvironment, typeof(ClangToolChain), ClangVersion));
-
-			// -Wimplicit-int-conversion also controls -Wshorten-64-to-32 as of clang 19
-			if (ClangVersion >= new VersionNumber(19))
-			{
-				Arguments.Add("-Wno-shorten-64-to-32");                 // https://clang.llvm.org/docs/DiagnosticsReference.html#wshorten-64-to-32						// new warning for clang 19
-			}
-
-			// The code base contains lots of places where we do "extern template X_API class ..." and we want to keep doing that to reduce compile times.
-			Arguments.Add("-Wno-dllexport-explicit-instantiation-decl");
 		}
 
 		internal static void GetHeaderDisabledWarnings(List<string> Arguments)
@@ -145,55 +28,6 @@ namespace UnrealBuildTool
 			Arguments.Add("-Wno-#pragma-messages");
 		}
 
-		// Additional disabled warnings for msvc. Everything below should be checked if it is necessary
-		internal static void GetVCDisabledWarnings(List<string> Arguments)
-		{
-			// Allow Microsoft-specific syntax to slide, even though it may be non-standard.  Needed for Windows headers.
-			Arguments.Add("-Wno-microsoft");
-
-			// @todo clang: Hack due to how we have our 'DummyPCH' wrappers setup when using unity builds.  This warning should not be disabled!!
-			Arguments.Add("-Wno-msvc-include");
-
-			// This is disabled because clang explicitly warns about changing pack alignment in a header and not
-			// restoring it afterwards, which is something we do with the Pre/PostWindowsApi.h headers.
-			// @todo clang: push/pop this in  Pre/PostWindowsApi.h headers instead?
-			Arguments.Add("-Wno-pragma-pack");
-
-			Arguments.Add("-Wno-inline-new-delete");    // @todo clang: We declare operator new as inline.  Clang doesn't seem to like that.
-			Arguments.Add("-Wno-implicit-exception-spec-mismatch");
-
-			// Sometimes we compare 'this' pointers against nullptr, which Clang warns about by default
-			Arguments.Add("-Wno-undefined-bool-conversion");
-
-			// @todo clang: Disabled warnings were copied from MacToolChain for the most part
-			Arguments.Add("-Wno-deprecated-writable-strings");
-			Arguments.Add("-Wno-deprecated-register");
-			Arguments.Add("-Wno-switch-enum");
-			Arguments.Add("-Wno-logical-op-parentheses");   // needed for external headers we shan't change
-			Arguments.Add("-Wno-null-arithmetic");          // needed for external headers we shan't change
-			Arguments.Add("-Wno-deprecated-declarations");  // needed for wxWidgets
-			Arguments.Add("-Wno-return-type-c-linkage");    // needed for PhysX
-			Arguments.Add("-Wno-ignored-attributes");       // needed for nvtesslib
-			Arguments.Add("-Wno-uninitialized");
-			Arguments.Add("-Wno-return-type");              // needed for external headers we shan't change 
-
-			// @todo clang: Sorry for adding more of these, but I couldn't read my output log. Most should probably be looked at
-			Arguments.Add("-Wno-unused-parameter");         // Unused function parameter. A lot are named 'bUnused'...
-			Arguments.Add("-Wno-ignored-qualifiers");       // const ignored when returning by value e.g. 'const int foo() { return 4; }'
-			Arguments.Add("-Wno-expansion-to-defined");     // Usage of 'defined(X)' in a macro definition. Gives different results under MSVC
-			Arguments.Add("-Wno-sign-compare");             // Signed/unsigned comparison - millions of these
-			Arguments.Add("-Wno-missing-field-initializers"); // Stupid warning, generated when you initialize with MyStruct A = {0};
-			Arguments.Add("-Wno-nonportable-include-path");
-			Arguments.Add("-Wno-invalid-token-paste");
-			Arguments.Add("-Wno-null-pointer-arithmetic");
-			Arguments.Add("-Wno-constant-logical-operand"); // Triggered by || of two template-derived values inside a static_assert
-			Arguments.Add("-Wno-unused-value");
-			Arguments.Add("-Wno-bitfield-enum-conversion");
-			Arguments.Add("-Wno-deprecated-copy-with-user-provided-copy");
-			Arguments.Add("-Wno-null-pointer-subtraction");
-			Arguments.Add("-Wno-dangling");
-		}
-
 		// Additional disabled warnings for Intel. Everything below should be checked if it is necessary
 		internal static void GetIntelDisabledWarnings(List<string> Arguments)
 		{
