Mac: Bump minOS for Editor/Mac runtime to 14.00

Also enables ld_prime (the new linker).

diff --git a/Engine/Config/Mac/Mac_SDK.json b/Engine/Config/Mac/Mac_SDK.json
--- a/Engine/Config/Mac/Mac_SDK.json
+++ b/Engine/Config/Mac/Mac_SDK.json
@@ -1,12 +1,12 @@
 {
 	"ParentSDKFile": "../Apple/Apple_SDK.json",
 
-	"EditorDeploymentTarget": "13.00",
-	"EditorBuildTarget": "13.00",
+	"EditorDeploymentTarget": "14.00",
+	"EditorBuildTarget": "14.00",
 
-	"DeploymentTarget": "13.00",
-	"BuildTarget": "13.00",
+	"DeploymentTarget": "14.00",
+	"BuildTarget": "14.00",
 
-	"//": "macOS Ventura 13.0 (Released October 24, 2022)",
-	"MinSoftwareVersion": "13.00"
+	"//": "macOS Ventura 14.0 (Released Sept 23, 2023)",
+	"MinSoftwareVersion": "14.00"
 }
diff --git a/Engine/Source/Programs/UnrealBuildTool/Matchers/CompileEventMatcher.cs b/Engine/Source/Programs/UnrealBuildTool/Matchers/CompileEventMatcher.cs
--- a/Engine/Source/Programs/UnrealBuildTool/Matchers/CompileEventMatcher.cs
+++ b/Engine/Source/Programs/UnrealBuildTool/Matchers/CompileEventMatcher.cs
@@ -66,7 +66,6 @@ namespace UnrealBuildTool.Matchers
 		static readonly Regex s_scriptCompilePattern = new Regex(@"^\s*[A-Za-z0-9_\.]+ ERROR:.* [A-Za-z_]+ failed to compile\.");
 		static readonly Regex s_cscSummaryPattern = new Regex(@"^\s+\d+ (?:Warning|Error)\(s\)");
 		static readonly Regex s_cscOutputPattern = new Regex(@"^  [^ ]+ -> ");
-		static readonly Regex s_xcodeClassicPattern = new Regex(@"ld_classic is deprecated and will be removed in a future release");
 
 		static readonly string[] s_invalidExtensions =
 		{
@@ -210,13 +209,6 @@ namespace UnrealBuildTool.Matchers
 				return true;
 			}
 
-			// Temporary XCode silencing when minOS is too low
-			if (builder.Current.IsMatch(s_xcodeClassicPattern))
-			{
-				outEvent = builder.ToMatch(LogEventPriority.Highest, LogLevel.Information, KnownLogEvents.Systemic_XCode);
-				return true;
-			}
-
 			outEvent = null;
 			return false;
 		}
diff --git a/Engine/Source/Programs/UnrealBuildTool/ToolChain/AppleToolChain.cs b/Engine/Source/Programs/UnrealBuildTool/ToolChain/AppleToolChain.cs
--- a/Engine/Source/Programs/UnrealBuildTool/ToolChain/AppleToolChain.cs
+++ b/Engine/Source/Programs/UnrealBuildTool/ToolChain/AppleToolChain.cs
@@ -819,7 +819,7 @@ namespace UnrealBuildTool
 				AppleClangVersion = ((AppleToolChainInfo)GetToolChainInfo()).AppleClangVersion;
 			}
 			// Temp solution for UE-191350
-			if (AppleClangVersion >= new VersionNumber(15) && AppleClangVersion < new VersionNumber(17))
+			if (AppleClangVersion >= new VersionNumber(15) && AppleClangVersion < new VersionNumber(16))
 			{
 				Arguments.Add(" -ld_classic");
 			}
