AppleToolChain: Silence duplicate lib warning.

diff --git a/Engine/Source/Programs/UnrealBuildTool/ToolChain/AppleToolChain.cs b/Engine/Source/Programs/UnrealBuildTool/ToolChain/AppleToolChain.cs
--- a/Engine/Source/Programs/UnrealBuildTool/ToolChain/AppleToolChain.cs
+++ b/Engine/Source/Programs/UnrealBuildTool/ToolChain/AppleToolChain.cs
@@ -485,6 +485,7 @@ namespace UnrealBuildTool
 
 			Arguments.Add("-Wno-unknown-warning-option");
 			Arguments.Add("-Wno-range-loop-analysis");
+			Arguments.Add("-no_warn_duplicate_libraries");
 			Arguments.Add("-Wno-single-bit-bitfield-constant-conversion");
 		}
 
