[UBT]
* Changed GenerateLTH action to use "copy <file> + NUL <file>" in batch file instead of +,, because this works on both windows and wine

diff --git a/Engine/Source/Programs/UnrealBuildTool/Platform/Windows/VCToolChain.cs b/Engine/Source/Programs/UnrealBuildTool/Platform/Windows/VCToolChain.cs
--- a/Engine/Source/Programs/UnrealBuildTool/Platform/Windows/VCToolChain.cs
+++ b/Engine/Source/Programs/UnrealBuildTool/Platform/Windows/VCToolChain.cs
@@ -2546,16 +2546,8 @@ namespace UnrealBuildTool
 			FileItem ResponseFile = Graph.CreateIntermediateTextFile(OutputFile.ChangeExtension(ResponseExt), String.Join(" ", Arguments));
 
 			// Build the command for touching the output file to update the time stamp
-			string TouchActionCommand;
 			string OutfileFullName = NormalizeCommandLinePath(OutputFile).Replace('/', '\\'); // Must be backward slash
-			if (BuildHostPlatform.Current.IsRunningOnWine())
-			{
-				TouchActionCommand = $"copy /Y \"{OutfileFullName}\" \"%TMP%/{OutputFile.GetFileName()}\" && type \"%TMP%/{OutputFile.GetFileName()}\" > \"{OutfileFullName}\"";
-			}
-			else
-			{
-				TouchActionCommand = $"copy /b \"{OutfileFullName}\"+,, \"{OutfileFullName}\"";
-			}
+			string TouchActionCommand = $"copy /b \"{OutfileFullName}\" + NUL \"{OutfileFullName}\">nul";
 
 			// Build the batch file
 			StringBuilder BatchFileContents = new();
