UBT: Silence no platform load command found in for mac builds

diff --git a/Engine/Source/Programs/UnrealBuildTool/Matchers/CompileEventMatcher.cs b/Engine/Source/Programs/UnrealBuildTool/Matchers/CompileEventMatcher.cs
--- a/Engine/Source/Programs/UnrealBuildTool/Matchers/CompileEventMatcher.cs
+++ b/Engine/Source/Programs/UnrealBuildTool/Matchers/CompileEventMatcher.cs
@@ -66,6 +66,8 @@ namespace UnrealBuildTool.Matchers
 		static readonly Regex s_scriptCompilePattern = new Regex(@"^\s*[A-Za-z0-9_\.]+ ERROR:.* [A-Za-z_]+ failed to compile\.");
 		static readonly Regex s_cscSummaryPattern = new Regex(@"^\s+\d+ (?:Warning|Error)\(s\)");
 		static readonly Regex s_cscOutputPattern = new Regex(@"^  [^ ]+ -> ");
+		static readonly Regex s_tempLibpvxMacLoadCmdPattern = new Regex(@"no platform load command found in");
+		static readonly Regex s_tempXCodeDupLibrariesPattern = new Regex(@"ignoring duplicate libraries");
 
 		static readonly string[] s_invalidExtensions =
 		{
@@ -209,6 +211,19 @@ namespace UnrealBuildTool.Matchers
 				return true;
 			}
 
+			// Temporary XCode silencing when linker makes assumption for macos.
+			if (builder.Current.IsMatch(s_tempLibpvxMacLoadCmdPattern))
+			{
+				outEvent = builder.ToMatch(LogEventPriority.Highest, LogLevel.Information, KnownLogEvents.Systemic_XCode);
+				return true;
+			}
+			if (builder.Current.IsMatch(s_tempXCodeDupLibrariesPattern))
+			{
+				outEvent = builder.ToMatch(LogEventPriority.Highest, LogLevel.Information, KnownLogEvents.Systemic_XCode);
+				return true;
+			}
+			
+
 			outEvent = null;
 			return false;
 		}
