[UBA]
* Enabled known inputs by default in UBT
* Enabled include crawler for uba helpers spawned by horde
* Enabled so proxy use local casdb storage

diff --git a/Engine/Source/Programs/UnrealBuildTool/Executors/UnrealBuildAccelerator/UBAAgentCoordinatorHorde.cs b/Engine/Source/Programs/UnrealBuildTool/Executors/UnrealBuildAccelerator/UBAAgentCoordinatorHorde.cs
--- a/Engine/Source/Programs/UnrealBuildTool/Executors/UnrealBuildAccelerator/UBAAgentCoordinatorHorde.cs
+++ b/Engine/Source/Programs/UnrealBuildTool/Executors/UnrealBuildAccelerator/UBAAgentCoordinatorHorde.cs
@@ -584,8 +584,8 @@ namespace UnrealBuildTool
 						arguments.Add("-MaxCon=8");
 						arguments.Add("-MaxWorkers=64");
 						arguments.Add("-MaxIdle=15");
-						//arguments.Add("-UseCrawler");
-						//arguments.Add("-ProxyUseLocalStorage"); // If set then agent's cas storage will be used by proxy
+						arguments.Add("-UseCrawler");
+						arguments.Add("-ProxyUseLocalStorage"); // If set then agent's cas storage will be used by proxy
 						//arguments.Add("-NoStore"); // This means that no cas files will be written to disk.. so no state is stored between runs
 
 						if (!String.IsNullOrEmpty(desc))
diff --git a/Engine/Source/Programs/UnrealBuildTool/Executors/UnrealBuildAccelerator/UnrealBuildAcceleratorConfig.cs b/Engine/Source/Programs/UnrealBuildTool/Executors/UnrealBuildAccelerator/UnrealBuildAcceleratorConfig.cs
--- a/Engine/Source/Programs/UnrealBuildTool/Executors/UnrealBuildAccelerator/UnrealBuildAcceleratorConfig.cs
+++ b/Engine/Source/Programs/UnrealBuildTool/Executors/UnrealBuildAccelerator/UnrealBuildAcceleratorConfig.cs
@@ -206,7 +206,7 @@ namespace UnrealBuildTool
 		/// </summary>
 		[XmlConfigFile(Category = "UnrealBuildAccelerator")]
 		[CommandLine("-UBAUseKnownInputs", Value = "true")]
-		public bool bUseKnownInputs { get; set; } = false;
+		public bool bUseKnownInputs { get; set; } = true;
 
 		/// <summary>
 		/// Write yaml file with all actions that are queued for build. This can be used to replay using "UbaCli.exe local file.yaml"
