[UBT]
* Disabled remote execution for preprocess steps when building for PVS.. since preprocess is just IO more or less there is no point sending the work remotely

diff --git a/Engine/Source/Programs/UnrealBuildTool/Platform/Windows/PVSToolChain.cs b/Engine/Source/Programs/UnrealBuildTool/Platform/Windows/PVSToolChain.cs
--- a/Engine/Source/Programs/UnrealBuildTool/Platform/Windows/PVSToolChain.cs
+++ b/Engine/Source/Programs/UnrealBuildTool/Platform/Windows/PVSToolChain.cs
@@ -979,6 +979,9 @@ namespace UnrealBuildTool
 					continue;
 				}
 
+				// We don't want to run these remotely since they are very lightweight but has lots of I/O
+				PreprocessAction.bCanExecuteRemotely = false;
+
 				AnalyzeCppFile(PreprocessAction, CompileEnvironment, OutputDir, Result, Graph);
 			}
 			return Result;
