UnrealBuildTool: Don't rewrite link scripts if unchanged

#rnx

diff --git a/Engine/Source/Programs/UnrealBuildTool/Platform/Linux/LinuxToolChain.cs b/Engine/Source/Programs/UnrealBuildTool/Platform/Linux/LinuxToolChain.cs
--- a/Engine/Source/Programs/UnrealBuildTool/Platform/Linux/LinuxToolChain.cs
+++ b/Engine/Source/Programs/UnrealBuildTool/Platform/Linux/LinuxToolChain.cs
@@ -1370,10 +1370,8 @@ namespace UnrealBuildTool
 
 			LinkAction.PrerequisiteItems.Add(FileItem.GetItemByFileReference(LinkScriptFullPath));
 
-			Logger.LogDebug("Creating link script: {LinkScriptFullPath}", LinkScriptFullPath.FullName);
-			Directory.CreateDirectory(Path.GetDirectoryName(LinkScriptFullPath.FullName)!);
-			using (StreamWriter LinkWriter = File.CreateText(LinkScriptFullPath.FullName))
 			{
+				StringWriter LinkWriter = new();
 				if (bUseCmdExe)
 				{
 					LinkWriter.NewLine = "\r\n";
@@ -1413,6 +1411,8 @@ namespace UnrealBuildTool
 						LinkWriter.WriteLine(GetDumpEncodeDebugCommand(LinkEnvironment, OutputFile, LinkAction));
 					}
 				}
+				Directory.CreateDirectory(Path.GetDirectoryName(LinkScriptFullPath.FullName)!);
+				Graph.CreateIntermediateTextFile(LinkScriptFullPath, LinkWriter.ToString());
 			}
 
 			LinkAction.CommandPath = ShellBinary;
@@ -1458,8 +1458,9 @@ namespace UnrealBuildTool
 
 				RelinkAction.PrerequisiteItems.Add(FileItem.GetItemByFileReference(RelinkScriptFullPath));
 
-				using (StreamWriter RelinkWriter = File.CreateText(RelinkScriptFullPath.FullName))
+				
 				{
+					StringWriter RelinkWriter = new();
 					string RelinkInvocation = LinkCommandString;
 					string Replace = "-Wl,--allow-shlib-undefined";
 					RelinkInvocation = RelinkInvocation.Replace(Replace, EngineAndGameLibrariesString);
@@ -1500,6 +1501,9 @@ namespace UnrealBuildTool
 						RelinkWriter.WriteLine(RelinkInvocation);
 						RelinkWriter.WriteLine(GetDumpEncodeDebugCommand(LinkEnvironment, RealOutputFile, RelinkAction));
 					}
+
+					Directory.CreateDirectory(Path.GetDirectoryName(RelinkScriptFullPath.FullName)!);
+					Graph.CreateIntermediateTextFile(RelinkScriptFullPath, RelinkWriter.ToString());
 				}
 
 				RelinkAction.CommandPath = ShellBinary;
