Stop calling GetToolChainInfo() over and over

#rnx

diff --git a/Engine/Source/Programs/UnrealBuildTool/ToolChain/AppleToolChain.cs b/Engine/Source/Programs/UnrealBuildTool/ToolChain/AppleToolChain.cs
--- a/Engine/Source/Programs/UnrealBuildTool/ToolChain/AppleToolChain.cs
+++ b/Engine/Source/Programs/UnrealBuildTool/ToolChain/AppleToolChain.cs
@@ -810,14 +810,9 @@ namespace UnrealBuildTool
 			return OutputFiles;
 		}
 
-		private static VersionNumber? AppleClangVersion = null;
-
 		protected virtual void GetLinkArguments_Global(LinkEnvironment LinkEnvironment, List<string> Arguments)
 		{
-			if (AppleClangVersion == null)
-			{
-				AppleClangVersion = ((AppleToolChainInfo)GetToolChainInfo()).AppleClangVersion;
-			}
+			VersionNumber AppleClangVersion = ((AppleToolChainInfo)Info).AppleClangVersion;
 			// Temp solution for UE-191350
 			if (AppleClangVersion >= new VersionNumber(15) && AppleClangVersion < new VersionNumber(16))
 			{
diff --git a/Engine/Source/Programs/UnrealBuildTool/ToolChain/ClangToolChain.cs b/Engine/Source/Programs/UnrealBuildTool/ToolChain/ClangToolChain.cs
--- a/Engine/Source/Programs/UnrealBuildTool/ToolChain/ClangToolChain.cs
+++ b/Engine/Source/Programs/UnrealBuildTool/ToolChain/ClangToolChain.cs
@@ -226,6 +226,9 @@ namespace UnrealBuildTool
 			LazyInfo = new Lazy<ClangToolChainInfo>(() => { return GetToolChainInfo(); }); // Don't change to => GetToolChainInfo().. it doesnt produce correct code
 		}
 
+		/// <summary>
+		/// Create toolchain info. Do not call this function, use "Info" to get results
+		/// </summary>
 		protected abstract ClangToolChainInfo GetToolChainInfo();
 
 		public override FileReference? GetCppCompilerPath()
@@ -1234,7 +1237,7 @@ namespace UnrealBuildTool
 			// the LLVM versioning scheme that Clang uses), but ensure that we rebuild the source if the compiler has changed.
 			if (CompileEnvironment.bUseAutoRTFMCompiler)
 			{
-				CompileAction.PrerequisiteItems.Add(FileItem.GetItemByFileReference(GetToolChainInfo().Clang));
+				CompileAction.PrerequisiteItems.Add(FileItem.GetItemByFileReference(Info.Clang));
 			}
 
 			CompileAction.Weight = CompileActionWeight;
