Correct description for the bUseAdaptiveUnityBuild property.

diff --git a/Engine/Source/Programs/UnrealBuildTool/Configuration/TargetRules.cs b/Engine/Source/Programs/UnrealBuildTool/Configuration/TargetRules.cs
--- a/Engine/Source/Programs/UnrealBuildTool/Configuration/TargetRules.cs
+++ b/Engine/Source/Programs/UnrealBuildTool/Configuration/TargetRules.cs
@@ -1762,9 +1762,10 @@ namespace UnrealBuildTool
 		public string[]? DisableMergingModuleAndGeneratedFilesInUnityFiles = null;
 
 		/// <summary>
-		/// Use a heuristic to determine which files are currently being iterated on and exclude them from unity blobs, result in faster
-		/// incremental compile times. The current implementation uses the read-only flag to distinguish the working set, assuming that files will
-		/// be made writable by the source control system if they are being modified. This is true for Perforce, but not for Git.
+		/// Use a heuristic to determine which files are currently being iterated on and exclude them from unity blobs. This results in faster
+		/// incremental compile times. For Perforce repositories, the current implementation uses the read-only flag to distinguish the working
+		/// set, assuming that files will be made writable by the source control system if they are being modified. For Git repositories, the
+		/// implementation uses the git status command. Source code archives downloaded from Git as .zip files are not supported.
 		/// </summary>
 		[CommandLine("-DisableAdaptiveUnity", Value = "false")]
 		[XmlConfigFile(Category = "BuildConfiguration")]
