[UBT]
* Changed so clang tool chain does not run preprocessor actions remotely since they are all IO

diff --git a/Engine/Source/Programs/UnrealBuildTool/ToolChain/ClangToolChain.cs b/Engine/Source/Programs/UnrealBuildTool/ToolChain/ClangToolChain.cs
--- a/Engine/Source/Programs/UnrealBuildTool/ToolChain/ClangToolChain.cs
+++ b/Engine/Source/Programs/UnrealBuildTool/ToolChain/ClangToolChain.cs
@@ -1307,9 +1307,15 @@ namespace UnrealBuildTool
 			CompileAction.bIsGCCCompiler = true;
 
 			// Don't farm out creation of pre-compiled headers as it is the critical path task.
-			CompileAction.bCanExecuteRemotely =
-				CompileEnvironment.PrecompiledHeaderAction != PrecompiledHeaderAction.Create ||
-				CompileEnvironment.bAllowRemotelyCompiledPCHs;
+			if (CompileEnvironment.PrecompiledHeaderAction == PrecompiledHeaderAction.Create)
+			{
+				CompileAction.bCanExecuteRemotely = CompileEnvironment.bAllowRemotelyCompiledPCHs;
+			}
+			else
+			{
+				// Don't farm out preprocess actions since they are more or less only IO
+				CompileAction.bCanExecuteRemotely = !CompileEnvironment.bPreprocessOnly;
+			}
 
 			CompileAction.ArtifactMode = ArtifactMode.Enabled;
 			CompileAction.CacheBucket = GetCacheBucket(TargetRules, (!CompileEnvironment.RootPaths.bUseVfs && CompileEnvironment.PrecompiledHeaderAction != PrecompiledHeaderAction.None) ? CompileEnvironment.RootPaths : null);
