UnrealBuildTool: Add upper bound to workaround for compiler crash which has been confirmed fixed

#rnx

diff --git a/Engine/Source/Programs/UnrealBuildTool/Platform/Windows/VCToolChain.cs b/Engine/Source/Programs/UnrealBuildTool/Platform/Windows/VCToolChain.cs
--- a/Engine/Source/Programs/UnrealBuildTool/Platform/Windows/VCToolChain.cs
+++ b/Engine/Source/Programs/UnrealBuildTool/Platform/Windows/VCToolChain.cs
@@ -341,7 +341,7 @@ namespace UnrealBuildTool
 		protected virtual void AppendCLArguments_Global(CppCompileEnvironment CompileEnvironment, List<string> Arguments)
 		{
 			// Workaround for MSVC 14.31 compiler crash
-			if (EnvVars.Compiler.IsMSVC() && EnvVars.ToolChainVersion >= new VersionNumber(14, 31))
+			if (EnvVars.Compiler.IsMSVC() && EnvVars.ToolChainVersion < new VersionNumber(14, 43) && EnvVars.ToolChainVersion >= new VersionNumber(14, 31))
 			{
 				Arguments.Add("/d2ssa-cfg-question-");
 			}
