UnrealBuildTool: Parallel module validation, saves ~2 seconds for UnrealEditor when the makefile is invalidated
Move validation on module creation to ValidateModule

#rnx

diff --git a/Engine/Source/Programs/UnrealBuildTool/Configuration/UEBuildModule.cs b/Engine/Source/Programs/UnrealBuildTool/Configuration/UEBuildModule.cs
--- a/Engine/Source/Programs/UnrealBuildTool/Configuration/UEBuildModule.cs
+++ b/Engine/Source/Programs/UnrealBuildTool/Configuration/UEBuildModule.cs
@@ -1248,18 +1248,32 @@ namespace UnrealBuildTool
 		/// Public entry point to validate a module
 		/// </summary>
 		/// <param name="logger"></param>
-		public bool ValidateModule(ILogger logger)
+		public virtual bool ValidateModule(ILogger logger)
 		{
 			bool anyErrors = false;
 
 			// Report any [Obsolete("...")] messages from referenced modules.
-			IEnumerable<ObsoleteAttribute> obsoleteAttributes = this.Rules.GetType().GetCustomAttributes<ObsoleteAttribute>();
+			IEnumerable<ObsoleteAttribute> obsoleteAttributes = Rules.GetType().GetCustomAttributes<ObsoleteAttribute>();
 			if (obsoleteAttributes.Any())
 			{
 				logger.LogWarning(KnownLogEvents.Validation, "Referenced Module '{Name}' is obsolete: '{Message}'", Name, String.Join(", ", obsoleteAttributes.Select(x => x.Message ?? "<unknown reason>")));
 			}
 
-			anyErrors |= ValidateCircularReferences(logger);
+			// Check if this module is invalid for the current target, based on Module attributes
+			if (Rules.CppCompileWarningSettings.ModuleUnsupportedWarningLevel == WarningLevel.Error && !ModuleRules.IsValidForTarget(Rules.GetType(), Rules.Target, out string? _))
+			{
+				anyErrors = true;
+			}
+			if (Rules.CppCompileWarningSettings.PluginModuleUnsupportedWarningLevel == WarningLevel.Error)
+			{
+				ModuleDescriptor? ModuleInfo = Rules.Plugin?.Descriptor.Modules?.FirstOrDefault(x => x.Name == Rules.Name);
+				if (ModuleInfo?.IsCompiledInConfiguration(Rules.Target.Platform, Rules.Target.Configuration, Rules.Target.Name, Rules.Target.Type, Rules.Target.bBuildDeveloperTools, Rules.Target.bBuildRequiresCookedData, out string? InvalidReason) == false)
+				{
+					anyErrors = true;
+				}
+			}
+
+			anyErrors = ValidateCircularReferences(logger) || anyErrors;
 
 			return anyErrors;
 		}
@@ -1274,20 +1288,6 @@ namespace UnrealBuildTool
 			List<CircularReference> circularReferences = new();
 			CheckForCircularReferences(this, new(), new(), circularReferences);
 
-			// Check if this module is invalid for the current target, based on Module attributes
-			if (Rules.CppCompileWarningSettings.ModuleUnsupportedWarningLevel == WarningLevel.Error && !ModuleRules.IsValidForTarget(Rules.GetType(), Rules.Target, out string? _))
-			{
-				return true;
-			}
-			if (Rules.CppCompileWarningSettings.PluginModuleUnsupportedWarningLevel == WarningLevel.Error)
-			{
-				ModuleDescriptor? ModuleInfo = Rules.Plugin?.Descriptor.Modules?.FirstOrDefault(x => x.Name == Rules.Name);
-				if (ModuleInfo?.IsCompiledInConfiguration(Rules.Target.Platform, Rules.Target.Configuration, Rules.Target.Name, Rules.Target.Type, Rules.Target.bBuildDeveloperTools, Rules.Target.bBuildRequiresCookedData, out string? InvalidReason) == false)
-				{
-					return true;
-				}
-			}
-
 			IEnumerable<CircularReference> allowed = circularReferences.Where(x => !x.IsError).OrderBy(x => x.ReferenceStack.Count);
 			IEnumerable<CircularReference> errors = circularReferences.Where(x => x.IsError).OrderBy(x => x.ReferenceStack.Count);
 
diff --git a/Engine/Source/Programs/UnrealBuildTool/Configuration/UEBuildModuleCPP.cs b/Engine/Source/Programs/UnrealBuildTool/Configuration/UEBuildModuleCPP.cs
--- a/Engine/Source/Programs/UnrealBuildTool/Configuration/UEBuildModuleCPP.cs
+++ b/Engine/Source/Programs/UnrealBuildTool/Configuration/UEBuildModuleCPP.cs
@@ -250,6 +250,25 @@ namespace UnrealBuildTool
 				bDependsOnVerse = true;
 			}
 
+			if (Rules.bAddDefaultIncludePaths)
+			{
+				AddDefaultIncludePaths();
+			}
+		}
+
+		public override bool ValidateModule(ILogger logger)
+		{
+			bool anyErrors = false;
+
+			{
+				IEnumerable<string> InvalidDependencies = Rules.DynamicallyLoadedModuleNames.Intersect(Rules.PublicDependencyModuleNames.Concat(Rules.PrivateDependencyModuleNames));
+				if (InvalidDependencies.Any())
+				{
+					logger.LogError("{Module} module should not be dependent on modules which are also dynamically loaded: {Invalid}", Name, String.Join(", ", InvalidDependencies));
+					anyErrors = true;
+				}
+			}
+
 			if (Rules.bValidateCircularDependencies || Rules.bTreatAsEngineModule)
 			{
 				foreach (string CircularlyReferencedModuleName in Rules.CircularlyReferencedDependentModules)
@@ -258,17 +277,29 @@ namespace UnrealBuildTool
 						!CircularDependenciesAllowList.Any(x =>
 							x.Key == Name && x.Value == CircularlyReferencedModuleName))
 					{
-						Logger.LogWarning(
+						logger.LogWarning(
 							"Found reference between '{Source}' and '{Target}'. Support for circular references is being phased out; please do not introduce new ones.",
 							Name, CircularlyReferencedModuleName);
 					}
 				}
 			}
 
-			if (Rules.bAddDefaultIncludePaths)
+			// Make sure that engine modules use shared PCHs or have an explicit private PCH
+			if (Rules.PCHUsage == ModuleRules.PCHUsageMode.NoSharedPCHs && Rules.PrivatePCHHeaderFile == null)
 			{
-				AddDefaultIncludePaths();
+				if (Rules.Target.ProjectFile == null || !Rules.File.IsUnderDirectory(Rules.Target.ProjectFile.Directory))
+				{
+					logger.LogWarning("{Module} module has shared PCHs disabled, but does not have a private PCH set", Name);
+				}
 			}
+
+			// If we can't use a shared PCH, check there's a private PCH set
+			if (Rules.PCHUsage != ModuleRules.PCHUsageMode.NoPCHs && Rules.PCHUsage != ModuleRules.PCHUsageMode.UseExplicitOrSharedPCHs && Rules.PrivatePCHHeaderFile == null)
+			{
+				logger.LogWarning("{Module} module must specify an explicit precompiled header for PCHUsage {PCHUsage} (eg. PrivatePCHHeaderFile = \"Private/{Module}PrivatePCH.h\")", Name, Rules.PCHUsage, Name);
+			}
+
+			return base.ValidateModule(logger) || anyErrors;
 		}
 
 		/// <summary>
diff --git a/Engine/Source/Programs/UnrealBuildTool/Configuration/UEBuildPlatform.cs b/Engine/Source/Programs/UnrealBuildTool/Configuration/UEBuildPlatform.cs
--- a/Engine/Source/Programs/UnrealBuildTool/Configuration/UEBuildPlatform.cs
+++ b/Engine/Source/Programs/UnrealBuildTool/Configuration/UEBuildPlatform.cs
@@ -804,8 +804,9 @@ namespace UnrealBuildTool
 		/// Validate a UEBuildModule before it's processed
 		/// <param name="Module">The UEBuildModule that needs to be validated</param>
 		/// <param name="Target">Options for the target being built</param>
+		/// <param name="Logger">The logger for output</param>
 		/// </summary>
-		public virtual void ValidateModule(UEBuildModule Module, ReadOnlyTargetRules Target)
+		public virtual void ValidateModule(UEBuildModule Module, ReadOnlyTargetRules Target, ILogger Logger)
 		{
 		}
 
@@ -815,7 +816,8 @@ namespace UnrealBuildTool
 		/// <param name="Module">The UEBuildModule that needs to be validated</param>
 		/// <param name="Target">Options for the target being built</param>
 		/// <param name="AllModules">Other modules to validate against, if needed</param>
-		public virtual bool ValidateModuleIncludePaths(UEBuildModule Module, ReadOnlyTargetRules Target, IEnumerable<UEBuildModule> AllModules)
+		/// <param name="Logger">The logger for output</param>
+		public virtual bool ValidateModuleIncludePaths(UEBuildModule Module, ReadOnlyTargetRules Target, IEnumerable<UEBuildModule> AllModules, ILogger Logger)
 		{
 			CppCompileWarnings ModuleWarningSettings = Module.Rules.CppCompileWarningSettings;
 			if (ModuleWarningSettings.ModuleIncludePathWarningLevel <= WarningLevel.Off
diff --git a/Engine/Source/Programs/UnrealBuildTool/Configuration/UEBuildTarget.cs b/Engine/Source/Programs/UnrealBuildTool/Configuration/UEBuildTarget.cs
--- a/Engine/Source/Programs/UnrealBuildTool/Configuration/UEBuildTarget.cs
+++ b/Engine/Source/Programs/UnrealBuildTool/Configuration/UEBuildTarget.cs
@@ -12,6 +12,7 @@ using System.Runtime.InteropServices;
 using System.Runtime.Serialization;
 using System.Text;
 using System.Text.Json.Serialization;
+using System.Threading;
 using System.Threading.Tasks;
 using EpicGames.Core;
 using EpicGames.UHT.Tables;
@@ -1480,15 +1481,26 @@ namespace UnrealBuildTool
 		/// <returns>true if there are any fatal errors</returns>
 		bool ValidateModules(ILogger logger)
 		{
-			bool anyErrors = false;
-			foreach (UEBuildModule module in Modules.Values.OrderBy(x => x.Name))
+			int errorCount = 0;
+			CaptureLogger[] captureLoggers = new CaptureLogger[Modules.Count];
+			Parallel.ForEach(Modules.Values.OrderBy(x => x.Name), (module, _, index) =>
 			{
-				UEBuildPlatform.GetBuildPlatform(Platform).ValidateModule(module, Rules);
-				anyErrors |= UEBuildPlatform.GetBuildPlatform(Platform).ValidateModuleIncludePaths(module, Rules, Modules.Values);
-				anyErrors |= ValidateModule(module, logger);
+				captureLoggers[index] = new();
+				UEBuildPlatform.GetBuildPlatform(Platform).ValidateModule(module, Rules, captureLoggers[index]);
+				bool anyErrors = UEBuildPlatform.GetBuildPlatform(Platform).ValidateModuleIncludePaths(module, Rules, Modules.Values, captureLoggers[index]);
+				anyErrors = ValidateModule(module, captureLoggers[index]) || anyErrors;
+				if (anyErrors)
+				{
+					Interlocked.Increment(ref errorCount);
+				}
+			});
+
+			foreach (CaptureLogger captureLogger in captureLoggers)
+			{
+				captureLogger.RenderTo(logger);
 			}
 
-			return anyErrors;
+			return errorCount > 0;
 		}
 
 		/// <summary>
@@ -5772,56 +5784,10 @@ namespace UnrealBuildTool
 			RemoveTrailingSlashes(RulesObject.PrivateIncludePaths);
 			RemoveTrailingSlashes(RulesObject.PublicSystemLibraryPaths);
 
-			// Validate rules object
-			if (RulesObject.Type == ModuleRules.ModuleType.CPlusPlus)
-			{
-				List<string> InvalidDependencies = RulesObject.DynamicallyLoadedModuleNames.Intersect(RulesObject.PublicDependencyModuleNames.Concat(RulesObject.PrivateDependencyModuleNames)).ToList();
-				if (InvalidDependencies.Count != 0)
-				{
-					throw new BuildException("Module rules for '{0}' should not be dependent on modules which are also dynamically loaded: {1}", ModuleName, String.Join(", ", InvalidDependencies));
-				}
-
-				// Make sure that engine modules use shared PCHs or have an explicit private PCH
-				if (RulesObject.PCHUsage == ModuleRules.PCHUsageMode.NoSharedPCHs && RulesObject.PrivatePCHHeaderFile == null)
-				{
-					if (ProjectFile == null || !RulesObject.File.IsUnderDirectory(ProjectFile.Directory))
-					{
-						Logger.LogWarning("{ModuleName} module has shared PCHs disabled, but does not have a private PCH set", ModuleName);
-					}
-				}
-
-				// If we can't use a shared PCH, check there's a private PCH set
-				if (RulesObject.PCHUsage != ModuleRules.PCHUsageMode.NoPCHs && RulesObject.PCHUsage != ModuleRules.PCHUsageMode.UseExplicitOrSharedPCHs && RulesObject.PrivatePCHHeaderFile == null)
-				{
-					// Try to figure out the legacy PCH file
-					FileReference? CppFile = DirectoryReference.EnumerateFiles(RulesObject.Directory, "*.cpp", SearchOption.AllDirectories).FirstOrDefault();
-					if (CppFile != null)
-					{
-						string? IncludeFile = MetadataCache.GetFirstInclude(FileItem.GetItemByFileReference(CppFile));
-						if (IncludeFile != null)
-						{
-							FileReference? PchIncludeFile = DirectoryReference.EnumerateFiles(RulesObject.Directory, Path.GetFileName(IncludeFile), SearchOption.AllDirectories).FirstOrDefault();
-							if (PchIncludeFile != null)
-							{
-								RulesObject.PrivatePCHHeaderFile = PchIncludeFile.MakeRelativeTo(RulesObject.Directory).Replace(Path.DirectorySeparatorChar, '/');
-							}
-						}
-					}
-
-					// Print a suggestion for which file to include
-					if (RulesObject.PrivatePCHHeaderFile == null)
-					{
-						Log.TraceWarningOnce(RulesObject.File, "Modules must specify an explicit precompiled header (eg. PrivatePCHHeaderFile = \"Private/{0}PrivatePCH.h\") from UE 4.21 onwards.", ModuleName);
-					}
-					else
-					{
-						Log.TraceWarningOnce(RulesObject.File, "Modules must specify an explicit precompiled header (eg. PrivatePCHHeaderFile = \"{0}\") from UE 4.21 onwards.", RulesObject.PrivatePCHHeaderFile);
-					}
-				}
-			}
+			RulesObject.PrivateDefinitions.Add($"UE_MODULE_NAME=\"{ModuleName}\"");
+			RulesObject.PrivateDefinitions.Add($"UE_PLUGIN_NAME=\"{(RulesObject.Plugin?.Name ?? String.Empty)}\"");
 
-			RulesObject.PrivateDefinitions.Add(String.Format("UE_MODULE_NAME=\"{0}\"", ModuleName));
-			RulesObject.PrivateDefinitions.Add(String.Format("UE_PLUGIN_NAME=\"{0}\"", RulesObject.IsPlugin && RulesObject.Plugin != null ? RulesObject.Plugin.Name : ""));
+			// Module validation can be found in UEBuildModuleCPP.ValidateModule()
 
 			return RulesObject;
 		}
diff --git a/Engine/Source/Programs/UnrealBuildTool/Platform/IOS/UEBuildIOS.cs b/Engine/Source/Programs/UnrealBuildTool/Platform/IOS/UEBuildIOS.cs
--- a/Engine/Source/Programs/UnrealBuildTool/Platform/IOS/UEBuildIOS.cs
+++ b/Engine/Source/Programs/UnrealBuildTool/Platform/IOS/UEBuildIOS.cs
@@ -809,9 +809,9 @@ namespace UnrealBuildTool
 			}
 		}
 
-		public override void ValidateModule(UEBuildModule Module, ReadOnlyTargetRules Target)
+		public override void ValidateModule(UEBuildModule Module, ReadOnlyTargetRules Target, ILogger Logger)
 		{
-			base.ValidateModule(Module, Target);
+			base.ValidateModule(Module, Target, Logger);
 
 			// @todo temporarily disabling due to VisionOS rquiring newer Xcode than this will allow - we may remove this entirely
 #if false
diff --git a/Engine/Source/Programs/UnrealBuildTool/Platform/Mac/UEBuildMac.cs b/Engine/Source/Programs/UnrealBuildTool/Platform/Mac/UEBuildMac.cs
--- a/Engine/Source/Programs/UnrealBuildTool/Platform/Mac/UEBuildMac.cs
+++ b/Engine/Source/Programs/UnrealBuildTool/Platform/Mac/UEBuildMac.cs
@@ -410,29 +410,6 @@ namespace UnrealBuildTool
 			Target.bCheckSystemHeadersForModification = BuildHostPlatform.Current.Platform != UnrealTargetPlatform.Mac;
 		}
 
-		static HashSet<FileReference> ValidatedLibs = new();
-		public override void ValidateModule(UEBuildModule Module, ReadOnlyTargetRules Target)
-		{
-			base.ValidateModule(Module, Target);
-
-			if (BuildHostPlatform.Current.Platform == UnrealTargetPlatform.Mac && !Target.MacPlatform.bSkipClangValidation)
-			{
-				ApplePlatformSDK SDK = (ApplePlatformSDK?)GetSDK() ?? new ApplePlatformSDK(Logger);
-				// Validate the added public libraries
-				lock (ValidatedLibs)
-				{
-					foreach (FileReference LibLoc in Module.PublicLibraries)
-					{
-						if (ValidatedLibs.Contains(LibLoc))
-						{
-							continue;
-						}
-						ValidatedLibs.Add(LibLoc);
-					}
-				}
-			}
-		}
-
 		/// <summary>
 		/// Determines if the given name is a build product for a target.
 		/// </summary>
