UnrealBuildTool: Fix some CompileEnvironment properties not being copied correctly

#rnx

diff --git a/Engine/Source/Programs/UnrealBuildTool/Configuration/UEBuildModuleCPP.cs b/Engine/Source/Programs/UnrealBuildTool/Configuration/UEBuildModuleCPP.cs
--- a/Engine/Source/Programs/UnrealBuildTool/Configuration/UEBuildModuleCPP.cs
+++ b/Engine/Source/Programs/UnrealBuildTool/Configuration/UEBuildModuleCPP.cs
@@ -813,13 +813,14 @@ namespace UnrealBuildTool
 			}
 
 			CppCompileEnvironment.bUseUnity = bModuleUsesUnityBuild && Rules.UnityFileTypes.HasFlag(ModuleRules.UnityFileType.Cpp);
-			GeneratedCPPCompileEnvironment.bUseUnity = bModuleUsesUnityBuild && Rules.UnityFileTypes.HasFlag(ModuleRules.UnityFileType.Cpp);
+			GeneratedCPPCompileEnvironment.bUseUnity = CppCompileEnvironment.bUseUnity;
 
 			// Create and register a special action that can be used to compile single files (even when unity is enabled) for generated files
 			if (GeneratedFileItems.Any())
 			{
 				CppCompileEnvironment SingleGeneratedCPPCompileEnvironment = new CppCompileEnvironment(GeneratedCPPCompileEnvironment);
 				SingleGeneratedCPPCompileEnvironment.PrecompiledHeaderAction = PrecompiledHeaderAction.None;
+				SingleGeneratedCPPCompileEnvironment.bUseUnity = false;
 				foreach (DirectoryReference Directory in GeneratedFileItems.Select(x => x.Location.Directory).Distinct())
 				{
 					ToolChain.CreateSpecificFileAction(SingleGeneratedCPPCompileEnvironment, Directory, IntermediateDirectory, Graph);
@@ -831,6 +832,7 @@ namespace UnrealBuildTool
 			{
 				CppCompileEnvironment SingleCompileEnvironment = new CppCompileEnvironment(CppCompileEnvironment);
 				SingleCompileEnvironment.PrecompiledHeaderAction = PrecompiledHeaderAction.None;
+				SingleCompileEnvironment.bUseUnity = false;
 				foreach (DirectoryReference Directory in ModuleDirectories)
 				{
 					ToolChain.CreateSpecificFileAction(SingleCompileEnvironment, Directory, IntermediateDirectory, Graph);
@@ -844,8 +846,11 @@ namespace UnrealBuildTool
 			}
 			else if (Rules.Target.bAlwaysUseUnityForGeneratedFiles && !Rules.Target.bEnableCppModules)
 			{
-				Unity.GenerateUnitySource(Target, GeneratedFileItems, new List<FileItem>(), CppCompileEnvironment, WorkingSet, (Rules.ShortName ?? Name) + ".gen", IntermediateDirectory, Graph, SourceFileToUnityFile,
+				GeneratedCPPCompileEnvironment.bUseUnity = true;
+				Unity.GenerateUnitySource(Target, GeneratedFileItems, new List<FileItem>(), GeneratedCPPCompileEnvironment, WorkingSet, (Rules.ShortName ?? Name) + ".gen", IntermediateDirectory, Graph, SourceFileToUnityFile,
 					out List<FileItem> NormalGeneratedFiles, out List<FileItem> AdaptiveGeneratedFiles, NumIncludedBytesPerUnityCPP);
+				ModuleCompileEnvironment.FileInlineGenCPPMap = new(GeneratedCPPCompileEnvironment.FileInlineGenCPPMap);
+				CppCompileEnvironment.FileInlineGenCPPMap = new(GeneratedCPPCompileEnvironment.FileInlineGenCPPMap);
 				LinkInputFiles.AddRange(CompileFilesWithToolChain(Target, ToolChain, GeneratedCPPCompileEnvironment, ModuleCompileEnvironment, NormalGeneratedFiles, AdaptiveGeneratedFiles, Graph, Logger).ObjectFiles);
 			}
 			else
@@ -895,6 +900,7 @@ namespace UnrealBuildTool
 				{
 					Unity.GenerateUnitySource(Target, CPPFiles, InputFiles.HeaderFiles, CppCompileEnvironment, WorkingSet, Rules.ShortName ?? Name, IntermediateDirectory, Graph, SourceFileToUnityFile,
 						out List<FileItem> NormalFiles, out List<FileItem> AdaptiveFiles, NumIncludedBytesPerUnityCPP);
+					ModuleCompileEnvironment.FileInlineGenCPPMap = new(CppCompileEnvironment.FileInlineGenCPPMap);
 					LinkInputFiles.AddRange(CompileFilesWithToolChain(Target, ToolChain, CppCompileEnvironment, ModuleCompileEnvironment, NormalFiles, AdaptiveFiles, Graph, Logger).ObjectFiles);
 				}
 				else
