[UBA]
* Added so it is possible to disable retry for uba actions. Use -UBANoRetry or set <bAllowRetry>false</bAllowRetry> in BuildConfiguration.xml

diff --git a/Engine/Source/Programs/UnrealBuildTool/Executors/UnrealBuildAccelerator/UBAExecutor.cs b/Engine/Source/Programs/UnrealBuildTool/Executors/UnrealBuildAccelerator/UBAExecutor.cs
--- a/Engine/Source/Programs/UnrealBuildTool/Executors/UnrealBuildAccelerator/UBAExecutor.cs
+++ b/Engine/Source/Programs/UnrealBuildTool/Executors/UnrealBuildAccelerator/UBAExecutor.cs
@@ -305,6 +305,15 @@ namespace UnrealBuildTool
 				UBAConfig.Zone = "local";
 			}
 
+			if (UBAConfig.bForcedRetry)
+			{
+				UBAConfig.bAllowRetry = true;
+			}
+			else if (UBAConfig.bStoreObjFilesCompressed)
+			{
+				UBAConfig.bAllowRetry = false;
+			}
+
 			if (!String.IsNullOrEmpty(UBAConfig.CacheServer))
 			{
 				foreach (LinkedAction action in inputActions)
@@ -895,7 +904,7 @@ namespace UnrealBuildTool
 				using (IProcess process = _session!.RunProcess(startInfo, false, null, enableDetour))
 				{
 					Interlocked.Add(ref _localProcessedActions, 1);
-					if (!UBAConfig.bStoreObjFilesCompressed && (process.ExitCode != 0 && UBAConfig.bForcedRetry || (process.ExitCode >= 9000 && process.ExitCode < 10000)))
+					if (UBAConfig.bAllowRetry && (process.ExitCode != 0 && UBAConfig.bForcedRetry || (process.ExitCode >= 9000 && process.ExitCode < 10000)))
 					{
 						_threadedLogger.LogInformation("{Description} {StatusDescription}: Exited with error code {ExitCode}. This action will retry without UBA", action.CommandDescription, action.StatusDescription, process.ExitCode);
 						_forcedRetryActions.AddOrUpdate(action, false, (k, v) => false);
diff --git a/Engine/Source/Programs/UnrealBuildTool/Executors/UnrealBuildAccelerator/UnrealBuildAcceleratorConfig.cs b/Engine/Source/Programs/UnrealBuildTool/Executors/UnrealBuildAccelerator/UnrealBuildAcceleratorConfig.cs
--- a/Engine/Source/Programs/UnrealBuildTool/Executors/UnrealBuildAccelerator/UnrealBuildAcceleratorConfig.cs
+++ b/Engine/Source/Programs/UnrealBuildTool/Executors/UnrealBuildAccelerator/UnrealBuildAcceleratorConfig.cs
@@ -26,7 +26,14 @@ namespace UnrealBuildTool
 		public bool bForceBuildAllRemote { get; set; } = false;
 
 		/// <summary>
-		/// When set to true, actions that fail locally with UBA will be retried without UBA.
+		/// When set to true, actions that fail locally with certain error codes will retry without uba.
+		/// </summary>
+		[XmlConfigFile(Category = "UnrealBuildAccelerator")]
+		[CommandLine("-UBANoRetry", Value = "false")]
+		public bool bAllowRetry { get; set; } = true;
+
+		/// <summary>
+		/// When set to true, all actions that fail locally with UBA will be retried without UBA.
 		/// </summary>
 		[XmlConfigFile(Category = "UnrealBuildAccelerator")]
 		[CommandLine("-UBAForcedRetry")]
