[UBT]
* Fixed some paths not being virtualized in rsp for ios

diff --git a/Engine/Source/Programs/UnrealBuildTool/Platform/IOS/IOSToolChain.cs b/Engine/Source/Programs/UnrealBuildTool/Platform/IOS/IOSToolChain.cs
--- a/Engine/Source/Programs/UnrealBuildTool/Platform/IOS/IOSToolChain.cs
+++ b/Engine/Source/Programs/UnrealBuildTool/Platform/IOS/IOSToolChain.cs
@@ -200,6 +200,8 @@ namespace UnrealBuildTool
 		{
 			base.GetCompileArguments_Global(CompileEnvironment, Arguments);
 
+			CppRootPaths RootPaths = CompileEnvironment.RootPaths;
+
 			if (CompileEnvironment.Configuration == CppConfiguration.Shipping)
 			{
 				Arguments.Add("-ffunction-sections");
@@ -216,7 +218,7 @@ namespace UnrealBuildTool
 			// What architecture(s) to build for
 			Arguments.Add(FormatArchitectureArg(CompileEnvironment.Architectures));
 
-			Arguments.Add($"-isysroot \"{Settings.GetSDKPath(CompileEnvironment.Architecture)}\"");
+			Arguments.Add($"-isysroot \"{NormalizeCommandLinePath(Settings.GetSDKPath(CompileEnvironment.Architecture), RootPaths)}\"");
 
 			// Add additional frameworks so that their headers can be found
 			foreach (UEBuildFramework Framework in CompileEnvironment.AdditionalFrameworks)
@@ -226,14 +228,13 @@ namespace UnrealBuildTool
 					DirectoryReference? FrameworkDirectory = Framework.GetFrameworkDirectory(CompileEnvironment.Platform, CompileEnvironment.Architecture, Logger);
 					if (FrameworkDirectory != null)
 					{
-						string FrameworkDir = FrameworkDirectory.FullName;
 						// embedded frameworks have a framework inside of this directory, so we use this directory. regular frameworks need to go one up to point to the 
 						// directory containing the framework. -F gives a path to look for the -framework
-						if (FrameworkDir.EndsWith(".framework"))
+						if (FrameworkDirectory.FullName.EndsWith(".framework"))
 						{
-							FrameworkDir = Path.GetDirectoryName(FrameworkDir)!;
+							FrameworkDirectory = FrameworkDirectory.ParentDirectory;
 						}
-						Arguments.Add($"-F\"{FrameworkDir}\"");
+						Arguments.Add($"-F \"{NormalizeCommandLinePath(FrameworkDirectory!, RootPaths)}\"");
 					}
 				}
 			}
