Simplify ApplyWarningAttributes with filter type cache / enum & null defaults

diff --git a/Engine/Source/Programs/Shared/UnrealBuildTool.Tests/CppCompileWarningsTests.cs b/Engine/Source/Programs/Shared/UnrealBuildTool.Tests/CppCompileWarningsTests.cs
--- a/Engine/Source/Programs/Shared/UnrealBuildTool.Tests/CppCompileWarningsTests.cs
+++ b/Engine/Source/Programs/Shared/UnrealBuildTool.Tests/CppCompileWarningsTests.cs
@@ -387,7 +387,7 @@ namespace UnrealBuildTool.Tests
 					baseWarnings.ShadowVariableWarningLevel = WarningLevel.Default;
 				}
 
-				// ToolChainVersionConstrainedWarningsClangToolChainAttribute
+				// ToolChainVersionConstrainedFilterAttribute
 				{
 					testTargetRules.WindowsPlatform.Compiler = WindowsCompiler.Clang;
 					baseWarnings.Shorten64To32WarningLevel = WarningLevel.Off;
@@ -404,7 +404,7 @@ namespace UnrealBuildTool.Tests
 					baseWarnings.Shorten64To32WarningLevel = WarningLevel.Default;
 				}
 
-				// ConfigurationConstrainedClangToolChainAttribute
+				// ConfigurationFilterAttribute
 				{
 					testTargetRules.WindowsPlatform.Compiler = WindowsCompiler.Clang;
 					baseWarnings.UnusedValueWarningLevel = WarningLevel.Off;
@@ -423,7 +423,7 @@ namespace UnrealBuildTool.Tests
 					baseWarnings.UnusedValueWarningLevel = WarningLevel.Default;
 				}
 
-				// CppStandardConstrainedClangToolChainAttribute
+				// CppStandardConstrainedFilterAttribute
 				{
 					testTargetRules.WindowsPlatform.Compiler = WindowsCompiler.Clang;
 					baseWarnings.AmbiguousReversedOperatorWarningLevel = WarningLevel.Off;
@@ -448,7 +448,7 @@ namespace UnrealBuildTool.Tests
 					baseWarnings.AmbiguousReversedOperatorWarningLevel = WarningLevel.Default;
 				}
 
-				// DeterministicClangToolChainAttribute
+				// DeterministcFlagSetFilter
 				{
 					testTargetRules.WindowsPlatform.Compiler = WindowsCompiler.Clang;
 					baseWarnings.DeterministicWarningLevel = WarningLevel.Off;
@@ -465,8 +465,14 @@ namespace UnrealBuildTool.Tests
 					args = baseWarnings.GenerateWarningCommandLineArgs(deterministicCompileEnvironment, typeof(ClangToolChain), newVersionNumber);
 					Assert.IsFalse(args.Any());
 
+					deterministicCompileEnvironment.bDeterministic = false;
 					baseWarnings.DeterministicWarningLevel = WarningLevel.Warning;
 					args = baseWarnings.GenerateWarningCommandLineArgs(deterministicCompileEnvironment, typeof(ClangToolChain), newVersionNumber);
+					Assert.IsFalse(args.Any());
+
+					deterministicCompileEnvironment.bDeterministic = true;
+					args = baseWarnings.GenerateWarningCommandLineArgs(deterministicCompileEnvironment, typeof(ClangToolChain), newVersionNumber);
+
 					Assert.IsTrue(args.Count() == 2);
 					Assert.IsTrue(args.Contains("-Wdate-time"));
 					Assert.IsTrue(args.Contains("-Wno-error=date-time"));
@@ -479,7 +485,7 @@ namespace UnrealBuildTool.Tests
 					baseWarnings.DeterministicWarningLevel = WarningLevel.Default;
 				}
 
-				// PGOOptimizedClangToolChainAttribute
+				// PGOOptimizedFilter
 				{
 					testTargetRules.WindowsPlatform.Compiler = WindowsCompiler.Clang;
 					baseWarnings.ProfileInstructWarningLevel = WarningLevel.Off;
@@ -502,12 +508,12 @@ namespace UnrealBuildTool.Tests
 					baseWarnings.ProfileInstructWarningLevel = WarningLevel.Off;
 				}
 
-				// AndroidNDKR28ToolChainVersionExclusionAttribute
+				// AndroidNDKR28ToolChainVersionExclusion
 				{
 					testTargetRules.WindowsPlatform.Compiler = WindowsCompiler.Clang;
 					baseWarnings.CastFunctionTypeMistmatchWarningLevel = WarningLevel.Off;
 
-					VersionNumber newVersionNumber = new VersionNumber(19,0,0);
+					VersionNumber newVersionNumber = new VersionNumber(19, 0, 0);
 					CppCompileEnvironment androidCompileEnvironment = new CppCompileEnvironment(UnrealTargetPlatform.Android, CppConfiguration.Development, new UnrealArchitectures(UnrealArch.X64), null!);
 					args = baseWarnings.GenerateWarningCommandLineArgs(androidCompileEnvironment, typeof(ClangToolChain), newVersionNumber);
 
@@ -527,7 +533,7 @@ namespace UnrealBuildTool.Tests
 					baseWarnings.CastFunctionTypeMistmatchWarningLevel = WarningLevel.Default;
 				}
 
-				// AndroidNDKR26ToolChainVersionExclusionAttribute
+				// AndroidNDKR26ToolChainVersionExclusion
 				{
 					testTargetRules.WindowsPlatform.Compiler = WindowsCompiler.Clang;
 					baseWarnings.InvalidUnevaluatedStringWarningLevel = WarningLevel.Off;
@@ -590,13 +596,19 @@ namespace UnrealBuildTool.Tests
 					Assert.IsTrue(args.Count() == 1 && args.Contains("-Wno-unused-value"));
 				}
 
-				// VCClangToolChainAttribute
+				// VCClangCompilerFilter
 				{
 					testTargetRules.WindowsPlatform.Compiler = WindowsCompiler.Clang;
 					baseWarnings.MicrosoftGroupWarningLevel = WarningLevel.Off;
 
 					args = baseWarnings.GenerateWarningCommandLineArgs(newCompileEnvironment, typeof(VCToolChain));
 					Assert.IsTrue(args.Any());
+					Assert.IsTrue(args.Contains("-Wno-microsoft"));
+
+					// make sure exclusively ONLY the clang stuff is in:
+					testTargetRules.WindowsPlatform.Compiler = WindowsCompiler.VisualStudio2022;
+					args = baseWarnings.GenerateWarningCommandLineArgs(newCompileEnvironment, typeof(VCToolChain));
+					Assert.IsFalse(args.Any());
 				}
 			}
 		}
diff --git a/Engine/Source/Programs/UnrealBuildTool/Configuration/CompileWarnings/ApplyWarningsAttribute.cs b/Engine/Source/Programs/UnrealBuildTool/Configuration/CompileWarnings/ApplyWarningsAttribute.cs
--- a/Engine/Source/Programs/UnrealBuildTool/Configuration/CompileWarnings/ApplyWarningsAttribute.cs
+++ b/Engine/Source/Programs/UnrealBuildTool/Configuration/CompileWarnings/ApplyWarningsAttribute.cs
@@ -15,6 +15,7 @@ namespace UnrealBuildTool
 		private readonly string[]? _offArgs;
 		private readonly string[]? _warningArgs;
 		private readonly string[]? _errorArgs;
+		protected readonly List<IApplyWarningsFilter> ActiveFilters = [];
 
 		/// <summary>
 		/// Defines the argument metadata around a <see cref="CppCompileWarnings"/> <see cref="WarningLevel"/> property.
@@ -22,11 +23,13 @@ namespace UnrealBuildTool
 		/// <param name="offArgs">The arguments to use when the property is defined as <see cref="WarningLevel.Off"/>.</param>
 		/// <param name="warningArgs">The arguments to use when the property is defined as <see cref="WarningLevel.Warning"/>.</param>
 		/// <param name="errorArgs">The arguments to use when the property is defined as <see cref="WarningLevel.Error"/>.</param>
-		internal ApplyWarningsAttribute(string[]? offArgs, string[]? warningArgs, string[]? errorArgs)
+		/// <param name="applicableFilters">The applicable filter criteria.</param>
+		internal ApplyWarningsAttribute(string[]? offArgs = null, string[]? warningArgs = null, string[]? errorArgs = null, string[]? applicableFilters = null)
 		{
 			_offArgs = offArgs;
 			_warningArgs = warningArgs;
 			_errorArgs = errorArgs;
+			AppendNewFilters(applicableFilters);
 		}
 
 		/// <summary>
@@ -77,7 +80,44 @@ namespace UnrealBuildTool
 		/// </summary>
 		/// <param name="toolChainContext">The toolchain context to be applied to.</param>
 		/// <returns>True if applicable, false otherwise.</returns>
-		internal abstract bool CanApplyToContext(CompilerWarningsToolChainContext toolChainContext);
+		internal virtual bool CanApplyToContext(CompilerWarningsToolChainContext toolChainContext)
+		{
+			return (ActiveFilters == null || ActiveFilters.TrueForAll(x => x.CanApply(toolChainContext)));
+		}
+
+		protected bool TryAppendNewFilter(string rawFilterName)
+		{
+			IApplyWarningsFilter? filter = ApplyWarningsFilterRegistry.RequestFilter(rawFilterName);
+			if (filter != null && !ActiveFilters.Contains(filter))
+			{
+				ActiveFilters.Add(filter);
+				return true;
+			}
+
+			return false;
+		}
+
+		protected void AppendNewFilters(IEnumerable<string>? newFilters)
+		{
+			if (newFilters != null)
+			{
+				foreach (string rawFilterName in newFilters)
+				{
+					TryAppendNewFilter(rawFilterName);
+				}
+			}
+		}
+
+		protected void AppendNewFilters(IEnumerable<FilterID>? newFilters)
+		{
+			if (newFilters != null)
+			{
+				foreach (FilterID filterID in newFilters)
+				{
+					TryAppendNewFilter(filterID.ToString());
+				}
+			}
+		}
 
 		/// <summary>
 		/// The list of off arugments corresponding to <see cref="WarningLevel.Off"/>
@@ -102,12 +142,9 @@ namespace UnrealBuildTool
 #pragma warning restore CA1813
 	{
 		/// <inheritdoc/>
-		internal WarningsVCToolChainAttribute(string[]? offArgs, string[]? warningArgs, string[]? errorArgs) : base(offArgs, warningArgs, errorArgs) { }
-
-		/// <inheritdoc/>
-		internal override bool CanApplyToContext(CompilerWarningsToolChainContext toolChainContext)
+		internal WarningsVCToolChainAttribute(string[]? offArgs = null, string[]? warningArgs = null, string[]? errorArgs = null, string[]? filterNames = null) : base(offArgs, warningArgs, errorArgs, filterNames)
 		{
-			return typeof(VCToolChain) == toolChainContext._toolChainType;
+			AppendNewFilters([nameof(FilterID.VCToolChain)]);
 		}
 	}
 
@@ -116,12 +153,9 @@ namespace UnrealBuildTool
 	{
 #pragma warning restore CA1813
 		/// <inheritdoc/>
-		internal WarningsClangToolChainAttribute(string[]? offArgs, string[]? warningArgs, string[]? errorArgs) : base(offArgs, warningArgs, errorArgs) { }
-
-		/// <inheritdoc/>
-		internal override bool CanApplyToContext(CompilerWarningsToolChainContext toolChainContext)
+		internal WarningsClangToolChainAttribute(string[]? offArgs = null, string[]? warningArgs = null, string[]? errorArgs = null, string[]? filterNames = null) : base(offArgs, warningArgs, errorArgs, filterNames)
 		{
-			return typeof(ClangToolChain) == toolChainContext._toolChainType;
+			AppendNewFilters([nameof(FilterID.ClangToolChain)]);
 		}
 	}
 
@@ -135,14 +169,9 @@ namespace UnrealBuildTool
 	internal sealed class WarningsMSCVAttribute : WarningsVCToolChainAttribute
 	{
 		/// <inheritdoc/>
-		internal WarningsMSCVAttribute(string[]? offArgs, string[]? warningArgs, string[]? errorArgs) : base(offArgs, warningArgs, errorArgs) { }
-
-		/// <inheritdoc/>
-		internal override bool CanApplyToContext(CompilerWarningsToolChainContext toolChainContext)
+		internal WarningsMSCVAttribute(string[]? offArgs = null, string[]? warningArgs = null, string[]? errorArgs = null, string[]? filterNames = null) : base(offArgs, warningArgs, errorArgs, filterNames)
 		{
-			ReadOnlyTargetRules? readOnlyTargetRules = toolChainContext._buildSystemContext.GetReadOnlyTargetRules();
-
-			return base.CanApplyToContext(toolChainContext) && readOnlyTargetRules != null && readOnlyTargetRules.WindowsPlatform != null && readOnlyTargetRules.WindowsPlatform.Compiler.IsMSVC();
+			AppendNewFilters([nameof(StandardFilters.MSVCCompilerFilter)]);
 		}
 	}
 
@@ -152,14 +181,9 @@ namespace UnrealBuildTool
 	internal sealed class WarningsVCClangAttribute : WarningsVCToolChainAttribute
 	{
 		/// <inheritdoc/>
-		internal WarningsVCClangAttribute(string[]? offArgs, string[]? warningArgs, string[]? errorArgs) : base(offArgs, warningArgs, errorArgs) { }
-
-		/// <inheritdoc/>
-		internal override bool CanApplyToContext(CompilerWarningsToolChainContext toolChainContext)
+		internal WarningsVCClangAttribute(string[]? offArgs = null, string[]? warningArgs = null, string[]? errorArgs = null) : base(offArgs, warningArgs, errorArgs)
 		{
-			ReadOnlyTargetRules? readOnlyTargetRules = toolChainContext._buildSystemContext.GetReadOnlyTargetRules();
-
-			return base.CanApplyToContext(toolChainContext) && readOnlyTargetRules != null && readOnlyTargetRules.WindowsPlatform != null && readOnlyTargetRules.WindowsPlatform.Compiler.IsClang();
+			AppendNewFilters([nameof(StandardFilters.VCClangCompilerFilter)]);
 		}
 	}
 
@@ -168,206 +192,15 @@ namespace UnrealBuildTool
 	#region -- ToolChain Specializations --
 
 	/// <summary>
-	/// Specialized <see cref="WarningsClangToolChainAttribute"/> to enforce toolchain version constraints.
-	/// </summary>
-#pragma warning disable CA1813
-	internal class ToolChainVersionConstrainedWarningsClangToolChainAttribute : WarningsClangToolChainAttribute
-#pragma warning restore CA1813
-	{
-		private readonly VersionNumber? _lowerBound;
-		private readonly VersionNumber? _upperBound;
-
-		/// <summary>
-		/// Constructor for ToolChainVersionConstrainedWarningsClangToolChainAttribute that enforces a clang toolchain version constraint.
-		/// </summary>
-		/// <param name="offArgs">The arguments to use when the property is defined as <see cref="WarningLevel.Off"/>.</param>
-		/// <param name="warningArgs">The arguments to use when the property is defined as <see cref="WarningLevel.Warning"/>.</param>
-		/// <param name="errorArgs">The arguments to use when the property is defined as <see cref="WarningLevel.Error"/>.</param>
-		/// <param name="lowerBoundVersionComponents">The lower bound of the <see cref="VersionNumber"/> to constrain this warning attribute to.</param>
-		/// <param name="upperBoundVersionComponents">The upper bound of the <see cref="VersionNumber"/> to constrain this warning attribute to.</param>
-		/// <remarks><see cref="VersionNumber"/> is not a viable parameter to this attribute, as it's a generic (via Comparable interface), and as such is not viable for attributes.</remarks>
-		internal ToolChainVersionConstrainedWarningsClangToolChainAttribute(string[]? offArgs, string[]? warningArgs, string[]? errorArgs, int[]? lowerBoundVersionComponents, int[]? upperBoundVersionComponents) : base(offArgs, warningArgs, errorArgs)
-		{
-			_lowerBound = lowerBoundVersionComponents == null ? null : new VersionNumber(lowerBoundVersionComponents);
-			_upperBound = upperBoundVersionComponents == null ? null : new VersionNumber(upperBoundVersionComponents);
-		}
-
-		/// <inheritdoc/>
-		internal override bool CanApplyToContext(CompilerWarningsToolChainContext toolChainContext)
-		{
-			bool baseCanApplyToContext = base.CanApplyToContext(toolChainContext);
-
-			return baseCanApplyToContext && IsWithinBounds(toolChainContext._toolChainVersion);
-		}
-
-		protected bool IsWithinBounds(VersionNumber? version)
-		{
-			// If our evaluation version is null, we can't possibly make a statement about whether it fits within a range
-			if (version == null)
-			{
-				return false;
-			}
-
-			bool isAboveLowerBound = _lowerBound == null || version >= _lowerBound;
-			bool isBelowUpperBound = _upperBound == null || version <= _upperBound;
-
-			return isAboveLowerBound && isBelowUpperBound;
-		}
-	}
-
-	/// <summary>
-	/// Specialized <see cref="WarningsClangToolChainAttribute"/> to enforce configuration constraints.
-	/// </summary>
-#pragma warning disable CA1813
-	internal class ConfigurationConstrainedClangToolChainAttribute : WarningsClangToolChainAttribute
-#pragma warning restore CA1813
-	{
-		private readonly CppConfiguration _cppConfiguration;
-
-		/// <summary>
-		/// Constructor for ConfigurationConstrainedClangToolChainAttribute that enforces specific configuration constraint.
-		/// </summary>
-		/// <param name="offArgs">The arguments to use when the property is defined as <see cref="WarningLevel.Off"/>.</param>
-		/// <param name="warningArgs">The arguments to use when the property is defined as <see cref="WarningLevel.Warning"/>.</param>
-		/// <param name="errorArgs">The arguments to use when the property is defined as <see cref="WarningLevel.Error"/>.</param>
-		/// <param name="enabledConfiguration">The configuration in which to apply this WarningsAttribute to.</param>
-		internal ConfigurationConstrainedClangToolChainAttribute(string[]? offArgs, string[]? warningArgs, string[]? errorArgs, CppConfiguration enabledConfiguration) : base(offArgs, warningArgs, errorArgs)
-		{
-			_cppConfiguration = enabledConfiguration;
-		}
-
-		/// <inheritdoc/>
-		internal override bool CanApplyToContext(CompilerWarningsToolChainContext toolChainContext)
-		{
-			return base.CanApplyToContext(toolChainContext) && (_cppConfiguration == toolChainContext._compileEnvironment.Configuration);
-		}
-	}
-
-	/// <summary>
-	/// Specialized <see cref="WarningsClangToolChainAttribute"/> to enforce cpp standard constraints.
-	/// </summary>
-	internal sealed class CppStandardConstrainedClangToolChainAttribute : WarningsClangToolChainAttribute
-	{
-		private readonly CppStandardVersion _lowerBound;
-		private readonly CppStandardVersion _upperBound;
-
-		/// <summary>
-		/// Constructor for CppStandardConstrainedClangToolChainAttribute that enforces a cpp standard constraint.
-		/// </summary>
-		/// <param name="offArgs">The arguments to use when the property is defined as <see cref="WarningLevel.Off"/>.</param>
-		/// <param name="warningArgs">The arguments to use when the property is defined as <see cref="WarningLevel.Warning"/>.</param>
-		/// <param name="errorArgs">The arguments to use when the property is defined as <see cref="WarningLevel.Error"/>.</param>
-		/// <param name="lowerBound">The lower bound of the <see cref="CppStandardVersion"/> to constrain this warning attribute to.</param>
-		/// <param name="upperBound">The upper bound of the <see cref="CppStandardVersion"/> to constrain this warning attribute to.</param>
-		internal CppStandardConstrainedClangToolChainAttribute(string[]? offArgs, string[]? warningArgs, string[]? errorArgs, CppStandardVersion lowerBound = CppStandardVersion.Cpp14, CppStandardVersion upperBound = CppStandardVersion.Latest) : base(offArgs, warningArgs, errorArgs)
-		{
-			_lowerBound = lowerBound;
-			_upperBound = upperBound;
-		}
-
-		/// <inheritdoc/>
-		internal override bool CanApplyToContext(CompilerWarningsToolChainContext toolChainContext)
-		{
-			bool baseCanApplyToContext = base.CanApplyToContext(toolChainContext);
-
-			return baseCanApplyToContext && IsWithinBounds(toolChainContext._compileEnvironment.CppStandard);
-		}
-
-		private bool IsWithinBounds(CppStandardVersion version)
-		{
-			bool isAboveLowerBound = version >= _lowerBound;
-			bool isBelowUpperBound = version <= _upperBound;
-
-			return isAboveLowerBound && isBelowUpperBound;
-		}
-	}
-
-	/// <summary>
-	///  Specialized <see cref="WarningsClangToolChainAttribute"/> to enforce deterministic compilation constraints.
-	/// </summary>
-	internal sealed class DeterministicClangToolChainAttribute : WarningsClangToolChainAttribute
-	{
-		/// <inheritdoc/>
-		internal DeterministicClangToolChainAttribute(string[]? offArgs, string[]? warningArgs, string[]? errorArgs) : base(offArgs, warningArgs, errorArgs) { }
-
-		/// <inheritdoc/>
-		internal override bool CanApplyToContext(CompilerWarningsToolChainContext toolChainContext)
-		{
-			bool baseCanApplyToContext = base.CanApplyToContext(toolChainContext);
-
-			return baseCanApplyToContext && toolChainContext._compileEnvironment.bDeterministic;
-		}
-	}
-
-	/// <summary>
-	/// Specialized <see cref="WarningsClangToolChainAttribute"/> profile guided optimization constraints.
-	/// </summary>
-	internal sealed class PGOOptimizedClangToolChainAttribute : WarningsClangToolChainAttribute
-	{
-		/// <inheritdoc/>
-		internal PGOOptimizedClangToolChainAttribute(string[]? offArgs, string[]? warningArgs, string[]? errorArgs) : base(offArgs, warningArgs, errorArgs) { }
-
-		/// <inheritdoc/>
-		internal override bool CanApplyToContext(CompilerWarningsToolChainContext toolChainContext)
-		{
-			return base.CanApplyToContext(toolChainContext) && toolChainContext._compileEnvironment.bPGOOptimize;
-		}
-	}
-
-	/// <summary>
-	/// Specialized <see cref="WarningsClangToolChainAttribute"/> to enforce Android as an excluded platform, at a specific toolchain version.
-	/// </summary>
-	internal sealed class AndroidNDKR28ToolChainVersionExclusionAttribute : WarningsClangToolChainAttribute
-	{
-		/// <inheritdoc/>
-		internal AndroidNDKR28ToolChainVersionExclusionAttribute(string[]? offArgs, string[]? warningArgs, string[]? errorArgs) : base(offArgs, warningArgs, errorArgs) { }
-
-		/// <inheritdoc/>
-		internal override bool CanApplyToContext(CompilerWarningsToolChainContext toolChainContext)
-		{
-			if (toolChainContext._toolChainVersion == null)
-			{
-				return false;
-			}
-
-			bool versionMin = toolChainContext._toolChainVersion >= new VersionNumber(19);
-			bool isExcludedVersion = toolChainContext._toolChainVersion == new VersionNumber(19, 0, 0);
-			bool bIsAndroidClang19 = isExcludedVersion && toolChainContext._compileEnvironment.Platform == UnrealTargetPlatform.Android;
-
-			return base.CanApplyToContext(toolChainContext) && versionMin && !bIsAndroidClang19;
-		}
-	}
-
-	/// <summary>
-	/// Specialized <see cref="WarningsClangToolChainAttribute"/> to enforce Android as an excluded platform, at a specific toolchain version, at a specific C++ standard.
+	/// / Specialized <see cref="WarningsClangToolChainAttribute"/> for Unused Value.
 	/// </summary>
-	internal sealed class AndroidNDKR26ToolChainVersionExclusionAttribute : WarningsClangToolChainAttribute
+	internal sealed class UnusedValueClangToolChainAttribute : WarningsClangToolChainAttribute
 	{
 		/// <inheritdoc/>
-		internal AndroidNDKR26ToolChainVersionExclusionAttribute(string[]? offArgs, string[]? warningArgs, string[]? errorArgs) : base(offArgs, warningArgs, errorArgs) { }
-
-		/// <inheritdoc/>
-		internal override bool CanApplyToContext(CompilerWarningsToolChainContext toolChainContext)
+		internal UnusedValueClangToolChainAttribute(string[]? offArgs = null, string[]? warningArgs = null, string[]? errorArgs = null) : base(offArgs, warningArgs, errorArgs)
 		{
-			if (toolChainContext._toolChainVersion == null)
-			{
-				return false;
-			}
-
-			bool versionMin = toolChainContext._toolChainVersion >= new VersionNumber(17);
-			bool bIsAndroidClang17 = toolChainContext._toolChainVersion == new VersionNumber(17, 0, 2) && toolChainContext._compileEnvironment.Platform == UnrealTargetPlatform.Android;
-
-			return base.CanApplyToContext(toolChainContext) && versionMin && !bIsAndroidClang17 && toolChainContext._compileEnvironment.CppStandard < CppStandardVersion.Latest;
+			AppendNewFilters([FilterID.ConfigShipping]);
 		}
-	}
-
-	/// <summary>
-	/// / Specialized <see cref="WarningsClangToolChainAttribute"/> for Unused Value.
-	/// </summary>
-	internal sealed class UnusedValueClangToolChainAttribute : ConfigurationConstrainedClangToolChainAttribute
-	{
-		/// <inheritdoc/>
-		internal UnusedValueClangToolChainAttribute(string[]? offArgs, string[]? warningArgs, string[]? errorArgs, CppConfiguration enabledConfiguration) : base(offArgs, warningArgs, errorArgs, enabledConfiguration) { }
 
 		/// <inheritdoc/>
 		internal override bool CanApplyToContext(CompilerWarningsToolChainContext toolChainContext)
@@ -402,11 +235,15 @@ namespace UnrealBuildTool
 				{
 					case WarningLevel.Warning:
 						if (WarningArgs != null)
+						{
 							arguments.AddRange(WarningArgs);
+						}
 						break;
 					case WarningLevel.Error:
 						if (ErrorArgs != null)
+						{
 							arguments.AddRange(ErrorArgs);
+						}
 						break;
 				}
 			}
diff --git a/Engine/Source/Programs/UnrealBuildTool/Configuration/CompileWarnings/ApplyWarningsFilterAttribute.cs b/Engine/Source/Programs/UnrealBuildTool/Configuration/CompileWarnings/ApplyWarningsFilterAttribute.cs
new file mode 100644
--- /dev/null
+++ b/Engine/Source/Programs/UnrealBuildTool/Configuration/CompileWarnings/ApplyWarningsFilterAttribute.cs
@@ -0,0 +1,283 @@
+﻿// Copyright Epic Games, Inc. All Rights Reserved.
+
+using System;
+using EpicGames.Core;
+
+namespace UnrealBuildTool
+{
+	/// <summary>
+	/// Shorthand enum to be used in the <see cref="ApplyWarningsFilterAttribute"/> extensibility context.
+	/// </summary>
+	public enum FilterID
+	{
+		/// <summary>
+		/// Filter ID for <see cref="ToolChainFilterAttribute"/>, specifically for <see cref="VCToolChain"/>.
+		/// </summary>
+		VCToolChain,
+
+		/// <summary>
+		/// Filter ID for <see cref="ToolChainFilterAttribute"/>, specifically for <see cref="ClangToolChain"/>.
+		/// </summary>
+		ClangToolChain,
+
+		/// <summary>
+		/// Filter ID for <see cref="ConfigurationFilterAttribute"/>, specifically for <see cref="CppConfiguration.Debug"/>.
+		/// </summary>
+		ConfigDebug,
+
+		/// <summary>
+		/// Filter ID for <see cref="ConfigurationFilterAttribute"/>, specifically for <see cref="CppConfiguration.Shipping"/>.
+		/// </summary>
+		ConfigShipping,
+
+		/// <summary>
+		/// Filter ID for <see cref="ToolChainVersionConstrainedFilterAttribute"/>, specifically for <see cref="VersionNumber"/> 14.0 as a min, with no upper bound.
+		/// </summary>
+		Version14Min,
+
+		/// <summary>
+		/// Filter ID for <see cref="ToolChainVersionConstrainedFilterAttribute"/>, specifically for <see cref="VersionNumber"/> 16.0 as a min, with no upper bound.
+		/// </summary>
+		Version16Min,
+
+		/// <summary>
+		/// Filter ID for <see cref="ToolChainVersionConstrainedFilterAttribute"/>, specifically for <see cref="VersionNumber"/> 18.0 as a min, with no upper bound.
+		/// </summary>
+		Version18Min,
+
+		/// <summary>
+		/// Filter ID for <see cref="ToolChainVersionConstrainedFilterAttribute"/>, specifically for <see cref="VersionNumber"/> 19.0 as a min, with no upper bound.
+		/// </summary>
+		Version19Min,
+
+		/// <summary>
+		/// Filter ID for <see cref="CppStandardConstrainedFilterAttribute"/>, specifically for <see cref="CppStandardVersion.Cpp20"/> as a minimum, with no upper bound.
+		/// </summary>
+		Cpp20Min
+	}
+
+	/// <summary>
+	///  Defines an attribute that is used in the class context for the purposes of ApplyWarnings filters. 
+	/// </summary>
+	/// <remarks>
+	/// Use this extensibility point when instance data is requred, and your need to have the same logic applied for different instacnce data.
+	/// The target class of the attribute is <see cref="ApplyWarningsFilterRegistry"/>.
+	/// </remarks>
+	[AttributeUsage(AttributeTargets.Class, AllowMultiple = true)]
+	internal abstract class ApplyWarningsFilterAttribute : Attribute, IApplyWarningsFilter
+	{
+		private readonly string _filterName;
+		public string FilterName => _filterName;
+
+		/// <summary>
+		/// Constructs the attribute.
+		/// </summary>
+		/// <param name="filterID">The filter id that describes the instance of the ApplyWarningsFilter attribute.</param>
+#pragma warning disable CA1019 // Define accessors for attribute arguments
+		public ApplyWarningsFilterAttribute(FilterID filterID)
+#pragma warning restore CA1019 // Define accessors for attribute arguments
+		{
+			_filterName = filterID.ToString();
+		}
+
+		/// <summary>
+		/// Constructs the attribute.
+		/// </summary>
+		/// <param name="filterName">The filter name that describes the instance of the ApplyWarningsFilter attribute.</param>
+		public ApplyWarningsFilterAttribute(string filterName)
+		{
+			_filterName = filterName;
+		}
+
+		/// <inheritdoc/>
+		public abstract bool CanApply(CompilerWarningsToolChainContext context);
+	}
+
+	#region -- Filter Attributes --
+
+	/// <summary>
+	/// Configuration filter attribute that will only apply the filter under a speciic <see cref="CppConfiguration"/>.
+	/// </summary>
+	internal sealed class ConfigurationFilterAttribute : ApplyWarningsFilterAttribute
+	{
+		private readonly CppConfiguration _config = CppConfiguration.Debug;
+		public CppConfiguration Config => _config;
+
+		/// Constructs a ConfigurationFilterAttribute.
+		/// <param name="filterID">The filter id that describes the instance of the ApplyWarningsFilter attribute.</param>
+		/// <param name="config">The config to filter on.</param>
+#pragma warning disable CA1019 // Define accessors for attribute arguments
+		public ConfigurationFilterAttribute(FilterID filterID, CppConfiguration config) : base(filterID)
+#pragma warning restore CA1019 // Define accessors for attribute arguments
+		{
+			_config = config;
+		}
+
+		/// Constructs a ConfigurationFilterAttribute.
+		/// <param name="filterName">The filter name that describes the instance of the ApplyWarningsFilter attribute.</param>
+		/// <param name="config">The config to filter on.</param>
+		public ConfigurationFilterAttribute(string filterName, CppConfiguration config) : base(filterName)
+		{
+			_config = config;
+		}
+
+		/// <inheritdoc/>
+		public override bool CanApply(CompilerWarningsToolChainContext context)
+		{
+			return _config == context._compileEnvironment.Configuration;
+		}
+	}
+
+	/// <summary>
+	/// Specialized <see cref="ApplyWarningsFilterAttribute"/> to toolchain version constraints.
+	/// </summary>
+	internal sealed class ToolChainVersionConstrainedFilterAttribute : ApplyWarningsFilterAttribute
+	{
+		private readonly VersionNumber? _lowerBound;
+		private readonly VersionNumber? _upperBound;
+
+		public VersionNumber? LowerBound => _lowerBound;
+		public VersionNumber? UpperBound => _upperBound;
+
+		/// <summary>
+		/// Constructs a ToolChainVersionConstrainedFilterAttribute.
+		/// </summary>
+		/// <param name="filterID">The filter id that describes the instance of the ApplyWarningsFilter attribute.</param>
+		/// <param name="lowerBoundVersionComponents">The lower bound of the <see cref="VersionNumber"/> to constrain this warning attribute to.</param>
+		/// <param name="upperBoundVersionComponents">The upper bound of the <see cref="VersionNumber"/> to constrain this warning attribute to.</param>
+		/// <remarks><see cref="VersionNumber"/> is not a viable parameter to this attribute, as it's a generic (via Comparable interface), and as such is not viable for attributes.</remarks>
+#pragma warning disable CA1019 // Define accessors for attribute arguments
+		public ToolChainVersionConstrainedFilterAttribute(FilterID filterID, int[]? lowerBoundVersionComponents = null, int[]? upperBoundVersionComponents = null) : base(filterID)
+#pragma warning restore CA1019 // Define accessors for attribute arguments
+		{
+			_lowerBound = lowerBoundVersionComponents == null ? null : new VersionNumber(lowerBoundVersionComponents);
+			_upperBound = upperBoundVersionComponents == null ? null : new VersionNumber(upperBoundVersionComponents);
+		}
+
+		/// <summary>
+		/// Constructs a ToolChainVersionConstrainedFilterAttribute.
+		/// </summary>
+		/// <param name="filterName">The filter name that describes the instance of the ApplyWarningsFilter attribute.</param>
+		/// <param name="lowerBoundVersionComponents">The lower bound of the <see cref="VersionNumber"/> to constrain this warning attribute to.</param>
+		/// <param name="upperBoundVersionComponents">The upper bound of the <see cref="VersionNumber"/> to constrain this warning attribute to.</param>
+		/// <remarks><see cref="VersionNumber"/> is not a viable parameter to this attribute, as it's a generic (via Comparable interface), and as such is not viable for attributes.</remarks>
+		public ToolChainVersionConstrainedFilterAttribute(string filterName, int[]? lowerBoundVersionComponents = null, int[]? upperBoundVersionComponents = null) : base(filterName)
+		{
+			_lowerBound = lowerBoundVersionComponents == null ? null : new VersionNumber(lowerBoundVersionComponents);
+			_upperBound = upperBoundVersionComponents == null ? null : new VersionNumber(upperBoundVersionComponents);
+		}
+
+		/// <inheritdoc/>
+		public override bool CanApply(CompilerWarningsToolChainContext context)
+		{
+			// If our evaluation version is null, we can't possibly make a statement about whether it fits within a range
+			if (context._toolChainVersion == null)
+			{
+				return false;
+			}
+
+			bool isAboveLowerBound = _lowerBound == null || context._toolChainVersion >= _lowerBound;
+			bool isBelowUpperBound = _upperBound == null || context._toolChainVersion <= _upperBound;
+
+			return isAboveLowerBound && isBelowUpperBound;
+		}
+	}
+
+	/// <summary>
+	/// Specialized <see cref="ApplyWarningsFilterAttribute"/> to enforce cpp standard constraints.
+	/// </summary>
+	internal sealed class CppStandardConstrainedFilterAttribute : ApplyWarningsFilterAttribute
+	{
+		private readonly CppStandardVersion _lowerBound;
+		private readonly CppStandardVersion _upperBound;
+
+		public CppStandardVersion LowerBound => _lowerBound;
+		public CppStandardVersion UpperBound => _upperBound;
+
+		/// <summary>
+		/// Constructor for CppStandardConstrainedFilterAttribute that enforces a cpp standard constraint.
+		/// </summary>
+		/// <param name="filterID">The filter id that describes the instance of the ApplyWarningsFilter attribute.</param>
+		/// <param name="lowerBound">The lower bound of the <see cref="CppStandardVersion"/> to constrain this warning attribute to.</param>
+		/// <param name="upperBound">The upper bound of the <see cref="CppStandardVersion"/> to constrain this warning attribute to.</param>
+#pragma warning disable CA1019 // Define accessors for attribute arguments
+		public CppStandardConstrainedFilterAttribute(FilterID filterID, CppStandardVersion lowerBound = CppStandardVersion.Cpp14, CppStandardVersion upperBound = CppStandardVersion.Latest) : base(filterID)
+#pragma warning restore CA1019 // Define accessors for attribute arguments
+		{
+			_lowerBound = lowerBound;
+			_upperBound = upperBound;
+		}
+
+		/// <summary>
+		/// Constructor for CppStandardConstrainedFilterAttribute that enforces a cpp standard constraint.
+		/// </summary>
+		/// <param name="filterName">The filter name that describes the instance of the ApplyWarningsFilter attribute.</param>
+		/// <param name="lowerBound">The lower bound of the <see cref="CppStandardVersion"/> to constrain this warning attribute to.</param>
+		/// <param name="upperBound">The upper bound of the <see cref="CppStandardVersion"/> to constrain this warning attribute to.</param>
+		public CppStandardConstrainedFilterAttribute(string filterName, CppStandardVersion lowerBound = CppStandardVersion.Cpp14, CppStandardVersion upperBound = CppStandardVersion.Latest) : base(filterName)
+		{
+			_lowerBound = lowerBound;
+			_upperBound = upperBound;
+		}
+
+		/// <inheritdoc/>
+		public override bool CanApply(CompilerWarningsToolChainContext context)
+		{
+			bool isAboveLowerBound = context._compileEnvironment.CppStandard >= _lowerBound;
+			bool isBelowUpperBound = context._compileEnvironment.CppStandard <= _upperBound;
+
+			return isAboveLowerBound && isBelowUpperBound;
+		}
+	}
+
+	/// <summary>
+	/// Specialized <see cref="ApplyWarningsFilterAttribute"/> to enforce specific tool chain contexts.
+	/// </summary>
+	internal sealed class ToolChainFilterAttribute : ApplyWarningsFilterAttribute
+	{
+		private readonly Type _toolChainType;
+		public Type ToolChainType => _toolChainType;
+
+		/// <summary>
+		/// Constructor for ToolChainFilterAttribute that enforces specific tool chain contexts.
+		/// </summary>
+		/// <param name="filterID">The filter id that describes the instance of the ApplyWarningsFilter attribute.</param>
+		/// <param name="toolChainType">The toolchain of which this filter should particiapte.</param>
+		/// <exception cref="ArgumentException">If the provided toolchain type is not derived from <see cref="UEToolChain"/>, throws.</exception>
+#pragma warning disable CA1019 // Define accessors for attribute arguments
+		public ToolChainFilterAttribute(FilterID filterID, Type toolChainType) : base(filterID)
+#pragma warning restore CA1019 // Define accessors for attribute arguments
+		{
+			if (!typeof(UEToolChain).IsAssignableFrom(toolChainType))
+			{
+				throw new ArgumentException($"Invalid type provided for toolchain: {toolChainType.ToString()}");
+			}
+
+			_toolChainType = toolChainType;
+		}
+
+		/// <summary>
+		/// Constructor for ToolChainFilterAttribute that enforces specific tool chain contexts.
+		/// </summary>
+		/// <param name="filterName">The filter name that describes the instance of the ApplyWarningsFilter attribute.</param>
+		/// <param name="toolChainType">The toolchain of which this filter should particiapte.</param>
+		/// <exception cref="ArgumentException">If the provided toolchain type is not derived from <see cref="UEToolChain"/>, throws.</exception>
+		public ToolChainFilterAttribute(string filterName, Type toolChainType) : base(filterName)
+		{
+			if (!typeof(UEToolChain).IsAssignableFrom(toolChainType))
+			{
+				throw new ArgumentException($"Invalid type provided for toolchain: {toolChainType}");
+			}
+
+			_toolChainType = toolChainType;
+		}
+
+		/// <inheritdoc/>
+		public override bool CanApply(CompilerWarningsToolChainContext context)
+		{
+			return _toolChainType == context._toolChainType;
+		}
+	}
+
+	#endregion
+}
\ No newline at end of file
diff --git a/Engine/Source/Programs/UnrealBuildTool/Configuration/CompileWarnings/ApplyWarningsFilterDelegateAttribute.cs b/Engine/Source/Programs/UnrealBuildTool/Configuration/CompileWarnings/ApplyWarningsFilterDelegateAttribute.cs
new file mode 100644
--- /dev/null
+++ b/Engine/Source/Programs/UnrealBuildTool/Configuration/CompileWarnings/ApplyWarningsFilterDelegateAttribute.cs
@@ -0,0 +1,73 @@
+// Copyright Epic Games, Inc. All Rights Reserved.
+
+using System;
+using System.Reflection;
+using EpicGames.Core;
+using Microsoft.Extensions.Logging;
+
+namespace UnrealBuildTool
+{
+	/// <summary>
+	/// Defines an attribute that is used to support reflective delegates for the purposes of ApplyWarnings filters.
+	/// </summary>
+	/// <remarks>Use this extensibility point when no such instance data is required, or you don't need to have the same logic with different instance data.</remarks>
+	[AttributeUsage(AttributeTargets.Method, AllowMultiple = false)]
+	internal sealed class ApplyWarningsFilterDelegateAttribute : Attribute
+	{
+		private readonly string _filterName;
+		public string FilterName => _filterName;
+
+		/// <summary>
+		/// Constructor for ApplyWarningsFilterDelegateAttribute.
+		/// </summary>
+		/// <param name="filterName">The name of the filter.</param>
+		public ApplyWarningsFilterDelegateAttribute(string filterName)
+		{
+			_filterName = filterName;
+		}
+	}
+
+	/// <summary>
+	/// A wrapper for the methods annotated by the <see cref="ApplyWarningsFilterDelegateAttribute"/> to support <see cref="IApplyWarningsFilter"/>.
+	/// </summary>
+	internal struct ApplyWarningsFilterWrapper : IApplyWarningsFilter
+	{
+		public string FilterName => _applyWarningsFilterDelegateAttribute.FilterName;
+
+		private readonly ILogger? _logger = Log.Logger;
+		private readonly ApplyWarningsFilterDelegateAttribute _applyWarningsFilterDelegateAttribute;
+		private readonly MethodInfo _method;
+
+		/// <summary>
+		/// Constructs a ApplyWarningsFilterWrapper.
+		/// </summary>
+		/// <param name="filterAttribute">Instance of the filter attribute that is being wrapped.</param>
+		/// <param name="method">The method that will be used for reflective invocation in <see cref="CanApply(CompilerWarningsToolChainContext)"/>.</param>
+		public ApplyWarningsFilterWrapper(ApplyWarningsFilterDelegateAttribute filterAttribute, MethodInfo method)
+		{
+			_applyWarningsFilterDelegateAttribute = filterAttribute;
+			_method = method;
+		}
+
+		/// <inheritdoc/>
+		public bool CanApply(CompilerWarningsToolChainContext context)
+		{
+			bool canApply = _method != null;
+
+			if (canApply)
+			{
+				try
+				{
+					object? result = _method?.Invoke(null, [context]);
+					canApply = result != null ? (bool)result : false;
+				}
+				catch (Exception e)
+				{
+					_logger?.LogWarning("Have failed to reflectively invoke (name:{MethodName}) in the CanApply context. Exception: {Exception}.", _method?.Name, e.Message);
+				}
+			}
+
+			return canApply;
+		}
+	}
+}
diff --git a/Engine/Source/Programs/UnrealBuildTool/Configuration/CompileWarnings/ApplyWarningsFilterRegistry.cs b/Engine/Source/Programs/UnrealBuildTool/Configuration/CompileWarnings/ApplyWarningsFilterRegistry.cs
new file mode 100644
--- /dev/null
+++ b/Engine/Source/Programs/UnrealBuildTool/Configuration/CompileWarnings/ApplyWarningsFilterRegistry.cs
@@ -0,0 +1,151 @@
+// Copyright Epic Games, Inc. All Rights Reserved.
+
+using System;
+using System.Collections.Generic;
+using System.Linq;
+using System.Reflection;
+using EpicGames.Core;
+using Microsoft.CodeAnalysis;
+using Microsoft.Extensions.Logging;
+
+namespace UnrealBuildTool
+{
+	/// <summary>
+	/// Interface that participates in defining whether a filter can be applied in a <see cref="ApplyWarningsAttribute"/>.
+	/// </summary>
+	internal interface IApplyWarningsFilter
+	{
+		/// <summary>
+		/// Determines whether the <see cref="IApplyWarningsFilter"/> should be applied for the given <see cref="CompilerWarningsToolChainContext"/>.
+		/// </summary>
+		/// <param name="context">The context in which to evaluate the applicability of the filter.</param>
+		/// <returns>True if the filter should be applied, false otherwise.</returns>
+		bool CanApply(CompilerWarningsToolChainContext context);
+	}
+
+	/// <summary>
+	/// ApplyWarningsFilterRegistry is the container for all of the filters applicable within the <see cref="ApplyWarningsAttribute.CanApplyToContext(CompilerWarningsToolChainContext)"/> context.
+	/// </summary>
+	/// <remarks>The two primary entry points for extensibility are <see cref="ApplyWarningsFilterDelegateAttribute"/> and <see cref="ApplyWarningsFilterAttribute"/>.</remarks>
+	[ConfigurationFilter(FilterID.ConfigDebug, CppConfiguration.Debug)]
+	[ConfigurationFilter(FilterID.ConfigShipping, CppConfiguration.Shipping)]
+	[ToolChainFilter(FilterID.VCToolChain, typeof(VCToolChain))]
+	[ToolChainFilter(FilterID.ClangToolChain, typeof(ClangToolChain))]
+	[ToolChainVersionConstrainedFilter(FilterID.Version14Min, [14])]
+	[ToolChainVersionConstrainedFilter(FilterID.Version16Min, [16])]
+	[ToolChainVersionConstrainedFilter(FilterID.Version18Min, [18])]
+	[ToolChainVersionConstrainedFilter(FilterID.Version19Min, [19])]
+	[CppStandardConstrainedFilter(FilterID.Cpp20Min, CppStandardVersion.Cpp20)]
+	internal class ApplyWarningsFilterRegistry
+	{
+		private readonly ILogger? _logger = Log.Logger;
+		private static ApplyWarningsFilterRegistry? s_filterRegistry;
+		private readonly Dictionary<string, IApplyWarningsFilter> _filterNameToBehaviour = [];
+
+		/// <summary>
+		/// The filter name to filter representation of the registry.
+		/// </summary>
+		internal static IReadOnlyDictionary<string, IApplyWarningsFilter> Filters
+		{
+			get
+			{
+				if (s_filterRegistry == null)
+				{
+					s_filterRegistry = new ApplyWarningsFilterRegistry();
+				}
+
+				return s_filterRegistry._filterNameToBehaviour;
+			}
+		}
+
+		/// <summary>
+		/// Requests a filter by name.
+		/// </summary>
+		/// <param name="filterName">The name of the filter to obtain.</param>
+		/// <returns>The ApplyWarningsFilter if it has been registered, null if it was not found.</returns>
+		internal static IApplyWarningsFilter? RequestFilter(string filterName)
+		{
+			IApplyWarningsFilter? returnFilter;
+			Filters.TryGetValue(filterName, out returnFilter);
+
+			return returnFilter;
+		}
+
+		private ApplyWarningsFilterRegistry()
+		{
+			// Process all ApplyWarningsFilterAttribute added to this class first, and store in the filterIDCacheDict
+			Type filterRegistryType = typeof(ApplyWarningsFilterRegistry);
+
+			IEnumerable<ApplyWarningsFilterAttribute> attributes = filterRegistryType
+				.GetCustomAttributes(typeof(ApplyWarningsFilterAttribute), inherit: false)
+				.Cast<ApplyWarningsFilterAttribute>();
+
+			foreach (ApplyWarningsFilterAttribute attribute in attributes)
+			{
+				if (!_filterNameToBehaviour.TryAdd(attribute.FilterName, attribute))
+				{
+					_logger?.LogWarning("Duplicate filter names ({DuplicateName}) encountered when building the ApplyWarningsFilterRegistry.", attribute.FilterName);
+				}
+			}
+
+			// Process all of the static methods in the assembly that are annotated wtih ApplyWarningsFilterDelegateAttribute, and store in _nameCacheDict given they match
+			// the execution delegate, which is effectively IApplyWarningsFilter.CanApply
+			MethodInfo? reflectiveInterfaceMethodInfo = typeof(IApplyWarningsFilter).GetMethod(nameof(IApplyWarningsFilter.CanApply));
+
+			if (reflectiveInterfaceMethodInfo != null)
+			{
+				Assembly assembly = Assembly.GetExecutingAssembly();
+				Type[] allTypes = assembly.GetTypes();
+				List<MethodInfo> delegatesForConsideration = new List<MethodInfo>();
+
+				foreach (Type type in allTypes)
+				{
+					// We could place all of the additional filter requirements in with the initial reflection call, but we want to provide better messaging to better guide users
+					// attempting to extend the functionality.
+					IEnumerable<MethodInfo> methodsWithAttribute = type
+						.GetMethods(BindingFlags.Static | BindingFlags.Public | BindingFlags.NonPublic)
+						.Where(m => m.GetCustomAttributes(typeof(ApplyWarningsFilterDelegateAttribute), false).Length > 0);
+
+					delegatesForConsideration.AddRange(methodsWithAttribute);
+				}
+
+				// Verify the method infos match the signature of our fundamental interface: IApplyWarningsFilter.CanApply
+				foreach (MethodInfo mi in delegatesForConsideration)
+				{
+					if (mi.ReturnType != reflectiveInterfaceMethodInfo.ReturnType)
+					{
+						_logger?.LogWarning("Skipping MethodInfo (name:{MethodName}) annotated with ApplyWarningsFilterDelegateAttribute as it doesn't return {ReturnType}.", mi.Name, reflectiveInterfaceMethodInfo.ReturnType.Name);
+						continue;
+					}
+
+					if (mi.GetParameters().Length != reflectiveInterfaceMethodInfo.GetParameters().Length)
+					{
+						_logger?.LogWarning("Skipping MethodInfo (name:{MethodName}) annotated with ApplyWarningsFilterDelegateAttribute as it doesn't contain the same number of parameters as the base method interface ({Count})", mi.Name, reflectiveInterfaceMethodInfo.GetParameters().Length);
+						continue;
+					}
+					if (!mi.GetParameters().SequenceEqual(reflectiveInterfaceMethodInfo.GetParameters(), EqualityComparer<ParameterInfo>.Create((a, b) => a?.ParameterType == b?.ParameterType)))
+					{
+						_logger?.LogWarning("Skipping MethodInfo (name:{MethodName}) annotated with ApplyWarningsFilterDelegateAttribute it's parameters don't match the base method interface.", mi.Name);
+						continue;
+					}
+
+					ApplyWarningsFilterDelegateAttribute? attributeInstance = mi.GetCustomAttribute<ApplyWarningsFilterDelegateAttribute>();
+					if (attributeInstance == null)
+					{
+						_logger?.LogWarning("Skipping MethodInfo (name:{MethodName}) annotated with ApplyWarningsFilterDelegateAttribute as reflection has failed to obtain the attribute instance.", mi.Name);
+						continue;
+					}
+
+					if (!_filterNameToBehaviour.TryAdd(attributeInstance.FilterName, new ApplyWarningsFilterWrapper(attributeInstance, mi)))
+					{
+						_logger?.LogWarning("Duplicate filter names ({DuplicateName}) encountered when building the ApplyWarningsFilterRegistry.", attributeInstance.FilterName);
+					}
+				}
+			}
+			else
+			{
+				_logger?.LogWarning("Base method interface was null. Ignoring all method reflective methods used for filters.");
+			}
+		}
+	}
+}
\ No newline at end of file
diff --git a/Engine/Source/Programs/UnrealBuildTool/Configuration/CompileWarnings/SpecializedApplyWarningsFilterDelegates.cs b/Engine/Source/Programs/UnrealBuildTool/Configuration/CompileWarnings/SpecializedApplyWarningsFilterDelegates.cs
new file mode 100644
--- /dev/null
+++ b/Engine/Source/Programs/UnrealBuildTool/Configuration/CompileWarnings/SpecializedApplyWarningsFilterDelegates.cs
@@ -0,0 +1,51 @@
+// Copyright Epic Games, Inc. All Rights Reserved.
+
+using EpicGames.Core;
+
+namespace UnrealBuildTool
+{
+	/// <summary>
+	/// Specialized filters with more complex logic that aren't typically reusable.
+	/// </summary>
+	internal static class SpecializedFilters
+	{
+		/// <summary>
+		///  Specialized filter to enforce Android as an excluded platform, at a specific toolchain version, at a specific C++ standard.
+		/// </summary>
+		/// <param name="context">The context in which to evaluate the applicability of the filter.</param>
+		/// <returns>True if the ApplyWarnings will apply, false otherwise.</returns>
+		[ApplyWarningsFilterDelegate(nameof(AndroidNDKR26ToolChainVersionExclusion))]
+		internal static bool AndroidNDKR26ToolChainVersionExclusion(CompilerWarningsToolChainContext context)
+		{
+			if (context._toolChainVersion == null)
+			{
+				return false;
+			}
+
+			bool versionMin = context._toolChainVersion >= new VersionNumber(17);
+			bool bIsAndroidClang17 = context._toolChainVersion == new VersionNumber(17, 0, 2) && context._compileEnvironment.Platform == UnrealTargetPlatform.Android;
+
+			return versionMin && !bIsAndroidClang17 && context._compileEnvironment.CppStandard < CppStandardVersion.Latest;
+		}
+
+		/// <summary>
+		/// Specialized filter to enforce Android as an excluded platform, at a specific toolchain version.
+		/// </summary>
+		/// <param name="context">The context in which to evaluate the applicability of the filter.</param>
+		/// <returns>True if the ApplyWarnings will apply, false otherwise.</returns>
+		[ApplyWarningsFilterDelegate(nameof(AndroidNDKR28ToolChainVersionExclusion))]
+		internal static bool AndroidNDKR28ToolChainVersionExclusion(CompilerWarningsToolChainContext context)
+		{
+			if (context._toolChainVersion == null)
+			{
+				return false;
+			}
+
+			bool versionMin = context._toolChainVersion >= new VersionNumber(19);
+			bool isExcludedVersion = context._toolChainVersion == new VersionNumber(19, 0, 0);
+			bool bIsAndroidClang19 = isExcludedVersion && context._compileEnvironment.Platform == UnrealTargetPlatform.Android;
+
+			return versionMin && !bIsAndroidClang19;
+		}
+	}
+}
diff --git a/Engine/Source/Programs/UnrealBuildTool/Configuration/CompileWarnings/StandardApplyWarningsFilterDelegates.cs b/Engine/Source/Programs/UnrealBuildTool/Configuration/CompileWarnings/StandardApplyWarningsFilterDelegates.cs
new file mode 100644
--- /dev/null
+++ b/Engine/Source/Programs/UnrealBuildTool/Configuration/CompileWarnings/StandardApplyWarningsFilterDelegates.cs
@@ -0,0 +1,58 @@
+// Copyright Epic Games, Inc. All Rights Reserved.
+
+namespace UnrealBuildTool
+{
+	/// <summary>
+	/// Standard filters that can be reused, and often in a composable manner.
+	/// </summary>
+	internal static class StandardFilters
+	{
+		/// <summary>
+		/// Filter to enforce deterministic compilation constraints.
+		/// </summary>
+		/// <param name="context">The context in which to evaluate the applicability of the filter.</param>
+		/// <returns>True if the ApplyWarnings will apply, false otherwise.</returns>
+		[ApplyWarningsFilterDelegate(nameof(DeterministcFlagSetFilter))]
+		internal static bool DeterministcFlagSetFilter(CompilerWarningsToolChainContext context)
+		{
+			return context._compileEnvironment.bDeterministic;
+		}
+
+		/// <summary>
+		/// Filter to profile guided optimization constraints.
+		/// </summary>
+		/// <param name="context">The context in which to evaluate the applicability of the filter.</param>
+		/// <returns>True if the ApplyWarnings will apply, false otherwise.</returns>
+		[ApplyWarningsFilterDelegate(nameof(PGOOptimizedFilter))]
+		internal static bool PGOOptimizedFilter(CompilerWarningsToolChainContext context)
+		{
+			return context._compileEnvironment.bPGOOptimize;
+		}
+
+		/// <summary>
+		/// Filter for MSVC compiler within the <see cref="VCToolChain"/>.
+		/// </summary>
+		/// <param name="context">The context in which to evaluate the applicability of the filter.</param>
+		/// <returns>True if the ApplyWarnings will apply, false otherwise.</returns>
+		[ApplyWarningsFilterDelegate(nameof(MSVCCompilerFilter))]
+		internal static bool MSVCCompilerFilter(CompilerWarningsToolChainContext context)
+		{
+			ReadOnlyTargetRules? readOnlyTargetRules = context._buildSystemContext.GetReadOnlyTargetRules();
+
+			return readOnlyTargetRules != null && readOnlyTargetRules.WindowsPlatform != null && readOnlyTargetRules.WindowsPlatform.Compiler.IsMSVC();
+		}
+
+		/// <summary>
+		/// Filter for Clang compiler within the <see cref="VCToolChain"/>.
+		/// </summary>
+		/// <param name="context">The context in which to evaluate the applicability of the filter.</param>
+		/// <returns>True if the ApplyWarnings will apply, false otherwise.</returns>
+		[ApplyWarningsFilterDelegate(nameof(VCClangCompilerFilter))]
+		internal static bool VCClangCompilerFilter(CompilerWarningsToolChainContext context)
+		{
+			ReadOnlyTargetRules? readOnlyTargetRules = context._buildSystemContext.GetReadOnlyTargetRules();
+
+			return readOnlyTargetRules != null && readOnlyTargetRules.WindowsPlatform != null && readOnlyTargetRules.WindowsPlatform.Compiler.IsClang();
+		}
+	}
+}
diff --git a/Engine/Source/Programs/UnrealBuildTool/Configuration/CppCompileWarnings.cs b/Engine/Source/Programs/UnrealBuildTool/Configuration/CppCompileWarnings.cs
--- a/Engine/Source/Programs/UnrealBuildTool/Configuration/CppCompileWarnings.cs
+++ b/Engine/Source/Programs/UnrealBuildTool/Configuration/CppCompileWarnings.cs
@@ -68,7 +68,7 @@ namespace UnrealBuildTool
 		/// </summary>
 		[CommandLine("-ShadowVariableErrors", Value = nameof(WarningLevel.Error))]
 		[RequiresUniqueBuildEnvironment]
-		[WarningsMSCV(["/wd4456", "/wd4458", "/wd4459"], [], ["/we4456", "/we4458", "/we4459"])]
+		[WarningsMSCV(["/wd4456", "/wd4458", "/wd4459"], null, ["/we4456", "/we4458", "/we4459"])]
 		[ShadowVariableWarningsClangToolChain(["-Wno-shadow"], ["-Wshadow", "-Wno-error=shadow"], ["-Wshadow"])]
 		public WarningLevel ShadowVariableWarningLevel
 		{
@@ -160,7 +160,7 @@ namespace UnrealBuildTool
 		/// Clang -
 		///		https://clang.llvm.org/docs/DiagnosticsReference.html#wenum-enum-conversion
 		/// </summary>
-		[WarningsClangToolChain(["-Wno-enum-enum-conversion"], null, null)]
+		[WarningsClangToolChain(["-Wno-enum-enum-conversion"])]
 		public WarningLevel EnumEnumConversionWarningLevel
 		{
 			get => ResolveWarning(ParentCppCompileWarnings?.EnumEnumConversionWarningLevel, _enumEnumConversionWarningLevel);
@@ -173,7 +173,7 @@ namespace UnrealBuildTool
 		/// Clang - 
 		///		https://clang.llvm.org/docs/DiagnosticsReference.html#wenum-float-conversion
 		/// </summary>
-		[WarningsClangToolChain(["-Wno-enum-float-conversion"], null, null)]
+		[WarningsClangToolChain(["-Wno-enum-float-conversion"])]
 		public WarningLevel EnumFloatConversionWarningLevel
 		{
 			get => ResolveWarning(ParentCppCompileWarnings?.EnumFloatConversionWarningLevel, _enumFloatConversionWarningLevel);
@@ -186,7 +186,7 @@ namespace UnrealBuildTool
 		/// Clang -
 		///		https://clang.llvm.org/docs/DiagnosticsReference.html#wambiguous-reversed-operator
 		/// </summary>
-		[CppStandardConstrainedClangToolChain(["-Wno-ambiguous-reversed-operator"], null, null, CppStandardVersion.Cpp20)]
+		[WarningsClangToolChain(["-Wno-ambiguous-reversed-operator"], null, null, [nameof(FilterID.Cpp20Min)])]
 		public WarningLevel AmbiguousReversedOperatorWarningLevel
 		{
 			get => ResolveWarning(ParentCppCompileWarnings?.AmbiguousReversedOperatorWarningLevel, _ambiguousReversedOperatorWarningLevel);
@@ -199,7 +199,7 @@ namespace UnrealBuildTool
 		/// Clang -
 		///		https://clang.llvm.org/docs/DiagnosticsReference.html#wdeprecated-anon-enum-enum-conversion
 		/// </summary>
-		[CppStandardConstrainedClangToolChain(["-Wno-deprecated-anon-enum-enum-conversion"], null, null, CppStandardVersion.Cpp20)]
+		[WarningsClangToolChain(["-Wno-deprecated-anon-enum-enum-conversion"], null, null, [nameof(FilterID.Cpp20Min)])]
 		public WarningLevel DeprecatedAnonEnumEnumConversionWarningLevel
 		{
 			get => ResolveWarning(ParentCppCompileWarnings?.DeprecatedAnonEnumEnumConversionWarningLevel, _deprecatedAnonEnumEnumConversionWarningLevel);
@@ -212,7 +212,7 @@ namespace UnrealBuildTool
 		/// Clang -
 		///		https://clang.llvm.org/docs/DiagnosticsReference.html#wdeprecated-volatile
 		/// </summary>
-		[CppStandardConstrainedClangToolChain(["-Wno-deprecated-volatile"], null, null, CppStandardVersion.Cpp20)]
+		[WarningsClangToolChain(["-Wno-deprecated-volatile"], null, null, [nameof(FilterID.Cpp20Min)])]
 		public WarningLevel DeprecatedVolatileWarningLevel
 		{
 			get => ResolveWarning(ParentCppCompileWarnings?.DeprecatedVolatileWarningLevel, _deprecatedVolatileWarningLevel);
@@ -225,7 +225,7 @@ namespace UnrealBuildTool
 		/// Clang -
 		///		https://clang.llvm.org/docs/DiagnosticsReference.html#wordered-compare-function-pointers
 		/// </summary>
-		[WarningsClangToolChain(["-Wno-ordered-compare-function-pointers"], null, null)]
+		[WarningsClangToolChain(["-Wno-ordered-compare-function-pointers"])]
 		public WarningLevel OrderedCompareFunctionPointers
 		{
 			get => ResolveWarning(ParentCppCompileWarnings?.OrderedCompareFunctionPointers, _orderedCompareFunctionPointers);
@@ -238,7 +238,7 @@ namespace UnrealBuildTool
 		/// Clang -
 		///		https://clang.llvm.org/docs/DiagnosticsReference.html#wbitwise-instead-of-logical
 		/// </summary>
-		[ToolChainVersionConstrainedWarningsClangToolChain(["-Wno-bitwise-instead-of-logical"], null, null, [14], null)]
+		[WarningsClangToolChain(["-Wno-bitwise-instead-of-logical"], null, null, [nameof(FilterID.Version14Min)])]
 		public WarningLevel BitwiseInsteadOfLogicalWarningLevel
 		{
 			get => ResolveWarning(ParentCppCompileWarnings?.BitwiseInsteadOfLogicalWarningLevel, _bitwiseInsteadOfLogicalWarningLevel);
@@ -251,7 +251,7 @@ namespace UnrealBuildTool
 		/// Clang -
 		///		https://clang.llvm.org/docs/DiagnosticsReference.html#wdeprecated-copy
 		/// </summary>
-		[ToolChainVersionConstrainedWarningsClangToolChain(["-Wno-deprecated-copy"], null, null, [16], null)]
+		[WarningsClangToolChain(["-Wno-deprecated-copy"], null, null, [nameof(FilterID.Version16Min)])]
 		public WarningLevel DeprecatedCopyWarningLevel
 		{
 			get => ResolveWarning(ParentCppCompileWarnings?.DeprecatedCopyWarningLevel, _deprecatedCopyWarningLevel);
@@ -264,7 +264,7 @@ namespace UnrealBuildTool
 		/// Clang -
 		///		https://clang.llvm.org/docs/DiagnosticsReference.html#wdeprecated-copy-with-user-provided-copy
 		/// </summary>
-		[ToolChainVersionConstrainedWarningsClangToolChain(["-Wno-deprecated-copy-with-user-provided-copy"], null, null, [16], null)]
+		[WarningsClangToolChain(["-Wno-deprecated-copy-with-user-provided-copy"], null, null, [nameof(FilterID.Version16Min)])]
 		public WarningLevel DeprecatedCopyWithUserProvidedCopyWarningLevel
 		{
 			get => ResolveWarning(ParentCppCompileWarnings?.DeprecatedCopyWithUserProvidedCopyWarningLevel, _deprecatedCopyWithUserProvidedCopyWarningLevel);
@@ -277,7 +277,7 @@ namespace UnrealBuildTool
 		/// Clang -
 		///		https://clang.llvm.org/docs/DiagnosticsReference.html#winvalid-unevaluated-string
 		/// </summary>
-		[AndroidNDKR26ToolChainVersionExclusion(["-Wno-invalid-unevaluated-string"], null, null)]
+		[WarningsClangToolChain(["-Wno-invalid-unevaluated-string"], null, null, [nameof(SpecializedFilters.AndroidNDKR26ToolChainVersionExclusion)])]
 		public WarningLevel InvalidUnevaluatedStringWarningLevel
 		{
 			get => ResolveWarning(ParentCppCompileWarnings?.InvalidUnevaluatedStringWarningLevel, _invalidUnevaluatedStringWarningLevel);
@@ -295,7 +295,7 @@ namespace UnrealBuildTool
 		/// It could be easily fixed in the engine, but would create the risk of Win64 building fine and the code failing just on Clang-based platforms. 
 		/// We tend to use NAN just as a bit pattern, and not rely on it in calculations, so it should be reasonably safe to disable.
 		/// </remarks>
-		[ToolChainVersionConstrainedWarningsClangToolChain(["-Wno-nan-infinity-disabled"], null, null, [18], null)]
+		[WarningsClangToolChain(["-Wno-nan-infinity-disabled"], null, null, [nameof(FilterID.Version18Min)])]
 		public WarningLevel NaNInfinityDisabledWarningLevel
 		{
 			get => ResolveWarning(ParentCppCompileWarnings?.NaNInfinityDisabledWarningLevel, _nanInfinityDisabledWarningLevel);
@@ -308,7 +308,7 @@ namespace UnrealBuildTool
 		/// Clang -
 		///		https://clang.llvm.org/docs/DiagnosticsReference.html#wextra-qualification
 		/// </summary>
-		[ToolChainVersionConstrainedWarningsClangToolChain(["-Wno-extra-qualification"], null, null, [19], null)]
+		[WarningsClangToolChain(["-Wno-extra-qualification"], null, null, [nameof(FilterID.Version19Min)])]
 		public WarningLevel LevelExtraQualificationWarningLevel
 		{
 			get => ResolveWarning(ParentCppCompileWarnings?.LevelExtraQualificationWarningLevel, _levelExtraQualificationWarningLevel);
@@ -321,7 +321,7 @@ namespace UnrealBuildTool
 		/// Clang -
 		///		https://clang.llvm.org/docs/DiagnosticsReference.html#wcast-function-type-mismatch
 		/// </summary>
-		[AndroidNDKR28ToolChainVersionExclusion(["-Wno-cast-function-type-mismatch"], null, null)]
+		[WarningsClangToolChain(["-Wno-cast-function-type-mismatch"], null, null, [nameof(SpecializedFilters.AndroidNDKR28ToolChainVersionExclusion)])]
 		public WarningLevel CastFunctionTypeMistmatchWarningLevel
 		{
 			get => ResolveWarning(ParentCppCompileWarnings?.CastFunctionTypeMistmatchWarningLevel, _castFunctionTypeMistmatchWarningLevel);
@@ -334,7 +334,7 @@ namespace UnrealBuildTool
 		/// Clang -
 		///		https://clang.llvm.org/docs/DiagnosticsReference.html#wmissing-template-arg-list-after-template-kw
 		/// </summary>
-		[AndroidNDKR28ToolChainVersionExclusion(["-Wno-missing-template-arg-list-after-template-kw"], null, null)]
+		[WarningsClangToolChain(["-Wno-missing-template-arg-list-after-template-kw"], null, null, [nameof(SpecializedFilters.AndroidNDKR28ToolChainVersionExclusion)])]
 		public WarningLevel MissingTemplateArgListAfterTemplateWarningLevel
 		{
 			get => ResolveWarning(ParentCppCompileWarnings?.MissingTemplateArgListAfterTemplateWarningLevel, _missingTemplateArgListAfterTemplateWarningLevel);
@@ -348,7 +348,7 @@ namespace UnrealBuildTool
 		///		https://clang.llvm.org/docs/DiagnosticsReference.html#wgnu-string-literal-operator-template
 		/// </summary>
 		/// <remarks>We use this feature to allow static FNames.</remarks>
-		[WarningsClangToolChain(["-Wno-gnu-string-literal-operator-template"], null, null)]
+		[WarningsClangToolChain(["-Wno-gnu-string-literal-operator-template"])]
 		public WarningLevel GNUStringLiteralOperatorTemplateWarningLevel
 		{
 			get => ResolveWarning(ParentCppCompileWarnings?.GNUStringLiteralOperatorTemplateWarningLevel, _gnuStringLiteralOperatorTemplateWarningLevel);
@@ -361,7 +361,7 @@ namespace UnrealBuildTool
 		/// Clang -
 		///		https://clang.llvm.org/docs/DiagnosticsReference.html#winconsistent-missing-override
 		/// </summary>
-		[WarningsClangToolChain(["-Wno-inconsistent-missing-override"], null, null)]
+		[WarningsClangToolChain(["-Wno-inconsistent-missing-override"])]
 		public WarningLevel InconsistentMissingOverrideWarningLevel
 		{
 			get => ResolveWarning(ParentCppCompileWarnings?.InconsistentMissingOverrideWarningLevel, _inconsistentMissingOverrideWarningLevel);
@@ -375,7 +375,7 @@ namespace UnrealBuildTool
 		///		https://clang.llvm.org/docs/DiagnosticsReference.html#winvalid-offsetof
 		/// </summary>
 		/// <remarks>Needed to suppress warnings about using offsetof on non-POD types.</remarks>
-		[WarningsClangToolChain(["-Wno-invalid-offsetof"], null, null)]
+		[WarningsClangToolChain(["-Wno-invalid-offsetof"])]
 		public WarningLevel InvalidOffsetWarningLevel
 		{
 			get => ResolveWarning(ParentCppCompileWarnings?.InvalidOffsetWarningLevel, _invalidOffsetOfWarningLevel);
@@ -389,7 +389,7 @@ namespace UnrealBuildTool
 		///		https://clang.llvm.org/docs/DiagnosticsReference.html#wswitch
 		/// </summary>
 		/// <remarks>Fhis hides the "enumeration value 'XXXXX' not handled in switch [-Wswitch]" warnings - we should maybe remove this at some point and add UE_LOG(, Fatal, ) to default cases.</remarks>
-		[WarningsClangToolChain(["-Wno-switch"], null, null)]
+		[WarningsClangToolChain(["-Wno-switch"])]
 		public WarningLevel SwitchWarningLevel
 		{
 			get => ResolveWarning(ParentCppCompileWarnings?.SwitchWarningLevel, _switchWarningLevel);
@@ -403,7 +403,7 @@ namespace UnrealBuildTool
 		///		https://clang.llvm.org/docs/DiagnosticsReference.html#wtautological-compare
 		/// </summary>
 		/// <remarks>This hides the "warning : comparison of unsigned expression &lt; 0 is always false" type warnings due to constant comparisons, which are possible with template arguments/</remarks>
-		[WarningsClangToolChain(["-Wno-tautological-compare"], null, null)]
+		[WarningsClangToolChain(["-Wno-tautological-compare"])]
 		public WarningLevel TautologicalCompareWarningLevel
 		{
 			get => ResolveWarning(ParentCppCompileWarnings?.TautologicalCompareWarningLevel, _tautologicalCompareWarningLevel);
@@ -417,7 +417,7 @@ namespace UnrealBuildTool
 		///		https://clang.llvm.org/docs/DiagnosticsReference.html#wunknown-pragmas
 		/// </summary>
 		/// <remarks>Slate triggers this (with its optimize on/off pragmas).</remarks>
-		[WarningsClangToolChain(["-Wno-unknown-pragmas"], null, null)]
+		[WarningsClangToolChain(["-Wno-unknown-pragmas"])]
 		public WarningLevel UnknownPragmasWarningLevel
 		{
 			get => ResolveWarning(ParentCppCompileWarnings?.UnknownPragmasWarningLevel, _unknownPragmasWarningLevel);
@@ -456,7 +456,7 @@ namespace UnrealBuildTool
 		///		https://clang.llvm.org/docs/DiagnosticsReference.html#wundefined-var-template
 		/// </summary>
 		/// <remarks>Not a good warning to be disabling.</remarks>
-		[WarningsClangToolChain(["-Wno-undefined-var-template"], null, null)]
+		[WarningsClangToolChain(["-Wno-undefined-var-template"])]
 		public WarningLevel UndefinedVarTemplateWarningLevel
 		{
 			get => ResolveWarning(ParentCppCompileWarnings?.UndefinedVarTemplateWarningLevel, _undefinedVarTemplateWarningLevel);
@@ -471,7 +471,7 @@ namespace UnrealBuildTool
 		///		https://clang.llvm.org/docs/DiagnosticsReference.html#wprofile-instr-unprofiled
 		/// </summary>
 		/// <remarks> Clang emits warnings for each compiled object file that doesn't have a matching entry in the profile data. This can happen when the profile data is older than the binaries we're compiling. Disable these warnings.</remarks>
-		[PGOOptimizedClangToolChain(["-Wno-profile-instr-out-of-date", "-Wno-profile-instr-unprofiled"], null, null)]
+		[WarningsClangToolChain(["-Wno-profile-instr-out-of-date", "-Wno-profile-instr-unprofiled"], null, null, [nameof(StandardFilters.PGOOptimizedFilter)])]
 		public WarningLevel ProfileInstructWarningLevel
 		{
 			get => ResolveWarning(ParentCppCompileWarnings?.ProfileInstructWarningLevel, _profileInstructWarningLevel);
@@ -484,7 +484,7 @@ namespace UnrealBuildTool
 		///		https://clang.llvm.org/docs/DiagnosticsReference.html#wbackend-plugin
 		/// </summary>
 		/// <remarks>Apparently there can be hashing conflicts with PGO which can result in: 'Function control flow change detected (hash mismatch)' warnings. </remarks>
-		[PGOOptimizedClangToolChain(["-Wno-backend-plugin"], null, null)]
+		[WarningsClangToolChain(["-Wno-backend-plugin"], null, null, [nameof(StandardFilters.PGOOptimizedFilter)])]
 		public WarningLevel BackendPluginWarningLevel
 		{
 			get => ResolveWarning(ParentCppCompileWarnings?.BackendPluginWarningLevel, _backendPluginWarningLevel);
@@ -498,7 +498,7 @@ namespace UnrealBuildTool
 		///		https://clang.llvm.org/docs/DiagnosticsReference.html#wunused-value
 		/// </summary>
 		/// <remarks>Shipping builds will cause this warning with "ensure", so disable only in those case.</remarks>
-		[UnusedValueClangToolChain(["-Wno-unused-value"], null, null, CppConfiguration.Shipping)]
+		[UnusedValueClangToolChain(["-Wno-unused-value"])]
 		public WarningLevel UnusedValueWarningLevel
 		{
 			get => ResolveWarning(ParentCppCompileWarnings?.UnusedValueWarningLevel, _unusedValueWarningLevel);
@@ -512,7 +512,7 @@ namespace UnrealBuildTool
 		///		https://clang.llvm.org/docs/DiagnosticsReference.html#wshorten-64-to-32
 		/// </summary>
 		/// <remarks>-Wimplicit-int-conversion also controls -Wshorten-64-to-32 as of clang 19.</remarks>
-		[ToolChainVersionConstrainedWarningsClangToolChain(["-Wno-shorten-64-to-32"], null, null, [19], null)]
+		[WarningsClangToolChain(["-Wno-shorten-64-to-32"], null, null, [nameof(FilterID.Version19Min)])]
 		public WarningLevel Shorten64To32WarningLevel
 		{
 			get => ResolveWarning(ParentCppCompileWarnings?.Shorten64To32WarningLevel, _shorten64To32WarningLevel);
@@ -526,7 +526,7 @@ namespace UnrealBuildTool
 		///		https://clang.llvm.org/docs/DiagnosticsReference.html#wdllexport-explicit-instantiation-decl
 		/// </summary>
 		/// <remarks>The code base contains lots of places where we do "extern template X_API class ..." and we want to keep doing that to reduce compile times.</remarks>
-		[WarningsClangToolChain(["-Wno-dllexport-explicit-instantiation-decl"], null, null)]
+		[WarningsClangToolChain(["-Wno-dllexport-explicit-instantiation-decl"])]
 		public WarningLevel DllExportExplicitInstantiationDeclWarningLevel
 		{
 			get => ResolveWarning(ParentCppCompileWarnings?.DllExportExplicitInstantiationDeclWarningLevel, _dllExportExplicitInstantiationDeclWarningLevel);
@@ -544,7 +544,7 @@ namespace UnrealBuildTool
 		///		https://clang.llvm.org/docs/DiagnosticsReference.html#wmicrosoft
 		/// </summary>
 		/// <remarks>Allow Microsoft-specific syntax to slide, even though it may be non-standard.  Needed for Windows headers.</remarks>
-		[WarningsVCClang(["-Wno-microsoft"], null, null)]
+		[WarningsVCClang(["-Wno-microsoft"])]
 		public WarningLevel MicrosoftGroupWarningLevel
 		{
 			get => ResolveWarning(ParentCppCompileWarnings?.MicrosoftGroupWarningLevel, _microsoftGroupWarningLevel);
@@ -558,7 +558,7 @@ namespace UnrealBuildTool
 		///		https://clang.llvm.org/docs/DiagnosticsReference.html#wmsvc-include
 		/// </summary>
 		/// <remarks> Hack due to how we have our 'DummyPCH' wrappers setup when using unity builds.  This warning should not be disabled!</remarks>
-		[WarningsVCClang(["-Wno-msvc-include"], null, null)]
+		[WarningsVCClang(["-Wno-msvc-include"])]
 		public WarningLevel MSVCIncludeWarningLevel
 		{
 			get => ResolveWarning(ParentCppCompileWarnings?.MSVCIncludeWarningLevel, _msvcIncludeWarningLevel);
@@ -575,7 +575,7 @@ namespace UnrealBuildTool
 		/// This is disabled because clang explicitly warns about changing pack alignment in a header and not restoring it afterwards, which is something we do with the Pre/PostWindowsApi.h headers.
 		/// @TODO clang: push/pop this in  Pre/PostWindowsApi.h headers instead?
 		/// </remarks>
-		[WarningsVCClang(["-Wno-pragma-pack"], null, null)]
+		[WarningsVCClang(["-Wno-pragma-pack"])]
 		public WarningLevel PragmaPackWarningLevel
 		{
 			get => ResolveWarning(ParentCppCompileWarnings?.PragmaPackWarningLevel, _pragmaPackWarningLevel);
@@ -589,7 +589,7 @@ namespace UnrealBuildTool
 		///		https://clang.llvm.org/docs/DiagnosticsReference.html#winline-new-delete
 		/// </summary>
 		/// <remarks>@todo clang: We declare operator new as inline.  Clang doesn't seem to like that.</remarks>
-		[WarningsVCClang(["-Wno-inline-new-delete"], null, null)]
+		[WarningsVCClang(["-Wno-inline-new-delete"])]
 		public WarningLevel InlineNewDeleteWarningLevel
 		{
 			get => ResolveWarning(ParentCppCompileWarnings?.InlineNewDeleteWarningLevel, _inlineNewDeleteWarningLevel);
@@ -603,7 +603,7 @@ namespace UnrealBuildTool
 		///		https://clang.llvm.org/docs/DiagnosticsReference.html#wimplicit-exception-spec-mismatch
 		/// </summary>
 
-		[WarningsVCClang(["-Wno-implicit-exception-spec-mismatch"], null, null)]
+		[WarningsVCClang(["-Wno-implicit-exception-spec-mismatch"])]
 		public WarningLevel ImplicitExceptionSpecMismatchWarningLevel
 		{
 			get => ResolveWarning(ParentCppCompileWarnings?.ImplicitExceptionSpecMismatchWarningLevel, _implicitExceptionSpecMismatchWarningLevel);
@@ -617,7 +617,7 @@ namespace UnrealBuildTool
 		///		https://clang.llvm.org/docs/DiagnosticsReference.html#wundefined-bool-conversion
 		/// </summary>
 		/// <remarks>Sometimes we compare 'this' pointers against nullptr, which Clang warns about by default.</remarks>
-		[WarningsVCClang(["-Wno-undefined-bool-conversion"], null, null)]
+		[WarningsVCClang(["-Wno-undefined-bool-conversion"])]
 		public WarningLevel UndefinedBoolConversionWarningLevel
 		{
 			get => ResolveWarning(ParentCppCompileWarnings?.UndefinedBoolConversionWarningLevel, _undefinedBoolConversionWarningLevel);
@@ -631,7 +631,7 @@ namespace UnrealBuildTool
 		///		https://clang.llvm.org/docs/DiagnosticsReference.html#wdeprecated-writable-strings
 		/// </summary>
 
-		[WarningsVCClang(["-Wno-deprecated-writable-strings"], null, null)]
+		[WarningsVCClang(["-Wno-deprecated-writable-strings"])]
 		public WarningLevel DeprecatedWriteableStringsWarningLevel
 		{
 			get => ResolveWarning(ParentCppCompileWarnings?.DeprecatedWriteableStringsWarningLevel, _deprecatedWriteableStringsWarningLevel);
@@ -644,7 +644,7 @@ namespace UnrealBuildTool
 		/// Clang -
 		///		https://clang.llvm.org/docs/DiagnosticsReference.html#wdeprecated-register
 		/// </summary>
-		[WarningsVCClang(["-Wno-deprecated-register"], null, null)]
+		[WarningsVCClang(["-Wno-deprecated-register"])]
 		public WarningLevel DeprecatedRegisterWarningLevel
 		{
 			get => ResolveWarning(ParentCppCompileWarnings?.DeprecatedRegisterWarningLevel, _deprecatedRegisterWarningLevel);
@@ -659,7 +659,7 @@ namespace UnrealBuildTool
 		///		https://clang.llvm.org/docs/DiagnosticsReference.html#wlogical-op-parentheses
 		/// </summary>
 		/// <remarks>Disable needed for external headers we shan't change.</remarks>
-		[WarningsVCClang(["-Wno-logical-op-parentheses"], null, null)]
+		[WarningsVCClang(["-Wno-logical-op-parentheses"])]
 		public WarningLevel LogicalOpParenthesesWarningLevel
 		{
 			get => ResolveWarning(ParentCppCompileWarnings?.LogicalOpParenthesesWarningLevel, _logicalOpParenthesesWarningLevel);
@@ -673,7 +673,7 @@ namespace UnrealBuildTool
 		///		https://clang.llvm.org/docs/DiagnosticsReference.html#wnull-arithmetic
 		/// </summary>
 		/// <remarks>Disable needed for external headers we shan't change.</remarks>
-		[WarningsVCClang(["-Wno-null-arithmetic"], null, null)]
+		[WarningsVCClang(["-Wno-null-arithmetic"])]
 		public WarningLevel NullArithmeticWarningLevel
 		{
 			get => ResolveWarning(ParentCppCompileWarnings?.NullArithmeticWarningLevel, _nullArithmeticWarningLevel);
@@ -687,7 +687,7 @@ namespace UnrealBuildTool
 		///		https://clang.llvm.org/docs/DiagnosticsReference.html#wreturn-type-c-linkage
 		/// </summary>
 		/// <remarks>Disable needed for PhysX.</remarks>
-		[WarningsVCClang(["-Wno-return-type-c-linkage"], null, null)]
+		[WarningsVCClang(["-Wno-return-type-c-linkage"])]
 		public WarningLevel ReturnTypeCLinkageWarningLevel
 		{
 			get => ResolveWarning(ParentCppCompileWarnings?.ReturnTypeCLinkageWarningLevel, _returnTypeCLinkageWarningLevel);
@@ -701,7 +701,7 @@ namespace UnrealBuildTool
 		///		https://clang.llvm.org/docs/DiagnosticsReference.html#wignored-attributes
 		/// </summary>
 		/// <remarks>Disable needed for nvtesslib.</remarks>
-		[WarningsVCClang(["-Wno-ignored-attributes"], null, null)]
+		[WarningsVCClang(["-Wno-ignored-attributes"])]
 		public WarningLevel IgnoredAttributesWarningLevel
 		{
 			get => ResolveWarning(ParentCppCompileWarnings?.IgnoredAttributesWarningLevel, _ignoredAttributesWarningLevel);
@@ -714,7 +714,7 @@ namespace UnrealBuildTool
 		/// Clang -
 		///		https://clang.llvm.org/docs/DiagnosticsReference.html#wuninitialized
 		/// </summary>
-		[WarningsVCClang(["-Wno-uninitialized"], null, null)]
+		[WarningsVCClang(["-Wno-uninitialized"])]
 		public WarningLevel UninitializedWarningLevel
 		{
 			get => ResolveWarning(ParentCppCompileWarnings?.UninitializedWarningLevel, _uninitializedWarningLevel);
@@ -728,7 +728,7 @@ namespace UnrealBuildTool
 		///		https://clang.llvm.org/docs/DiagnosticsReference.html#wreturn-type
 		/// </summary>
 		/// <remarks>Disable needed for external headers we shan't change.</remarks>
-		[WarningsVCClang(["-Wno-return-type"], null, null)]
+		[WarningsVCClang(["-Wno-return-type"])]
 		public WarningLevel ReturnTypeWarningLevel
 		{
 			get => ResolveWarning(ParentCppCompileWarnings?.ReturnTypeWarningLevel, _returnTypeWarningLevel);
@@ -742,7 +742,7 @@ namespace UnrealBuildTool
 		///		https://clang.llvm.org/docs/DiagnosticsReference.html#wunused-parameter
 		/// </summary>
 		/// <remarks>Unused function parameter. A lot are named 'bUnused'...</remarks>
-		[WarningsVCClang(["-Wno-unused-parameter"], null, null)]
+		[WarningsVCClang(["-Wno-unused-parameter"])]
 		public WarningLevel UnusedParameterWarningLevel
 		{
 			get => ResolveWarning(ParentCppCompileWarnings?.UnusedParameterWarningLevel, _unusedParameterWarningLevel);
@@ -756,7 +756,7 @@ namespace UnrealBuildTool
 		///		https://clang.llvm.org/docs/DiagnosticsReference.html#wignored-qualifiers
 		/// </summary>
 		/// <remarks>const ignored when returning by value e.g. 'const int foo() { return 4; }'.</remarks>
-		[WarningsVCClang(["-Wno-ignored-qualifiers"], null, null)]
+		[WarningsVCClang(["-Wno-ignored-qualifiers"])]
 		public WarningLevel IgnoredQualifiersWarningLevel
 		{
 			get => ResolveWarning(ParentCppCompileWarnings?.IgnoredQualifiersWarningLevel, _ignoredQualifiersWarningLevel);
@@ -770,7 +770,7 @@ namespace UnrealBuildTool
 		///		https://clang.llvm.org/docs/DiagnosticsReference.html#wexpansion-to-defined
 		/// </summary>
 		/// <remarks>Usage of 'defined(X)' in a macro definition. Gives different results under MSVC.</remarks>
-		[WarningsVCClang(["-Wno-expansion-to-defined"], null, null)]
+		[WarningsVCClang(["-Wno-expansion-to-defined"])]
 		public WarningLevel ExpansionToDefined
 		{
 			get => ResolveWarning(ParentCppCompileWarnings?.ExpansionToDefined, _expansionToDefined);
@@ -784,7 +784,7 @@ namespace UnrealBuildTool
 		///		https://clang.llvm.org/docs/DiagnosticsReference.html#wsign-compare
 		/// </summary>
 		/// <remarks>Signed/unsigned comparison - millions of these.</remarks>
-		[WarningsVCClang(["-Wno-sign-compare"], null, null)]
+		[WarningsVCClang(["-Wno-sign-compare"])]
 		public WarningLevel SignCompareWarningLevel
 		{
 			get => ResolveWarning(ParentCppCompileWarnings?.SignCompareWarningLevel, _signCompareWarningLevel);
@@ -798,7 +798,7 @@ namespace UnrealBuildTool
 		///		https://clang.llvm.org/docs/DiagnosticsReference.html#wmissing-field-initializers
 		/// </summary>
 		/// <remarks>Excessive warning, generated when you initialize with MyStruct A = {0};</remarks>
-		[WarningsVCClang(["-Wno-missing-field-initializers"], null, null)]
+		[WarningsVCClang(["-Wno-missing-field-initializers"])]
 		public WarningLevel MissingFieldInitializersWarningLevel
 		{
 			get => ResolveWarning(ParentCppCompileWarnings?.MissingFieldInitializersWarningLevel, _missingFieldInitializersWarningLevel);
@@ -811,7 +811,7 @@ namespace UnrealBuildTool
 		/// Clang -
 		///		https://clang.llvm.org/docs/DiagnosticsReference.html#wnonportable-include-path
 		/// </summary>
-		[WarningsVCClang(["-Wno-nonportable-include-path"], null, null)]
+		[WarningsVCClang(["-Wno-nonportable-include-path"])]
 		public WarningLevel NonPortableIncludePathWarningLevel
 		{
 			get => ResolveWarning(ParentCppCompileWarnings?.NonPortableIncludePathWarningLevel, _nonPortableIncludePathWarningLevel);
@@ -826,7 +826,7 @@ namespace UnrealBuildTool
 		///	MASVC -
 		///		https://learn.microsoft.com/en-us/cpp/preprocessor/token-pasting-operator-hash-hash?view=msvc-170
 		/// </summary>
-		[WarningsVCClang(["-Wno-invalid-token-paste"], null, null)]
+		[WarningsVCClang(["-Wno-invalid-token-paste"])]
 		public WarningLevel InvalidTokenPasteWarningLevel
 		{
 			get => ResolveWarning(ParentCppCompileWarnings?.InvalidTokenPasteWarningLevel, _invalidTokenPasteWarningLevel);
@@ -839,7 +839,7 @@ namespace UnrealBuildTool
 		/// Clang -
 		///		https://clang.llvm.org/docs/DiagnosticsReference.html#wnull-pointer-arithmetic
 		/// </summary>
-		[WarningsVCClang(["-Wno-null-pointer-arithmetic"], null, null)]
+		[WarningsVCClang(["-Wno-null-pointer-arithmetic"])]
 		public WarningLevel NullPointerArithmeticWarningLevel
 		{
 			get => ResolveWarning(ParentCppCompileWarnings?.NullPointerArithmeticWarningLevel, _nullPointerArithmeticWarningLevel);
@@ -853,7 +853,7 @@ namespace UnrealBuildTool
 		///		https://clang.llvm.org/docs/DiagnosticsReference.html#wconstant-logical-operand
 		/// </summary>
 		/// <remarks>Triggered by || of two template-derived values inside a static_assert.</remarks>
-		[WarningsVCClang(["-Wno-constant-logical-operand"], null, null)]
+		[WarningsVCClang(["-Wno-constant-logical-operand"])]
 		public WarningLevel ConstantLogicalOperandWarningLevel
 		{
 			get => ResolveWarning(ParentCppCompileWarnings?.ConstantLogicalOperandWarningLevel, _constantLogicalOperandWarningLevel);
@@ -866,7 +866,7 @@ namespace UnrealBuildTool
 		/// Clang -
 		///		https://clang.llvm.org/docs/DiagnosticsReference.html#wbitfield-enum-conversion
 		/// </summary>
-		[WarningsVCClang(["-Wno-bitfield-enum-conversion"], null, null)]
+		[WarningsVCClang(["-Wno-bitfield-enum-conversion"])]
 		public WarningLevel BitfieldEnumConversion
 		{
 			get => ResolveWarning(ParentCppCompileWarnings?.BitfieldEnumConversion, _bitfieldEnumConversion);
@@ -879,7 +879,7 @@ namespace UnrealBuildTool
 		/// Clang -
 		///		https://clang.llvm.org/docs/DiagnosticsReference.html#wnull-pointer-subtraction
 		/// </summary>
-		[WarningsVCClang(["-Wno-null-pointer-subtraction"], null, null)]
+		[WarningsVCClang(["-Wno-null-pointer-subtraction"])]
 		public WarningLevel NullPointerSubtractionWarningLevel
 		{
 			get => ResolveWarning(ParentCppCompileWarnings?.NullPointerSubtractionWarningLevel, _nullPointerSubtractionWarningLevel);
@@ -892,7 +892,7 @@ namespace UnrealBuildTool
 		/// Clang -
 		///		https://clang.llvm.org/docs/DiagnosticsReference.html#wdangling
 		/// </summary>	
-		[WarningsVCClang(["-Wno-dangling"], null, null)]
+		[WarningsVCClang(["-Wno-dangling"])]
 		public WarningLevel DanglingWarningLevel
 		{
 			get => ResolveWarning(ParentCppCompileWarnings?.DanglingWarningLevel, _danglingWarningLevel);
@@ -915,7 +915,7 @@ namespace UnrealBuildTool
 		/// The reason for duplication is that the clang field will only participate in the shipping configuration, whereas the MSVC variant has been disabled in all.
 		/// In order to ensure no change of behaviour for MSVC toolchain within the clang compilation context, we will keep this in, and deprecate with assignments to the clang setting.
 		/// </remarks>
-		[WarningsVCClang(["-Wno-switch-enum"], null, null)]
+		[WarningsVCClang(["-Wno-switch-enum"])]
 		public WarningLevel MSVCSwitchEnumWarningLevel
 		{
 			get => ResolveWarning(ParentCppCompileWarnings?.MSVCSwitchEnumWarningLevel, _msvcSwitchEnumWarningLevel);
@@ -933,7 +933,7 @@ namespace UnrealBuildTool
 		/// The reason for duplication is that the clang field will only participate in the shipping configuration, whereas the MSVC variant has been disabled in all.
 		/// In order to ensure no change of behaviour for MSVC toolchain within the clang compilation context, we will keep this in, and deprecate with assignments to the clang setting.
 		/// </remarks>
-		[WarningsVCClang(["-Wno-unused-value"], null, null)]
+		[WarningsVCClang(["-Wno-unused-value"])]
 		public WarningLevel MSVCUnusedValueWarningLevel
 		{
 			get => ResolveWarning(ParentCppCompileWarnings?.MSVCUnusedValueWarningLevel, _msvcUnusedValueWarningLevel);
@@ -951,7 +951,7 @@ namespace UnrealBuildTool
 		/// The reason for duplication is that the clang field will be defaulted to WARNING in clang, but not for msvc which is disabled.
 		/// In order to ensure no change of behaviour for MSVC toolchain within the clang compilation context, we will keep this in, and deprecate with assignments to the clang setting.
 		/// </remarks>
-		[WarningsVCClang(["-Wno-deprecated-declarations"], null, null)]
+		[WarningsVCClang(["-Wno-deprecated-declarations"])]
 		public WarningLevel MSVCDeprecationWarningLevel
 		{
 			get => ResolveWarning(ParentCppCompileWarnings?.MSVCDeprecationWarningLevel, _msvcDeprecationWarningLevel);
@@ -984,7 +984,7 @@ namespace UnrealBuildTool
 		/// Clang -
 		///		https://clang.llvm.org/docs/DiagnosticsReference.html#wdate-time
 		/// </summary>
-		[DeterministicClangToolChain(null, ["-Wdate-time", "-Wno-error=date-time"], ["-Wdate-time"])]
+		[WarningsClangToolChain(null, ["-Wdate-time", "-Wno-error=date-time"], ["-Wdate-time"], [nameof(StandardFilters.DeterministcFlagSetFilter)])]
 		public WarningLevel DeterministicWarningLevel
 		{
 			get => ResolveWarning(ParentCppCompileWarnings?.DeterministicWarningLevel, _deterministicWarningLevel);
@@ -1494,7 +1494,7 @@ namespace UnrealBuildTool
 			_bitfieldEnumConversion = WarningLevel.Off;
 			_nullPointerSubtractionWarningLevel = WarningLevel.Off;
 			_danglingWarningLevel = WarningLevel.Off;
-			
+
 			// Special case duplicates, for MSVC guaranteed disable.
 			_msvcUnusedValueWarningLevel = WarningLevel.Off;
 			_msvcDeprecationWarningLevel = WarningLevel.Off;
