UnrealBuildTool: Make pvs license renewal notice output consistent

#rnx

diff --git a/Engine/Source/Programs/UnrealBuildTool/Platform/Windows/PVSToolChain.cs b/Engine/Source/Programs/UnrealBuildTool/Platform/Windows/PVSToolChain.cs
--- a/Engine/Source/Programs/UnrealBuildTool/Platform/Windows/PVSToolChain.cs
+++ b/Engine/Source/Programs/UnrealBuildTool/Platform/Windows/PVSToolChain.cs
@@ -460,6 +460,13 @@ namespace UnrealBuildTool
 								ErrorInfo = JsonConvert.DeserializeObject<PVSErrorInfo>(Line) ?? throw new FormatException();
 							}
 
+							if (ErrorInfo.Code?.Equals("Renew", StringComparison.Ordinal) == true)
+							{
+								Logger.LogInformation("PVS-Studio Renewal Notice: {WarningMessage}", ErrorInfo.Message);
+								Logger.LogWarning("Warning: PVS-Studio license will expire soon. See output log for details.");
+								continue;
+							}
+
 							string FileName = ErrorInfo.Positions![0].File!;
 							int LineNumber = ErrorInfo.Positions![0].Lines![0];
 
