Enable warnings for [[nodiscard]] in clang.

diff --git a/Engine/Source/Programs/UnrealBuildTool/Configuration/CppCompileWarnings.cs b/Engine/Source/Programs/UnrealBuildTool/Configuration/CppCompileWarnings.cs
--- a/Engine/Source/Programs/UnrealBuildTool/Configuration/CppCompileWarnings.cs
+++ b/Engine/Source/Programs/UnrealBuildTool/Configuration/CppCompileWarnings.cs
@@ -497,8 +497,11 @@ namespace UnrealBuildTool
 		/// Clang -
 		///		https://clang.llvm.org/docs/DiagnosticsReference.html#wunused-value
 		/// </summary>
-		/// <remarks>Shipping builds will cause this warning with "ensure", so disable only in those case.</remarks>
-		[UnusedValueClangToolChain(["-Wno-unused-value"])]
+		/// <remarks>
+		/// Shipping builds will cause -Wno-unused-value to warn with "ensure", so disable only in those case.
+		/// -Wunused-result - we will always re-enable this as it's controlled by the same set. This will be moved out to a separate property. It is on by default, so it only must be re-enabled in the disable context.
+		/// </remarks>
+		[UnusedValueClangToolChain(["-Wno-unused-value", "-Wunused-result"])]
 		public WarningLevel UnusedValueWarningLevel
 		{
 			get => ResolveWarning(ParentCppCompileWarnings?.UnusedValueWarningLevel, _unusedValueWarningLevel);
@@ -932,8 +935,9 @@ namespace UnrealBuildTool
 		/// This is a special case where we have a duplicate field for clang <see cref="CppCompileWarnings.UnusedValueWarningLevel"/>.
 		/// The reason for duplication is that the clang field will only participate in the shipping configuration, whereas the MSVC variant has been disabled in all.
 		/// In order to ensure no change of behaviour for MSVC toolchain within the clang compilation context, we will keep this in, and deprecate with assignments to the clang setting.
+		/// -Wunused-result - we will always re-enable this as it's controlled by the same set. This will be moved out to a separate property. It is on by default, so it only must be re-enabled in the disable context.
 		/// </remarks>
-		[WarningsVCClang(["-Wno-unused-value"])]
+		[WarningsVCClang(["-Wno-unused-value", "-Wunused-result"])]
 		public WarningLevel MSVCUnusedValueWarningLevel
 		{
 			get => ResolveWarning(ParentCppCompileWarnings?.MSVCUnusedValueWarningLevel, _msvcUnusedValueWarningLevel);
diff --git a/Engine/Source/Runtime/CoreUObject/Private/UObject/GarbageCollection.cpp b/Engine/Source/Runtime/CoreUObject/Private/UObject/GarbageCollection.cpp
--- a/Engine/Source/Runtime/CoreUObject/Private/UObject/GarbageCollection.cpp
+++ b/Engine/Source/Runtime/CoreUObject/Private/UObject/GarbageCollection.cpp
@@ -2580,7 +2580,7 @@ TArrayView<UObject*> FAROQueue::PopImpl()
 		// can be removed the next time we do a proper profiling pass of these bits on all platforms.
 		if constexpr ((USING_THREAD_SANITISER || USING_INSTRUMENTATION) && Fence == AcquireFence)
 		{
-			Head.load(std::memory_order_acquire);
+			(void) Head.load(std::memory_order_acquire);
 		}
 
 		UObject** TailData = &GAROBlocks[TailNow];
