UnrealBuildTool: Check all module descriptors of the same name to check if a warning needs to be shown, rather than just the first instance, as there can be dupes with different types

#rnx

diff --git a/Engine/Source/Programs/UnrealBuildTool/Configuration/UEBuildModule.cs b/Engine/Source/Programs/UnrealBuildTool/Configuration/UEBuildModule.cs
--- a/Engine/Source/Programs/UnrealBuildTool/Configuration/UEBuildModule.cs
+++ b/Engine/Source/Programs/UnrealBuildTool/Configuration/UEBuildModule.cs
@@ -1171,10 +1171,11 @@ namespace UnrealBuildTool
 					}
 				}
 			}
-			if (Rules.CppCompileWarningSettings.PluginModuleUnsupportedWarningLevel != WarningLevel.Off)
+
+			// Check if this module is invalid for the current target, based on Plugin attributes
+			if (Rules.CppCompileWarningSettings.PluginModuleUnsupportedWarningLevel != WarningLevel.Off && Rules.Plugin != null)
 			{
-				ModuleDescriptor? ModuleInfo = Rules.Plugin?.Descriptor.Modules?.FirstOrDefault(x => x.Name == Rules.Name);
-				if (ModuleInfo?.IsCompiledInConfiguration(Rules.Target.Platform, Rules.Target.Configuration, Rules.Target.Name, Rules.Target.Type, Rules.Target.bBuildDeveloperTools, Rules.Target.bBuildRequiresCookedData, out string? InvalidReason) == false)
+				if (!Rules.Plugin.Descriptor.IsModuleCompiledInConfiguration(Rules.Name, Rules.Target.Platform, Rules.Target.Configuration, Rules.Target.Name, Rules.Target.Type, Rules.Target.bBuildDeveloperTools, Rules.Target.bBuildRequiresCookedData, out string? InvalidReason))
 				{
 					if (Rules.CppCompileWarningSettings.PluginModuleUnsupportedWarningLevel == WarningLevel.Warning)
 					{
@@ -1264,10 +1265,9 @@ namespace UnrealBuildTool
 			{
 				anyErrors = true;
 			}
-			if (Rules.CppCompileWarningSettings.PluginModuleUnsupportedWarningLevel == WarningLevel.Error)
+			if (Rules.CppCompileWarningSettings.PluginModuleUnsupportedWarningLevel == WarningLevel.Error && Rules.Plugin != null)
 			{
-				ModuleDescriptor? ModuleInfo = Rules.Plugin?.Descriptor.Modules?.FirstOrDefault(x => x.Name == Rules.Name);
-				if (ModuleInfo?.IsCompiledInConfiguration(Rules.Target.Platform, Rules.Target.Configuration, Rules.Target.Name, Rules.Target.Type, Rules.Target.bBuildDeveloperTools, Rules.Target.bBuildRequiresCookedData, out string? InvalidReason) == false)
+				if (!Rules.Plugin.Descriptor.IsModuleCompiledInConfiguration(Rules.Name, Rules.Target.Platform, Rules.Target.Configuration, Rules.Target.Name, Rules.Target.Type, Rules.Target.bBuildDeveloperTools, Rules.Target.bBuildRequiresCookedData, out string? _) == false)
 				{
 					anyErrors = true;
 				}
diff --git a/Engine/Source/Programs/UnrealBuildTool/System/PluginDescriptor.cs b/Engine/Source/Programs/UnrealBuildTool/System/PluginDescriptor.cs
--- a/Engine/Source/Programs/UnrealBuildTool/System/PluginDescriptor.cs
+++ b/Engine/Source/Programs/UnrealBuildTool/System/PluginDescriptor.cs
@@ -1,11 +1,12 @@
 // Copyright Epic Games, Inc. All Rights Reserved.
 
+using EpicGames.Core;
 using System;
 using System.Collections.Generic;
+using System.Diagnostics.CodeAnalysis;
 using System.IO;
 using System.Linq;
 using System.Text.Json;
-using EpicGames.Core;
 
 namespace UnrealBuildTool
 {
@@ -765,5 +766,44 @@ namespace UnrealBuildTool
 			};
 			return AllowedMissingPlatforms.Contains(Platform);
 		}
+
+		/// <summary>
+		/// Determines whether the given plugin module is part of the current build.
+		/// </summary>
+		/// <param name="Name">The name of the module</param>
+		/// <param name="Platform">The platform being compiled for</param>
+		/// <param name="Configuration">The target configuration being compiled for</param>
+		/// <param name="TargetName">Name of the target being built</param>
+		/// <param name="TargetType">The type of the target being compiled</param>
+		/// <param name="bBuildDeveloperTools">Whether the configuration includes developer tools (typically UEBuildConfiguration.bBuildDeveloperTools for UBT callers)</param>
+		/// <param name="bBuildRequiresCookedData">Whether the configuration requires cooked content (typically UEBuildConfiguration.bBuildRequiresCookedData for UBT callers)</param>
+		/// <param name="invalidReason">Out parameter, reason why this plugin module is invalid</param>
+		public bool IsModuleCompiledInConfiguration(string Name, UnrealTargetPlatform Platform, UnrealTargetConfiguration Configuration, string TargetName, TargetType TargetType, bool bBuildDeveloperTools, bool bBuildRequiresCookedData, [NotNullWhen(false)] out string? invalidReason)
+		{
+			IEnumerable<ModuleDescriptor>? Descriptors = Modules?.Where(x => x.Name == Name);
+			if (Descriptors == null)
+			{
+				invalidReason = null;
+				return true;
+			}
+
+			HashSet<string> invalidReasons = new();
+
+			foreach (ModuleDescriptor ModuleInfo in Descriptors)
+			{
+				if (!ModuleInfo.IsCompiledInConfiguration(Platform, Configuration, Name, TargetType, bBuildDeveloperTools, bBuildRequiresCookedData, out string? InvalidReason))
+				{
+					invalidReasons.Add(InvalidReason);
+				}
+				else
+				{
+					invalidReason = null;
+					return true;
+				}
+			}
+
+			invalidReason = invalidReasons.Any() ? String.Join(", ", invalidReasons.Order()) : null;
+			return invalidReason == null;
+		}
 	}
 }
