UnrealBuildTool: Fix typo in last change

#rnx

diff --git a/Engine/Source/Programs/UnrealBuildTool/Configuration/UEBuildModule.cs b/Engine/Source/Programs/UnrealBuildTool/Configuration/UEBuildModule.cs
--- a/Engine/Source/Programs/UnrealBuildTool/Configuration/UEBuildModule.cs
+++ b/Engine/Source/Programs/UnrealBuildTool/Configuration/UEBuildModule.cs
@@ -1267,7 +1267,7 @@ namespace UnrealBuildTool
 			}
 			if (Rules.CppCompileWarningSettings.PluginModuleUnsupportedWarningLevel == WarningLevel.Error && Rules.Plugin != null)
 			{
-				if (!Rules.Plugin.Descriptor.IsModuleCompiledInConfiguration(Rules.Name, Rules.Target.Platform, Rules.Target.Configuration, Rules.Target.Name, Rules.Target.Type, Rules.Target.bBuildDeveloperTools, Rules.Target.bBuildRequiresCookedData, out string? _) == false)
+				if (!Rules.Plugin.Descriptor.IsModuleCompiledInConfiguration(Rules.Name, Rules.Target.Platform, Rules.Target.Configuration, Rules.Target.Name, Rules.Target.Type, Rules.Target.bBuildDeveloperTools, Rules.Target.bBuildRequiresCookedData, out string? _))
 				{
 					anyErrors = true;
 				}
