UnrealBuildTool: Invalidate Linux relink action if object file input to first past link is updated

#rnx

diff --git a/Engine/Source/Programs/UnrealBuildTool/Platform/Linux/LinuxToolChain.cs b/Engine/Source/Programs/UnrealBuildTool/Platform/Linux/LinuxToolChain.cs
--- a/Engine/Source/Programs/UnrealBuildTool/Platform/Linux/LinuxToolChain.cs
+++ b/Engine/Source/Programs/UnrealBuildTool/Platform/Linux/LinuxToolChain.cs
@@ -1456,6 +1456,7 @@ namespace UnrealBuildTool
 				string RelinkScriptName = String.Format((bUseCmdExe ? "Relink-{0}.bat" : "Relink-{0}.sh"), RealOutputFile.Location.GetFileName());
 				FileReference RelinkScriptFullPath = FileReference.Combine(LinkEnvironment.IntermediateDirectory!, RelinkScriptName);
 
+				RelinkAction.PrerequisiteItems.UnionWith(LinkAction.PrerequisiteItems);
 				RelinkAction.PrerequisiteItems.Add(FileItem.GetItemByFileReference(RelinkScriptFullPath));
 
 				
