[UBA]
* Changed so failing actions will still retry even though compressed obj files are enabled. It will retry with uba

diff --git a/Engine/Source/Programs/UnrealBuildTool/Executors/UnrealBuildAccelerator/UBAExecutor.cs b/Engine/Source/Programs/UnrealBuildTool/Executors/UnrealBuildAccelerator/UBAExecutor.cs
--- a/Engine/Source/Programs/UnrealBuildTool/Executors/UnrealBuildAccelerator/UBAExecutor.cs
+++ b/Engine/Source/Programs/UnrealBuildTool/Executors/UnrealBuildAccelerator/UBAExecutor.cs
@@ -309,10 +309,6 @@ namespace UnrealBuildTool
 			{
 				UBAConfig.bAllowRetry = true;
 			}
-			else if (UBAConfig.bStoreObjFilesCompressed)
-			{
-				UBAConfig.bAllowRetry = false;
-			}
 
 			if (!String.IsNullOrEmpty(UBAConfig.CacheServer))
 			{
@@ -906,8 +902,15 @@ namespace UnrealBuildTool
 					Interlocked.Add(ref _localProcessedActions, 1);
 					if (UBAConfig.bAllowRetry && (process.ExitCode != 0 && UBAConfig.bForcedRetry || (process.ExitCode >= 9000 && process.ExitCode < 10000)))
 					{
-						_threadedLogger.LogInformation("{Description} {StatusDescription}: Exited with error code {ExitCode}. This action will retry without UBA", action.CommandDescription, action.StatusDescription, process.ExitCode);
-						_forcedRetryActions.AddOrUpdate(action, false, (k, v) => false);
+						if (UBAConfig.bStoreObjFilesCompressed)
+						{
+							_threadedLogger.LogInformation("{Description} {StatusDescription}: Exited with error code {ExitCode}. This action will retry", action.CommandDescription, action.StatusDescription, process.ExitCode);
+						}
+						else
+						{
+							_threadedLogger.LogInformation("{Description} {StatusDescription}: Exited with error code {ExitCode}. This action will retry without UBA", action.CommandDescription, action.StatusDescription, process.ExitCode);
+							_forcedRetryActions.AddOrUpdate(action, false, (k, v) => false);
+						}
 						queue.RequeueAction(action);
 						return Task.CompletedTask;
 					}
