- Fixed Xcode project GUID generation to be threadsafe, for parallel project generation

diff --git a/Engine/Source/Programs/UnrealBuildTool/ProjectFiles/Xcode/XcodeProjectFileGenerator.cs b/Engine/Source/Programs/UnrealBuildTool/ProjectFiles/Xcode/XcodeProjectFileGenerator.cs
--- a/Engine/Source/Programs/UnrealBuildTool/ProjectFiles/Xcode/XcodeProjectFileGenerator.cs
+++ b/Engine/Source/Programs/UnrealBuildTool/ProjectFiles/Xcode/XcodeProjectFileGenerator.cs
@@ -164,7 +164,10 @@ namespace UnrealBuildTool
 			string Guid = "";
 
 			byte[] Randoms = new byte[12];
-			Rand.NextBytes(Randoms);
+			lock(Rand)
+			{
+				Rand.NextBytes(Randoms);
+			}
 			for (int Index = 0; Index < 12; Index++)
 			{
 				Guid += Randoms[Index].ToString("X2");
