UnrealBuildTool: In SupportedTargetTypesAttribute, Client or Server implies Game is supported

diff --git a/Engine/Source/Programs/UnrealBuildTool/System/SupportedTargetTypesAttribute.cs b/Engine/Source/Programs/UnrealBuildTool/System/SupportedTargetTypesAttribute.cs
--- a/Engine/Source/Programs/UnrealBuildTool/System/SupportedTargetTypesAttribute.cs
+++ b/Engine/Source/Programs/UnrealBuildTool/System/SupportedTargetTypesAttribute.cs
@@ -1,19 +1,25 @@
 // Copyright Epic Games, Inc. All Rights Reserved.
 
 using System;
+using System.Collections.Generic;
 
 namespace UnrealBuildTool
 {
 	/// <summary>
-	/// Attribute which can be applied to a TargetRules-dervied class to indicate which target types it supports
+	/// Attribute which can be applied to a TargetRules-derived class to indicate which target types it supports
 	/// </summary>
 	[AttributeUsage(AttributeTargets.Class, AllowMultiple = true)]
 	public sealed class SupportedTargetTypesAttribute : Attribute
 	{
 		/// <summary>
-		/// Array of supported target types
+		/// Enumerable of supported target types
 		/// </summary>
-		public TargetType[] TargetTypes { get; }
+		public IEnumerable<TargetType> TargetTypes => _targetTypes;
+
+		/// <summary>
+		/// HashSet of supported target types
+		/// </summary>
+		private HashSet<TargetType> _targetTypes;
 
 		/// <summary>
 		/// Initialize the attribute with a list of target types
@@ -21,7 +27,13 @@ namespace UnrealBuildTool
 		/// <param name="targetTypes">Variable-length array of target type arguments</param>
 		public SupportedTargetTypesAttribute(params TargetType[] targetTypes)
 		{
-			TargetTypes = targetTypes;
+			_targetTypes = [.. targetTypes];
+
+			// Client or Server imply Game is supported
+			if (_targetTypes.Contains(TargetType.Client) || _targetTypes.Contains(TargetType.Server))
+			{
+				_targetTypes.Add(TargetType.Game);
+			}
 		}
 	}
 }
