UnrealBuildTool: Support updating last modified time of action outputs.
If an executor does not support updating mod time the files will be deleted before the action runs (same as Action.DeleteItems)

#rnx

diff --git a/Engine/Source/Programs/UnrealBuildTool/Executors/ImmediateActionQueue.cs b/Engine/Source/Programs/UnrealBuildTool/Executors/ImmediateActionQueue.cs
--- a/Engine/Source/Programs/UnrealBuildTool/Executors/ImmediateActionQueue.cs
+++ b/Engine/Source/Programs/UnrealBuildTool/Executors/ImmediateActionQueue.cs
@@ -918,6 +918,14 @@ namespace UnrealBuildTool
 			// done so the artifact system sees any compiled files.
 			if (success)
 			{
+				if (action.TouchItems.Any())
+				{
+					DateTime lastWriteTime = DateTime.Now;
+					foreach (FileItem touch in action.TouchItems)
+					{
+						touch.Location.ToFileInfo().LastWriteTime = lastWriteTime;
+					}
+				}
 				foreach (FileItem output in action.ProducedItems)
 				{
 					output.ResetCachedInfo();
diff --git a/Engine/Source/Programs/UnrealBuildTool/Executors/ParallelExecutor.cs b/Engine/Source/Programs/UnrealBuildTool/Executors/ParallelExecutor.cs
--- a/Engine/Source/Programs/UnrealBuildTool/Executors/ParallelExecutor.cs
+++ b/Engine/Source/Programs/UnrealBuildTool/Executors/ParallelExecutor.cs
@@ -158,6 +158,9 @@ namespace UnrealBuildTool
 		/// <inheritdoc/>
 		public override TelemetryExecutorEvent? GetTelemetryEvent() => telemetryEvent;
 
+		/// <inheritdoc/>
+		public override bool CanModifyFiles => true;
+
 		/// <summary>
 		/// Create an action queue
 		/// </summary>
diff --git a/Engine/Source/Programs/UnrealBuildTool/Platform/Windows/VCCompileAction.cs b/Engine/Source/Programs/UnrealBuildTool/Platform/Windows/VCCompileAction.cs
--- a/Engine/Source/Programs/UnrealBuildTool/Platform/Windows/VCCompileAction.cs
+++ b/Engine/Source/Programs/UnrealBuildTool/Platform/Windows/VCCompileAction.cs
@@ -189,6 +189,12 @@ namespace UnrealBuildTool
 		/// </summary>
 		public List<FileItem> DeleteItems { get; } = new List<FileItem>();
 
+		/// <summary>
+		/// Items that should have their modified time updated after this action runs successfully
+		/// Will be deleted before actions execute of the ActionExecutor does not support updating modified time
+		/// </summary>
+		public List<FileItem> TouchItems { get; } = new List<FileItem>();
+
 		/// <summary>
 		/// Root paths for this action (generally engine root project root, toolchain root, sdk root)
 		/// </summary>
@@ -228,6 +234,7 @@ namespace UnrealBuildTool
 		#region Implementation of IAction
 
 		IEnumerable<FileItem> IExternalAction.DeleteItems => DeleteItems;
+		IEnumerable<FileItem> IExternalAction.TouchItems => TouchItems;
 		public DirectoryReference WorkingDirectory => Unreal.EngineSourceDirectory;
 		string IExternalAction.CommandDescription
 		{
@@ -430,6 +437,7 @@ namespace UnrealBuildTool
 			AdditionalPrerequisiteItems = new List<FileItem>(InAction.AdditionalPrerequisiteItems);
 			AdditionalProducedItems = new List<FileItem>(InAction.AdditionalProducedItems);
 			DeleteItems = new List<FileItem>(InAction.DeleteItems);
+			TouchItems = new List<FileItem>(InAction.TouchItems);
 			RootPaths = new(InAction.RootPaths);
 		}
 
@@ -479,6 +487,7 @@ namespace UnrealBuildTool
 			AdditionalPrerequisiteItems = Reader.ReadList(() => Reader.ReadFileItem())!;
 			AdditionalProducedItems = Reader.ReadList(() => Reader.ReadFileItem())!;
 			DeleteItems = Reader.ReadList(() => Reader.ReadFileItem())!;
+			TouchItems = Reader.ReadList(() => Reader.ReadFileItem())!;
 			RootPaths = new(Reader);
 		}
 
@@ -525,6 +534,7 @@ namespace UnrealBuildTool
 			Writer.WriteList(AdditionalPrerequisiteItems, Item => Writer.WriteFileItem(Item));
 			Writer.WriteList(AdditionalProducedItems, Item => Writer.WriteFileItem(Item));
 			Writer.WriteList(DeleteItems, Item => Writer.WriteFileItem(Item));
+			Writer.WriteList(TouchItems, Item => Writer.WriteFileItem(Item));
 			RootPaths.Write(Writer);
 		}
 
diff --git a/Engine/Source/Programs/UnrealBuildTool/Platform/Windows/VCToolChain.cs b/Engine/Source/Programs/UnrealBuildTool/Platform/Windows/VCToolChain.cs
--- a/Engine/Source/Programs/UnrealBuildTool/Platform/Windows/VCToolChain.cs
+++ b/Engine/Source/Programs/UnrealBuildTool/Platform/Windows/VCToolChain.cs
@@ -2983,7 +2983,7 @@ namespace UnrealBuildTool
 			// VS 15.3+ does not touch lib files if they do not contain any modifications, but we need to ensure the timestamps are updated to avoid repeatedly building them.
 			if (bBuildImportLibraryOnly || (LinkEnvironment.bHasExports && !bIsBuildingLibraryOrImportLibrary))
 			{
-				LinkAction.DeleteItems.UnionWith(LinkAction.ProducedItems.Where(x => x.Location.HasExtension(".lib") || x.Location.HasExtension(".exp")));
+				LinkAction.TouchItems.UnionWith(LinkAction.ProducedItems.Where(x => x.Location.HasExtension(".lib") || x.Location.HasExtension(".exp")));
 			}
 
 			// Delete PDB files for all produced items, since incremental updates are slower than full ones.
diff --git a/Engine/Source/Programs/UnrealBuildTool/System/Action.cs b/Engine/Source/Programs/UnrealBuildTool/System/Action.cs
--- a/Engine/Source/Programs/UnrealBuildTool/System/Action.cs
+++ b/Engine/Source/Programs/UnrealBuildTool/System/Action.cs
@@ -96,6 +96,12 @@ namespace UnrealBuildTool
 		/// </summary>
 		IEnumerable<FileItem> DeleteItems { get; }
 
+		/// <summary>
+		/// Items that should have their modified time updated after this action runs successfully
+		/// Will be deleted before actions execute of the ActionExecutor does not support updating modified time
+		/// </summary>
+		IEnumerable<FileItem> TouchItems { get; }
+
 		/// <summary>
 		/// Root paths for this action (generally engine root project root, toolchain root, sdk root)
 		/// </summary>
@@ -254,9 +260,15 @@ namespace UnrealBuildTool
 		/// </summary>
 		public SortedSet<FileItem> DeleteItems { get; set; } = new();
 
+		/// <summary>
+		/// Items that should have their modified time updated after this action runs successfully
+		/// Will be deleted before actions execute of the ActionExecutor does not support updating modified time
+		/// </summary>
+		public SortedSet<FileItem> TouchItems { get; set; } = new();
+
 		/// <summary>
 		/// Root paths for this action (generally engine root project root, toolchain root, sdk root)
-        /// Order needs to be mainted, so a list is used
+		/// Order needs to be mainted, so a list is used
 		/// </summary>
 		public CppRootPaths RootPaths { get; set; } = new();
 
@@ -367,7 +379,9 @@ namespace UnrealBuildTool
 		IEnumerable<FileItem> IExternalAction.PrerequisiteItems => PrerequisiteItems;
 		IEnumerable<FileItem> IExternalAction.ProducedItems => ProducedItems;
 		IEnumerable<FileItem> IExternalAction.DeleteItems => DeleteItems;
-        public Action(ActionType InActionType)
+
+		IEnumerable<FileItem> IExternalAction.TouchItems => TouchItems;
+		public Action(ActionType InActionType)
 		{
 			ActionType = InActionType;
 
@@ -385,6 +399,7 @@ namespace UnrealBuildTool
 			PrerequisiteItems = new(InOther.PrerequisiteItems);
 			ProducedItems = new(InOther.ProducedItems);
 			DeleteItems = new(InOther.DeleteItems);
+			TouchItems = new(InOther.TouchItems);
 			RootPaths = new(InOther.RootPaths);
 			DependencyListFile = InOther.DependencyListFile;
 			WorkingDirectory = InOther.WorkingDirectory;
@@ -434,6 +449,7 @@ namespace UnrealBuildTool
 			PrerequisiteItems = Reader.ReadSortedSet(() => Reader.ReadFileItem())!;
 			ProducedItems = Reader.ReadSortedSet(() => Reader.ReadFileItem())!;
 			DeleteItems = Reader.ReadSortedSet(() => Reader.ReadFileItem())!;
+			TouchItems = Reader.ReadSortedSet(() => Reader.ReadFileItem())!;
 			RootPaths = new(Reader);
 			DependencyListFile = Reader.ReadFileItem();
 			bUseActionHistory = Reader.ReadBool();
@@ -469,6 +485,7 @@ namespace UnrealBuildTool
 			Writer.WriteSortedSet(PrerequisiteItems, Item => Writer.WriteFileItem(Item));
 			Writer.WriteSortedSet(ProducedItems, Item => Writer.WriteFileItem(Item));
 			Writer.WriteSortedSet(DeleteItems, Item => Writer.WriteFileItem(Item));
+			Writer.WriteSortedSet(TouchItems, Item => Writer.WriteFileItem(Item));
 			RootPaths.Write(Writer);
 			Writer.WriteFileItem(DependencyListFile);
 			Writer.WriteBool(bUseActionHistory);
@@ -590,6 +607,11 @@ namespace UnrealBuildTool
 				Action.DeleteItems.UnionWith(DeleteItems.Select(x => FileItem.GetItemByPath(x)));
 			}
 
+			if (Object.TryGetStringArrayField("TouchItems", out string[]? TouchItems))
+			{
+				Action.TouchItems.UnionWith(TouchItems.Select(x => FileItem.GetItemByPath(x)));
+			}
+
 			if (Object.TryGetObjectField("RootPaths", out JsonObject? RootPaths))
 			{
 				Action.RootPaths.Read(RootPaths);
@@ -702,6 +724,13 @@ namespace UnrealBuildTool
 			}
 			Writer.WriteArrayEnd();
 
+			Writer.WriteArrayStart("TouchItems");
+			foreach (FileItem TouchItem in Action.TouchItems)
+			{
+				Writer.WriteValue(TouchItem.AbsolutePath);
+			}
+			Writer.WriteArrayEnd();
+
 			Writer.WriteObjectStart("RootPaths");
 			Action.RootPaths.Write(Writer);
 			Writer.WriteObjectEnd();
@@ -791,6 +820,7 @@ namespace UnrealBuildTool
 		public IEnumerable<FileItem> PrerequisiteItems => Inner.PrerequisiteItems;
 		public IEnumerable<FileItem> ProducedItems => Inner.ProducedItems;
 		public IEnumerable<FileItem> DeleteItems => Inner.DeleteItems;
+		public IEnumerable<FileItem> TouchItems => Inner.TouchItems;
 		public CppRootPaths RootPaths => Inner.RootPaths;
 		public FileItem? DependencyListFile => Inner.DependencyListFile;
 		public DirectoryReference WorkingDirectory => Inner.WorkingDirectory;
diff --git a/Engine/Source/Programs/UnrealBuildTool/System/ActionExecutor.cs b/Engine/Source/Programs/UnrealBuildTool/System/ActionExecutor.cs
--- a/Engine/Source/Programs/UnrealBuildTool/System/ActionExecutor.cs
+++ b/Engine/Source/Programs/UnrealBuildTool/System/ActionExecutor.cs
@@ -87,6 +87,11 @@ namespace UnrealBuildTool
 		/// </summary>
 		public virtual bool VerifyOutputs => true;
 
+		/// <summary>
+		/// If this executor supports modifying produced items
+		/// </summary>
+		public virtual bool CanModifyFiles => false;
+
 		protected void WriteToolOutput(string Line)
 		{
 			lock (Parser)
diff --git a/Engine/Source/Programs/UnrealBuildTool/System/ActionGraph.cs b/Engine/Source/Programs/UnrealBuildTool/System/ActionGraph.cs
--- a/Engine/Source/Programs/UnrealBuildTool/System/ActionGraph.cs
+++ b/Engine/Source/Programs/UnrealBuildTool/System/ActionGraph.cs
@@ -459,6 +459,20 @@ namespace UnrealBuildTool
 				using ActionExecutor Executor = SelectExecutor(BuildConfiguration, ActionsToExecute.Count, TargetDescriptors, Logger);
 				Logger.LogInformation("Using {ExecutorName} executor to run {ActionCount} action(s)", Executor.Name, ActionsToExecute.Count);
 
+				if (!Executor.CanModifyFiles)
+				{
+					IEnumerable<FileItem> touchItems = ActionsToExecute.SelectMany(x => x.TouchItems).Where(x => x.Exists);
+					if (touchItems.Any())
+					{
+						Logger.LogDebug("Executor {ExecutorName} does not support updating last modified time. Deleting files to ensure up to date timestamps...", Executor.Name);
+						foreach (FileItem TouchItem in touchItems)
+						{
+							Logger.LogDebug("Deleting item: {TouchItem}", TouchItem.Location);
+							TouchItem.Delete(Logger);
+						}
+					}
+				}
+
 				// Execute the build
 				Stopwatch Timer = Stopwatch.StartNew();
 				bool Result = await Executor.ExecuteActionsAsync(ActionsToExecute, Logger, actionArtifactCache);
