UnrealBuildTool: Don't set bIgnoreOutdatedImportLibraries for build machines

#rnx

diff --git a/Engine/Source/Programs/UnrealBuildTool/Configuration/BuildConfiguration.cs b/Engine/Source/Programs/UnrealBuildTool/Configuration/BuildConfiguration.cs
--- a/Engine/Source/Programs/UnrealBuildTool/Configuration/BuildConfiguration.cs
+++ b/Engine/Source/Programs/UnrealBuildTool/Configuration/BuildConfiguration.cs
@@ -2,6 +2,7 @@
 
 using System;
 using EpicGames.Core;
+using UnrealBuildBase;
 
 namespace UnrealBuildTool
 {
@@ -18,7 +19,7 @@ namespace UnrealBuildTool
 		/// in which case the target would automatically be rebuilt.
 		/// </summary>
 		[XmlConfigFile]
-		public bool bIgnoreOutdatedImportLibraries { get; set; } = true;
+		public bool bIgnoreOutdatedImportLibraries { get; set; } = !Unreal.IsBuildMachine();
 
 		/// <summary>
 		/// Use existing static libraries for all engine modules in this target.
