[Backout] - CL40639526
#rnx
Original CL Desc
-----------------------------------------------------------------
UnrealBuildTool: Don't set bIgnoreOutdatedImportLibraries for build machines

#rnx

diff --git a/Engine/Source/Programs/UnrealBuildTool/Configuration/BuildConfiguration.cs b/Engine/Source/Programs/UnrealBuildTool/Configuration/BuildConfiguration.cs
--- a/Engine/Source/Programs/UnrealBuildTool/Configuration/BuildConfiguration.cs
+++ b/Engine/Source/Programs/UnrealBuildTool/Configuration/BuildConfiguration.cs
@@ -2,7 +2,6 @@
 
 using System;
 using EpicGames.Core;
-using UnrealBuildBase;
 
 namespace UnrealBuildTool
 {
@@ -19,7 +18,7 @@ namespace UnrealBuildTool
 		/// in which case the target would automatically be rebuilt.
 		/// </summary>
 		[XmlConfigFile]
-		public bool bIgnoreOutdatedImportLibraries { get; set; } = !Unreal.IsBuildMachine();
+		public bool bIgnoreOutdatedImportLibraries { get; set; } = true;
 
 		/// <summary>
 		/// Use existing static libraries for all engine modules in this target.
