UnrealBuildTool: Print exception callstack if a .modules manfiest file fails to load

#rnx

diff --git a/Engine/Source/Programs/UnrealBuildTool/Modes/WriteMetadataMode.cs b/Engine/Source/Programs/UnrealBuildTool/Modes/WriteMetadataMode.cs
--- a/Engine/Source/Programs/UnrealBuildTool/Modes/WriteMetadataMode.cs
+++ b/Engine/Source/Programs/UnrealBuildTool/Modes/WriteMetadataMode.cs
@@ -299,7 +299,7 @@ namespace UnrealBuildTool
 				catch (Exception Ex)
 				{
 					Logger.LogWarning("Unable to read '{ManifestFileName}'; ignoring.", ManifestFileName);
-					Logger.LogDebug("{Ex}", ExceptionUtils.FormatExceptionDetails(Ex));
+					Logger.LogInformation("{Ex}", ExceptionUtils.FormatExceptionDetails(Ex));
 				}
 			}
 
