UnrealBuildTool: Clean up handling of PER_MODULE_INLINE_FILE when merging modules
* Check all references in a binary's modules for any PER_MODULE_INLINE_FILE and update the primary module's defines and includes to use it
* If there are conflicts prioritize project modules first and show a warning

#rnx

diff --git a/Engine/Source/Programs/UnrealBuildTool/Configuration/UEBuildBinary.cs b/Engine/Source/Programs/UnrealBuildTool/Configuration/UEBuildBinary.cs
--- a/Engine/Source/Programs/UnrealBuildTool/Configuration/UEBuildBinary.cs
+++ b/Engine/Source/Programs/UnrealBuildTool/Configuration/UEBuildBinary.cs
@@ -778,9 +778,6 @@ namespace UnrealBuildTool
 
 			HashSet<FileItem> InputFilesLookup = new();
 
-			int KeepLevel = 0;
-			FileItem? PerModuleItemToKeep = null;
-
 			foreach (UEBuildModule Module in UEBuildModule.StableTopologicalSort(Modules))
 			{
 				if (Module.Binary == null || Module.Binary == this)
@@ -788,57 +785,6 @@ namespace UnrealBuildTool
 					// Compile each module.
 					Logger.LogDebug("Compile module: {ModuleName}", Module.Name);
 					List<FileItem> LinkInputFiles = Module.Compile(Target, ToolChain, BinaryCompileEnvironment, WorkingSet, Graph, Logger);
-
-					// If modules are merged we will have multiple PerModuleInline which will cause duplicated symbols. So let's only keep one of those
-					// The used PerModuleInline file should be picked in order "Any", "Core", "If owning module is in project", "If module has define set"
-					// Ideally it should be "Any", "Core", "If module has define set" but unfortunately define is not set for depending modules at this point
-					// so we can have merged binaries that has 10 modules where all depends on the module that has the define... and in that case any of them would work to pick the PerModuleInline file
-					// This code is very inefficient and need to be improved.
-					if (Target.bMergeModules)
-					{
-						FileItem? PerModuleItem = LinkInputFiles.FirstOrDefault(x => x!.Name.Contains("PerModuleInline.gen.cpp"), null);
-						if (PerModuleItem != null)
-						{
-							// If we don't have a PerModuleItem, use any
-							if (KeepLevel == 0)
-							{
-								if (PerModuleItemToKeep == null)
-								{
-									PerModuleItemToKeep = PerModuleItem;
-									KeepLevel = 1;
-								}
-							}
-							else if (KeepLevel == 1 && Module.Name == "Core") // If we have one we look for Core
-							{
-								PerModuleItemToKeep = PerModuleItem;
-								KeepLevel = 2;
-
-							}
-							else if (KeepLevel <= 2) 
-							{
-								if (Target.ProjectFile != null && Module.RulesFile.IsUnderDirectory(Target.ProjectFile.Directory))  // TODO: THIS IS WRONG AND CAN EASILY GO WRONG
-								{
-									PerModuleItemToKeep = PerModuleItem;
-									KeepLevel = 3;
-								}
-							}
-							else if (KeepLevel <= 3)
-							{
-								foreach (string Def in Module.PublicDefinitions)
-								{
-									if (Def.StartsWith("PER_MODULE_INLINE_FILE"))
-									{
-										PerModuleItemToKeep = PerModuleItem;
-										KeepLevel = 4;
-										break;
-									}
-								}
-							}
-
-							LinkInputFiles.Remove(PerModuleItem);
-						}
-					}
-
 					InputObjects.UnionWith(LinkInputFiles.Where(x => x.HasExtension(".obj") || x.HasExtension(".o")));
 
 					// Save the module outputs. In monolithic builds, this is just the object files.
@@ -877,12 +823,6 @@ namespace UnrealBuildTool
 				Module.SetupPrivateLinkEnvironment(this, BinaryLinkEnvironment, BinaryDependencies, LinkEnvironmentVisitedModules, ExeDir, Logger);
 			}
 
-			if (PerModuleItemToKeep != null)
-			{
-				BinaryLinkEnvironment.InputFiles.Add(PerModuleItemToKeep);
-				InputObjects.Add(PerModuleItemToKeep);
-			}
-
 			// Allow the binary dependencies to modify the link environment.
 			foreach (UEBuildBinary BinaryDependency in BinaryDependencies)
 			{
diff --git a/Engine/Source/Programs/UnrealBuildTool/Configuration/UEBuildModuleCPP.cs b/Engine/Source/Programs/UnrealBuildTool/Configuration/UEBuildModuleCPP.cs
--- a/Engine/Source/Programs/UnrealBuildTool/Configuration/UEBuildModuleCPP.cs
+++ b/Engine/Source/Programs/UnrealBuildTool/Configuration/UEBuildModuleCPP.cs
@@ -720,39 +720,33 @@ namespace UnrealBuildTool
 			}
 
 			// The operator new/delete overrides need to exist in all modules.
-			// We put these overrides in a separate file in order to be able to filter out so only one file is added when modules are merged
-			if (Target.LinkType == TargetLinkType.Modular)
+			// We put these overrides in a separate file in order to be able to ensure only one file is added per binary for the primary module
+			if (Target.LinkType == TargetLinkType.Modular && Binary.PrimaryModule == this)
 			{
 				FileReference NewDeleteOverrides = FileReference.Combine(IntermediateDirectory, "PerModuleInline.gen.cpp");
-				string Content =
-					"#if !defined(PER_MODULE_INLINE_FILE) && defined(CORE_API)\r\n" +
-					"#define PER_MODULE_INLINE_FILE \"HAL/PerModuleInline.inl\"\r\n" +
-					"#endif\r\n" +
-					"#if defined(PER_MODULE_INLINE_FILE) && !defined(SUPPRESS_PER_MODULE_INLINE_FILE)\r\n" +
-					"#include PER_MODULE_INLINE_FILE\r\n" +
-					"#endif\r\n"
-					;
+				List<string> Content =
+				[
+					"// Generated by UnrealBuildTool (UEBuildModuleCPP.cs)",
+					"#if !defined(PER_MODULE_INLINE_FILE) && defined(CORE_API)",
+					"#define PER_MODULE_INLINE_FILE \"HAL/PerModuleInline.inl\"",
+					"#endif",
+					"#if defined(PER_MODULE_INLINE_FILE) && !defined(SUPPRESS_PER_MODULE_INLINE_FILE)",
+					"#include PER_MODULE_INLINE_FILE",
+					"#endif",
+				];
 
 				if (Target.WindowsPlatform.bEnableInstrumentation)
 				{
-					// Give the ability for core to hotpatch itself for the per module functions to save another set of JMPs.
-					if (Name == "Core")
-					{ 
-						Content += "#define USE_INSTRUMENTATION_PERMODULE_HOTPATCH 1\r\n";
+					// Give the ability for Core to hotpatch itself for the per module functions to save another set of JMPs.
+					if (Binary.Modules.Any(x => x.Name == "Core"))
+					{
+						Content.Add("#define USE_INSTRUMENTATION_PERMODULE_HOTPATCH 1");
 					}
-					Content += "#include \"Instrumentation/PerModuleInline.inl\"\r\n";
+					Content.Add("#include \"Instrumentation/PerModuleInline.inl\"");
 				}
 
-				bool bMergeModules = Target.bMergeModules;
-				FileItem PerModuleFile = Graph.CreateIntermediateTextFile(NewDeleteOverrides, Content, bMergeModules);
-				if (bMergeModules)
-				{
-					LinkInputFiles.AddRange(ToolChain.CompileAllCPPFiles(CppCompileEnvironment, new[] { PerModuleFile }, IntermediateDirectory, Name, Graph).ObjectFiles);
-				}
-				else
-				{
-					CPPFiles.Add(PerModuleFile);
-				}
+				FileItem PerModuleFile = Graph.CreateIntermediateTextFile(NewDeleteOverrides, Content, false);
+				CPPFiles.Add(PerModuleFile);
 			}
 
 			// Engine modules will always use unity build mode unless MinSourceFilesForUnityBuildOverride is specified in
diff --git a/Engine/Source/Programs/UnrealBuildTool/Configuration/UEBuildTarget.cs b/Engine/Source/Programs/UnrealBuildTool/Configuration/UEBuildTarget.cs
--- a/Engine/Source/Programs/UnrealBuildTool/Configuration/UEBuildTarget.cs
+++ b/Engine/Source/Programs/UnrealBuildTool/Configuration/UEBuildTarget.cs
@@ -3472,6 +3472,60 @@ namespace UnrealBuildTool
 		/// <param name="Logger"></param>
 		protected void SetupMergeModules(ILogger Logger)
 		{
+			void UpdatePrimaryModule(UEBuildBinary binary)
+			{
+				// Check all modules in this binary to see if any references set PER_MODULE_INLINE_FILE
+				Dictionary<UEBuildModuleCPP, string> perModuleInlineFiles = new();
+				{
+					HashSet<UEBuildModule> allReferencedModules = [];
+					foreach (UEBuildModuleCPP module in binary.Modules)
+					{
+						if (allReferencedModules.Contains(module))
+						{
+							continue;
+						}
+						List<UEBuildModule> referencedModules = [];
+						allReferencedModules.Add(module);
+						module.GetAllDependencyModules(referencedModules, allReferencedModules, false, false, false);
+						allReferencedModules.UnionWith(referencedModules);
+						allReferencedModules.UnionWith(referencedModules.SelectMany(x => x.PublicIncludePathModules ?? new()));
+					}
+
+					foreach (UEBuildModuleCPP reference in allReferencedModules.OfType<UEBuildModuleCPP>())
+					{
+						string? perModuleInlineFile = reference.Rules.PublicDefinitions.LastOrDefault(x => x.StartsWith("PER_MODULE_INLINE_FILE"));
+						if (perModuleInlineFile != null)
+						{
+							perModuleInlineFiles.Add(reference, perModuleInlineFile);
+						}
+					}
+				}
+
+				// If PER_MODULE_INLINE_FILE is found, move that definition to the primary module
+				if (perModuleInlineFiles.Any())
+				{
+					// Warn if there are PER_MODULE_INLINE_FILE conflicts
+					UEBuildModuleCPP module = perModuleInlineFiles.First().Key;
+					string definition = perModuleInlineFiles.First().Value;
+					if (perModuleInlineFiles.Values.Distinct().Count() > 1)
+					{
+						IOrderedEnumerable<UEBuildModuleCPP> ordered = perModuleInlineFiles.Keys
+							.OrderBy(m => !m.RulesFile.IsUnderDirectory(Unreal.EngineDirectory))
+							.ThenBy(m => m.Name);
+						Logger.LogWarning("Conflicting definitions of PER_MODULE_INLINE_FILE define when merging '{Modules}', first instance '{First}' will be used", String.Join(", ", ordered), perModuleInlineFiles[ordered.First()]);
+						module = ordered.First();
+						definition = perModuleInlineFiles[ordered.First()];
+					}
+
+					if (!binary.PrimaryModule.Rules.PublicDefinitions.Contains(definition))
+					{
+						Logger.LogDebug("Updating {Binary} primary module {Module} to include define {Define} from {SourceModule}", binary.OutputFilePath.GetFileNameWithoutExtension(), binary.PrimaryModule.Name, definition, module.Name);
+						binary.PrimaryModule.PublicDefinitions.Add(definition);
+						binary.PrimaryModule.PublicIncludePaths.UnionWith(module.PublicIncludePaths);
+					}
+				}
+			}
+
 			void MergeModulesToDLL(string name, IEnumerable<UEBuildModuleCPP> modules)
 			{
 				if (!modules.Any())
@@ -3479,12 +3533,13 @@ namespace UnrealBuildTool
 					return;
 				}
 
-				List<UEBuildModuleCPP> modulesList = modules.OrderBy(m => m.Name).ToList();
-				UEBuildBinary mergedBinary = CreateDynamicLibraryForModules(name, modules.OrderBy(x => x.Name));
+				IEnumerable<UEBuildModuleCPP> modulesList = modules;
+				UEBuildBinary mergedBinary = CreateDynamicLibraryForModules(name, modulesList);
 				foreach (UEBuildModuleCPP module in modules)
 				{
 					module.Binary = mergedBinary;
 				}
+				UpdatePrimaryModule(mergedBinary);
 
 				Binaries.Add(mergedBinary);
 			}
@@ -3500,6 +3555,7 @@ namespace UnrealBuildTool
 					Binaries[0].AddModule(module);
 					module.Binary = Binaries[0];
 				}
+				UpdatePrimaryModule(Binaries[0]);
 			}
 
 			void ConsumeModules(string name, IEnumerable<UEBuildModuleCPP> modules)
