UnrealBuildTool: Support UBA caching for PVS analyze actions

#rnx

diff --git a/Engine/Source/Programs/UnrealBuildTool/Platform/Windows/PVSToolChain.cs b/Engine/Source/Programs/UnrealBuildTool/Platform/Windows/PVSToolChain.cs
--- a/Engine/Source/Programs/UnrealBuildTool/Platform/Windows/PVSToolChain.cs
+++ b/Engine/Source/Programs/UnrealBuildTool/Platform/Windows/PVSToolChain.cs
@@ -948,6 +948,10 @@ namespace UnrealBuildTool
 			AnalyzeAction.bCanExecuteRemotelyWithSNDBS = false;
 			AnalyzeAction.bCanExecuteRemotelyWithXGE = false;
 
+			AnalyzeAction.RootPaths = CompileEnvironment.RootPaths;
+			AnalyzeAction.CacheBucket = GetCacheBucket(Target, null);
+			AnalyzeAction.ArtifactMode = ArtifactMode.Enabled;
+
 			Result.ObjectFiles.AddRange(AnalyzeAction.ProducedItems);
 		}
 
