[UBT]
* Fixed so linux toolchain batch scripts check if file exists before attempting to delete it. This will prevent annoying logging errors on wine

diff --git a/Engine/Source/Programs/UnrealBuildTool/Platform/Linux/LinuxToolChain.cs b/Engine/Source/Programs/UnrealBuildTool/Platform/Linux/LinuxToolChain.cs
--- a/Engine/Source/Programs/UnrealBuildTool/Platform/Linux/LinuxToolChain.cs
+++ b/Engine/Source/Programs/UnrealBuildTool/Platform/Linux/LinuxToolChain.cs
@@ -157,7 +157,19 @@ namespace UnrealBuildTool
 		{
 			bool bUseCmdExe = BuildHostPlatform.Current.Platform == UnrealTargetPlatform.Win64;
 			string DumpCommand = bUseCmdExe ? "\"{0}\" \"{1}\" \"{2}\" 2>NUL" : "\"{0}\" -c -o \"{2}\" \"{1}\"";
-			string DeleteCommand = bUseCmdExe ? "del /F" : "rm";
+			
+			Func<string, string> DeleteCommand = fileName =>
+				{
+					if (bUseCmdExe)
+					{
+						return $"if exist \"{fileName}\" del /F \"{fileName}\"";
+					}
+					else
+					{
+						return $"rm \"{fileName}\"";
+					}
+				};
+
 			FileItem EncodedBinarySymbolsFile = FileItem.GetItemByPath(Path.Combine(LinkEnvironment.OutputDirectory!.FullName, OutputFile.Location.GetFileNameWithoutExtension() + ".sym"));
 			FileItem SymbolsFile = FileItem.GetItemByPath(Path.Combine(LinkEnvironment.IntermediateDirectory!.FullName, OutputFile.Location.GetFileName() + ".psym"));
 			FileItem StrippedFile = FileItem.GetItemByPath(Path.Combine(LinkEnvironment.IntermediateDirectory.FullName, OutputFile.Location.GetFileName() + "_nodebug"));
@@ -194,7 +206,7 @@ namespace UnrealBuildTool
 				}
 				else
 				{
-					Out.WriteLine($"{DeleteCommand} \"{SymbolsFile.AbsolutePath}\"");
+					Out.WriteLine(DeleteCommand(SymbolsFile.AbsolutePath));
 				}
 			}
 			else
@@ -238,7 +250,7 @@ namespace UnrealBuildTool
 					OutputFile.AbsolutePath
 				);
 
-				Out.WriteLine($"{DeleteCommand} \"{StrippedFile.AbsolutePath}\"");
+				Out.WriteLine(DeleteCommand(StrippedFile.AbsolutePath));
 
 				if (bUseCmdExe)
 				{
@@ -1380,7 +1392,7 @@ namespace UnrealBuildTool
 					LinkWriter.WriteLine();
 					LinkWriter.WriteLine("set Retries=0");
 					LinkWriter.WriteLine(":linkloop");
-					LinkWriter.WriteLine("if %Retries% GEQ 10 goto failedtorelink");
+					LinkWriter.WriteLine("if %Retries% GEQ 1 goto failedtorelink");
 					LinkWriter.WriteLine(LinkCommandString);
 					LinkWriter.WriteLine("if %errorlevel% neq 0 goto sleepandretry");
 					if (!needsRelink)
