UnrealBuildTool: Support bVFS for pvs toolchain

#rnx

diff --git a/Engine/Source/Programs/UnrealBuildTool/Configuration/TargetRules.cs b/Engine/Source/Programs/UnrealBuildTool/Configuration/TargetRules.cs
--- a/Engine/Source/Programs/UnrealBuildTool/Configuration/TargetRules.cs
+++ b/Engine/Source/Programs/UnrealBuildTool/Configuration/TargetRules.cs
@@ -3190,13 +3190,6 @@ namespace UnrealBuildTool
 			{
 				// Always override environment for analyzing, regardless of other settings
 				IntermediateEnvironment = UnrealIntermediateEnvironment.Analyze;
-
-				// Need to disable vfs for pvsstudio since there is a postbuild action that reads output which is not detoured
-				// and will not be able to interpret virtual file paths
-				if (StaticAnalyzer == StaticAnalyzer.PVSStudio)
-				{
-					bUseVFS = false;
-				}
 			}
 			else if (IntermediateEnvironment == UnrealIntermediateEnvironment.Default)
 			{
diff --git a/Engine/Source/Programs/UnrealBuildTool/Platform/Windows/PVSToolChain.cs b/Engine/Source/Programs/UnrealBuildTool/Platform/Windows/PVSToolChain.cs
--- a/Engine/Source/Programs/UnrealBuildTool/Platform/Windows/PVSToolChain.cs
+++ b/Engine/Source/Programs/UnrealBuildTool/Platform/Windows/PVSToolChain.cs
@@ -12,6 +12,7 @@ using System.Threading.Tasks;
 using System.Xml.Serialization;
 using EpicGames.Core;
 using EpicGames.Horde.Jobs.Graphs;
+using EpicGames.Horde.Logs;
 using Microsoft.Extensions.Logging;
 using Newtonsoft.Json;
 using UnrealBuildBase;
@@ -392,6 +393,12 @@ namespace UnrealBuildTool
 		[CommandLine("-Ignored")]
 		FileReference? IgnoredFile = null;
 
+		/// <summary>
+		/// Path to file list of rootpaths,realpath mapping
+		/// </summary>
+		[CommandLine("-RootPaths")]
+		FileReference? RootPathsFile = null;
+
 		/// <summary>
 		/// The maximum level of warnings to print
 		/// </summary>
@@ -428,6 +435,20 @@ namespace UnrealBuildTool
 			string[] IgnoreFileLines = IgnoredFile != null ? FileReference.ReadAllLines(IgnoredFile) : Array.Empty<string>();
 			DirectoryReference[] IgnoredDirectories = IgnoreFileLines.Select(x => x.Trim()).Where(x => x.Length > 0).Select(x => new DirectoryReference(x)).ToArray();
 
+			// Read the root paths file
+			string[] RootPathFileLines = RootPathsFile != null ? FileReference.ReadAllLines(RootPathsFile) : Array.Empty<string>();
+			Dictionary<DirectoryReference, DirectoryReference> RootPaths = [];
+			foreach (string line in RootPathFileLines.Select(x => x.Trim()).Where(x => x.Length > 0))
+			{
+				string[] split = line.Split(',');
+				if (split.Length != 2)
+				{
+					Logger.LogWarning("Unexpected number of items {Count} for line '{Line}', expected 2.", split.Count(), line);
+					continue;
+				}
+				RootPaths.Add(new DirectoryReference(split.First()), new DirectoryReference(split.Last()));
+			}
+
 			// Create the combined output file, and print the diagnostics to the log
 			HashSet<string> UniqueItems = new HashSet<string>();
 			List<string> OutputLines = new List<string>();
@@ -467,6 +488,22 @@ namespace UnrealBuildTool
 								continue;
 							}
 
+							if (RootPaths.Any())
+							{
+								foreach (PVSPosition position in ErrorInfo.Positions!.Where(x => !String.IsNullOrWhiteSpace(x.File)))
+								{
+									FileReference positionFile = new FileReference(position!.File!);
+									foreach (KeyValuePair<DirectoryReference, DirectoryReference> item in RootPaths)
+									{
+										if (positionFile.IsUnderDirectory(item.Key))
+										{
+											position!.File = FileReference.Combine(item.Value, positionFile.MakeRelativeTo(item.Key)).FullName;
+											break;
+										}
+									}
+								}
+							}
+
 							string FileName = ErrorInfo.Positions![0].File!;
 							int LineNumber = ErrorInfo.Positions![0].Lines![0];
 
@@ -493,7 +530,7 @@ namespace UnrealBuildTool
 							}
 
 							// Output the line to the raw output file
-							RawWriter.WriteLine(Line);
+							RawWriter.WriteLine(JsonConvert.SerializeObject(ErrorInfo));
 
 							// Output the line to the log
 							if (ErrorInfo.FalseAlarm != true && ErrorInfo.Level <= PrintLevel)
@@ -1027,19 +1064,21 @@ namespace UnrealBuildTool
 
 			// Collect the sourcefile items off of the Compile action added in CompileCPPFiles so that in SingleFileCompile mode the PVSGather step is also not filtered out
 			ImmutableSortedSet<FileItem> CompileSourceFiles = Makefile.Actions.OfType<VCCompileAction>().Select(x => x.SourceFile!).ToImmutableSortedSet();
+			CppRootPaths? RootPaths = Makefile.Actions.OfType<VCCompileAction>().FirstOrDefault(x => x.RootPaths.Any())?.RootPaths;
 
 			// Store list of system paths that should be excluded
 			ImmutableSortedSet<DirectoryReference> SystemIncludePaths = Makefile.Actions.OfType<VCCompileAction>().SelectMany(x => x.SystemIncludePaths).ToImmutableSortedSet();
 
 			FileItem InputFileListItem = MakefileBuilder.CreateIntermediateTextFile(OutputFile.ChangeExtension(".input"), InputFiles.Select(x => x.FullName));
-			FileItem IgnoredFileListeItem = MakefileBuilder.CreateIntermediateTextFile(OutputFile.ChangeExtension(".ignored"), SystemIncludePaths.Select(x => x.FullName));
+			FileItem IgnoredFileListItem = MakefileBuilder.CreateIntermediateTextFile(OutputFile.ChangeExtension(".ignored"), SystemIncludePaths.Select(x => x.FullName));
+			FileItem RootPathsItem = MakefileBuilder.CreateIntermediateTextFile(OutputFile.ChangeExtension(".rootpaths"), RootPaths?.Select(x => $"{x.vfs},{x.local}") ?? []);
 
-			string Arguments = $"-Input=\"{InputFileListItem.Location}\" -Output=\"{OutputFile}\" -Ignored=\"{IgnoredFileListeItem.Location}\" -PrintLevel={Target.StaticAnalyzerPVSPrintLevel} -AnalyzerVersion={AnalyzerVersion}";
+			string Arguments = $"-Input=\"{InputFileListItem.Location}\" -Output=\"{OutputFile}\" -Ignored=\"{IgnoredFileListItem.Location}\" -RootPaths=\"{RootPathsItem.Location}\" -PrintLevel={Target.StaticAnalyzerPVSPrintLevel} -AnalyzerVersion={AnalyzerVersion}";
 
 			Action AnalyzeAction = MakefileBuilder.CreateRecursiveAction<PVSGatherMode>(ActionType.PostBuildStep, Arguments);
 			AnalyzeAction.CommandDescription = "Process PVS-Studio Results";
 			AnalyzeAction.PrerequisiteItems.Add(InputFileListItem);
-			AnalyzeAction.PrerequisiteItems.Add(IgnoredFileListeItem);
+			AnalyzeAction.PrerequisiteItems.Add(IgnoredFileListItem);
 			AnalyzeAction.PrerequisiteItems.UnionWith(Makefile.OutputItems);
 			AnalyzeAction.PrerequisiteItems.UnionWith(CompileSourceFiles);
 			AnalyzeAction.ProducedItems.Add(FileItem.GetItemByFileReference(OutputFile));
diff --git a/Engine/Source/Programs/UnrealBuildTool/System/CppRootPaths.cs b/Engine/Source/Programs/UnrealBuildTool/System/CppRootPaths.cs
--- a/Engine/Source/Programs/UnrealBuildTool/System/CppRootPaths.cs
+++ b/Engine/Source/Programs/UnrealBuildTool/System/CppRootPaths.cs
@@ -374,8 +374,6 @@ namespace UnrealBuildTool
 		/// <param name="name">The new string name</param>
 		public void AddFolderName(CppRootPathFolder folder, string name)
 		{
-			if (name == "Clang")
-				Console.WriteLine();
 			_folderNames[folder] = name;
 		}
 
