UnrealBuildTool: Fix default static analyzer property when other compilers are used

#rnx

diff --git a/Engine/Source/Programs/UnrealBuildTool/Configuration/TargetRules.cs b/Engine/Source/Programs/UnrealBuildTool/Configuration/TargetRules.cs
--- a/Engine/Source/Programs/UnrealBuildTool/Configuration/TargetRules.cs
+++ b/Engine/Source/Programs/UnrealBuildTool/Configuration/TargetRules.cs
@@ -237,7 +237,7 @@ namespace UnrealBuildTool
 		/// <summary>
 		/// Use the built-in Visual C++ static analyzer
 		/// </summary>
-		VisualCpp = Default,
+		VisualCpp,
 
 		/// <summary>
 		/// Use PVS-Studio for static analysis
