UnrealBuildTool: Support bVFS for pvs toolchain
Optimize post-process log parsing

#rnx

diff --git a/Engine/Source/Programs/UnrealBuildTool/Configuration/TargetRules.cs b/Engine/Source/Programs/UnrealBuildTool/Configuration/TargetRules.cs
--- a/Engine/Source/Programs/UnrealBuildTool/Configuration/TargetRules.cs
+++ b/Engine/Source/Programs/UnrealBuildTool/Configuration/TargetRules.cs
@@ -3190,13 +3190,6 @@ namespace UnrealBuildTool
 			{
 				// Always override environment for analyzing, regardless of other settings
 				IntermediateEnvironment = UnrealIntermediateEnvironment.Analyze;
-
-				// Need to disable vfs for pvsstudio since there is a postbuild action that reads output which is not detoured
-				// and will not be able to interpret virtual file paths
-				if (StaticAnalyzer == StaticAnalyzer.PVSStudio)
-				{
-					bUseVFS = false;
-				}
 			}
 			else if (IntermediateEnvironment == UnrealIntermediateEnvironment.Default)
 			{
diff --git a/Engine/Source/Programs/UnrealBuildTool/Platform/Windows/PVSToolChain.cs b/Engine/Source/Programs/UnrealBuildTool/Platform/Windows/PVSToolChain.cs
--- a/Engine/Source/Programs/UnrealBuildTool/Platform/Windows/PVSToolChain.cs
+++ b/Engine/Source/Programs/UnrealBuildTool/Platform/Windows/PVSToolChain.cs
@@ -1,21 +1,20 @@
 // Copyright Epic Games, Inc. All Rights Reserved.
 
+using EpicGames.Core;
+using Microsoft.Extensions.Logging;
+using Newtonsoft.Json;
 using System;
 using System.Collections.Generic;
 using System.Collections.Immutable;
 using System.Diagnostics;
+using System.Diagnostics.CodeAnalysis;
 using System.IO;
 using System.Linq;
 using System.Text;
 using System.Text.RegularExpressions;
 using System.Threading.Tasks;
 using System.Xml.Serialization;
-using EpicGames.Core;
-using EpicGames.Horde.Jobs.Graphs;
-using Microsoft.Extensions.Logging;
-using Newtonsoft.Json;
 using UnrealBuildBase;
-using static UnrealBuildTool.XmlConfig;
 
 namespace UnrealBuildTool
 {
@@ -56,19 +55,19 @@ namespace UnrealBuildTool
 	public enum AnalysisTimeoutFlags
 	{
 		/// <summary>
-		/// Analisys timeout for file 10 minutes (600 seconds)
+		/// Analysis timeout for file 10 minutes (600 seconds)
 		/// </summary>
 		After_10_minutes = 600,
 		/// <summary>
-		/// Analisys timeout for file 30 minutes (1800 seconds)
+		/// Analysis timeout for file 30 minutes (1800 seconds)
 		/// </summary>
 		After_30_minutes = 1800,
 		/// <summary>
-		/// Analisys timeout for file 60 minutes (3600 seconds)
+		/// Analysis timeout for file 60 minutes (3600 seconds)
 		/// </summary>
 		After_60_minutes = 3600,
 		/// <summary>
-		/// Analisys timeout when not set (a lot of seconds)
+		/// Analysis timeout when not set (a lot of seconds)
 		/// </summary>
 		No_timeout = 999999
 	}
@@ -175,11 +174,9 @@ namespace UnrealBuildTool
 			{
 				try
 				{
-					XmlSerializer Serializer = new XmlSerializer(typeof(PVSApplicationSettings));
-					using (FileStream Stream = new FileStream(SettingsPath.FullName, FileMode.Open, FileAccess.Read, FileShare.Read))
-					{
-						return (PVSApplicationSettings?)Serializer.Deserialize(Stream);
-					}
+					XmlSerializer Serializer = new(typeof(PVSApplicationSettings));
+					using FileStream Stream = new(SettingsPath.FullName, FileMode.Open, FileAccess.Read, FileShare.Read);
+					return (PVSApplicationSettings?)Serializer.Deserialize(Stream);
 				}
 				catch (Exception Ex)
 				{
@@ -321,22 +318,12 @@ namespace UnrealBuildTool
 	/// <summary>
 	/// Read-only version of the PVS toolchain settings
 	/// </summary>
-	public class ReadOnlyPVSTargetSettings
+	/// <remarks>
+	/// Constructor
+	/// </remarks>
+	/// <param name="Inner">The inner object</param>
+	public class ReadOnlyPVSTargetSettings(PVSTargetSettings Inner)
 	{
-		/// <summary>
-		/// Inner settings
-		/// </summary>
-		PVSTargetSettings Inner;
-
-		/// <summary>
-		/// Constructor
-		/// </summary>
-		/// <param name="Inner">The inner object</param>
-		public ReadOnlyPVSTargetSettings(PVSTargetSettings Inner)
-		{
-			this.Inner = Inner;
-		}
-
 		/// <summary>
 		/// Accessor for the Application settings
 		/// </summary>
@@ -378,164 +365,146 @@ namespace UnrealBuildTool
 		/// Path to the input file list
 		/// </summary>
 		[CommandLine("-Input", Required = true)]
-		FileReference? InputFileList = null;
+		FileReference? _inputFileList = null;
 
 		/// <summary>
 		/// Output file to generate
 		/// </summary>
 		[CommandLine("-Output", Required = true)]
-		FileReference? OutputFile = null;
+		FileReference? _outputFile = null;
 
 		/// <summary>
 		/// Path to file list of paths to ignore
 		/// </summary>
-		[CommandLine("-Ignored")]
-		FileReference? IgnoredFile = null;
+		[CommandLine("-Ignored", Required = true)]
+		FileReference? _ignoredFile = null;
+
+		/// <summary>
+		/// Path to file list of rootpaths,realpath mapping
+		/// </summary>
+		[CommandLine("-RootPaths", Required = true)]
+		FileReference? _rootPathsFile = null;
 
 		/// <summary>
 		/// The maximum level of warnings to print
 		/// </summary>
 		[CommandLine("-PrintLevel")]
-		public int PrintLevel { get; set; } = 1;
+		int _printLevel = 1;
 
 		/// <summary>
 		/// Version of the analyzers
 		/// </summary>
-		[CommandLine("-AnalyzerVersion")]
-		public Version? AnalyzerVersion { get; set; } = null;
+		[CommandLine("-AnalyzerVersion", Required = true)]
+
+		Version? _analyzerVersion = null;
 
 		/// <summary>
 		/// If all ThirdParty code should be ignored
 		/// </summary>
-		bool IgnoreThirdParty = true;
+		bool _ignoreThirdParty = true;
+
+		readonly CaptureLogger _parseLogger = new();
+
+		IEnumerable<DirectoryReference> _ignoredDirectories = [];
+
+		CppRootPaths _rootPaths = new();
+
+		bool isOldErrorInfo = false;
 
 		/// <summary>
 		/// Execute the command
 		/// </summary>
-		/// <param name="Arguments">List of command line arguments</param>
+		/// <param name="arguments">List of command line arguments</param>
 		/// <returns>Always zero, or throws an exception</returns>
-		/// <param name="Logger"></param>
-		public override Task<int> ExecuteAsync(CommandLineArguments Arguments, ILogger Logger)
+		/// <param name="logger"></param>
+		public override async Task<int> ExecuteAsync(CommandLineArguments arguments, ILogger logger)
 		{
-			Arguments.ApplyTo(this);
-			Arguments.CheckAllArgumentsUsed();
+			arguments.ApplyTo(this);
+			arguments.CheckAllArgumentsUsed();
+
+			if (_inputFileList == null || _outputFile == null || _ignoredFile == null || _rootPathsFile == null || _analyzerVersion == null)
+			{
+				throw new NullReferenceException();
+			}
 
 			// Read the input files
-			string[] InputFileLines = FileReference.ReadAllLines(InputFileList!);
-			FileReference[] InputFiles = InputFileLines.Select(x => x.Trim()).Where(x => x.Length > 0).Select(x => new FileReference(x)).ToArray();
+			string[] inputFileLines = await FileReference.ReadAllLinesAsync(_inputFileList);
+			IEnumerable<FileReference> inputFiles = inputFileLines.Select(x => x.Trim()).Where(x => x.Length > 0).Select(x => new FileReference(x));
 
 			// Read the ignore file
-			string[] IgnoreFileLines = IgnoredFile != null ? FileReference.ReadAllLines(IgnoredFile) : Array.Empty<string>();
-			DirectoryReference[] IgnoredDirectories = IgnoreFileLines.Select(x => x.Trim()).Where(x => x.Length > 0).Select(x => new DirectoryReference(x)).ToArray();
+			string[] ignoreFileLines = await FileReference.ReadAllLinesAsync(_ignoredFile);
+			_ignoredDirectories = ignoreFileLines.Select(x => x.Trim()).Where(x => x.Length > 0).Select(x => new DirectoryReference(x));
 
-			// Create the combined output file, and print the diagnostics to the log
-			HashSet<string> UniqueItems = new HashSet<string>();
-			List<string> OutputLines = new List<string>();
+			// TODO: Disallow old versions
+			isOldErrorInfo = _analyzerVersion.CompareTo(new Version("7.30")) < 0;
+
+			// Read the root paths file
+			_rootPaths = new(new BinaryArchiveReader(_rootPathsFile));
+
+			// Remove analyzedSourceFiles array from each line so all the lines can be more efficiently deduped.
+			Regex removeRegex = new("\"analyzedSourceFiles\":\\[.*?\\],", RegexOptions.Compiled);
 
-			using (StreamWriter RawWriter = new StreamWriter(OutputFile!.FullName))
+			ParallelQuery<string> allLines = inputFiles.AsParallel().SelectMany(FileReference.ReadAllLines)
+				.Select(x => removeRegex.Replace(x, String.Empty))
+				.Distinct();
+
+			// Task.Run to prevent blocking on parallel query
+			Task writeTask = Task.Run(() => FileReference.WriteAllLines(_outputFile, allLines.OrderBy(x => x)));
+
+			ParallelQuery<PVSErrorInfo> allErrors = allLines.Select(GetErrorInfo)
+				.OfType<PVSErrorInfo>();
+
+			ParallelQuery<PVSErrorInfo> renewErrors = allErrors.Where(x => x.Code.Equals("Renew", StringComparison.Ordinal));
+
+			OrderedParallelQuery<PVSErrorInfo> filteredErrors = allErrors.Except(renewErrors)
+				.Where(x => x.FalseAlarm != true && x.Level <= _printLevel) // Ignore false alarm warnings, and limit printing by PrintLevel
+				.OrderBy(x => x.Positions.FirstOrDefault()?.File)
+				.ThenBy(x => x.Positions.FirstOrDefault()?.Lines.FirstOrDefault());
+
+			// Create the combined output file, and print the diagnostics to the log
+			foreach (PVSErrorInfo errorInfo in filteredErrors)
 			{
-				foreach (FileReference InputFile in InputFiles)
-				{
-					string[] Lines = FileReference.ReadAllLines(InputFile);
-					for (int LineIdx = 0; LineIdx < Lines.Length; LineIdx++)
-					{
-						string Line = Lines[LineIdx];
-						if (String.IsNullOrWhiteSpace(Line) || !UniqueItems.Add(Line))
-						{
-							continue;
-						}
+				string file = errorInfo.Positions.FirstOrDefault()?.File ?? Unreal.EngineDirectory.FullName;
+				int lineNumber = errorInfo.Positions.FirstOrDefault()?.Lines.FirstOrDefault() ?? 1;
+				LogValue logValue = new(LogValueType.SourceFile, file, new Dictionary<Utf8String, object?> { [LogEventPropertyName.File] = file });
+				logger.LogWarning(KnownLogEvents.Compiler, "{Path}({LineNumber}): warning {WarningCode}: {WarningMessage}", logValue, lineNumber, errorInfo.Code, errorInfo.Message);
+			}
 
-						try
-						{
-							PVSErrorInfo? ErrorInfo;
-							if (AnalyzerVersion!.CompareTo(new Version("7.30")) < 0)
-							{
-								if (!TryParseErrorInfo(Line, out ErrorInfo))
-								{
-									throw new FormatException();
-								}
-							}
-							else
-							{
-								ErrorInfo = JsonConvert.DeserializeObject<PVSErrorInfo>(Line) ?? throw new FormatException();
-							}
-
-							if (ErrorInfo.Code?.Equals("Renew", StringComparison.Ordinal) == true)
-							{
-								Logger.LogInformation("PVS-Studio Renewal Notice: {WarningMessage}", ErrorInfo.Message);
-								Logger.LogWarning("Warning: PVS-Studio license will expire soon. See output log for details.");
-								continue;
-							}
-
-							string FileName = ErrorInfo.Positions![0].File!;
-							int LineNumber = ErrorInfo.Positions![0].Lines![0];
-
-							FileReference file;
-							if (!String.IsNullOrWhiteSpace(FileName))
-							{
-								file = new FileReference(FileName);
-								// Ignore anything in the IgnoredDirectories folders
-								if (IgnoredDirectories.Any() && IgnoredDirectories.Any(x => file.IsUnderDirectory(x)))
-								{
-									continue;
-								}
-
-								if (IgnoreThirdParty && file.FullName.Contains("ThirdParty", StringComparison.OrdinalIgnoreCase))
-								{
-									continue;
-								}
-							}
-							else
-							{
-								file = InputFile;
-								FileName = InputFile.FullName;
-								LineNumber = LineIdx + 1;
-							}
-
-							// Output the line to the raw output file
-							RawWriter.WriteLine(Line);
-
-							// Output the line to the log
-							if (ErrorInfo.FalseAlarm != true && ErrorInfo.Level <= PrintLevel)
-							{
-								Logger.LogWarning(KnownLogEvents.Compiler, "{Path}({LineNumber}): warning {WarningCode}: {WarningMessage}", LogValue.SourceFile(file, FileName), LineNumber, ErrorInfo.Code, ErrorInfo.Message);
-							}
-						}
-						catch (Exception Ex)
-						{
-							Logger.LogWarning(KnownLogEvents.Compiler, "{Path}({LineNumber}): warning: Unable to parse PVS output line '{Line}' ({Message})", LogValue.SourceFile(InputFile, InputFile.GetFileName()), LineIdx + 1, Line, Ex.Message);
-						}
-					}
-				}
+			if (renewErrors.Any())
+			{
+				logger.LogInformation("PVS-Studio Renewal Notice: {WarningMessage}", renewErrors.First().Message);
+				logger.LogWarning("Warning: PVS-Studio license will expire soon. See output log for details.");
 			}
-			Logger.LogInformation("Written {NumItems} {Noun} to {File}.", UniqueItems.Count, (UniqueItems.Count == 1) ? "diagnostic" : "diagnostics", OutputFile.FullName);
-			return Task.FromResult(0);
+
+			await writeTask;
+			int count = allLines.Count();
+			logger.LogInformation("Written {NumItems} {Noun} to {File}.", count, (count == 1) ? "diagnostic" : "diagnostics", _outputFile.FullName);
+			_parseLogger.RenderTo(logger);
+			return 0;
 		}
 
-		bool TryParseErrorInfo(string Line, out PVSErrorInfo ErrorInfo)
+		static readonly string[] errorInfoSplit = ["<#~>"];
+		static bool TryParseErrorInfo(string line, [NotNullWhen(true)] out PVSErrorInfo? errorInfo)
 		{
-			string[] Tokens = Line.Split(new string[] { "<#~>" }, StringSplitOptions.None);
-			if (Tokens.Length >= 9)
+			string[] tokens = line.Split(errorInfoSplit, StringSplitOptions.None);
+			if (tokens.Length >= 9)
 			{
-				string LineNumberStr = Tokens[2];
-				string FileName = Tokens[3];
-				string WarningCode = Tokens[5];
-				string WarningMessage = Tokens[6];
-				string FalseAlarmStr = Tokens[7];
-				string LevelStr = Tokens[8];
-
-				int LineNumber;
-				bool bFalseAlarm;
-				int Level;
-				if (Int32.TryParse(LineNumberStr, out LineNumber) && Boolean.TryParse(FalseAlarmStr, out bFalseAlarm) && Int32.TryParse(LevelStr, out Level))
+				string lineNumberStr = tokens[2];
+				string fileName = tokens[3];
+				string warningCode = tokens[5];
+				string warningMessage = tokens[6];
+				string falseAlarmStr = tokens[7];
+				string levelStr = tokens[8];
+
+				if (Int32.TryParse(lineNumberStr, out int LineNumber) && Boolean.TryParse(falseAlarmStr, out bool bFalseAlarm) && Int32.TryParse(levelStr, out int Level))
 				{
-					ErrorInfo = new PVSErrorInfo()
+					errorInfo = new PVSErrorInfo()
 					{
-						Code = WarningCode,
-						Message = WarningMessage,
+						Code = warningCode,
+						Message = warningMessage,
 						FalseAlarm = bFalseAlarm,
 						Level = Level,
-						Positions = new PVSPosition[] { new PVSPosition { File = FileName, Lines = new int[] { LineNumber } } },
+						Positions = [new PVSPosition { File = fileName, Lines = [LineNumber] }],
 
 					};
 
@@ -543,48 +512,91 @@ namespace UnrealBuildTool
 				}
 			}
 
-			ErrorInfo = new PVSErrorInfo();
+			errorInfo = null;
 			return false;
 		}
+
+		// Ignore anything in the IgnoredDirectories folders or ThirdParty if ignored
+		bool IsFileIgnored(FileReference? fileReference) => fileReference == null
+			|| (_ignoreThirdParty && fileReference.FullName.Contains("ThirdParty", StringComparison.OrdinalIgnoreCase))
+			|| (_ignoredDirectories.Any() && _ignoredDirectories.Any(fileReference.IsUnderDirectory));
+
+		PVSErrorInfo? GetErrorInfo(string line)
+		{
+			try
+			{
+				PVSErrorInfo? errorInfo = null;
+				if (isOldErrorInfo)
+				{
+					if (!TryParseErrorInfo(line, out errorInfo))
+					{
+						throw new FormatException();
+					}
+				}
+				else
+				{
+					errorInfo = JsonConvert.DeserializeObject<PVSErrorInfo>(line) ?? throw new FormatException();
+				}
+
+				FileReference? fileReference = errorInfo.Positions.FirstOrDefault()?.UpdateFilePath(_rootPaths);
+				return IsFileIgnored(fileReference) ? null : errorInfo;
+			}
+			catch (Exception ex)
+			{
+				_parseLogger.LogDebug(KnownLogEvents.Compiler, "warning: Unable to parse PVS output line '{Line}' ({Message})", line, ex.Message);
+			}
+			return null;
+		}
 	}
 
 	class PVSPosition
 	{
 		[JsonProperty(Required = Required.Always)]
-		public string? File;
+		public required string File;
 
 		[JsonProperty(Required = Required.Always)]
-		public int[]? Lines;
+		public required int[] Lines;
+
+		public FileReference? UpdateFilePath(CppRootPaths rootPaths)
+		{
+			FileReference? fileReference = !String.IsNullOrWhiteSpace(File) ? FileReference.FromString(File) : null;
+			if (fileReference != null && rootPaths.bUseVfs)
+			{
+				fileReference = rootPaths.GetLocalPath(fileReference);
+				File = fileReference.FullName;
+			}
+			return fileReference;
+		}
 	}
 
 	class PVSErrorInfo
 	{
 		[JsonProperty(Required = Required.Always)]
-		public string? Code;
+		public required string Code;
 
 		[JsonProperty(Required = Required.Always)]
-		public bool? FalseAlarm;
+		public required bool FalseAlarm;
 
 		[JsonProperty(Required = Required.Always)]
-		public int? Level;
+		public required int Level;
 
 		[JsonProperty(Required = Required.Always)]
-		public string? Message;
+		public required string Message;
 
 		[JsonProperty(Required = Required.Always)]
-		public PVSPosition[]? Positions;
+		public required PVSPosition[] Positions;
 	}
 
 	class PVSToolChain : ISPCToolChain
 	{
-		ReadOnlyTargetRules Target;
-		ReadOnlyPVSTargetSettings Settings;
-		PVSApplicationSettings? ApplicationSettings;
-		VCToolChain InnerToolChain;
-		FileReference AnalyzerFile;
-		FileReference? LicenseFile;
-		UnrealTargetPlatform Platform;
-		Version AnalyzerVersion;
+		readonly ReadOnlyTargetRules Target;
+		readonly ReadOnlyPVSTargetSettings Settings;
+		readonly PVSApplicationSettings? ApplicationSettings;
+		readonly VCToolChain InnerToolChain;
+		readonly FileReference AnalyzerFile;
+		readonly FileReference? LicenseFile;
+		readonly UnrealTargetPlatform Platform;
+		readonly Version AnalyzerVersion;
 
 		string OutputFileExtension => AnalyzerVersion.CompareTo(new Version("7.30")) >= 0 ? ".PVS-Studio.log" : ".pvslog";
 
@@ -628,10 +640,10 @@ namespace UnrealBuildTool
 			}
 			else
 			{
-				FileReference DefaultLicenseFile = AnalyzerFile.ChangeExtension(".lic");
-				if (FileReference.Exists(DefaultLicenseFile))
+				FileReference defaultLicenseFile = AnalyzerFile.ChangeExtension(".lic");
+				if (FileReference.Exists(defaultLicenseFile))
 				{
-					LicenseFile = DefaultLicenseFile;
+					LicenseFile = defaultLicenseFile;
 				}
 			}
 
@@ -643,10 +655,10 @@ namespace UnrealBuildTool
 
 		public override void GetVersionInfo(List<string> Lines)
 		{
-			base.GetVersionInfo(Lines);
+			InnerToolChain.GetVersionInfo(Lines);
 
-			ReadOnlyPVSTargetSettings Settings = Target.WindowsPlatform.PVS;
-			Lines.Add(String.Format("Using PVS-Studio {0} at {1} with analysis mode {2} ({3})", AnalyzerVersion, AnalyzerFile, (uint)Settings.ModeFlags, Settings.ModeFlags.ToString()));
+			ReadOnlyPVSTargetSettings settings = Target.WindowsPlatform.PVS;
+			Lines.Add(String.Format("Using PVS-Studio {0} at {1} with analysis mode {2} ({3})", AnalyzerVersion, AnalyzerFile, (uint)settings.ModeFlags, settings.ModeFlags.ToString()));
 		}
 
 		public override void SetUpGlobalEnvironment(ReadOnlyTargetRules Target, CppCompileEnvironment GlobalCompileEnvironment, LinkEnvironment GlobalLinkEnvironment)
@@ -674,12 +686,12 @@ namespace UnrealBuildTool
 			 
 		static Version GetAnalyzerVersion(FileReference AnalyzerPath)
 		{
-			string Output = String.Empty;
-			Version? AnalyzerVersion = new Version(0, 0);
+			string output = String.Empty;
+			Version? analyzerVersion = new(0, 0);
 
 			try
 			{
-				using (Process PvsProc = new Process())
+				using (Process PvsProc = new())
 				{
 					PvsProc.StartInfo.FileName = AnalyzerPath.FullName;
 					PvsProc.StartInfo.Arguments = "--version";
@@ -688,297 +700,271 @@ namespace UnrealBuildTool
 					PvsProc.StartInfo.RedirectStandardOutput = true;
 
 					PvsProc.Start();
-					Output = PvsProc.StandardOutput.ReadToEnd();
+					output = PvsProc.StandardOutput.ReadToEnd();
 					PvsProc.WaitForExit();
 				}
 
-				const string VersionPattern = @"\d+(?:\.\d+)+";
-				Match Match = Regex.Match(Output, VersionPattern);
+				const string versionPattern = @"\d+(?:\.\d+)+";
+				Match match = Regex.Match(output, versionPattern);
 
-				if (Match.Success)
+				if (match.Success)
 				{
-					string VersionStr = Match.Value;
-					if (!Version.TryParse(VersionStr, out AnalyzerVersion))
+					string versionStr = match.Value;
+					if (!Version.TryParse(versionStr, out analyzerVersion))
 					{
-						throw new BuildException(String.Format("Failed to parse PVS-Studio version: {0}", VersionStr));
+						throw new BuildLogEventException("Failed to parse PVS-Studio version: {Version}", versionStr);
 					}
 				}
 			}
-			catch (Exception Ex)
+			catch (Exception ex)
 			{
-				if (Ex is BuildException)
+				if (ex is BuildException)
 				{
 					throw;
 				}
 
-				throw new BuildException(Ex, "Failed to obtain PVS-Studio version.");
+				throw new BuildException(ex, "Failed to obtain PVS-Studio version.");
 			}
 
-			return AnalyzerVersion;
+			return analyzerVersion;
 		}
 
-		class ActionGraphCapture : ForwardingActionGraphBuilder
+		class ActionGraphCapture(IActionGraphBuilder inner, List<IExternalAction> actions) : ForwardingActionGraphBuilder(inner)
 		{
-			List<IExternalAction> Actions;
-
-			public ActionGraphCapture(IActionGraphBuilder Inner, List<IExternalAction> Actions)
-				: base(Inner)
-			{
-				this.Actions = Actions;
-			}
-
 			public override void AddAction(IExternalAction Action)
 			{
 				base.AddAction(Action);
-
-				Actions.Add(Action);
+				actions.Add(Action);
 			}
 		}
 
-		public static readonly VersionNumber CLVerWithCPP20Support = new VersionNumber(14, 23);
+		public static readonly VersionNumber CLVerWithCPP20Support = new(14, 23);
 
 		const string CPP_20 = "c++20";
 		const string CPP_17 = "c++17";
+		const string CPP_14 = "c++14";
 
 		public static string GetLangStandForCfgFile(CppStandardVersion cppStandard, VersionNumber compilerVersion)
 		{
-			string cppCfgStandard;
-
-			switch (cppStandard)
+			return cppStandard switch
 			{
-				case CppStandardVersion.Cpp17:
-					cppCfgStandard = CPP_17;
-					break;
-				case CppStandardVersion.Cpp20:
-					cppCfgStandard = CPP_20;
-					break;
-				case CppStandardVersion.Latest:
-					cppCfgStandard = VersionNumber.Compare(compilerVersion, CLVerWithCPP20Support) >= 0 ? CPP_20 : CPP_17;
-					break;
-				default:
-					cppCfgStandard = "c++14";
-					break;
-			}
-
-			return cppCfgStandard;
+				CppStandardVersion.Cpp17 => CPP_17,
+				CppStandardVersion.Cpp20 => CPP_20,
+				CppStandardVersion.Latest => VersionNumber.Compare(compilerVersion, CLVerWithCPP20Support) >= 0 ? CPP_20 : CPP_17,
+				_ => CPP_14,
+			};
 		}
 
 		public static bool ShouldCompileAsC(string compilerCommandLine, string sourceFileName)
 		{
-			int CFlagLastPosition = Math.Max(Math.Max(compilerCommandLine.LastIndexOf("/TC "), compilerCommandLine.LastIndexOf("/Tc ")),
+			int cFlagLastPosition = Math.Max(Math.Max(compilerCommandLine.LastIndexOf("/TC "), compilerCommandLine.LastIndexOf("/Tc ")),
 												Math.Max(compilerCommandLine.LastIndexOf("-TC "), compilerCommandLine.LastIndexOf("-Tc ")));
 
-			int CppFlagLastPosition = Math.Max(Math.Max(compilerCommandLine.LastIndexOf("/TP "), compilerCommandLine.LastIndexOf("/Tp ")),
+			int cppFlagLastPosition = Math.Max(Math.Max(compilerCommandLine.LastIndexOf("/TP "), compilerCommandLine.LastIndexOf("/Tp ")),
 												Math.Max(compilerCommandLine.LastIndexOf("-TP "), compilerCommandLine.LastIndexOf("-Tp ")));
 
-			bool compileAsCCode;
-			if (CFlagLastPosition == CppFlagLastPosition)
-			{
-				//ни один флаг, определяющий язык, не задан. Определяем по расширению файла
-				compileAsCCode = Path.GetExtension(sourceFileName).Equals(".c", StringComparison.InvariantCultureIgnoreCase);
-			}
-			else
-			{
-				compileAsCCode = CFlagLastPosition > CppFlagLastPosition;
-			}
+			bool compileAsCCode = cFlagLastPosition == cppFlagLastPosition
+				? Path.GetExtension(sourceFileName).Equals(".c", StringComparison.InvariantCultureIgnoreCase)
+				: cFlagLastPosition > cppFlagLastPosition;
 
 			return compileAsCCode;
 		}
 
-		public override CppCompileEnvironment CreateSharedResponseFile(CppCompileEnvironment CompileEnvironment, FileReference OutResponseFile, IActionGraphBuilder Graph)
+		public override CppCompileEnvironment CreateSharedResponseFile(CppCompileEnvironment compileEnvironment, FileReference outResponseFile, IActionGraphBuilder graph)
 		{
-			return CompileEnvironment;
+			return compileEnvironment;
 		}
 
-		CPPOutput PreprocessCppFiles(CppCompileEnvironment CompileEnvironment, IEnumerable<FileItem> InputFiles, DirectoryReference OutputDir, string ModuleName, IActionGraphBuilder Graph, out List<IExternalAction> PreprocessActions)
+		CPPOutput PreprocessCppFiles(CppCompileEnvironment compileEnvironment, IEnumerable<FileItem> inputFiles, DirectoryReference outputDir, string moduleName, IActionGraphBuilder graph, out List<IExternalAction> preprocessActions)
 		{
 			// Preprocess the source files with the regular toolchain
-			CppCompileEnvironment PreprocessCompileEnvironment = new CppCompileEnvironment(CompileEnvironment);
-			PreprocessCompileEnvironment.bPreprocessOnly = true;
-			PreprocessCompileEnvironment.CppCompileWarnings.UndefinedIdentifierWarningLevel = WarningLevel.Off; // Not sure why THIRD_PARTY_INCLUDES_START doesn't pick this up; the _Pragma appears in the preprocessed output. Perhaps in preprocess-only mode the compiler doesn't respect these?
-			PreprocessCompileEnvironment.AdditionalArguments += " /wd4005 /wd4828 /wd5105";
-			PreprocessCompileEnvironment.Definitions.Add("PVS_STUDIO");
-
-			PreprocessActions = new List<IExternalAction>();
-			return InnerToolChain.CompileAllCPPFiles(PreprocessCompileEnvironment, InputFiles, OutputDir, ModuleName, new ActionGraphCapture(Graph, PreprocessActions));
+			CppCompileEnvironment preprocessCompileEnvironment = new(compileEnvironment)
+			{
+				bPreprocessOnly = true
+			};
+			preprocessCompileEnvironment.AdditionalArguments += " /wd4005 /wd4828 /wd5105";
+			preprocessCompileEnvironment.Definitions.Add("PVS_STUDIO");
+			preprocessCompileEnvironment.CppCompileWarnings.UndefinedIdentifierWarningLevel = WarningLevel.Off; // Not sure why THIRD_PARTY_INCLUDES_START doesn't pick this up; the _Pragma appears in the preprocessed output. Perhaps in preprocess-only mode the compiler doesn't respect these?
+
+			preprocessActions = [];
+			return InnerToolChain.CompileAllCPPFiles(preprocessCompileEnvironment, inputFiles, outputDir, moduleName, new ActionGraphCapture(graph, preprocessActions));
 		}
 
-		void AnalyzeCppFile(VCCompileAction PreprocessAction, CppCompileEnvironment CompileEnvironment, DirectoryReference OutputDir, CPPOutput Result, IActionGraphBuilder Graph)
+		void AnalyzeCppFile(VCCompileAction preprocessAction, CppCompileEnvironment compileEnvironment, DirectoryReference outputDir, CPPOutput result, IActionGraphBuilder graph)
 		{
-			FileItem SourceFileItem = PreprocessAction.SourceFile!;
-			FileItem PreprocessedFileItem = PreprocessAction.PreprocessedFile!;
+			FileItem sourceFileItem = preprocessAction.SourceFile!;
+			FileItem preprocessedFileItem = preprocessAction.PreprocessedFile!;
 
 			// Write the PVS studio config file
-			StringBuilder ConfigFileContents = new StringBuilder();
-			foreach (DirectoryReference IncludePath in Target.WindowsPlatform.Environment!.IncludePaths)
+			StringBuilder configFileContents = new();
+			foreach (DirectoryReference includePath in Target.WindowsPlatform.Environment!.IncludePaths)
 			{
-				ConfigFileContents.AppendFormat("exclude-path={0}\n", IncludePath.FullName);
+				configFileContents.AppendFormat("exclude-path={0}\n", includePath.FullName);
 			}
 			if (ApplicationSettings != null && ApplicationSettings.PathMasks != null)
 			{
-				foreach (string PathMask in ApplicationSettings.PathMasks)
+				foreach (string pathMask in ApplicationSettings.PathMasks)
 				{
-					if (PathMask.Contains(':') || PathMask.Contains('\\') || PathMask.Contains('/'))
+					if (pathMask.Contains(':') || pathMask.Contains('\\') || pathMask.Contains('/'))
 					{
-						if (Path.IsPathRooted(PathMask) && !PathMask.Contains(':'))
+						if (Path.IsPathRooted(pathMask) && !pathMask.Contains(':'))
 						{
-							ConfigFileContents.AppendFormat("exclude-path=*{0}*\n", PathMask);
+							configFileContents.AppendFormat("exclude-path=*{0}*\n", pathMask);
 						}
 						else
 						{
-							ConfigFileContents.AppendFormat("exclude-path={0}\n", PathMask);
+							configFileContents.AppendFormat("exclude-path={0}\n", pathMask);
 						}
 					}
 				}
 			}
 			if (Platform.IsInGroup(UnrealPlatformGroup.Microsoft))
 			{
-				ConfigFileContents.Append("platform=x64\n");
+				configFileContents.Append("platform=x64\n");
 			}
 			else
 			{
 				throw new BuildException("PVS-Studio does not support this platform");
 			}
-			ConfigFileContents.Append("preprocessor=visualcpp\n");
+			configFileContents.Append("preprocessor=visualcpp\n");
 
-			bool shouldCompileAsC = ShouldCompileAsC(String.Join(" ", PreprocessAction.Arguments), SourceFileItem.AbsolutePath);
-			ConfigFileContents.AppendFormat("language={0}\n", shouldCompileAsC ? "C" : "C++");
+			bool shouldCompileAsC = ShouldCompileAsC(String.Join(" ", preprocessAction.Arguments), sourceFileItem.AbsolutePath);
+			configFileContents.AppendFormat("language={0}\n", shouldCompileAsC ? "C" : "C++");
 
-			ConfigFileContents.Append("skip-cl-exe=yes\n");
+			configFileContents.Append("skip-cl-exe=yes\n");
 
-			WindowsCompiler WindowsCompiler = Target.WindowsPlatform.Compiler;
-			bool isVisualCppCompiler = WindowsCompiler.IsMSVC();
+			WindowsCompiler windowsCompiler = Target.WindowsPlatform.Compiler;
+			bool isVisualCppCompiler = windowsCompiler.IsMSVC();
 			if (AnalyzerVersion.CompareTo(new Version("7.07")) >= 0 && !shouldCompileAsC)
 			{
 				VersionNumber compilerVersion = Target.WindowsPlatform.Environment.CompilerVersion;
-				string languageStandardForCfg = GetLangStandForCfgFile(CompileEnvironment.CppStandard, compilerVersion);
+				string languageStandardForCfg = GetLangStandForCfgFile(compileEnvironment.CppStandard, compilerVersion);
 
-				ConfigFileContents.AppendFormat("std={0}\n", languageStandardForCfg);
+				configFileContents.AppendFormat("std={0}\n", languageStandardForCfg);
 
-				bool disableMsExtentinsFromArgs = PreprocessAction.Arguments.Any(arg => arg.Equals("/Za") || arg.Equals("-Za") || arg.Equals("/permissive-"));
-				bool disableMsExtentions = isVisualCppCompiler && (languageStandardForCfg == CPP_20 || disableMsExtentinsFromArgs);
-				ConfigFileContents.AppendFormat("disable-ms-extensions={0}\n", disableMsExtentions ? "yes" : "no");
+				bool disableMsExtensionsFromArgs = preprocessAction.Arguments.Any(arg => arg.Equals("/Za") || arg.Equals("-Za") || arg.Equals("/permissive-"));
+				bool disableMsExtensions = isVisualCppCompiler && (languageStandardForCfg == CPP_20 || disableMsExtensionsFromArgs);
+				configFileContents.AppendFormat("disable-ms-extensions={0}\n", disableMsExtensions ? "yes" : "no");
 			}
 
-			if (isVisualCppCompiler && PreprocessAction.Arguments.Any(arg => arg.StartsWith("/await")))
+			if (isVisualCppCompiler && preprocessAction.Arguments.Any(arg => arg.StartsWith("/await")))
 			{
-				ConfigFileContents.Append("msvc-await=yes\n");
+				configFileContents.Append("msvc-await=yes\n");
 			}
 
 			if (Settings.EnableNoNoise)
 			{
-				ConfigFileContents.Append("no-noise=yes\n");
+				configFileContents.Append("no-noise=yes\n");
 			}
 
 			if (Settings.EnableReportDisabledRules)
 			{
-				ConfigFileContents.Append("report-disabled-rules=yes\n");
+				configFileContents.Append("report-disabled-rules=yes\n");
 			}
 
 			// TODO: Investigate into this disabled error
-			if (SourceFileItem.Location.IsUnderDirectory(Unreal.RootDirectory))
+			if (sourceFileItem.Location.IsUnderDirectory(Unreal.RootDirectory))
 			{
-				ConfigFileContents.AppendFormat("errors-off=V1102\n");
+				configFileContents.AppendFormat("errors-off=V1102\n");
 			}
 
-			foreach (string error in CompileEnvironment.StaticAnalyzerPVSDisabledErrors.OrderBy(x => x))
+			foreach (string error in compileEnvironment.StaticAnalyzerPVSDisabledErrors.OrderBy(x => x))
 			{
-				ConfigFileContents.AppendFormat($"errors-off={error}\n");
+				configFileContents.AppendFormat($"errors-off={error}\n");
 			}
 
-			int Timeout = (int)(Settings.AnalysisTimeoutFlag == AnalysisTimeoutFlags.No_timeout ? 0 : Settings.AnalysisTimeoutFlag);
-			ConfigFileContents.AppendFormat("timeout={0}\n", Timeout);
+			int timeout = (int)(Settings.AnalysisTimeoutFlag == AnalysisTimeoutFlags.No_timeout ? 0 : Settings.AnalysisTimeoutFlag);
+			configFileContents.AppendFormat("timeout={0}\n", timeout);
 
 			if (AnalyzerVersion.CompareTo(new Version("7.28")) >= 0)
 			{
-				ConfigFileContents.Append("silent-exit-code-mode=yes\n");
+				configFileContents.Append("silent-exit-code-mode=yes\n");
 			}
 
 			if (AnalyzerVersion.CompareTo(new Version("7.30")) >= 0)
 			{
-				ConfigFileContents.Append("new-output-format=yes\n");
+				configFileContents.Append("new-output-format=yes\n");
 			}
 
 			if (AnalyzerVersion.CompareTo(new Version("7.34")) >= 0)
 			{
 				if (Target.bStaticAnalyzerProjectOnly)
 				{
-					ConfigFileContents.Append($"analysis-paths=skip={Unreal.EngineSourceDirectory}\n");
+					configFileContents.Append($"analysis-paths=skip={Unreal.EngineSourceDirectory}\n");
 				}
 				if (!Target.bStaticAnalyzerIncludeGenerated)
 				{
-					ConfigFileContents.Append($"analysis-paths=skip=*.gen.cpp\n");
-					ConfigFileContents.Append($"analysis-paths=skip=*.generated.h\n");
+					configFileContents.Append($"analysis-paths=skip=*.gen.cpp\n");
+					configFileContents.Append($"analysis-paths=skip=*.generated.h\n");
 				}
 			}
 
-			string BaseFileName = PreprocessedFileItem.Location.GetFileName();
+			string baseFileName = preprocessedFileItem.Location.GetFileName();
 
-			FileReference ConfigFileLocation = FileReference.Combine(OutputDir, BaseFileName + ".cfg");
-			FileItem ConfigFileItem = Graph.CreateIntermediateTextFile(ConfigFileLocation, ConfigFileContents.ToString());
+			FileReference configFileLocation = FileReference.Combine(outputDir, baseFileName + ".cfg");
+			FileItem configFileItem = graph.CreateIntermediateTextFile(configFileLocation, configFileContents.ToString());
 
 			// Run the analyzer on the preprocessed source file
-			FileReference OutputFileLocation = FileReference.Combine(OutputDir, BaseFileName + OutputFileExtension);
-			FileItem OutputFileItem = FileItem.GetItemByFileReference(OutputFileLocation);
+			FileReference outputFileLocation = FileReference.Combine(outputDir, baseFileName + OutputFileExtension);
+			FileItem outputFileItem = FileItem.GetItemByFileReference(outputFileLocation);
 
-			Action AnalyzeAction = Graph.CreateAction(ActionType.Compile);
-			AnalyzeAction.CommandDescription = "Analyzing";
-			AnalyzeAction.StatusDescription = BaseFileName;
-			AnalyzeAction.WorkingDirectory = Unreal.EngineSourceDirectory;
-			AnalyzeAction.CommandPath = AnalyzerFile;
+			Action analyzeAction = graph.CreateAction(ActionType.Compile);
+			analyzeAction.CommandDescription = "Analyzing";
+			analyzeAction.StatusDescription = baseFileName;
+			analyzeAction.WorkingDirectory = Unreal.EngineSourceDirectory;
+			analyzeAction.CommandPath = AnalyzerFile;
 
-			List<string> Arguments =
+			List<string> arguments =
 			[
-				$"--source-file \"{SourceFileItem.AbsolutePath}\"",
-				$"--output-file \"{OutputFileItem.AbsolutePath}\"",
-				$"--cfg \"{ConfigFileItem.AbsolutePath}\"",
-				$"--i-file=\"{PreprocessedFileItem.AbsolutePath}\"",
+				$"--source-file \"{sourceFileItem.AbsolutePath}\"",
+				$"--output-file \"{outputFileItem.AbsolutePath}\"",
+				$"--cfg \"{configFileItem.AbsolutePath}\"",
+				$"--i-file=\"{preprocessedFileItem.AbsolutePath}\"",
 				$"--analysis-mode {(uint)Settings.ModeFlags}",
 				$"--lic-name \"{ApplicationSettings?.UserName}\" --lic-key \"{ApplicationSettings?.SerialNumber}\"",
 			];
 
-			AnalyzeAction.CommandArguments = String.Join(' ', Arguments);
+			analyzeAction.CommandArguments = String.Join(' ', arguments);
 
-			AnalyzeAction.PrerequisiteItems.Add(SourceFileItem);
-			AnalyzeAction.PrerequisiteItems.Add(ConfigFileItem);
-			AnalyzeAction.PrerequisiteItems.Add(PreprocessedFileItem);
-			AnalyzeAction.ProducedItems.Add(OutputFileItem);
-			AnalyzeAction.DeleteItems.Add(OutputFileItem); // PVS Studio will append by default, so need to delete produced items
-			AnalyzeAction.bCanExecuteRemotely = true;
-			AnalyzeAction.bCanExecuteRemotelyWithSNDBS = false;
-			AnalyzeAction.bCanExecuteRemotelyWithXGE = false;
+			analyzeAction.PrerequisiteItems.Add(sourceFileItem);
+			analyzeAction.PrerequisiteItems.Add(configFileItem);
+			analyzeAction.PrerequisiteItems.Add(preprocessedFileItem);
+			analyzeAction.ProducedItems.Add(outputFileItem);
+			analyzeAction.DeleteItems.Add(outputFileItem); // PVS Studio will append by default, so need to delete produced items
+			analyzeAction.bCanExecuteRemotely = true;
+			analyzeAction.bCanExecuteRemotelyWithSNDBS = false;
+			analyzeAction.bCanExecuteRemotelyWithXGE = false;
 
-			AnalyzeAction.RootPaths = CompileEnvironment.RootPaths;
-			AnalyzeAction.CacheBucket = GetCacheBucket(Target, null);
-			AnalyzeAction.ArtifactMode = ArtifactMode.Enabled;
+			analyzeAction.RootPaths = compileEnvironment.RootPaths;
+			analyzeAction.CacheBucket = GetCacheBucket(Target, null);
+			analyzeAction.ArtifactMode = ArtifactMode.Enabled;
 
-			Result.ObjectFiles.AddRange(AnalyzeAction.ProducedItems);
+			result.ObjectFiles.AddRange(analyzeAction.ProducedItems);
 		}
 
-		protected override CPPOutput CompileCPPFiles(CppCompileEnvironment CompileEnvironment, IEnumerable<FileItem> InputFiles, DirectoryReference OutputDir, string ModuleName, IActionGraphBuilder Graph)
+		protected override CPPOutput CompileCPPFiles(CppCompileEnvironment compileEnvironment, IEnumerable<FileItem> inputFiles, DirectoryReference outputDir, string moduleName, IActionGraphBuilder graph)
 		{
-			if (CompileEnvironment.bDisableStaticAnalysis)
+			if (compileEnvironment.bDisableStaticAnalysis)
 			{
 				return new CPPOutput();
 			}
 
 			// Use a subdirectory for PVS output, to avoid clobbering regular build artifacts
-			OutputDir = DirectoryReference.Combine(OutputDir, "PVS");
+			outputDir = DirectoryReference.Combine(outputDir, "PVS");
 
 			// Preprocess the source files with the regular toolchain
-			CPPOutput Result = PreprocessCppFiles(CompileEnvironment, InputFiles, OutputDir, ModuleName, Graph, out List<IExternalAction> PreprocessActions);
+			CPPOutput result = PreprocessCppFiles(compileEnvironment, inputFiles, outputDir, moduleName, graph, out List<IExternalAction> PreprocessActions);
 
 			// Run the source files through PVS-Studio
 			for (int Idx = 0; Idx < PreprocessActions.Count; Idx++)
 			{
-				VCCompileAction? PreprocessAction = PreprocessActions[Idx] as VCCompileAction;
-				if (PreprocessAction == null)
+				if (PreprocessActions[Idx] is not VCCompileAction PreprocessAction)
 				{
 					continue;
 				}
 
-				FileItem? SourceFileItem = PreprocessAction.SourceFile;
-				if (SourceFileItem == null)
+				FileItem? sourceFileItem = PreprocessAction.SourceFile;
+				if (sourceFileItem == null)
 				{
 					Logger.LogWarning("Unable to find source file from command producing: {File}", String.Join(", ", PreprocessActions[Idx].ProducedItems.Select(x => x.Location.GetFileName())));
 					continue;
@@ -986,67 +972,69 @@ namespace UnrealBuildTool
 
 				if (PreprocessAction.PreprocessedFile == null)
 				{
-					Logger.LogWarning("Unable to find preprocessed output file from {File}", SourceFileItem.Location.GetFileName());
+					Logger.LogWarning("Unable to find preprocessed output file from {File}", sourceFileItem.Location.GetFileName());
 					continue;
 				}
 
 				// We don't want to run these remotely since they are very lightweight but has lots of I/O
 				PreprocessAction.bCanExecuteRemotely = false;
 
-				AnalyzeCppFile(PreprocessAction, CompileEnvironment, OutputDir, Result, Graph);
+				AnalyzeCppFile(PreprocessAction, compileEnvironment, outputDir, result, graph);
 			}
-			return Result;
+			return result;
 		}
 
-		public override void GenerateTypeLibraryHeader(CppCompileEnvironment CompileEnvironment, ModuleRules.TypeLibrary TypeLibrary, FileReference OutputFile, FileReference? OutputHeader, IActionGraphBuilder Graph)
+		public override void GenerateTypeLibraryHeader(CppCompileEnvironment compileEnvironment, ModuleRules.TypeLibrary typeLibrary, FileReference outputFile, FileReference? outputHeader, IActionGraphBuilder graph)
 		{
-			InnerToolChain.GenerateTypeLibraryHeader(CompileEnvironment, TypeLibrary, OutputFile, OutputHeader, Graph);
+			InnerToolChain.GenerateTypeLibraryHeader(compileEnvironment, typeLibrary, outputFile, outputHeader, graph);
 		}
 
-		public override FileItem LinkFiles(LinkEnvironment LinkEnvironment, bool bBuildImportLibraryOnly, IActionGraphBuilder Graph)
+		public override FileItem LinkFiles(LinkEnvironment linkEnvironment, bool bBuildImportLibraryOnly, IActionGraphBuilder graph)
 		{
 			throw new BuildException("Unable to link with PVS toolchain.");
 		}
 
-		public override void FinalizeOutput(ReadOnlyTargetRules Target, TargetMakefileBuilder MakefileBuilder)
+		public override void FinalizeOutput(ReadOnlyTargetRules target, TargetMakefileBuilder makefileBuilder)
 		{
-			FileReference OutputFile;
 			string outputFileExtension = OutputFileExtension;
+			FileReference outputFile = target.ProjectFile == null
+				? FileReference.Combine(Unreal.EngineDirectory, "Saved", "PVS-Studio", $"{target.Name}{outputFileExtension}")
+				: FileReference.Combine(target.ProjectFile.Directory, "Saved", "PVS-Studio", $"{target.Name}{outputFileExtension}");
 
-			if (Target.ProjectFile == null)
-			{
-				OutputFile = FileReference.Combine(Unreal.EngineDirectory, "Saved", "PVS-Studio", $"{Target.Name}{outputFileExtension}");
-			}
-			else
-			{
-				OutputFile = FileReference.Combine(Target.ProjectFile.Directory, "Saved", "PVS-Studio", $"{Target.Name}{outputFileExtension}");
-			}
-
-			TargetMakefile Makefile = MakefileBuilder.Makefile;
-			ImmutableSortedSet<FileReference> InputFiles = Makefile.OutputItems.Select(x => x.Location).Where(x => x.HasExtension(outputFileExtension)).ToImmutableSortedSet();
+			TargetMakefile makefile = makefileBuilder.Makefile;
+			IEnumerable<FileReference> inputFiles = [.. makefile.OutputItems.Select(x => x.Location).Where(x => x.HasExtension(outputFileExtension))];
 
 			// Collect the sourcefile items off of the Compile action added in CompileCPPFiles so that in SingleFileCompile mode the PVSGather step is also not filtered out
-			ImmutableSortedSet<FileItem> CompileSourceFiles = Makefile.Actions.OfType<VCCompileAction>().Select(x => x.SourceFile!).ToImmutableSortedSet();
+			IEnumerable<FileItem> compileSourceFiles = [.. makefile.Actions.OfType<VCCompileAction>().Select(x => x.SourceFile!)];
+			CppRootPaths rootPaths = makefile.Actions.OfType<VCCompileAction>().FirstOrDefault(x => x.RootPaths.Any())?.RootPaths ?? new();
 
 			// Store list of system paths that should be excluded
-			ImmutableSortedSet<DirectoryReference> SystemIncludePaths = Makefile.Actions.OfType<VCCompileAction>().SelectMany(x => x.SystemIncludePaths).ToImmutableSortedSet();
+			IEnumerable<DirectoryReference> systemIncludePaths = [.. makefile.Actions.OfType<VCCompileAction>().SelectMany(x => x.SystemIncludePaths)];
 
-			FileItem InputFileListItem = MakefileBuilder.CreateIntermediateTextFile(OutputFile.ChangeExtension(".input"), InputFiles.Select(x => x.FullName));
-			FileItem IgnoredFileListeItem = MakefileBuilder.CreateIntermediateTextFile(OutputFile.ChangeExtension(".ignored"), SystemIncludePaths.Select(x => x.FullName));
+			FileItem inputFileListItem = makefileBuilder.CreateIntermediateTextFile(FileReference.Combine(makefile.ProjectIntermediateDirectory, outputFile.ChangeExtension(".input").GetFileName()), inputFiles.Select(x => x.FullName).Distinct().Order());
+			FileItem ignoredFileListItem = makefileBuilder.CreateIntermediateTextFile(FileReference.Combine(makefile.ProjectIntermediateDirectory, outputFile.ChangeExtension(".ignored").GetFileName()), systemIncludePaths.Select(x => x.FullName).Distinct().Order());
+			FileItem rootPathsItem = FileItem.GetItemByFileReference(FileReference.Combine(makefile.ProjectIntermediateDirectory, outputFile.ChangeExtension(".rootpaths").GetFileName()));
+			{
+				using MemoryStream stream = new();
+				using BinaryArchiveWriter binaryArchiveWriter = new(stream);
+				rootPaths.Write(binaryArchiveWriter);
+				binaryArchiveWriter.Flush();
+				FileReference.WriteAllBytesIfDifferent(rootPathsItem.Location, stream.ToArray());
+			}
 
-			string Arguments = $"-Input=\"{InputFileListItem.Location}\" -Output=\"{OutputFile}\" -Ignored=\"{IgnoredFileListeItem.Location}\" -PrintLevel={Target.StaticAnalyzerPVSPrintLevel} -AnalyzerVersion={AnalyzerVersion}";
+			string arguments = $"-Input=\"{inputFileListItem.Location}\" -Output=\"{outputFile}\" -Ignored=\"{ignoredFileListItem.Location}\" -RootPaths=\"{rootPathsItem.Location}\" -PrintLevel={target.StaticAnalyzerPVSPrintLevel} -AnalyzerVersion={AnalyzerVersion}";
 
-			Action AnalyzeAction = MakefileBuilder.CreateRecursiveAction<PVSGatherMode>(ActionType.PostBuildStep, Arguments);
-			AnalyzeAction.CommandDescription = "Process PVS-Studio Results";
-			AnalyzeAction.PrerequisiteItems.Add(InputFileListItem);
-			AnalyzeAction.PrerequisiteItems.Add(IgnoredFileListeItem);
-			AnalyzeAction.PrerequisiteItems.UnionWith(Makefile.OutputItems);
-			AnalyzeAction.PrerequisiteItems.UnionWith(CompileSourceFiles);
-			AnalyzeAction.ProducedItems.Add(FileItem.GetItemByFileReference(OutputFile));
-			AnalyzeAction.ProducedItems.Add(FileItem.GetItemByPath(OutputFile.FullName + "_does_not_exist")); // Force the gather step to always execute
-			AnalyzeAction.DeleteItems.UnionWith(AnalyzeAction.ProducedItems);
+			Action finalizeAction = makefileBuilder.CreateRecursiveAction<PVSGatherMode>(ActionType.PostBuildStep, arguments);
+			finalizeAction.CommandDescription = "Process PVS-Studio Results";
+			finalizeAction.PrerequisiteItems.Add(inputFileListItem);
+			finalizeAction.PrerequisiteItems.Add(ignoredFileListItem);
+			finalizeAction.PrerequisiteItems.UnionWith(makefile.OutputItems);
+			finalizeAction.PrerequisiteItems.UnionWith(compileSourceFiles);
+			finalizeAction.ProducedItems.Add(FileItem.GetItemByFileReference(outputFile));
+			finalizeAction.ProducedItems.Add(FileItem.GetItemByPath(outputFile.FullName + "_does_not_exist")); // Force the gather step to always execute
+			finalizeAction.DeleteItems.UnionWith(finalizeAction.ProducedItems);
 
-			Makefile.OutputItems.AddRange(AnalyzeAction.ProducedItems);
+			makefile.OutputItems.AddRange(finalizeAction.ProducedItems);
 		}
 	}
 }
diff --git a/Engine/Source/Programs/UnrealBuildTool/System/CppRootPaths.cs b/Engine/Source/Programs/UnrealBuildTool/System/CppRootPaths.cs
--- a/Engine/Source/Programs/UnrealBuildTool/System/CppRootPaths.cs
+++ b/Engine/Source/Programs/UnrealBuildTool/System/CppRootPaths.cs
@@ -374,8 +374,6 @@ namespace UnrealBuildTool
 		/// <param name="name">The new string name</param>
 		public void AddFolderName(CppRootPathFolder folder, string name)
 		{
-			if (name == "Clang")
-				Console.WriteLine();
 			_folderNames[folder] = name;
 		}
 
