UnrealBuildTool: Fix PVS renewal message not printing

#rnx

diff --git a/Engine/Source/Programs/UnrealBuildTool/Platform/Windows/PVSToolChain.cs b/Engine/Source/Programs/UnrealBuildTool/Platform/Windows/PVSToolChain.cs
--- a/Engine/Source/Programs/UnrealBuildTool/Platform/Windows/PVSToolChain.cs
+++ b/Engine/Source/Programs/UnrealBuildTool/Platform/Windows/PVSToolChain.cs
@@ -3,6 +3,7 @@
 using EpicGames.Core;
 using Microsoft.Extensions.Logging;
 using Newtonsoft.Json;
+using OpenTracing.Tag;
 using System;
 using System.Collections.Generic;
 using System.Collections.Immutable;
@@ -454,10 +455,9 @@ namespace UnrealBuildTool
 			ParallelQuery<PVSErrorInfo> allErrors = allLines.Select(GetErrorInfo)
 				.OfType<PVSErrorInfo>();
 
-			ParallelQuery<PVSErrorInfo> renewErrors = allErrors.Where(x => x.Code.Equals("Renew", StringComparison.Ordinal));
-
-			OrderedParallelQuery<PVSErrorInfo> filteredErrors = allErrors.Except(renewErrors)
+			OrderedParallelQuery<PVSErrorInfo> filteredErrors = allErrors
 				.Where(x => x.FalseAlarm != true && x.Level <= _printLevel) // Ignore false alarm warnings, and limit printing by PrintLevel
+				.Where(x => !String.IsNullOrWhiteSpace(x.Positions.FirstOrDefault()?.File)) // Ignore files with no position
 				.OrderBy(x => x.Positions.FirstOrDefault()?.File)
 				.ThenBy(x => x.Positions.FirstOrDefault()?.Lines.FirstOrDefault());
 
@@ -470,9 +470,10 @@ namespace UnrealBuildTool
 				logger.LogWarning(KnownLogEvents.Compiler, "{Path}({LineNumber}): warning {WarningCode}: {WarningMessage}", logValue, lineNumber, errorInfo.Code, errorInfo.Message);
 			}
 
-			if (renewErrors.Any())
+			PVSErrorInfo? renewError = allErrors.FirstOrDefault(x => x.Code.Equals("Renew", StringComparison.Ordinal));
+			if (renewError != null)
 			{
-				logger.LogInformation("PVS-Studio Renewal Notice: {WarningMessage}", renewErrors.First().Message);
+				logger.LogInformation("PVS-Studio Renewal Notice: {WarningMessage}", renewError.Message);
 				logger.LogWarning("Warning: PVS-Studio license will expire soon. See output log for details.");
 			}
 
@@ -517,9 +518,8 @@ namespace UnrealBuildTool
 		}
 
 		// Ignore anything in the IgnoredDirectories folders or ThirdParty if ignored
-		bool IsFileIgnored(FileReference? fileReference) => fileReference == null
-			|| (_ignoreThirdParty && fileReference.FullName.Contains("ThirdParty", StringComparison.OrdinalIgnoreCase))
-			|| (_ignoredDirectories.Any() && _ignoredDirectories.Any(fileReference.IsUnderDirectory));
+		bool IsFileIgnored(FileReference? fileReference) => fileReference != null
+			&& ((_ignoreThirdParty && fileReference.FullName.Contains("ThirdParty", StringComparison.OrdinalIgnoreCase)) || (_ignoredDirectories.Any() && _ignoredDirectories.Any(fileReference.IsUnderDirectory)));
 
 		PVSErrorInfo? GetErrorInfo(string line)
 		{
