UnrealBuildTool: Add support for debugging optimized code with MSVC

diff --git a/Engine/Source/Programs/UnrealBuildTool/Platform/Windows/UEBuildWindows.cs b/Engine/Source/Programs/UnrealBuildTool/Platform/Windows/UEBuildWindows.cs
--- a/Engine/Source/Programs/UnrealBuildTool/Platform/Windows/UEBuildWindows.cs
+++ b/Engine/Source/Programs/UnrealBuildTool/Platform/Windows/UEBuildWindows.cs
@@ -267,6 +267,15 @@ namespace UnrealBuildTool
 		[XmlConfigFile(Category = "WindowsPlatform")]
 		public int InlineFunctionExpansionLevel { get; set; } = 2;
 
+		/// <summary>
+		/// If specificed Visual Studio Dynamic Debugging support will be enabled, unless the compiler or linker is not MSVC or LTCG is enabled
+		/// </summary>
+		[RequiresUniqueBuildEnvironment]
+		[ConfigFile(ConfigHierarchyType.Engine, "/Script/WindowsTargetPlatform.WindowsTargetSettings", "bDynamicDebugging")]
+		[XmlConfigFile(Category = "WindowsPlatform")]
+		[CommandLine("-DynamicDebugging")]
+		public bool bDynamicDebugging { get; set; }
+
 		/// <summary>
 		/// Version of the compiler toolchain to use on Windows platform. A value of "default" will be changed to a specific version at UBT start up.
 		/// </summary>
@@ -906,6 +915,8 @@ namespace UnrealBuildTool
 
 		public int InlineFunctionExpansionLevel => Inner.InlineFunctionExpansionLevel;
 
+		public bool bDynamicDebugging => Inner.bDynamicDebugging;
+
 		public WindowsCompiler Compiler => Inner.Compiler;
 
 		public WindowsCompiler ToolChain => Inner.ToolChain;
@@ -1318,6 +1329,9 @@ namespace UnrealBuildTool
 
 				// Enable extended warnings when analyzing
 				Target.WindowsPlatform.bVCExtendedWarningInfo = true;
+
+				// Disable Dynamic Debugging when analyzing
+				Target.WindowsPlatform.bDynamicDebugging = false;
 			}
 
 			// Disable chaining PCH files
@@ -1347,6 +1361,24 @@ namespace UnrealBuildTool
 				Target.bDisableDebugInfoForGeneratedCode = false;
 			}
 
+			if (Target.WindowsPlatform.bDynamicDebugging)
+			{
+				// Monolithic targets always require increasing the pdb page size with Dynamic Debugging or linking may fail
+				if (Target.LinkType == TargetLinkType.Monolithic && (Target.WindowsPlatform.PdbPageSize ?? 0) < 16384)
+				{
+					Target.WindowsPlatform.PdbPageSize = 16384;
+				}
+
+				// '/DYNAMICDEOPT' combined with '/OPT:ICF' has a degraded debugging experience;
+				if (Target.WindowsPlatform.bMergeIdenticalCOMDATs)
+				{
+					if (Target.WindowsPlatform.bMergeIdenticalCOMDATs)
+					{
+						Target.WindowsPlatform.bMergeIdenticalCOMDATs = false;
+					}
+				}
+			}
+
 			Target.bCompileISPC = true;
 
 			if (OperatingSystem.IsWindows())
@@ -1667,20 +1699,26 @@ namespace UnrealBuildTool
 		/// <returns>True if the string matches the name of a build product, false otherwise</returns>
 		public override bool IsBuildProduct(string FileName, string[] NamePrefixes, string[] NameSuffixes)
 		{
-			return IsBuildProductName(FileName, NamePrefixes, NameSuffixes, ".exe")
-				|| IsBuildProductName(FileName, NamePrefixes, NameSuffixes, ".dll")
-				|| IsBuildProductName(FileName, NamePrefixes, NameSuffixes, ".dll.response")
-				|| IsBuildProductName(FileName, NamePrefixes, NameSuffixes, ".dll.rsp")
-				|| IsBuildProductName(FileName, NamePrefixes, NameSuffixes, ".lib")
-				|| IsBuildProductName(FileName, NamePrefixes, NameSuffixes, ".pdb")
-				|| IsBuildProductName(FileName, NamePrefixes, NameSuffixes, ".full.pdb")
-				|| IsBuildProductName(FileName, NamePrefixes, NameSuffixes, ".exp")
-				|| IsBuildProductName(FileName, NamePrefixes, NameSuffixes, ".obj")
-				|| IsBuildProductName(FileName, NamePrefixes, NameSuffixes, ".map")
-				|| IsBuildProductName(FileName, NamePrefixes, NameSuffixes, ".objpaths")
-				|| IsBuildProductName(FileName, NamePrefixes, NameSuffixes, ".natvis")
-				|| IsBuildProductName(FileName, NamePrefixes, NameSuffixes, ".natstepfilter")
-				|| IsBuildProductName(FileName, NamePrefixes, NameSuffixes, ".natjmc");
+			IEnumerable<string> extensions = [
+				".exe",
+				".dll",
+				".dll.response",
+				".dll.rsp",
+				".lib",
+				".pdb",
+				".full.pdb",
+				".exp",
+				".obj",
+				".map",
+				".objpaths",
+				".natvis",
+				".natstepfilter",
+				".natjmc"
+			];
+
+			return extensions.Any(x => IsBuildProductName(FileName, NamePrefixes, NameSuffixes, x))
+				|| extensions.Select(x => $".alt{x}").Any(x => IsBuildProductName(FileName, NamePrefixes, NameSuffixes, x));
+
 		}
 
 		/// <summary>
diff --git a/Engine/Source/Programs/UnrealBuildTool/Platform/Windows/VCToolChain.cs b/Engine/Source/Programs/UnrealBuildTool/Platform/Windows/VCToolChain.cs
--- a/Engine/Source/Programs/UnrealBuildTool/Platform/Windows/VCToolChain.cs
+++ b/Engine/Source/Programs/UnrealBuildTool/Platform/Windows/VCToolChain.cs
@@ -1,16 +1,15 @@
 // Copyright Epic Games, Inc. All Rights Reserved.
 
+using EpicGames.Core;
+using Microsoft.CodeAnalysis;
+using Microsoft.Extensions.Logging;
 using System;
 using System.Collections.Generic;
 using System.IO;
 using System.Linq;
 using System.Runtime.Versioning;
 using System.Text;
-using EpicGames.Core;
-using Microsoft.CodeAnalysis;
-using Microsoft.Extensions.Logging;
 using UnrealBuildBase;
-using UnrealBuildTool.VerseVMBytecode;
 
 namespace UnrealBuildTool
 {
@@ -59,6 +58,57 @@ namespace UnrealBuildTool
 			return EnvVars.IncludePaths;
 		}
 
+		protected bool IsDynamicDebuggingEnabled => Target.WindowsPlatform.bDynamicDebugging
+			&& Target.WindowsPlatform.Compiler.IsMSVC()
+			&& !(Target.WindowsPlatform.bAllowRadLinker)
+			&& !(Target.WindowsPlatform.Environment != null && Target.WindowsPlatform.Environment.CompilerVersion < new VersionNumber(14, 44, 34918))
+			&& !(Target.bAllowLTCG || Target.bPGOOptimize || Target.bPGOProfile)
+			&& !(Target.Configuration == UnrealTargetConfiguration.Debug || Target.Configuration == UnrealTargetConfiguration.DebugGame)
+			&& !(Target.WindowsPlatform.bEnableAddressSanitizer)
+			&& !(Target.bSupportEditAndContinue);
+
+		protected void DynamicDebuggingInfo(List<string> Lines)
+		{
+			if (Target.WindowsPlatform.bDynamicDebugging)
+			{
+				if (IsDynamicDebuggingEnabled)
+				{
+					Lines.Add("Visual Studio Dynamic Debugging is enabled");
+					return;
+				}
+
+				Lines.Add("Visual Studio Dynamic Debugging has been disabled due to the following:");
+				if (!Target.WindowsPlatform.Compiler.IsMSVC())
+				{
+					Lines.Add($" * Target.{nameof(Target.WindowsPlatform)}.{nameof(Target.WindowsPlatform.Compiler)} '{Target.WindowsPlatform.Compiler}' is not MSVC");
+				}
+				if (Target.WindowsPlatform.bAllowRadLinker)
+				{
+					Lines.Add($" * Target.{nameof(Target.WindowsPlatform)}.{nameof(Target.WindowsPlatform.bAllowRadLinker)} is enabled");
+				}
+				if (Target.WindowsPlatform.Environment != null && Target.WindowsPlatform.Environment.CompilerVersion < new VersionNumber(14, 44, 34918))
+				{
+					Lines.Add($" * MSVC compiler version {Target.WindowsPlatform.Environment.CompilerVersion} older than 14.44.34918");
+				}
+				if (Target.bAllowLTCG || Target.bPGOOptimize || Target.bPGOProfile)
+				{
+					Lines.Add($" * Target.{nameof(Target.bAllowLTCG)}, Target.{nameof(Target.bPGOOptimize)}, or Target.{nameof(Target.bPGOProfile)} is enabled");
+				}
+				if (Target.Configuration == UnrealTargetConfiguration.Debug || Target.Configuration == UnrealTargetConfiguration.DebugGame)
+				{
+					Lines.Add($" * Target.{nameof(Target.Configuration)} '{Target.Configuration}' is Debug or DebugGame. Development, Test, or Shipping configuration is recommended when Dynamic Debugging is requested");
+				}
+				if (Target.WindowsPlatform.bEnableAddressSanitizer)
+				{
+					Lines.Add($" * Target.{nameof(Target.WindowsPlatform)}.{nameof(Target.WindowsPlatform.bEnableAddressSanitizer)} is enabled");
+				}
+				if (Target.bSupportEditAndContinue)
+				{
+					Lines.Add($" * Target.{nameof(Target.bSupportEditAndContinue)} is enabled");
+				}
+			}
+		}
+
 		/// <summary>
 		/// Prepares the environment for building
 		/// </summary>
@@ -106,6 +156,8 @@ namespace UnrealBuildTool
 				Lines.Add($"Using {WindowsPlatform.GetCompilerName(EnvVars.Compiler)} {EnvVars.CompilerVersion} compiler ({EnvVars.CompilerDir}) with {WindowsPlatform.GetCompilerName(EnvVars.ToolChain)} {EnvVars.ToolChainVersion} runtime ({EnvVars.ToolChainDir}) and Windows {EnvVars.WindowsSdkVersion} SDK ({EnvVars.WindowsSdkDir}).");
 			}
 
+			DynamicDebuggingInfo(Lines);
+
 			if (Target.WindowsPlatform.ToolchainVersionWarningLevel == WarningLevel.Warning)
 			{
 				if (!MicrosoftPlatformSDK.IsPreferredVersion(EnvVars.Compiler, EnvVars.CompilerVersion))
@@ -563,6 +615,11 @@ namespace UnrealBuildTool
 				Arguments.Add($"/d2ReducedOptimizeThreshold:{Target.WindowsPlatform.ReducedOptimizeHugeFunctionsThreshold}");
 			}
 
+			if (IsDynamicDebuggingEnabled)
+			{
+				Arguments.Add("/dynamicdeopt");
+			}
+
 			// Separate functions for linker.
 			Arguments.Add("/Gy");
 
@@ -1936,6 +1993,11 @@ namespace UnrealBuildTool
 						}
 					}
 
+					if (IsDynamicDebuggingEnabled)
+					{
+						CompileAction.AdditionalProducedItems.Add(FileItem.GetItemByFileReference(FileReference.Combine(OutputDir, FileName + ".alt.obj")));
+					}
+
 					if (Target.WindowsPlatform.Compiler.IsMSVC() && CompileEnvironment.bWithAssembly && SourceFile.HasExtension(".cpp"))
 					{
 						CompileAction.AssemblyFile = FileItem.GetItemByFileReference(FileReference.Combine(OutputDir, FileName + ".asm"));
@@ -2720,6 +2782,11 @@ namespace UnrealBuildTool
 				Arguments.Add("/time+");
 			}
 
+			if (IsDynamicDebuggingEnabled)
+			{
+				Arguments.Add("/dynamicdeopt");
+			}
+
 			// Rad linker should not be used for import libs
 			if (Target.WindowsPlatform.bAllowRadLinker && !bBuildImportLibraryOnly)
 			{
@@ -2950,6 +3017,15 @@ namespace UnrealBuildTool
 			// Allow the toolchain to adjust/process the link arguments
 			ModifyFinalLinkArguments(LinkEnvironment, Arguments, bBuildImportLibraryOnly);
 
+			if (IsDynamicDebuggingEnabled)
+			{
+				HashSet<string> exts = [".dll", ".exe", ".exp", ".lib", ".map", ".pdb"];
+				List<FileItem> additional = [.. ProducedItems.Where(x => exts.Contains(x.Location.GetExtension()))
+					.Select(x => FileItem.GetItemByFileReference(x.Location.ChangeExtension($".alt{x.Location.GetExtension()}")))
+				];
+				ProducedItems.AddRange(additional);
+			}
+
 			// Create a response file for the linker, unless we're generating IntelliSense data
 			FileReference ResponseFileName = GetResponseFileName(LinkEnvironment, OutputFile);
 			if (!ProjectFileGenerator.bGenerateProjectFiles)
@@ -3421,6 +3497,14 @@ namespace UnrealBuildTool
 					BuildProducts.Add(FileReference.Combine(OutputFilePath.Directory, OutputFilePath.GetFileNameWithoutExtension() + ".objpaths"), BuildProductType.MapFile);
 				}
 			}
+
+			if (IsDynamicDebuggingEnabled)
+			{
+				foreach (KeyValuePair<FileReference, BuildProductType> item in BuildProducts.ToList())
+				{
+					BuildProducts.Add(item.Key.ChangeExtension($".alt{item.Key.GetExtension()}"), item.Value);
+				}
+			}
 		}
 	}
 }
