[UBA]
* Moved code applying target descriptor arguments to uba to UbaExecutor instead of horde coordinator
* Removed logic that disabled obj file compression if building for hot reload.

diff --git a/Engine/Source/Programs/UnrealBuildTool/Executors/UnrealBuildAccelerator/UBAAgentCoordinatorHorde.cs b/Engine/Source/Programs/UnrealBuildTool/Executors/UnrealBuildAccelerator/UBAAgentCoordinatorHorde.cs
--- a/Engine/Source/Programs/UnrealBuildTool/Executors/UnrealBuildAccelerator/UBAAgentCoordinatorHorde.cs
+++ b/Engine/Source/Programs/UnrealBuildTool/Executors/UnrealBuildAccelerator/UBAAgentCoordinatorHorde.cs
@@ -716,13 +716,6 @@ namespace UnrealBuildTool
 			List<UBAAgentCoordinatorHorde> hordeAgentCoordinators = new();
 			foreach (TargetDescriptor targetDescriptor in targetDescriptors)
 			{
-				if (targetDescriptor.HotReloadMode != HotReloadMode.Disabled)
-				{
-					ubaConfig.bStoreObjFilesCompressed = false;
-					ubaConfig.CacheServer = null;
-				}
-				targetDescriptor.AdditionalArguments.ApplyTo(executor);
-				targetDescriptor.AdditionalArguments.ApplyTo(ubaConfig);
 				hordeAgentCoordinators.Add(new UBAAgentCoordinatorHorde(logger, ubaConfig, targetDescriptor.AdditionalArguments, targetDescriptor.ProjectFile?.Directory));
 			}
 			hordeAgentCoordinators.RemoveAll(x => !x.Enabled);
diff --git a/Engine/Source/Programs/UnrealBuildTool/Executors/UnrealBuildAccelerator/UBAExecutor.cs b/Engine/Source/Programs/UnrealBuildTool/Executors/UnrealBuildAccelerator/UBAExecutor.cs
--- a/Engine/Source/Programs/UnrealBuildTool/Executors/UnrealBuildAccelerator/UBAExecutor.cs
+++ b/Engine/Source/Programs/UnrealBuildTool/Executors/UnrealBuildAccelerator/UBAExecutor.cs
@@ -124,6 +124,16 @@ namespace UnrealBuildTool
 			XmlConfig.ApplyTo(UBAConfig);
 			CommandLine.ParseArguments(Environment.GetCommandLineArgs(), this, logger);
 
+			foreach (TargetDescriptor targetDescriptor in targetDescriptors)
+			{
+				if (targetDescriptor.HotReloadMode != HotReloadMode.Disabled)
+				{
+					UBAConfig.CacheServer = null; // Disable cache server when using hot reload path (there would likely never be any cache hits anyway
+				}
+				targetDescriptor.AdditionalArguments.ApplyTo(this);
+				targetDescriptor.AdditionalArguments.ApplyTo(UBAConfig);
+			}
+
 			_agentCoordinators.AddRange(UBAAgentCoordinatorHorde.Init(logger, this, targetDescriptors, UBAConfig, ref _remoteConnectionMode));
 			_threadedLogger = new ThreadedLogger(logger);
 		}
