Add -NoIntellisense argument to skip full intellisense generation for target details and just return binary/launch details.

diff --git a/Engine/Source/Programs/UnrealBuildTool/Modes/QueryMode.cs b/Engine/Source/Programs/UnrealBuildTool/Modes/QueryMode.cs
--- a/Engine/Source/Programs/UnrealBuildTool/Modes/QueryMode.cs
+++ b/Engine/Source/Programs/UnrealBuildTool/Modes/QueryMode.cs
@@ -117,6 +117,9 @@ namespace UnrealBuildTool
 		[CommandLine("-Indented")]
 		public bool bIndented;
 
+		[CommandLine("-NoIntellisense")]
+		public bool bNoIntellisense = false;
+
 		private BuildConfiguration BuildConfiguration = new();
 
 		public override Task<int> ExecuteAsync(CommandLineArguments Arguments, ILogger Logger)
@@ -381,42 +384,6 @@ namespace UnrealBuildTool
 					throw new BuildException("Unable to find launch binary for target");
 				}
 
-				// Create the makefile for the target and export the module information
-				{
-					using ISourceFileWorkingSet WorkingSet = new EmptySourceFileWorkingSet();
-					TargetMakefile Makefile;
-					try
-					{
-						Makefile = await BuildMode.CreateMakefileAsync(BuildConfiguration, TargetDescriptors[0], WorkingSet, Logger);
-					}
-					finally
-					{
-						SourceFileMetadataCache.SaveAll();
-					}
-				}
-
-				TargetIntellisenseInfo CurrentTargetIntellisenseInfo = new TargetIntellisenseInfo();
-
-				// Partially duplicated from UEBuildTarget.Build because we just want to get C++ compile actions without running UHT
-				// or generating link actions / full dependency graph 
-				CppConfiguration CppConfiguration = UEBuildTarget.GetCppConfiguration(CurrentTarget.Configuration);
-				SourceFileMetadataCache MetadataCache = SourceFileMetadataCache.CreateHierarchy(null, Logger);
-				CppCompileEnvironment GlobalCompileEnvironment = new CppCompileEnvironment(CurrentTarget.Platform, CppConfiguration, CurrentTarget.Architectures, MetadataCache);
-				LinkEnvironment GlobalLinkEnvironment = new LinkEnvironment(GlobalCompileEnvironment.Platform, GlobalCompileEnvironment.Configuration, GlobalCompileEnvironment.Architectures);
-
-				UEToolChain TargetToolChain = CurrentTarget.CreateToolchain(CurrentTarget.Platform, Logger);
-				TargetToolChain.SetEnvironmentVariables();
-				CurrentTarget.SetupGlobalEnvironment(TargetToolChain, GlobalCompileEnvironment, GlobalLinkEnvironment);
-
-				if (CurrentTarget.Rules.bUseSharedPCHs)
-				{
-					// Find all the shared PCHs.
-					CurrentTarget.FindSharedPCHs(CurrentTarget.Binaries, GlobalCompileEnvironment, Logger);
-
-					// Create all the shared PCH instances before processing the modules
-					CurrentTarget.CreateSharedPCHInstances(CurrentTarget.Rules, TargetToolChain, CurrentTarget.Binaries, GlobalCompileEnvironment, new NullActionGraphBuilder(Logger), Logger);
-				}
-
 				LaunchSettings CurrentLaunchSettings = new LaunchSettings();
 				CurrentLaunchSettings.Description = $"{CurrentTarget.TargetName} {CurrentTarget.Configuration} {CurrentTarget.Platform}";
 				CurrentLaunchSettings.BinaryPath = LaunchBinary.OutputFilePath.FullName;
@@ -425,66 +392,106 @@ namespace UnrealBuildTool
 					CurrentLaunchSettings.Arguments.Add(CurrentTarget.ProjectFile.FullName);
 				}
 
-				await Parallel.ForEachAsync(CurrentTarget.Binaries, async (Binary, CancellationToken) =>
+				TargetIntellisenseInfo? CurrentTargetIntellisenseInfo = null;
+				if (!bNoIntellisense)
 				{
-					CppCompileEnvironment BinaryCompileEnvironment = Binary.CreateBinaryCompileEnvironment(GlobalCompileEnvironment);
-
-					IEnumerable<UEBuildModuleCPP> Modules = Binary.Modules.OfType<UEBuildModuleCPP>().Where(x => x.Binary == Binary);
-
-					await Parallel.ForEachAsync(Modules, (Module, CancellationToken) =>
+					// Create the makefile for the target and export the module information
 					{
-						if (CurrentTargetIntellisenseInfo.ModuleToCompileSettings.ContainsKey(Module))
-						{
-							return ValueTask.CompletedTask;
-						}
-
-						foreach (DirectoryReference Dir in Module.ModuleDirectories)
+						using ISourceFileWorkingSet WorkingSet = new EmptySourceFileWorkingSet();
+						TargetMakefile Makefile;
+						try
 						{
-							CurrentTargetIntellisenseInfo.DirToModule.TryAdd(Dir, Module);
+							Makefile = await BuildMode.CreateMakefileAsync(BuildConfiguration, TargetDescriptors[0], WorkingSet, Logger);
 						}
-						if (Module.GeneratedCodeDirectory != null)
+						finally
 						{
-							CurrentTargetIntellisenseInfo.DirToModule.TryAdd(Module.GeneratedCodeDirectory, Module);
+							SourceFileMetadataCache.SaveAll();
 						}
+					}
 
-						CppCompileEnvironment ModuleCompileEnvironment = Module.CreateCompileEnvironmentForIntellisense(CurrentTarget.Rules, BinaryCompileEnvironment, Logger);
-						if (ModuleCompileEnvironment.PrecompiledHeaderAction == PrecompiledHeaderAction.Include)
-						{
-							FileItem IncludeHeader = FileItem.GetItemByFileReference(ModuleCompileEnvironment.PrecompiledHeaderIncludeFilename!);
-							ModuleCompileEnvironment.ForceIncludeFiles.Insert(0, IncludeHeader);
-						}
+					CurrentTargetIntellisenseInfo = new TargetIntellisenseInfo();
+					// Partially duplicated from UEBuildTarget.Build because we just want to get C++ compile actions without running UHT
+					// or generating link actions / full dependency graph 
+					CppConfiguration CppConfiguration = UEBuildTarget.GetCppConfiguration(CurrentTarget.Configuration);
+					SourceFileMetadataCache MetadataCache = SourceFileMetadataCache.CreateHierarchy(null, Logger);
+					CppCompileEnvironment GlobalCompileEnvironment = new CppCompileEnvironment(CurrentTarget.Platform, CppConfiguration, CurrentTarget.Architectures, MetadataCache);
+					LinkEnvironment GlobalLinkEnvironment = new LinkEnvironment(GlobalCompileEnvironment.Platform, GlobalCompileEnvironment.Configuration, GlobalCompileEnvironment.Architectures);
+
+					UEToolChain TargetToolChain = CurrentTarget.CreateToolchain(CurrentTarget.Platform, Logger);
+					TargetToolChain.SetEnvironmentVariables();
+					CurrentTarget.SetupGlobalEnvironment(TargetToolChain, GlobalCompileEnvironment, GlobalLinkEnvironment);
+
+					if (CurrentTarget.Rules.bUseSharedPCHs)
+					{
+						// Find all the shared PCHs.
+						CurrentTarget.FindSharedPCHs(CurrentTarget.Binaries, GlobalCompileEnvironment, Logger);
+
+						// Create all the shared PCH instances before processing the modules
+						CurrentTarget.CreateSharedPCHInstances(CurrentTarget.Rules, TargetToolChain, CurrentTarget.Binaries, GlobalCompileEnvironment, new NullActionGraphBuilder(Logger), Logger);
+					}
 
-						// Remove include paths and defintiions from an environment used to get the command line args
-						CppCompileEnvironment ModuleCompileEnvironmentForArgs = new CppCompileEnvironment(ModuleCompileEnvironment);
-						ModuleCompileEnvironmentForArgs.SystemIncludePaths.Clear();
-						ModuleCompileEnvironmentForArgs.UserIncludePaths.Clear();
-						ModuleCompileEnvironmentForArgs.Definitions.Clear();
 
-						TargetIntellisenseInfo.CompileSettings Settings = new TargetIntellisenseInfo.CompileSettings();
-						Settings.IncludePaths.AddRange(ModuleCompileEnvironment.UserIncludePaths.Select(x => x.ToString()));
-						Settings.IncludePaths.AddRange(ModuleCompileEnvironment.SystemIncludePaths.Select(x => x.ToString()));
-						if (TargetToolChain is VCToolChain TargetVCToolChain)
+					await Parallel.ForEachAsync(CurrentTarget.Binaries, async (Binary, CancellationToken) =>
+					{
+						CppCompileEnvironment BinaryCompileEnvironment = Binary.CreateBinaryCompileEnvironment(GlobalCompileEnvironment);
+
+						IEnumerable<UEBuildModuleCPP> Modules = Binary.Modules.OfType<UEBuildModuleCPP>().Where(x => x.Binary == Binary);
+
+						await Parallel.ForEachAsync(Modules, (Module, CancellationToken) =>
 						{
-							Settings.IncludePaths.AddRange(TargetVCToolChain.GetVCIncludePaths().Select(x => x.ToString()));
-						}
-						Settings.Defines = ModuleCompileEnvironment.Definitions;
-						Settings.Standard = ModuleCompileEnvironment.CppStandard.ToString();
-						Settings.ForcedIncludes = ModuleCompileEnvironment.ForceIncludeFiles.Select(x => x.ToString()).ToList();
-						Settings.CompilerPath = TargetToolChain.GetCppCompilerPath()?.ToString();
-						Settings.CompilerArgs = TargetToolChain.GetGlobalCommandLineArgs(ModuleCompileEnvironmentForArgs).ToList();
-						Settings.CompilerAdditionalArgs.Add("c", TargetToolChain.GetCCommandLineArgs(ModuleCompileEnvironmentForArgs).ToList());
-						Settings.CompilerAdditionalArgs.Add("cpp", TargetToolChain.GetCPPCommandLineArgs(ModuleCompileEnvironmentForArgs).ToList());
-						Settings.WindowsSdkVersion = CurrentTarget.Rules.WindowsPlatform.WindowsSdkVersion;
-						CurrentTargetIntellisenseInfo.ModuleToCompileSettings.TryAdd(Module, Settings);
+							if (CurrentTargetIntellisenseInfo.ModuleToCompileSettings.ContainsKey(Module))
+							{
+								return ValueTask.CompletedTask;
+							}
 
-						return ValueTask.CompletedTask;
+							foreach (DirectoryReference Dir in Module.ModuleDirectories)
+							{
+								CurrentTargetIntellisenseInfo.DirToModule.TryAdd(Dir, Module);
+							}
+							if (Module.GeneratedCodeDirectory != null)
+							{
+								CurrentTargetIntellisenseInfo.DirToModule.TryAdd(Module.GeneratedCodeDirectory, Module);
+							}
+
+							CppCompileEnvironment ModuleCompileEnvironment = Module.CreateCompileEnvironmentForIntellisense(CurrentTarget.Rules, BinaryCompileEnvironment, Logger);
+							if (ModuleCompileEnvironment.PrecompiledHeaderAction == PrecompiledHeaderAction.Include)
+							{
+								FileItem IncludeHeader = FileItem.GetItemByFileReference(ModuleCompileEnvironment.PrecompiledHeaderIncludeFilename!);
+								ModuleCompileEnvironment.ForceIncludeFiles.Insert(0, IncludeHeader);
+							}
+
+							// Remove include paths and defintiions from an environment used to get the command line args
+							CppCompileEnvironment ModuleCompileEnvironmentForArgs = new CppCompileEnvironment(ModuleCompileEnvironment);
+							ModuleCompileEnvironmentForArgs.SystemIncludePaths.Clear();
+							ModuleCompileEnvironmentForArgs.UserIncludePaths.Clear();
+							ModuleCompileEnvironmentForArgs.Definitions.Clear();
+
+							TargetIntellisenseInfo.CompileSettings Settings = new TargetIntellisenseInfo.CompileSettings();
+							Settings.IncludePaths.AddRange(ModuleCompileEnvironment.UserIncludePaths.Select(x => x.ToString()));
+							Settings.IncludePaths.AddRange(ModuleCompileEnvironment.SystemIncludePaths.Select(x => x.ToString()));
+							if (TargetToolChain is VCToolChain TargetVCToolChain)
+							{
+								Settings.IncludePaths.AddRange(TargetVCToolChain.GetVCIncludePaths().Select(x => x.ToString()));
+							}
+							Settings.Defines = ModuleCompileEnvironment.Definitions;
+							Settings.Standard = ModuleCompileEnvironment.CppStandard.ToString();
+							Settings.ForcedIncludes = ModuleCompileEnvironment.ForceIncludeFiles.Select(x => x.ToString()).ToList();
+							Settings.CompilerPath = TargetToolChain.GetCppCompilerPath()?.ToString();
+							Settings.CompilerArgs = TargetToolChain.GetGlobalCommandLineArgs(ModuleCompileEnvironmentForArgs).ToList();
+							Settings.CompilerAdditionalArgs.Add("c", TargetToolChain.GetCCommandLineArgs(ModuleCompileEnvironmentForArgs).ToList());
+							Settings.CompilerAdditionalArgs.Add("cpp", TargetToolChain.GetCPPCommandLineArgs(ModuleCompileEnvironmentForArgs).ToList());
+							Settings.WindowsSdkVersion = CurrentTarget.Rules.WindowsPlatform.WindowsSdkVersion;
+							CurrentTargetIntellisenseInfo.ModuleToCompileSettings.TryAdd(Module, Settings);
+
+							return ValueTask.CompletedTask;
+						});
 					});
-				});
+				}
 
 				var Result = new
 				{
-					DirToModule = CurrentTargetIntellisenseInfo.DirToModule.ToImmutableSortedDictionary(x => x.Key.ToString(), x => x.Value.Name),
-					ModuleToCompileSettings = CurrentTargetIntellisenseInfo.ModuleToCompileSettings.ToImmutableSortedDictionary(x => x.Key.Name, x => x.Value),
+					DirToModule = CurrentTargetIntellisenseInfo?.DirToModule.ToImmutableSortedDictionary(x => x.Key.ToString(), x => x.Value.Name),
+					ModuleToCompileSettings = CurrentTargetIntellisenseInfo?.ModuleToCompileSettings.ToImmutableSortedDictionary(x => x.Key.Name, x => x.Value),
 					LaunchSettings = CurrentLaunchSettings,
 				};
 
