ISPC architecture set to x86-x64 for arm64ec.

#rnx

diff --git a/Engine/Source/Programs/UnrealBuildTool/ToolChain/ISPCToolChain.cs b/Engine/Source/Programs/UnrealBuildTool/ToolChain/ISPCToolChain.cs
--- a/Engine/Source/Programs/UnrealBuildTool/ToolChain/ISPCToolChain.cs
+++ b/Engine/Source/Programs/UnrealBuildTool/ToolChain/ISPCToolChain.cs
@@ -46,7 +46,7 @@ namespace UnrealBuildTool
 				UEBuildPlatform.IsPlatformInGroup(Platform, UnrealPlatformGroup.Unix) ||
 				UEBuildPlatform.IsPlatformInGroup(Platform, UnrealPlatformGroup.Apple))
 			{
-				if (Arch.bIsX64)
+				if (Arch.bIsX64 || Arch == UnrealArch.Arm64ec)
 				{
 					ISPCTargets.AddRange(new string[] { "avx512skx-i32x8", "avx2", "avx", "sse4" });
 				}
@@ -123,7 +123,7 @@ namespace UnrealBuildTool
 				UEBuildPlatform.IsPlatformInGroup(Platform, UnrealPlatformGroup.Apple) ||
 				UEBuildPlatform.IsPlatformInGroup(Platform, UnrealPlatformGroup.Android))
 			{
-				if (Arch.bIsX64)
+				if (Arch.bIsX64 || Arch == UnrealArch.Arm64ec)
 				{
 					return "x86-64";
 				}
