Add initial support for iOS App Extensions, starting with the Notification Service Extension.
- add some missing remote notification API methods

diff --git a/Engine/Source/Programs/UnrealBuildTool/Platform/IOS/IOSCodeSigning.cs b/Engine/Source/Programs/UnrealBuildTool/Platform/IOS/IOSCodeSigning.cs
--- a/Engine/Source/Programs/UnrealBuildTool/Platform/IOS/IOSCodeSigning.cs
+++ b/Engine/Source/Programs/UnrealBuildTool/Platform/IOS/IOSCodeSigning.cs
@@ -564,531 +564,6 @@ namespace UnrealBuildTool
 		}
 	}
 
-	/// <summary>
-	/// The XML doc
-	/// </summary>
-	public class PListHelper
-	{
-		/// <summary>
-		/// Helper for info plist
-		/// </summary>
-		public XmlDocument Doc;
-
-		bool bReadOnly = false;
-
-		/// <summary>
-		/// Set read only
-		/// </summary>
-		public void SetReadOnly(bool bNowReadOnly)
-		{
-			bReadOnly = bNowReadOnly;
-		}
-		/// <summary>
-		/// constructor
-		/// </summary>
-		public PListHelper(string Source)
-		{
-			Doc = new XmlDocument();
-			Doc.XmlResolver = null;
-			Doc.LoadXml(Source);
-		}
-
-		/// <summary>
-		/// Create plist from file
-		/// </summary>
-		public static PListHelper CreateFromFile(string Filename)
-		{
-			byte[] RawPList = File.ReadAllBytes(Filename);
-			return new PListHelper(Encoding.UTF8.GetString(RawPList));
-		}
-
-		/// <summary>
-		/// save the plist to file
-		/// </summary>
-		public void SaveToFile(string Filename)
-		{
-			File.WriteAllText(Filename, SaveToString(), Encoding.UTF8);
-		}
-
-		/// <summary>
-		/// Constructor
-		/// </summary>
-		public PListHelper()
-		{
-			string EmptyFileText =
-				"<?xml version=\"1.0\" encoding=\"UTF-8\"?>\n" +
-				"<!DOCTYPE plist PUBLIC \"-//Apple//DTD PLIST 1.0//EN\" \"http://www.apple.com/DTDs/PropertyList-1.0.dtd\">\n" +
-				"<plist version=\"1.0\">\n" +
-				"<dict>\n" +
-				"</dict>\n" +
-				"</plist>\n";
-
-			Doc = new XmlDocument();
-			Doc.XmlResolver = null;
-			Doc.LoadXml(EmptyFileText);
-		}
-
-		/// <summary>
-		/// Convert value to plist format
-		/// </summary>
-		public XmlElement? ConvertValueToPListFormat(object Value)
-		{
-			XmlElement ValueElement;
-
-			if (Value is string str)
-			{
-				ValueElement = Doc.CreateElement("string");
-				ValueElement.InnerText = str;
-			}
-			else if (Value is Dictionary<string, object> dict)
-			{
-				ValueElement = Doc.CreateElement("dict");
-				foreach (KeyValuePair<string, object> KVP in dict)
-				{
-					AddKeyValuePair(ValueElement, KVP.Key, KVP.Value);
-				}
-			}
-			else if (Value is PListHelper)
-			{
-				PListHelper? PList = Value as PListHelper;
-
-				ValueElement = Doc.CreateElement("dict");
-
-				XmlNode? docElement = PList?.Doc?.DocumentElement;
-				XmlNode? SourceDictionaryNode = docElement?.SelectSingleNode("/plist/dict");
-				if (SourceDictionaryNode == null)
-				{
-					throw new InvalidDataException("The PListHelper object contains a document without a dictionary.");
-				}
-
-				foreach (XmlNode TheirChild in SourceDictionaryNode)
-				{
-					ValueElement.AppendChild(Doc.ImportNode(TheirChild, true));
-				}
-			}
-			else if (Value is Array)
-			{
-				if (Value is byte[])
-				{
-					ValueElement = Doc.CreateElement("data");
-					ValueElement.InnerText = (Value is byte[] byteArray) ? Convert.ToBase64String(byteArray) : String.Empty;
-				}
-				else
-				{
-					ValueElement = Doc.CreateElement("array");
-					if (Value is Array array)
-					{
-						foreach (object? A in array)
-						{
-							XmlElement? childNode = ConvertValueToPListFormat(A);
-							if (childNode != null)
-							{
-								ValueElement.AppendChild(childNode);
-							}
-						}
-					}
-				}
-			}
-			else if (Value is IList)
-			{
-				ValueElement = Doc.CreateElement("array");
-				if (Value is IList list && ValueElement != null)
-				{
-					foreach (object? A in list)
-					{
-						XmlElement? child = ConvertValueToPListFormat(A);
-						if (child != null)
-						{
-							ValueElement.AppendChild(child);
-						}
-					}
-				}
-			}
-			else if (Value is bool vbool)
-			{
-				ValueElement = Doc.CreateElement(vbool ? "true" : "false");
-			}
-			else if (Value is double vdouble)
-			{
-				ValueElement = Doc.CreateElement("real");
-				ValueElement.InnerText = vdouble.ToString();
-			}
-			else if (Value is int vint)
-			{
-				ValueElement = Doc.CreateElement("integer");
-				ValueElement.InnerText = vint.ToString();
-			}
-			else
-			{
-				throw new InvalidDataException(String.Format("Object '{0}' is in an unknown type that cannot be converted to PList format", Value));
-			}
-			return ValueElement;
-		}
-
-		/// <summary>
-		/// Add key value pair
-		/// </summary>
-		public void AddKeyValuePair(XmlNode DictRoot, string KeyName, object Value)
-		{
-			if (bReadOnly)
-			{
-				throw new AccessViolationException("PList has been set to read only and may not be modified");
-			}
-
-			XmlElement KeyElement = Doc.CreateElement("key");
-			KeyElement.InnerText = KeyName;
-
-			DictRoot.AppendChild(KeyElement);
-			XmlElement? newChild = ConvertValueToPListFormat(Value);
-			if (newChild != null)
-			{
-				DictRoot.AppendChild(newChild);
-			}
-		}
-
-		/// <summary>
-		/// Add key value pair
-		/// </summary>
-		public void AddKeyValuePair(string KeyName, object? Value)
-		{
-			XmlNode? dictRoot = null;
-
-			if (Doc.DocumentElement != null)
-			{
-				dictRoot = Doc.DocumentElement.SelectSingleNode("/plist/dict");
-			}
-
-			if (dictRoot != null)
-			{
-				XmlNode? childNode = Doc.CreateElement("null");
-				if (Value != null)
-				{
-					childNode = ConvertValueToPListFormat(Value);
-				}
-				else
-				{
-					childNode = Doc.CreateElement("null");
-				}
-				if (childNode != null)
-				{
-					dictRoot.AppendChild(childNode);
-				}
-				if (Value != null)
-				{
-					AddKeyValuePair(dictRoot, KeyName, Value);
-				}
-			}
-		}
-
-		/// <summary>
-		/// Clones a dictionary from an existing .plist into a new one.  Root should point to the dict key in the source plist.
-		/// </summary>
-		public static PListHelper CloneDictionaryRootedAt(XmlNode Root)
-		{
-			// Create a new empty dictionary
-			PListHelper Result = new PListHelper();
-
-			// Copy all of the entries in the source dictionary into the new one
-			XmlNode? NewDictRoot = Result.Doc.DocumentElement?.SelectSingleNode("/plist/dict");
-			if (NewDictRoot != null)
-			{
-				foreach (XmlNode TheirChild in Root)
-				{
-					XmlNode? importedNode = Result.Doc.ImportNode(TheirChild, true);
-					if (importedNode != null)
-					{
-						NewDictRoot.AppendChild(importedNode);
-					}
-				}
-			}
-			return Result;
-		}
-
-		/// <summary>
-		/// Get String
-		/// </summary>
-		public bool GetString(string Key, out string Value)
-		{
-			string PathToValue = String.Format("/plist/dict/key[.='{0}']/following-sibling::string[1]", Key);
-
-			XmlNode? ValueNode = Doc.DocumentElement?.SelectSingleNode(PathToValue);
-			if (ValueNode == null)
-			{
-				Value = "";
-				return false;
-			}
-
-			Value = ValueNode.InnerText ?? "";
-			return true;
-		}
-
-		/// <summary>
-		/// Get date
-		/// </summary>
-		public bool GetDate(string Key, out string Value)
-		{
-			string PathToValue = String.Format("/plist/dict/key[.='{0}']/following-sibling::date[1]", Key);
-
-			XmlNode? ValueNode = Doc.DocumentElement?.SelectSingleNode(PathToValue);
-			if (ValueNode == null)
-			{
-				Value = "";
-				return false;
-			}
-
-			Value = ValueNode.InnerText;
-			return true;
-		}
-
-		/// <summary>
-		/// Get bool
-		/// </summary>
-		public bool GetBool(string Key)
-		{
-			string PathToValue = String.Format("/plist/dict/key[.='{0}']/following-sibling::node()", Key);
-
-			XmlNode? ValueNode = Doc.DocumentElement?.SelectSingleNode(PathToValue);
-			if (ValueNode == null)
-			{
-				return false;
-			}
-
-			return ValueNode.Name == "true";
-		}
-
-		/// <summary>
-		/// Process one node event
-		/// </summary>
-		public delegate void ProcessOneNodeEvent(XmlNode ValueNode);
-
-		/// <summary>
-		/// Process the value for the key
-		/// </summary>
-		public void ProcessValueForKey(string Key, string ExpectedValueType, ProcessOneNodeEvent ValueHandler)
-		{
-			string PathToValue = String.Format("/plist/dict/key[.='{0}']/following-sibling::{1}[1]", Key, ExpectedValueType);
-
-			XmlNode? ValueNode = Doc.DocumentElement?.SelectSingleNode(PathToValue);
-			if (ValueNode != null)
-			{
-				ValueHandler(ValueNode);
-			}
-		}
-
-		/// <summary>
-		/// Merge two plists together.  Whenever both have the same key, the value in the dominant source list wins.
-		/// </summary>
-		public void MergePlistIn(string dominantPlist)
-		{
-			if (bReadOnly)
-			{
-				throw new AccessViolationException("PList has been set to read only and may not be modified");
-			}
-
-			XmlDocument dominant = new XmlDocument();
-			dominant.XmlResolver = null;
-			dominant.LoadXml(dominantPlist);
-
-			XmlNode? dictionaryNode = Doc.DocumentElement?.SelectSingleNode("/plist/dict");
-			if (dictionaryNode == null)
-			{
-				throw new InvalidOperationException("Invalid PList format: missing dictionary node");
-			}
-
-			// Merge any key-value pairs in the strong .plist into the weak .plist
-			XmlNodeList? strongKeys = dominant.DocumentElement?.SelectNodes("/plist/dict/key");
-			if (strongKeys != null)
-			{
-				foreach (XmlNode strongKeyNode in strongKeys)
-				{
-					string strongKey = strongKeyNode.InnerText;
-
-					XmlNode? weakNode = Doc.DocumentElement?.SelectSingleNode($"/plist/dict/key[.='{strongKey}']");
-					if (weakNode == null)
-					{
-						// Doesn't exist in dominant plist, inject key-value pair
-						XmlNode? valueNode = strongKeyNode.NextSibling;
-						if (valueNode != null)
-						{
-							dictionaryNode.AppendChild(Doc.ImportNode(strongKeyNode, true));
-							dictionaryNode.AppendChild(Doc.ImportNode(valueNode, true));
-						}
-					}
-					else
-					{
-						// Remove the existing value node from the weak file
-						XmlNode? existingValueNode = weakNode.NextSibling;
-						if (existingValueNode != null && existingValueNode.Name == "string")
-						{
-							weakNode.ParentNode?.RemoveChild(existingValueNode);
-
-							// Insert a clone of the dominant value node
-							XmlNode? dominantValueNode = strongKeyNode.NextSibling;
-							if (dominantValueNode != null && dominantValueNode.Name == "string")
-							{
-								weakNode.ParentNode?.InsertAfter(Doc.ImportNode(dominantValueNode, true), weakNode);
-							}
-						}
-					}
-				}
-			}
-		}
-
-		/// <summary>
-		/// Returns each of the entries in the value tag of type array for a given key
-		/// If the key is missing, an empty array is returned.
-		/// Only entries of a given type within the array are returned.
-		/// </summary>
-		public List<string> GetArray(string Key, string EntryType)
-		{
-			List<string> Result = new List<string>();
-
-			ProcessValueForKey(Key, "array",
-				delegate (XmlNode ValueNode)
-				{
-					foreach (XmlNode ChildNode in ValueNode.ChildNodes)
-					{
-						if (EntryType == ChildNode.Name)
-						{
-							string Value = ChildNode.InnerText;
-							Result.Add(Value);
-						}
-					}
-				});
-
-			return Result;
-		}
-
-		/// <summary>
-		/// Returns true if the key exists (and has a value) and false otherwise
-		/// </summary>
-		public bool HasKey(string KeyName)
-		{
-			string PathToKey = String.Format("/plist/dict/key[.='{0}']", KeyName);
-
-			XmlNode? KeyNode = Doc.DocumentElement?.SelectSingleNode(PathToKey);
-			return (KeyNode != null);
-		}
-
-		/// <summary>
-		/// Remove the key value
-		/// </summary>
-		public void RemoveKeyValue(string KeyName)
-		{
-			if (bReadOnly)
-			{
-				throw new AccessViolationException("PList has been set to read only and may not be modified");
-			}
-
-			XmlNode? DictionaryNode = Doc.DocumentElement?.SelectSingleNode("/plist/dict");
-
-			string PathToKey = String.Format("/plist/dict/key[.='{0}']", KeyName);
-			XmlNode? KeyNode = Doc.DocumentElement?.SelectSingleNode(PathToKey);
-			if (KeyNode != null && KeyNode.ParentNode != null)
-			{
-				XmlNode? ValueNode = KeyNode.NextSibling;
-				//remove value
-				if (ValueNode != null)
-				{
-					ValueNode.RemoveAll();
-					ValueNode?.ParentNode?.RemoveChild(ValueNode);
-				}
-				//remove key
-				KeyNode.RemoveAll();
-				KeyNode.ParentNode.RemoveChild(KeyNode);
-			}
-		}
-
-		/// <summary>
-		/// Set the value for the key
-		/// </summary>
-		public void SetValueForKey(string KeyName, object Value)
-		{
-			if (bReadOnly)
-			{
-				throw new AccessViolationException("PList has been set to read only and may not be modified");
-			}
-
-			XmlNode? DictionaryNode = Doc.DocumentElement?.SelectSingleNode("/plist/dict");
-
-			string PathToKey = String.Format("/plist/dict/key[.='{0}']", KeyName);
-			XmlNode? KeyNode = Doc.DocumentElement?.SelectSingleNode(PathToKey);
-
-			XmlNode? ValueNode = null;
-			if (KeyNode != null)
-			{
-				ValueNode = KeyNode.NextSibling;
-			}
-
-			if (ValueNode == null)
-			{
-				KeyNode = Doc.CreateNode(XmlNodeType.Element, "key", null);
-				KeyNode.InnerText = KeyName;
-
-				ValueNode = ConvertValueToPListFormat(Value);
-
-				if (DictionaryNode != null && KeyNode != null && ValueNode != null)
-				{
-					DictionaryNode.AppendChild(KeyNode);
-					DictionaryNode.AppendChild(ValueNode);
-				}
-			}
-			else
-			{
-				// Remove the existing value and create a new one
-				ValueNode.ParentNode?.RemoveChild(ValueNode);
-				ValueNode = ConvertValueToPListFormat(Value);
-
-				// Insert the value after the key
-				KeyNode?.ParentNode?.InsertAfter(ValueNode!, KeyNode);
-
-			}
-		}
-
-		/// <summary>
-		/// Set the string
-		/// </summary>
-		public void SetString(string Key, string Value)
-		{
-			SetValueForKey(Key, Value);
-		}
-
-		/// <summary>
-		/// Save the string
-		/// </summary>
-		public string SaveToString()
-		{
-			// Convert the XML back to text in the same style as the original .plist
-			StringBuilder TextOut = new StringBuilder();
-
-			// Work around the fact it outputs the wrong encoding by default (and set some other settings to get something similar to the input file)
-			TextOut.Append("<?xml version=\"1.0\" encoding=\"UTF-8\"?>\n");
-			XmlWriterSettings Settings = new XmlWriterSettings();
-			Settings.Indent = true;
-			Settings.IndentChars = "\t";
-			Settings.NewLineChars = "\n";
-			Settings.NewLineHandling = NewLineHandling.Replace;
-			Settings.OmitXmlDeclaration = true;
-			Settings.Encoding = new UTF8Encoding(false);
-
-			// Work around the fact that it embeds an empty declaration list to the document type which codesign dislikes...
-			// Replacing InternalSubset with null if it's empty.  The property is readonly, so we have to reconstruct it entirely
-			Doc.ReplaceChild(Doc.CreateDocumentType(
-				Doc.DocumentType!.Name,
-				Doc.DocumentType!.PublicId,
-				Doc.DocumentType!.SystemId,
-				String.IsNullOrEmpty(Doc.DocumentType!.InternalSubset) ? null : Doc.DocumentType!.InternalSubset),
-				Doc.DocumentType!);
-
-			XmlWriter Writer = XmlWriter.Create(TextOut, Settings);
-
-			Doc.Save(Writer);
-
-			// Remove the space from any standalone XML elements because the iOS parser does not handle them
-			return Regex.Replace(TextOut.ToString(), @"<(?<tag>\S+) />", "<${tag}/>");
-		}
-	}
-
 	internal class FileOperations
 	{
 		/// <summary>
diff --git a/Engine/Source/Programs/UnrealBuildTool/Platform/IOS/PListHelper.cs b/Engine/Source/Programs/UnrealBuildTool/Platform/IOS/PListHelper.cs
new file mode 100644
--- /dev/null
+++ b/Engine/Source/Programs/UnrealBuildTool/Platform/IOS/PListHelper.cs
@@ -0,0 +1,543 @@
+// Copyright Epic Games, Inc. All Rights Reserved.
+
+using System;
+using System.Collections;
+using System.Collections.Generic;
+using System.IO;
+using System.Text;
+using System.Text.RegularExpressions;
+using System.Xml;
+
+namespace UnrealBuildTool
+{
+	/// <summary>
+	/// A helper class to handler modifying text based PList files 
+	/// </summary>
+	public class PListHelper
+	{
+		/// <summary>
+		/// Helper for info plist
+		/// </summary>
+		public XmlDocument Doc;
+
+		bool bReadOnly = false;
+
+		/// <summary>
+		/// Set read only
+		/// </summary>
+		public void SetReadOnly(bool bNowReadOnly)
+		{
+			bReadOnly = bNowReadOnly;
+		}
+
+		/// <summary>
+		/// constructor
+		/// </summary>
+		public PListHelper(string Source)
+		{
+			Doc = new XmlDocument();
+			Doc.XmlResolver = null;
+			Doc.LoadXml(Source);
+		}
+
+		/// <summary>
+		/// Create plist from file
+		/// </summary>
+		public static PListHelper CreateFromFile(string Filename)
+		{
+			byte[] RawPList = File.ReadAllBytes(Filename);
+			return new PListHelper(Encoding.UTF8.GetString(RawPList));
+		}
+
+		/// <summary>
+		/// save the plist to file
+		/// </summary>
+		public void SaveToFile(string Filename)
+		{
+			File.WriteAllText(Filename, SaveToString(), Encoding.UTF8);
+		}
+
+		/// <summary>
+		/// Constructor
+		/// </summary>
+		public PListHelper()
+		{
+			string EmptyFileText =
+				"<?xml version=\"1.0\" encoding=\"UTF-8\"?>\n" +
+				"<!DOCTYPE plist PUBLIC \"-//Apple//DTD PLIST 1.0//EN\" \"http://www.apple.com/DTDs/PropertyList-1.0.dtd\">\n" +
+				"<plist version=\"1.0\">\n" +
+				"<dict>\n" +
+				"</dict>\n" +
+				"</plist>\n";
+
+			Doc = new XmlDocument();
+			Doc.XmlResolver = null;
+			Doc.LoadXml(EmptyFileText);
+		}
+
+		/// <summary>
+		/// Convert value to plist format
+		/// </summary>
+		public XmlElement? ConvertValueToPListFormat(object Value)
+		{
+			XmlElement ValueElement;
+
+			if (Value is string str)
+			{
+				ValueElement = Doc.CreateElement("string");
+				ValueElement.InnerText = str;
+			}
+			else if (Value is Dictionary<string, object> dict)
+			{
+				ValueElement = Doc.CreateElement("dict");
+				foreach (KeyValuePair<string, object> KVP in dict)
+				{
+					AddKeyValuePair(ValueElement, KVP.Key, KVP.Value);
+				}
+			}
+			else if (Value is PListHelper)
+			{
+				PListHelper? PList = Value as PListHelper;
+
+				ValueElement = Doc.CreateElement("dict");
+
+				XmlNode? docElement = PList?.Doc?.DocumentElement;
+				XmlNode? SourceDictionaryNode = docElement?.SelectSingleNode("/plist/dict");
+				if (SourceDictionaryNode == null)
+				{
+					throw new InvalidDataException("The PListHelper object contains a document without a dictionary.");
+				}
+
+				foreach (XmlNode TheirChild in SourceDictionaryNode)
+				{
+					ValueElement.AppendChild(Doc.ImportNode(TheirChild, true));
+				}
+			}
+			else if (Value is Array)
+			{
+				if (Value is byte[])
+				{
+					ValueElement = Doc.CreateElement("data");
+					ValueElement.InnerText = (Value is byte[] byteArray) ? Convert.ToBase64String(byteArray) : String.Empty;
+				}
+				else
+				{
+					ValueElement = Doc.CreateElement("array");
+					if (Value is Array array)
+					{
+						foreach (object? A in array)
+						{
+							XmlElement? childNode = ConvertValueToPListFormat(A);
+							if (childNode != null)
+							{
+								ValueElement.AppendChild(childNode);
+							}
+						}
+					}
+				}
+			}
+			else if (Value is IList)
+			{
+				ValueElement = Doc.CreateElement("array");
+				if (Value is IList list && ValueElement != null)
+				{
+					foreach (object? A in list)
+					{
+						XmlElement? child = ConvertValueToPListFormat(A);
+						if (child != null)
+						{
+							ValueElement.AppendChild(child);
+						}
+					}
+				}
+			}
+			else if (Value is bool vbool)
+			{
+				ValueElement = Doc.CreateElement(vbool ? "true" : "false");
+			}
+			else if (Value is double vdouble)
+			{
+				ValueElement = Doc.CreateElement("real");
+				ValueElement.InnerText = vdouble.ToString();
+			}
+			else if (Value is int vint)
+			{
+				ValueElement = Doc.CreateElement("integer");
+				ValueElement.InnerText = vint.ToString();
+			}
+			else
+			{
+				throw new InvalidDataException(String.Format("Object '{0}' is in an unknown type that cannot be converted to PList format", Value));
+			}
+
+			return ValueElement;
+		}
+
+		/// <summary>
+		/// Add key value pair
+		/// </summary>
+		public void AddKeyValuePair(XmlNode DictRoot, string KeyName, object Value)
+		{
+			if (bReadOnly)
+			{
+				throw new AccessViolationException("PList has been set to read only and may not be modified");
+			}
+
+			XmlElement KeyElement = Doc.CreateElement("key");
+			KeyElement.InnerText = KeyName;
+
+			DictRoot.AppendChild(KeyElement);
+			XmlElement? newChild = ConvertValueToPListFormat(Value);
+			if (newChild != null)
+			{
+				DictRoot.AppendChild(newChild);
+			}
+		}
+
+		/// <summary>
+		/// Add key value pair
+		/// </summary>
+		public void AddKeyValuePair(string KeyName, object? Value)
+		{
+			XmlNode? dictRoot = null;
+
+			if (Doc.DocumentElement != null)
+			{
+				dictRoot = Doc.DocumentElement.SelectSingleNode("/plist/dict");
+			}
+
+			if (dictRoot != null)
+			{
+				XmlNode? childNode = Doc.CreateElement("null");
+				if (Value != null)
+				{
+					childNode = ConvertValueToPListFormat(Value);
+				}
+				else
+				{
+					childNode = Doc.CreateElement("null");
+				}
+
+				if (childNode != null)
+				{
+					dictRoot.AppendChild(childNode);
+				}
+
+				if (Value != null)
+				{
+					AddKeyValuePair(dictRoot, KeyName, Value);
+				}
+			}
+		}
+
+		/// <summary>
+		/// Clones a dictionary from an existing .plist into a new one.  Root should point to the dict key in the source plist.
+		/// </summary>
+		public static PListHelper CloneDictionaryRootedAt(XmlNode Root)
+		{
+			// Create a new empty dictionary
+			PListHelper Result = new PListHelper();
+
+			// Copy all of the entries in the source dictionary into the new one
+			XmlNode? NewDictRoot = Result.Doc.DocumentElement?.SelectSingleNode("/plist/dict");
+			if (NewDictRoot != null)
+			{
+				foreach (XmlNode TheirChild in Root)
+				{
+					XmlNode? importedNode = Result.Doc.ImportNode(TheirChild, true);
+					if (importedNode != null)
+					{
+						NewDictRoot.AppendChild(importedNode);
+					}
+				}
+			}
+
+			return Result;
+		}
+
+		/// <summary>
+		/// Get String
+		/// </summary>
+		public bool GetString(string Key, out string Value)
+		{
+			string PathToValue = String.Format("/plist/dict/key[.='{0}']/following-sibling::string[1]", Key);
+
+			XmlNode? ValueNode = Doc.DocumentElement?.SelectSingleNode(PathToValue);
+			if (ValueNode == null)
+			{
+				Value = "";
+				return false;
+			}
+
+			Value = ValueNode.InnerText ?? "";
+			return true;
+		}
+
+		/// <summary>
+		/// Get date
+		/// </summary>
+		public bool GetDate(string Key, out string Value)
+		{
+			string PathToValue = String.Format("/plist/dict/key[.='{0}']/following-sibling::date[1]", Key);
+
+			XmlNode? ValueNode = Doc.DocumentElement?.SelectSingleNode(PathToValue);
+			if (ValueNode == null)
+			{
+				Value = "";
+				return false;
+			}
+
+			Value = ValueNode.InnerText;
+			return true;
+		}
+
+		/// <summary>
+		/// Get bool
+		/// </summary>
+		public bool GetBool(string Key)
+		{
+			string PathToValue = String.Format("/plist/dict/key[.='{0}']/following-sibling::node()", Key);
+
+			XmlNode? ValueNode = Doc.DocumentElement?.SelectSingleNode(PathToValue);
+			if (ValueNode == null)
+			{
+				return false;
+			}
+
+			return ValueNode.Name == "true";
+		}
+
+		/// <summary>
+		/// Process one node event
+		/// </summary>
+		public delegate void ProcessOneNodeEvent(XmlNode ValueNode);
+
+		/// <summary>
+		/// Process the value for the key
+		/// </summary>
+		public void ProcessValueForKey(string Key, string ExpectedValueType, ProcessOneNodeEvent ValueHandler)
+		{
+			string PathToValue = String.Format("/plist/dict/key[.='{0}']/following-sibling::{1}[1]", Key, ExpectedValueType);
+
+			XmlNode? ValueNode = Doc.DocumentElement?.SelectSingleNode(PathToValue);
+			if (ValueNode != null)
+			{
+				ValueHandler(ValueNode);
+			}
+		}
+
+		/// <summary>
+		/// Merge two plists together.  Whenever both have the same key, the value in the dominant source list wins.
+		/// </summary>
+		public void MergePlistIn(string dominantPlist)
+		{
+			if (bReadOnly)
+			{
+				throw new AccessViolationException("PList has been set to read only and may not be modified");
+			}
+
+			XmlDocument dominant = new XmlDocument();
+			dominant.XmlResolver = null;
+			dominant.LoadXml(dominantPlist);
+
+			XmlNode? dictionaryNode = Doc.DocumentElement?.SelectSingleNode("/plist/dict");
+			if (dictionaryNode == null)
+			{
+				throw new InvalidOperationException("Invalid PList format: missing dictionary node");
+			}
+
+			// Merge any key-value pairs in the strong .plist into the weak .plist
+			XmlNodeList? strongKeys = dominant.DocumentElement?.SelectNodes("/plist/dict/key");
+			if (strongKeys != null)
+			{
+				foreach (XmlNode strongKeyNode in strongKeys)
+				{
+					string strongKey = strongKeyNode.InnerText;
+
+					XmlNode? weakNode = Doc.DocumentElement?.SelectSingleNode($"/plist/dict/key[.='{strongKey}']");
+					if (weakNode == null)
+					{
+						// Doesn't exist in dominant plist, inject key-value pair
+						XmlNode? valueNode = strongKeyNode.NextSibling;
+						if (valueNode != null)
+						{
+							dictionaryNode.AppendChild(Doc.ImportNode(strongKeyNode, true));
+							dictionaryNode.AppendChild(Doc.ImportNode(valueNode, true));
+						}
+					}
+					else
+					{
+						// Remove the existing value node from the weak file
+						XmlNode? existingValueNode = weakNode.NextSibling;
+						if (existingValueNode != null && existingValueNode.Name == "string")
+						{
+							weakNode.ParentNode?.RemoveChild(existingValueNode);
+
+							// Insert a clone of the dominant value node
+							XmlNode? dominantValueNode = strongKeyNode.NextSibling;
+							if (dominantValueNode != null && dominantValueNode.Name == "string")
+							{
+								weakNode.ParentNode?.InsertAfter(Doc.ImportNode(dominantValueNode, true), weakNode);
+							}
+						}
+					}
+				}
+			}
+		}
+
+		/// <summary>
+		/// Returns each of the entries in the value tag of type array for a given key
+		/// If the key is missing, an empty array is returned.
+		/// Only entries of a given type within the array are returned.
+		/// </summary>
+		public List<string> GetArray(string Key, string EntryType)
+		{
+			List<string> Result = new List<string>();
+
+			ProcessValueForKey(Key, "array",
+				delegate(XmlNode ValueNode)
+				{
+					foreach (XmlNode ChildNode in ValueNode.ChildNodes)
+					{
+						if (EntryType == ChildNode.Name)
+						{
+							string Value = ChildNode.InnerText;
+							Result.Add(Value);
+						}
+					}
+				});
+
+			return Result;
+		}
+
+		/// <summary>
+		/// Returns true if the key exists (and has a value) and false otherwise
+		/// </summary>
+		public bool HasKey(string KeyName)
+		{
+			string PathToKey = String.Format("/plist/dict/key[.='{0}']", KeyName);
+
+			XmlNode? KeyNode = Doc.DocumentElement?.SelectSingleNode(PathToKey);
+			return (KeyNode != null);
+		}
+
+		/// <summary>
+		/// Remove the key value
+		/// </summary>
+		public void RemoveKeyValue(string KeyName)
+		{
+			if (bReadOnly)
+			{
+				throw new AccessViolationException("PList has been set to read only and may not be modified");
+			}
+
+			XmlNode? DictionaryNode = Doc.DocumentElement?.SelectSingleNode("/plist/dict");
+
+			string PathToKey = String.Format("/plist/dict/key[.='{0}']", KeyName);
+			XmlNode? KeyNode = Doc.DocumentElement?.SelectSingleNode(PathToKey);
+			if (KeyNode != null && KeyNode.ParentNode != null)
+			{
+				XmlNode? ValueNode = KeyNode.NextSibling;
+				//remove value
+				if (ValueNode != null)
+				{
+					ValueNode.RemoveAll();
+					ValueNode?.ParentNode?.RemoveChild(ValueNode);
+				}
+
+				//remove key
+				KeyNode.RemoveAll();
+				KeyNode.ParentNode.RemoveChild(KeyNode);
+			}
+		}
+
+		/// <summary>
+		/// Set the value for the key
+		/// </summary>
+		public void SetValueForKey(string KeyName, object Value)
+		{
+			if (bReadOnly)
+			{
+				throw new AccessViolationException("PList has been set to read only and may not be modified");
+			}
+
+			XmlNode? DictionaryNode = Doc.DocumentElement?.SelectSingleNode("/plist/dict");
+
+			string PathToKey = String.Format("/plist/dict/key[.='{0}']", KeyName);
+			XmlNode? KeyNode = Doc.DocumentElement?.SelectSingleNode(PathToKey);
+
+			XmlNode? ValueNode = null;
+			if (KeyNode != null)
+			{
+				ValueNode = KeyNode.NextSibling;
+			}
+
+			if (ValueNode == null)
+			{
+				KeyNode = Doc.CreateNode(XmlNodeType.Element, "key", null);
+				KeyNode.InnerText = KeyName;
+
+				ValueNode = ConvertValueToPListFormat(Value);
+
+				if (DictionaryNode != null && KeyNode != null && ValueNode != null)
+				{
+					DictionaryNode.AppendChild(KeyNode);
+					DictionaryNode.AppendChild(ValueNode);
+				}
+			}
+			else
+			{
+				// Remove the existing value and create a new one
+				ValueNode.ParentNode?.RemoveChild(ValueNode);
+				ValueNode = ConvertValueToPListFormat(Value);
+
+				// Insert the value after the key
+				KeyNode?.ParentNode?.InsertAfter(ValueNode!, KeyNode);
+
+			}
+		}
+
+		/// <summary>
+		/// Set the string
+		/// </summary>
+		public void SetString(string Key, string Value)
+		{
+			SetValueForKey(Key, Value);
+		}
+
+		/// <summary>
+		/// Save the string
+		/// </summary>
+		public string SaveToString()
+		{
+			// Convert the XML back to text in the same style as the original .plist
+			StringBuilder TextOut = new StringBuilder();
+
+			// Work around the fact it outputs the wrong encoding by default (and set some other settings to get something similar to the input file)
+			TextOut.Append("<?xml version=\"1.0\" encoding=\"UTF-8\"?>\n");
+			XmlWriterSettings Settings = new XmlWriterSettings();
+			Settings.Indent = true;
+			Settings.IndentChars = "\t";
+			Settings.NewLineChars = "\n";
+			Settings.NewLineHandling = NewLineHandling.Replace;
+			Settings.OmitXmlDeclaration = true;
+			Settings.Encoding = new UTF8Encoding(false);
+
+			// Work around the fact that it embeds an empty declaration list to the document type which codesign dislikes...
+			// Replacing InternalSubset with null if it's empty.  The property is readonly, so we have to reconstruct it entirely
+			Doc.ReplaceChild(Doc.CreateDocumentType(
+					Doc.DocumentType!.Name,
+					Doc.DocumentType!.PublicId,
+					Doc.DocumentType!.SystemId,
+					String.IsNullOrEmpty(Doc.DocumentType!.InternalSubset) ? null : Doc.DocumentType!.InternalSubset),
+				Doc.DocumentType!);
+
+			XmlWriter Writer = XmlWriter.Create(TextOut, Settings);
+
+			Doc.Save(Writer);
+
+			// Remove the space from any standalone XML elements because the iOS parser does not handle them
+			return Regex.Replace(TextOut.ToString(), @"<(?<tag>\S+) />", "<${tag}/>");
+		}
+	}
+}
\ No newline at end of file
diff --git a/Engine/Source/Programs/UnrealBuildTool/ProjectFiles/Xcode/XcodeProject.File.cs b/Engine/Source/Programs/UnrealBuildTool/ProjectFiles/Xcode/XcodeProject.File.cs
--- a/Engine/Source/Programs/UnrealBuildTool/ProjectFiles/Xcode/XcodeProject.File.cs
+++ b/Engine/Source/Programs/UnrealBuildTool/ProjectFiles/Xcode/XcodeProject.File.cs
@@ -259,8 +259,9 @@ namespace UnrealBuildTool.XcodeProjectXcconfig
 					if (BuildFileOffset != -1)
 					{
 						// first find the build.cs file 
-						// look if it contains any of the exluded names, and is so, don't include it
-						if (!ExcludedFolders.Any(x => File.Reference.ContainsName(x, BuildFileOffset)))
+						// look if it contains any of the excluded names, and is so, don't include it
+						bool isExtension = File.Reference.FullName.Contains("Extensions");
+						if (isExtension || !ExcludedFolders.Any(x => File.Reference.ContainsName(x, BuildFileOffset)))
 						{
 							// will fill in the response file later
 							BuildableFilesToResponseFile[File] = new();
@@ -389,11 +390,21 @@ namespace UnrealBuildTool.XcodeProjectXcconfig
 				string RelativePath = Utils.CleanDirectorySeparators(File.Reference.MakeRelativeTo(ProjectDirectory), '/');
 				string SourceTree = "SOURCE_ROOT";
 
+				if (FileName.Contains("UserNotifications.framework"))
+				{
+					SourceTree = "DEVELOPER_DIR";
+					RelativePath = File.Reference.FullName;
+				}
 				Content.WriteLine($"\t\t{File.FileRefGuid} = {{isa = PBXFileReference; {FileKey} = {FileType}; name = \"{FileName}\"; path = \"{RelativePath}\"; sourceTree = {SourceTree}; }};");
 			}
 			foreach (ManualFileReference File in ManualFiles)
 			{
 				string FileName = Path.GetFileName(File.RelativePath);
+				if (FileName.Contains("UserNotifications.framework"))
+				{
+					File.SourceTree = "PLATFORM_DEVELOPER_SDK_DIR";
+					File.RelativePath = "iPhoneOS.sdk/System/Library/Frameworks/UserNotifications.framework";
+				}
 				Content.WriteLine($"\t\t{File.Guid} = {{isa = PBXFileReference; {File.FileKey} = {File.FileType}; name = \"{FileName}\"; path = \"{File.RelativePath}\"; sourceTree = {File.SourceTree}; }};");
 			}
 			Content.WriteLine("/* End PBXFileReference section */");
diff --git a/Engine/Source/Programs/UnrealBuildTool/ProjectFiles/Xcode/XcodeProject.cs b/Engine/Source/Programs/UnrealBuildTool/ProjectFiles/Xcode/XcodeProject.cs
--- a/Engine/Source/Programs/UnrealBuildTool/ProjectFiles/Xcode/XcodeProject.cs
+++ b/Engine/Source/Programs/UnrealBuildTool/ProjectFiles/Xcode/XcodeProject.cs
@@ -1001,6 +1001,29 @@
 		}
 	}
 
+	class XcodeCopyAppExtFilesBuildPhase : XcodeBuildPhase
+	{
+		private XcodeFileCollection FileCollection;
+
+		public XcodeCopyAppExtFilesBuildPhase(XcodeFileCollection FileCollection)
+			: base("Embed Foundation Extensions", "PBXCopyFilesBuildPhase")
+		{
+			this.FileCollection = FileCollection;
+
+			// MiscItems.Add($"dstPath = \"${{EXTENSIONS_FOLDER_PATH}}\";");
+			MiscItems.Add($"dstPath = \"\";");
+			MiscItems.Add($"dstSubfolderSpec = 13;");
+			MiscItems.Add($"name = \"{Name}\";");
+		}
+
+		public void AddExtension(DirectoryReference Extension, string FileRefGuid)
+		{
+			XcodeSourceFile ExtensionSource = new XcodeSourceFile(new FileReference(Extension.FullName), null, FileRefGuid);
+			FileCollection.ProcessFile(ExtensionSource, true, true, "", "settings = {ATTRIBUTES = (RemoveHeadersOnCopy, ); };");
+			FileItems.Add(ExtensionSource);
+		}
+	}
+
 	class XcodeShellScriptBuildPhase : XcodeBuildPhase
 	{
 		public XcodeShellScriptBuildPhase(string Name, IEnumerable<string> ScriptLines, IEnumerable<string> Inputs, IEnumerable<string> Outputs, bool bInstallOnly = false)
@@ -1037,7 +1060,7 @@
 	{
 		public string Guid = XcodeProjectFileGenerator.MakeXcodeGuid();
 		public UnrealBuildConfig Info;
-		// Because we don't make project-wide .xcconfig files, we need to specify at the projet level that all of the platforms are supported,
+		// Because we don't make project-wide .xcconfig files, we need to specify at the project level that all of the platforms are supported,
 		// so that the _Build targets will have any platform as a supported platform, otherwise it can only compile for Mac (default Xcode platform)
 		private bool bIncludeAllPlatforms;
 
@@ -1140,6 +1163,7 @@
 		public enum Type
 		{
 			Run_App,
+			App_Ext,
 			Run_Tool,
 			Build,
 			Index,
@@ -1187,6 +1211,11 @@
 					TargetTypeName = "PBXNativeTarget";
 					ConfigName = "_Run";
 					break;
+				case Type.App_Ext:
+					ProductType = "com.apple.product-type.app-extension";
+					TargetTypeName = "PBXNativeTarget";
+					ConfigName = "_Run";
+					break;
 				case Type.Run_Tool:
 					ProductType = "com.apple.product-type.tool";
 					TargetTypeName = "PBXNativeTarget";
@@ -1220,7 +1249,7 @@
 
 		public override void Write(StringBuilder Content)
 		{
-			Content.WriteLine($"/* Begin {TargetType} section */");
+			Content.WriteLine($"/* Begin PBXNativeTarget section ({TargetType}) */");
 
 			Content.WriteLine(2, $"{Guid} /* {Name} */ = {{");
 			Content.WriteLine(3, $"isa = {TargetTypeName};");
@@ -1229,7 +1258,7 @@
 
 			if (TargetType == Type.Build)
 			{
-				// get paths to Unreal bits to be able ro tun UBT
+				// get paths to Unreal bits to be able to run UBT
 				string UProjectParam = GameProject == null ? "" : $"{GameProject.FullName.Replace(" ", "\\ ")}";
 				string UEDir = XcodeFileCollection.ConvertPath(Path.GetFullPath(Directory.GetCurrentDirectory() + "../../.."));
 				string BuildToolPath = UEDir + "/Engine/Build/BatchFiles/Mac/XcodeBuild.sh";
@@ -1257,7 +1286,7 @@
 			WriteExtraTargetProperties(Content);
 			Content.WriteLine(2, "};");
 
-			Content.WriteLine($"/* End {TargetType} section */");
+			Content.WriteLine($"/* End PBXNativeTarget section ({TargetType})*/");
 		}
 
 		/// <summary>
@@ -1267,6 +1296,49 @@
 		{
 			// nothing by default
 		}
+		
+		private static Dictionary<string, string> PlistFileMap = new(StringComparer.OrdinalIgnoreCase);
+		protected string GetPlistSigningName(string ProvisionSetting, UnrealData UnrealData, ILogger Logger)
+		{
+			string? SigningName;
+			lock (PlistFileMap)
+			{
+				if (!PlistFileMap.TryGetValue(ProvisionSetting, out SigningName))
+				{
+					FileReference ProfileFile = AppleExports.ConvertFilePath(UnrealData.UProjectFileLocation?.Directory, ProvisionSetting);
+
+					// get mobile provision UUID, either directly or from the probision
+					if (ProfileFile.GetExtension() == ".mobileprovision")
+					{
+						// security will read the provision and dump out the text (plist) bits of the profile, and plutil will extract the UUID to stdout
+						string UUID = Utils.RunLocalProcessAndReturnStdOut("/bin/sh", $"-c 'security cms -D -i {ProfileFile.FullName.Replace(" ", "\\ ")} | plutil -extract UUID raw -'");
+
+						// make sure it's living in user's library
+						// (note: i couldn't find a way to install it without opening Xcode, so copying it works well enough)
+						DirectoryReference UserDir = new DirectoryReference(Environment.GetFolderPath(Environment.SpecialFolder.UserProfile));
+						FileReference InstalledProfileFile = FileReference.Combine(UserDir, "Library", "MobileDevice", "Provisioning Profiles", $"{UUID}.mobileprovision");
+
+						// copy if needed
+						if (!FileReference.Exists(InstalledProfileFile))
+						{
+							DirectoryReference.CreateDirectory(InstalledProfileFile.Directory);
+							FileReference.Copy(ProfileFile, InstalledProfileFile);
+							Logger.LogInformation("Copying project's provision '{SourceProvision}' to your libary: '{TargetProvision}'", ProfileFile, InstalledProfileFile);
+						}
+
+						SigningName = UUID;
+					}
+					else
+					{
+						SigningName = ProvisionSetting;
+					}
+
+					PlistFileMap[ProvisionSetting] = SigningName;
+				}
+			}
+
+			return SigningName!;
+		}
 	}
 
 	class XcodeRunTarget : XcodeTarget
@@ -1306,6 +1378,63 @@
 				AddDependency(BuildTarget, Project);
 			}
 
+			// Check if there's any app extensions for iOS 
+			if (Platform == UnrealTargetPlatform.IOS)
+			{
+				if (Project.UnrealData.UProjectFileLocation != null)
+				{
+					DirectoryReference ProjectDirectory = Project.UnrealData.ProductDirectory;
+					string ProjectExtensionsDir = Path.Combine(ProjectDirectory.FullName, "Build/IOS/Extensions");
+					if (Directory.Exists(ProjectExtensionsDir))
+					{
+						foreach (DirectoryInfo DI in new System.IO.DirectoryInfo(ProjectExtensionsDir).EnumerateDirectories())
+						{
+							Logger.LogInformation("Creating App Extension Target for '{Name}' in '{FullName}'", DI.Name, DI.FullName);
+
+							// assume each Extension in here will create a resulting Extension.appex
+							string Extension = DI.Name + ".appex";
+							string ExtensionGuid = XcodeProjectFileGenerator.MakeXcodeGuid();
+
+							// is it an extension we support? Grab the Plist file for it and check
+							// Also, only some AppExtensions need to embed
+							String AppExtPlist = DI.FullName + "/Info.plist";
+							PListHelper plistDoc = PListHelper.CreateFromFile(AppExtPlist);
+							
+							XcodeCopyAppExtFilesBuildPhase EmbedAppExt = new XcodeCopyAppExtFilesBuildPhase(Project.FileCollection);
+							
+							// Get the source files
+							// XcodeFileCollection AppExtFiles = new XcodeFileCollection();
+							foreach (FileSystemInfo FSI in DI.EnumerateFileSystemInfos())
+							{
+								if (FSI.Name.StartsWith(".") || FSI.Name.EndsWith(".plist"))
+								{
+									continue;
+								}
+
+								Logger.LogInformation("Ext File - init: '{Name}'", FSI.Name);
+								FileReference FileRef = new FileReference(FSI.FullName);
+								XcodeSourceFile SourceFile = new XcodeSourceFile(FileRef, FileRef.Directory, XcodeProjectFileGenerator.MakeXcodeGuid());
+								Project.FileCollection.ProcessFile(SourceFile, true, false, "Source");
+
+							}
+							
+							Project.FileCollection.AddFileReference(ExtensionGuid, Extension, "explicitFileType", "wrapper.app-extension", "BUILT_PRODUCTS_DIR", "Products");
+							
+							DirectoryReference FinalAppExtDir = DirectoryReference.Combine(ProjectDirectory, DI.Name);
+							EmbedAppExt.AddExtension(FinalAppExtDir, ExtensionGuid);
+							
+							BuildPhases.Add(EmbedAppExt);
+							References.Add(EmbedAppExt);
+							
+							// Also need a (Native)Run Target for this extension
+							XcodeAppExtRunTarget RunTarget = new XcodeAppExtRunTarget(Project, DI.Name, ExtensionGuid, UnrealData.AllConfigs[0].ProjectTarget!.TargetRules!.Type, UnrealTargetPlatform.IOS, BuildTarget, ProjectFile, DI, Logger);
+							Project.AppExtRunTargets.Add(RunTarget);
+							Project.References.Add(RunTarget);
+						}
+					}
+				}
+			}
+
 			CreateXcconfigFile(Project, Platform, $"{Name}");
 			// create per-config Xcconfig files
 			foreach (XcodeBuildConfig Config in BuildConfigList.BuildConfigs)
@@ -1749,49 +1878,6 @@
 			Content.WriteLine($"\t\t\tproductName = \"{UnrealData.ProductName}\";");
 		}
 
-		private static Dictionary<string, string> PlistFileMap = new(StringComparer.OrdinalIgnoreCase);
-		private string GetPlistSigningName(string ProvisionSetting, ILogger Logger)
-		{
-			string? SigningName;
-			lock (PlistFileMap)
-			{
-				if (!PlistFileMap.TryGetValue(ProvisionSetting, out SigningName))
-				{
-					FileReference ProfileFile = AppleExports.ConvertFilePath(UnrealData.UProjectFileLocation?.Directory, ProvisionSetting);
-
-					// get mobile provision UUID, either directly or from the probision
-					if (ProfileFile.GetExtension() == ".mobileprovision")
-					{
-						// security will read the provision and dump out the text (plist) bits of the profile, and plutil will extract the UUID to stdout
-						string UUID = Utils.RunLocalProcessAndReturnStdOut("/bin/sh", $"-c 'security cms -D -i {ProfileFile.FullName.Replace(" ", "\\ ")} | plutil -extract UUID raw -'");
-
-						// make sure it's living in user's library
-						// (note: i couldn't find a way to install it without opening Xcode, so copying it works well enough)
-						DirectoryReference UserDir = new DirectoryReference(Environment.GetFolderPath(Environment.SpecialFolder.UserProfile));
-						FileReference InstalledProfileFile = FileReference.Combine(UserDir, "Library", "MobileDevice", "Provisioning Profiles", $"{UUID}.mobileprovision");
-
-						// copy if needed
-						if (!FileReference.Exists(InstalledProfileFile))
-						{
-							DirectoryReference.CreateDirectory(InstalledProfileFile.Directory);
-							FileReference.Copy(ProfileFile, InstalledProfileFile);
-							Logger.LogInformation("Copying project's provision '{SourceProvision}' to your libary: '{TargetProvision}'", ProfileFile, InstalledProfileFile);
-						}
-
-						SigningName = UUID;
-					}
-					else
-					{
-						SigningName = ProvisionSetting;
-					}
-
-					PlistFileMap[ProvisionSetting] = SigningName;
-				}
-			}
-
-			return SigningName!;
-		}
-
 		public override void WriteXcconfigFile(ILogger Logger)
 		{
 			// gather general, all-platform, data we are doing to put into the configs
@@ -2029,7 +2115,7 @@
 				Xcconfig.AppendLine($"INFOPLIST_KEY_LSApplicationCategoryType = {AppCategory}");
 			}
 
-			// always use defualt codesigning for Editor, which is for running locally
+			// always use default codesigning for Editor, which is for running locally
 			if (UnrealData.bWriteCodeSigningSettings && UnrealData.TargetRules.Type != TargetType.Editor)
 			{
 				Xcconfig.AppendLine("");
@@ -2058,7 +2144,7 @@
 						// only use profile on IOS
 						if (!String.IsNullOrEmpty(ProvisioningProfile))
 						{
-							string SigningName = GetPlistSigningName(ProvisioningProfile, Logger);
+							string SigningName = GetPlistSigningName(ProvisioningProfile, UnrealData, Logger);
 							Xcconfig.AppendLine($"PROVISIONING_PROFILE_SPECIFIER = {SigningName}");
 						}
 						if (SigningIdentity != null)
@@ -2169,7 +2255,519 @@
 			}
 		}
 	}
+	
+	class XcodeAppExtRunTarget : XcodeTarget
+	{
+		public enum ExtType
+		{
+			BackgroundAssets,
+			NotificationService
+		}
+		
+		private string ProductGuid = XcodeProjectFileGenerator.MakeXcodeGuid();
+		private UnrealData UnrealData;
+		private DirectoryInfo ExtensionPath;
+		public UnrealTargetPlatform Platform;
+		public TargetType TargetType;
+		public string TargetName;
+		public string AppGuid;
+		
+		public XcodeAppExtRunTarget(XcodeProject Project, string TargetName, string AppGuid, TargetType TargetType, UnrealTargetPlatform Platform, XcodeBuildTarget? BuildTarget, XcodeProjectFile ProjectFile, DirectoryInfo ExtensionPath, ILogger Logger)
+			: base(XcodeTarget.Type.App_Ext, Project.UnrealData,
+				  TargetName + (XcodeProjectFileGenerator.PerPlatformMode == XcodePerPlatformMode.OneWorkspacePerPlatform ? "" : $"_{Platform}"))
+		{
+			this.TargetType = TargetType;
+			this.Platform = Platform;
+			this.AppGuid = AppGuid;
+			this.TargetName = TargetName;
+			this.ExtensionPath = ExtensionPath;
+			UnrealData = Project.UnrealData;
 
+			BuildConfigList = new XcodeBuildConfigList(Platform, Name, Project.UnrealData, bIncludeAllPlatformsInConfig: false);
+			References.Add(BuildConfigList);
+
+			if (Project.UnrealData.IsAppBundle(Platform))
+			{
+				XcodeFileCollection AppExtFileList = new XcodeFileCollection(ProjectFile);	
+
+				// add all of the files to be natively compiled by this target
+				XcodeSourcesBuildPhase SourcesBuildPhase = new XcodeSourcesBuildPhase();
+				BuildPhases.Add(SourcesBuildPhase);
+				References.Add(SourcesBuildPhase);
+
+				foreach (FileSystemInfo FSI in ExtensionPath.EnumerateFileSystemInfos())
+				{
+					if (FSI.Name.EndsWith(".h") || FSI.Name.StartsWith(".") || FSI.Name.EndsWith(".plist"))
+					{
+						continue;
+					}
+
+					string FileRef = XcodeProjectFileGenerator.MakeXcodeGuid();
+					FileReference FR = new FileReference(FSI.FullName);
+					XcodeSourceFile xsf = new XcodeSourceFile(FR, FR.Directory, FileRef);
+					SourcesBuildPhase.AddFile(xsf);
+					Project.FileCollection.AllFiles.Add(xsf);
+					Project.FileCollection.BuildableFilesToResponseFile[xsf] = new();
+				}
+				
+				foreach (KeyValuePair<XcodeSourceFile, Dictionary<UnrealTargetPlatform, FileReference>> Pair in Project.FileCollection.BuildableFilesToResponseFile)
+				{
+					// only add files that found a module to be part of (since we can't build without the build settings that come from a module)
+					if (Pair.Value.ContainsKey(Platform))
+					{
+						SourcesBuildPhase.AddFile(Pair.Key);
+					}
+				}
+				
+				XcodeResourcesBuildPhase ResourcesBuildPhase = new XcodeResourcesBuildPhase(AppExtFileList);
+				BuildPhases.Add(ResourcesBuildPhase);
+				References.Add(ResourcesBuildPhase);
+
+				ProcessFrameworks(ResourcesBuildPhase, Project, ProjectFile, Logger);
+				ProcessAssets(ResourcesBuildPhase);
+			}
+
+			if (BuildTarget != null)
+			{
+				AddDependency(BuildTarget, Project);
+			}
+
+			CreateXcconfigFile(Project, Platform, $"{Name}");
+			// create per-config Xcconfig files
+			foreach (XcodeBuildConfig Config in BuildConfigList.BuildConfigs)
+			{
+				Config.CreateXcconfigFile(Project, Platform, $"{Name}_{Config.Info.DisplayName.Replace(" ", "")}");
+			}
+		}
+
+		private void ProcessAssets(XcodeResourcesBuildPhase ResourcesBuildPhase)
+		{
+			// Get the source files
+			foreach (FileSystemInfo FSI in ExtensionPath.EnumerateFileSystemInfos())
+			{
+				if (FSI.Name.StartsWith(".") || FSI.Name.EndsWith(".plist"))
+				{
+					continue;
+				}
+				ResourcesBuildPhase.AddResource(new FileReference(FSI.FullName)); 
+			}
+		}
+		
+		protected void ProcessFrameworks(XcodeResourcesBuildPhase ResourcesBuildPhase, XcodeProject Project, XcodeProjectFile ProjectFile, ILogger Logger)
+		{
+			// look up to see if we had cached any Frameworks
+			Tuple<ProjectFile, UnrealTargetPlatform> FrameworkKey = Tuple.Create((ProjectFile)ProjectFile, Platform);
+			IEnumerable<UEBuildBundleResource>? Bundles;
+			if (XcodeProjectFileGenerator.TargetBundles.TryGetValue(FrameworkKey, out Bundles))
+			{
+				foreach (UEBuildBundleResource Bundle in Bundles)
+				{
+					ResourcesBuildPhase.AddFolderResource(new DirectoryReference(Bundle.ResourcePath!), "Resources");
+				}
+			}
+
+			IEnumerable<UEBuildFramework>? Frameworks;
+			if (XcodeProjectFileGenerator.TargetFrameworks.TryGetValue(FrameworkKey, out Frameworks))
+			{
+				XcodeFrameworkBuildPhase FrameworkPhase = new XcodeFrameworkBuildPhase(Project.FileCollection);
+
+				// filter frameworks that need to installed into the .app (either the framework or a bundle inside a .zip)
+				IEnumerable<UEBuildFramework> InstalledFrameworks = Frameworks.Where(x => x.bCopyFramework || !String.IsNullOrEmpty(x.CopyBundledAssets));
+
+				// only look at frameworks that need anything copied into 
+				foreach (UEBuildFramework Framework in InstalledFrameworks)
+				{
+					DirectoryReference? FinalFrameworkDir = Framework.GetFrameworkDirectory(null, null, Logger);
+					if (FinalFrameworkDir == null)
+					{
+						continue;
+					}
+
+					DirectoryReference BundleRootDir = Framework.GetFrameworkDirectory(null, null, Logger)!;
+
+					// set up the CopyBundle to be copied
+					if (!String.IsNullOrEmpty(Framework.CopyBundledAssets))
+					{
+						ResourcesBuildPhase.AddFolderResource(DirectoryReference.Combine(BundleRootDir, Framework.CopyBundledAssets), "Resources");
+					}
+				}
+
+				DirectoryReference userNotifFw = new DirectoryReference("UserNotifications.framework");
+				string FileRefGuid = XcodeProjectFileGenerator.MakeXcodeGuid();
+				FrameworkPhase.AddFramework(userNotifFw, FileRefGuid);
+
+				BuildPhases.Add(FrameworkPhase);
+				References.Add(FrameworkPhase);
+			}
+		}
+		
+		public override void WriteXcconfigFile(ILogger Logger)
+		{
+			// gather general, all-platform, data we are doing to put into the configs
+			UnrealBuildConfig BuildConfig = BuildConfigList!.BuildConfigs[0].Info;
+			TargetRules TargetRules = BuildConfig.ProjectTarget!.TargetRules!;
+			DirectoryReference ProjectOrEngineDir = UnrealData.UProjectFileLocation?.Directory ?? Unreal.EngineDirectory;
+
+			// point to the shader Engine/Binaries for content only project (sadly, the TargetRules.OutputFile is not filled out)
+			bool bUseEngineDirectory = TargetRules.Type == TargetType.Editor || UnrealData.bIsContentOnlyProject;
+			if (TargetRules.BuildEnvironment == TargetBuildEnvironment.Unique)
+			{
+				// unique build environment editors will put their binaries in the project
+				bUseEngineDirectory = false;
+			}
+
+			DirectoryReference ConfigBuildDir = bUseEngineDirectory ? Unreal.EngineDirectory : ProjectOrEngineDir;
+			if (TargetRules.Type == TargetType.Program && TargetRules.File!.IsUnderDirectory(Unreal.EngineDirectory))
+			{
+				ConfigBuildDir = BuildConfig.RootDirectory;
+			}
+			string BinariesBaseDir = ConfigBuildDir.GetDirectoryName();
+
+			ConfigBuildDir = DirectoryReference.Combine(ConfigBuildDir, "Binaries", Platform.ToString(), TargetRules.ExeBinariesSubFolder);
+
+			MetadataPlatform MetadataPlatform;
+			string TargetPlatformName = Platform.ToString();
+			if (TargetRules.Type != TargetType.Game && TargetRules.Type != TargetType.Program)
+			{
+				TargetPlatformName += TargetRules.Type.ToString();
+			}
+
+			// get ini file for the platform
+			ConfigHierarchy PlatformIni = ConfigCache.ReadHierarchy(ConfigHierarchyType.Engine, UnrealData.ConfigDirectory, Platform);
+
+			// settings for all platforms
+			bool bAutomaticSigning;
+			bool bMacSignToRunLocally;
+			bool bUseEntitlementsForPrograms;
+			string? SigningTeam;
+			string? SigningPrefix;
+			string? AppCategory;
+			string SupportedPlatforms;
+			string SDKRoot;
+			string DeploymentTarget;
+			string DeploymentTargetKey;
+			string? SigningIdentity = null;
+			string? ProvisioningProfile = null;
+			string? SupportedDevices = null;
+			string? MarketingVersion = null;
+			string? BundleIdentifier;
+			string? ApplicationDisplayName = null;
+			List<string> ExtraConfigLines = new();
+
+			// get signing settings
+			PlatformIni.TryGetValue("/Script/MacTargetPlatform.XcodeProjectSettings", "bUseAutomaticCodeSigning", out bAutomaticSigning);
+			PlatformIni.TryGetValue("/Script/MacTargetPlatform.XcodeProjectSettings", "bUseEntitlementsForPrograms", out bUseEntitlementsForPrograms);
+			PlatformIni.TryGetValue("/Script/MacTargetPlatform.XcodeProjectSettings", "bMacSignToRunLocally", out bMacSignToRunLocally);
+			PlatformIni.TryGetValue("/Script/MacTargetPlatform.XcodeProjectSettings", "CodeSigningTeam", out SigningTeam);
+			PlatformIni.TryGetValue("/Script/MacTargetPlatform.XcodeProjectSettings", "CodeSigningPrefix", out SigningPrefix);
+			PlatformIni.TryGetValue("/Script/MacTargetPlatform.XcodeProjectSettings", $"{Platform}ProvisioningProfile", out ProvisioningProfile);
+			PlatformIni.TryGetValue("/Script/MacTargetPlatform.XcodeProjectSettings", $"{Platform}SigningIdentity", out SigningIdentity);
+			PlatformIni.TryGetValue("/Script/MacTargetPlatform.XcodeProjectSettings", "AppCategory", out AppCategory);
+			PlatformIni.TryGetValue("/Script/MacTargetPlatform.XcodeProjectSettings", "ApplicationDisplayName", out ApplicationDisplayName);
+
+			if (Platform == UnrealTargetPlatform.Mac)
+			{
+				// editor vs game metadata
+				bool bIsEditor = UnrealData.TargetRules.Type == TargetType.Editor;
+				MetadataPlatform = bIsEditor ? MetadataPlatform.MacEditor : MetadataPlatform.Mac;
+
+				SDKRoot = "macosx";
+				SupportedPlatforms = "macosx";
+				DeploymentTargetKey = "MACOSX_DEPLOYMENT_TARGET";
+				DeploymentTarget = MacToolChainSettings.MinMacDeploymentVersion(UnrealData.TargetRules.Type);
+				BundleIdentifier = bIsEditor ? "com.epicgames.UnrealEditor" : UnrealData.BundleIdentifier;
+
+				PlatformIni.TryGetValue($"/Script/MacRuntimeSettings.MacRuntimeSettings", "VersionInfo", out MarketingVersion);
+				if (String.IsNullOrEmpty(MarketingVersion))
+				{
+					MarketingVersion = MacToolChain.LoadEngineDisplayVersion();
+				}
+			}
+			else
+			{
+				MetadataPlatform = MetadataPlatform.IOS;
+
+				// get IOS (same as TVOS) BundleID, and if there's a specified plist with a bundleID, use it, as Xcode would warn if they don't match
+				BundleIdentifier = UnrealData.BundleIdentifier;
+
+				// short version string
+				PlatformIni.GetString($"/Script/IOSRuntimeSettings.IOSRuntimeSettings", "VersionInfo", out MarketingVersion);
+
+				if (Platform == UnrealTargetPlatform.IOS)
+				{
+					SDKRoot = "iphoneos";
+					SupportedPlatforms = "iphoneos";
+					DeploymentTargetKey = "IPHONEOS_DEPLOYMENT_TARGET";
+					SupportedDevices = UnrealData.IOSProjectSettings!.RuntimeDevices;
+					DeploymentTarget = UnrealData.IOSProjectSettings.RuntimeVersion;
+				}
+				else if (Platform == UnrealTargetPlatform.TVOS) // tvos
+				{
+					SDKRoot = "appletvos";
+					SupportedPlatforms = "appletvos"; // appletvsimulator
+					DeploymentTargetKey = "TVOS_DEPLOYMENT_TARGET";
+					SupportedDevices = UnrealData.TVOSProjectSettings!.RuntimeDevices;
+					DeploymentTarget = UnrealData.TVOSProjectSettings.RuntimeVersion;
+				}
+				else if (Platform == UnrealTargetPlatform.VisionOS)
+				{
+					SDKRoot = "xros";
+					SupportedPlatforms = "xrsimulator xros";
+					DeploymentTargetKey = "XROS_DEPLOYMENT_TARGET";
+					SupportedDevices = UnrealData.VisionOSProjectSettings!.RuntimeDevices;
+					DeploymentTarget = UnrealData.VisionOSProjectSettings.RuntimeVersion;
+
+					ExtraConfigLines.Add($"VALID_ARCHS = arm64");
+					ExtraConfigLines.Add($"ARCHS = arm64");
+
+					// if we are doing immersive with SwiftUI we need a plist key
+					bool bUseSwiftUIMain;
+					AppleExports.GetSwiftIntegrationSettings(UnrealData.UProjectFileLocation, UnrealData.TargetRules.Type, Platform, out bUseSwiftUIMain, out _);
+					if (bUseSwiftUIMain)
+					{
+						ExtraConfigLines.Add($"INFOPLIST_KEY_UIApplicationSceneManifest_Generation = YES");
+					}
+				}
+				else
+				{
+					throw new BuildException($"Unsupported platform {Platform}");
+				}
+			}
+
+			// get metadata for the platform set above
+			MetadataItem PlistMetadata = UnrealData.Metadata!.PlistFiles[MetadataPlatform];
+			// now pull the bundle id's out, as xcode will warn if they don't match (this has to happen after each platform set bundle id above)
+			XcodeUtils.FindPlistId(PlistMetadata, "CFBundleIdentifier", ref BundleIdentifier);
+			// if the user had a ini setting, but also set it in the template plist, the plist one should win (this is only used if we are updating a template plist)
+			XcodeUtils.FindPlistId(PlistMetadata, "LSApplicationCategoryType", ref AppCategory);
+
+			// include another xcconfig for versions that UBT writes out
+			Xcconfig!.AppendLine($"#include? \"{ProjectOrEngineDir}/Intermediate/Build/Versions.xcconfig\"");
+			if (UnrealData.UProjectFileLocation == null)
+			{
+				Xcconfig!.AppendLine($"#include? \"{XcodeProjectFileGenerator.ContentOnlySettingsFile}\"");
+			}
+
+			// write out some UE variables that can be used in premade .plist files, etc
+			Xcconfig.AppendLine("");
+			Xcconfig.AppendLine("// Unreal project-wide variables");
+			Xcconfig.AppendLine($"UE_XCODE_BUILD_MODE = Default");
+			Xcconfig.AppendLine($"UE_PRODUCT_NAME = {Name}");
+			Xcconfig.AppendLine($"UE_PRODUCT_NAME_STRIPPED = {Name.Replace("_", "").Replace(" ", "")}");
+			// Xcconfig.AppendLine($"UE_PRODUCT_NAME = {UnrealData.ProductName}");
+			// Xcconfig.AppendLine($"UE_PRODUCT_NAME_STRIPPED = {UnrealData.ProductName.Replace("_", "").Replace(" ", "")}");
+			Xcconfig.AppendLine($"UE_DISPLAY_NAME = {UnrealData.DisplayName}");
+			Xcconfig.AppendLine($"UE_SIGNING_PREFIX = {SigningPrefix}");
+			Xcconfig.AppendLine($"UE_PLATFORM_NAME = {Platform}");
+			Xcconfig.AppendLine($"UE_TARGET_NAME = {TargetName}");
+			// Xcconfig.AppendLine($"UE_TARGET_NAME = {UnrealData.TargetRules.OriginalName}");
+			Xcconfig.AppendLine($"UE_TARGET_PLATFORM_NAME = {TargetPlatformName}");
+			Xcconfig.AppendLine($"UE_ENGINE_DIR = {Unreal.EngineDirectory}");
+			Xcconfig.AppendLine($"UE_BINARIES_DIR = {ConfigBuildDir}");
+			Xcconfig.AppendLine($"UE_STAGED_BINARIES_DIR_BASE = {BinariesBaseDir}");
+			if (UnrealData.UProjectFileLocation != null)
+			{
+				Xcconfig.AppendLine($"UE_PROJECT_DIR = {UnrealData.UProjectFileLocation.Directory}");
+			}
+
+			Xcconfig.AppendLine("");
+			Xcconfig.AppendLine("// Constant settings (same for all platforms and targets)");
+			Xcconfig.AppendLine("INFOPLIST_OUTPUT_FORMAT = xml");
+			Xcconfig.AppendLine("COMBINE_HIDPI_IMAGES = YES");
+			Xcconfig.AppendLine("USE_HEADERMAP = NO");
+			Xcconfig.AppendLine("ONLY_ACTIVE_ARCH = YES");
+
+			Xcconfig.AppendLine("");
+			Xcconfig.AppendLine("// Platform settings");
+			Xcconfig.AppendLine($"SUPPORTED_PLATFORMS = {SupportedPlatforms}");
+			Xcconfig.AppendLine($"SDKROOT = {SDKRoot}");
+
+			// Xcode creates the Build Dir (where the .app is) by combining {SYMROOT}/{CONFIGURATION}{EFFECTIVE_PLATFORM_NAME}, so we set SYMROOT
+			// to the Parent directory of the diectory the binary is in, CONFIGURATION to nothing, and EFFECTIVE_PLATFORM_NAME to the directory
+			// the binary is in
+			Xcconfig.AppendLine("");
+			Xcconfig.AppendLine($"// These settings combined will tell Xcode to write to Binaries/{Platform} (instead of something like Binaries/Development-iphoneos)");
+			Xcconfig.AppendLine($"SYMROOT = {ConfigBuildDir.ParentDirectory}");
+			Xcconfig.AppendLine($"CONFIGURATION = ");
+			Xcconfig.AppendLine($"EFFECTIVE_PLATFORM_NAME = {ConfigBuildDir.GetDirectoryName()}");
+
+			if (ExtraConfigLines.Count > 0)
+			{
+				Xcconfig.AppendLine("");
+				Xcconfig.AppendLine("// Misc settings");
+				foreach (string Line in ExtraConfigLines)
+				{
+					// Xcconfig.AppendLine(Line);
+				}
+			}
+
+			Xcconfig.AppendLine("");
+			Xcconfig.AppendLine("// Project settings");
+			Xcconfig.AppendLine($"TARGETED_DEVICE_FAMILY = {SupportedDevices}");
+			Xcconfig.AppendLine($"PRODUCT_BUNDLE_IDENTIFIER = {SigningPrefix}.{UnrealData.ProductName}.{Name}");
+			Xcconfig.AppendLine($"{DeploymentTargetKey} = {DeploymentTarget}");
+
+			Xcconfig.AppendLine("");
+			Xcconfig.AppendLine("// Plist settings");
+			// Xcconfig.AppendLine($"INFOPLIST_FILE = {PlistMetadata.XcodeProjectRelative}");
+			Xcconfig.AppendLine($"INFOPLIST_FILE = ../../Build/IOS/Extensions/{Name}/Info.plist");
+			Xcconfig.AppendLine($"INFOPLIST_KEY_CFBundleDisplayName = {Name}");
+			if (PlistMetadata.Mode == MetadataMode.UpdateTemplate)
+			{
+				// allow Xcode to generate the final plist file from our input, some INFOPLIST settings and other settings 
+				Xcconfig.AppendLine("GENERATE_INFOPLIST_FILE = YES");
+				Xcconfig.AppendLine($"CURRENT_PROJECT_VERSION = $(UE_{Platform.ToString().ToUpper()}_BUILD_VERSION)");
+				Xcconfig.AppendLine($"MARKETING_VERSION = {MarketingVersion}");
+			}
+
+			// always use default codesigning for Editor, which is for running locally
+			if (UnrealData.bWriteCodeSigningSettings && UnrealData.TargetRules.Type != TargetType.Editor)
+			{
+				Xcconfig.AppendLine("");
+				Xcconfig.AppendLine("// Code-signing settings");
+				Xcconfig.AppendLine("CODE_SIGN_STYLE = " + (bAutomaticSigning ? "Automatic" : "Manual"));
+				Xcconfig.AppendLine($"DEVELOPMENT_TEAM = {SigningTeam}");
+
+				// Mac has Sign to Run Locally to deal with
+				if (Platform == UnrealTargetPlatform.Mac)
+				{
+					// when set, always use the - identity, that's how Xcode selects "Sign to Run Locally"
+					// if it is not set, and Automatic is on, then it will use it's default of "Apple Development", but we don't specify it in case default changes
+					if (bMacSignToRunLocally)
+					{
+						Xcconfig.AppendLine($"CODE_SIGN_IDENTITY = -");
+					}
+					else if (!bAutomaticSigning && SigningIdentity != null)
+					{
+						Xcconfig.AppendLine($"CODE_SIGN_IDENTITY = {SigningIdentity}");
+					}
+				}
+				else
+				{
+					if (!bAutomaticSigning)
+					{
+						// only use profile on IOS
+						if (!String.IsNullOrEmpty(ProvisioningProfile))
+						{
+							string SigningName = GetPlistSigningName(ProvisioningProfile, UnrealData, Logger);
+							Xcconfig.AppendLine($"PROVISIONING_PROFILE_SPECIFIER = {SigningName}");
+						}
+						if (SigningIdentity != null)
+						{
+							Xcconfig.AppendLine($"CODE_SIGN_IDENTITY = {SigningIdentity}");
+						}
+					}
+				}
+			}
+
+			Xcconfig.Write();
+
+			// Now for each config write out the specific settings
+			DirectoryReference? GameDir = UnrealData.UProjectFileLocation?.Directory;
+			string? GamePath = GameDir != null ? XcodeFileCollection.ConvertPath(GameDir.FullName) : null;
+			foreach (UnrealBuildConfig Config in UnrealData.AllConfigs)
+			{
+				XcodeBuildConfig? MatchedConfig = BuildConfigList!.BuildConfigs.FirstOrDefault(x => x.Info == Config);
+				if (MatchedConfig == null || !Config.Supports(Platform))
+				{
+					continue;
+				}
+
+				// hook up the Buildconfig that matches this info to this xcconfig file
+				XcconfigFile ConfigXcconfig = MatchedConfig.Xcconfig!;
+
+				string ExecutableName = Name;
+				string ExecutableSubPath = FileReference.Combine(ConfigBuildDir, ExecutableName).MakeRelativeTo(ConfigBuildDir);
+				string ExecutableKey = $"UE_{Platform.ToString().ToUpper()}_EXECUTABLE_NAME";
+
+				// we want to make Foo.app, not FooGame.app, since we added on the target type when making targets
+				string PerTargetTypeProductName = UnrealData.ProductName;
+				if (UnrealData.bIsContentOnlyProject && TargetRules.Type == TargetType.Game)
+				{
+					PerTargetTypeProductName = UnrealData.UProjectFileLocation!.GetFileNameWithoutAnyExtensions();
+				}
+				string ProductName = ExecutableName;
+				// content only projects don't want UnrealGame, etc as the ProductName
+				if (UnrealData.bIsContentOnlyProject && TargetRules.Type != TargetType.Editor)
+				{
+					ProductName = AppleExports.MakeBinaryFileName(PerTargetTypeProductName, TargetRules.DecoratedSeparator, Platform, Config.BuildConfig, TargetRules.Architectures, TargetRules.UndecoratedConfiguration, null);
+				}
+
+				MetadataItem? EntitlementsMetadata = UnrealData.Metadata!.EntitlementsFiles[MetadataPlatform];
+
+				// if we have shipping specified, use it instead
+				if (Config.BuildConfig == UnrealTargetConfiguration.Shipping)
+				{
+					MetadataItem ShippingEntitlementsMetadata = UnrealData.Metadata!.ShippingEntitlementsFiles[MetadataPlatform];
+					if (ShippingEntitlementsMetadata.Mode == MetadataMode.UsePremade)
+					{
+						EntitlementsMetadata = ShippingEntitlementsMetadata;
+					}
+				}
+
+				// programs almost never want entitlements, so handle them specially
+				if ((UnrealData.bNoEntitlements) || (UnrealData.TargetRules.Type == TargetType.Program && !bUseEntitlementsForPrograms))
+				{
+					EntitlementsMetadata = null;
+				}
+
+				ConfigXcconfig.AppendLine("// pull in the shared settings for all configs for this target");
+				ConfigXcconfig.AppendLine($"#include \"{Xcconfig.Name}.xcconfig\"");
+				ConfigXcconfig.AppendLine("");
+				ConfigXcconfig.AppendLine("// Unreal per-config variables");
+				ConfigXcconfig.AppendLine($"UE_TARGET_CONFIG = {Config.BuildConfig}");
+				ConfigXcconfig.AppendLine($"UE_UBT_BINARY_SUBPATH = {ExecutableSubPath}");
+				ConfigXcconfig.AppendLine($"{ExecutableKey} = {ExecutableName}");
+				if (EntitlementsMetadata != null && EntitlementsMetadata.Mode == MetadataMode.UsePremade)
+				{
+					ConfigXcconfig.AppendLine($"CODE_SIGN_ENTITLEMENTS = {EntitlementsMetadata.XcodeProjectRelative}");
+				}
+
+				// debug settings
+				if (Config.BuildConfig == UnrealTargetConfiguration.Debug)
+				{
+					ConfigXcconfig.AppendLine("ENABLE_TESTABILITY = YES");
+				}
+
+				if (Platform == UnrealTargetPlatform.Mac)
+				{
+					// on Mac, we need to name the .app nicely before pushing to App store, otherwise distributing, so use the ini setting if it's there ("Unreal Match 3"), otherwise use the uproject name (ie "Lyra" instead of "LyraGame")
+					ConfigXcconfig.AppendLine("");
+					ConfigXcconfig.AppendLine($"// this variable trickery will set the proper name for debugging, building, and archiving,");
+					ConfigXcconfig.AppendLine($"// where archiving (the '_install' action type) may need a differnet name so it shows up nicely");
+					ConfigXcconfig.AppendLine($"// on end-users machines in Finder, Spotlight, etc. The trailing _ on the next line is correct.");
+
+					ConfigXcconfig.AppendLine($"PRODUCT_NAME_ = {ProductName}");
+					ConfigXcconfig.AppendLine($"PRODUCT_NAME_build = $(PRODUCT_NAME_)");
+					if (String.IsNullOrEmpty(ApplicationDisplayName))
+					{
+						ConfigXcconfig.AppendLine($"PRODUCT_NAME_install = {(UnrealData.UProjectFileLocation == null ? ProductName : UnrealData.UProjectFileLocation!.GetFileNameWithoutAnyExtensions())}");
+					}
+					else
+					{
+						ConfigXcconfig.AppendLine($"PRODUCT_NAME_install = {ApplicationDisplayName}");
+					}
+
+					// this will choose the proper PRODUCT_NAME when archiving vs normal building
+					ConfigXcconfig.AppendLine("PRODUCT_NAME = $(PRODUCT_NAME_$(ACTION))");
+				}
+				else
+				{
+					ConfigXcconfig.AppendLine($"PRODUCT_NAME = {ProductName}");
+				}
+
+				ConfigXcconfig.Write();
+			}
+		}
+		
+
+		protected override void WriteExtraTargetProperties(StringBuilder Content)
+		{
+			Content.WriteLine($"\t\t\tproductReference = {AppGuid};");
+			Content.WriteLine($"\t\t\tproductName = \"{TargetName}\";");
+		}
+	}
+
 	class XcodeBuildTarget : XcodeTarget
 	{
 		public XcodeBuildTarget(XcodeProject Project)
@@ -2243,6 +2841,7 @@
 		private string ProvisioningStyle;
 
 		public List<XcodeRunTarget> RunTargets = new();
+		public List<XcodeAppExtRunTarget> AppExtRunTargets = new();
 
 		public XcodeBuildConfigList ProjectBuildConfigs;
 
@@ -2316,6 +2915,12 @@
 				Content.WriteLine(6, $"ProvisioningStyle = {ProvisioningStyle};");
 				Content.WriteLine(5, "};");
 			}
+			foreach (XcodeAppExtRunTarget RunTarget in AppExtRunTargets)
+			{
+				Content.WriteLine(5, $"{RunTarget.Guid} = {{");
+				Content.WriteLine(6, $"ProvisioningStyle = {ProvisioningStyle};");
+				Content.WriteLine(5, "};");
+			}
 			Content.WriteLine(4, "};");
 			Content.WriteLine(3, "};");
 			Content.WriteLine(3, $"buildConfigurationList = {ProjectBuildConfigs.Guid} /* Build configuration list for PBXProject \"{ProjectBuildConfigs.TargetName}\" */;");
@@ -2327,7 +2932,7 @@
 			Content.WriteLine(3, ");");
 			Content.WriteLine(3, $"mainGroup = {FileCollection.MainGroupGuid};");
 			// for stub editor projects, we don't have a run target, so we don't have a product folder
-			if (RunTargets.Count > 0)
+			if (RunTargets.Count > 0 || AppExtRunTargets.Count > 0)
 			{
 				Content.WriteLine(3, $"productRefGroup = {FileCollection.GetProductGroupGuid()};");
 			}
diff --git a/Engine/Source/Runtime/ApplicationCore/Private/IOS/IOSAppDelegate.cpp b/Engine/Source/Runtime/ApplicationCore/Private/IOS/IOSAppDelegate.cpp
--- a/Engine/Source/Runtime/ApplicationCore/Private/IOS/IOSAppDelegate.cpp
+++ b/Engine/Source/Runtime/ApplicationCore/Private/IOS/IOSAppDelegate.cpp
@@ -1746,6 +1746,12 @@ extern double GCStartTime;
 
 #if !PLATFORM_TVOS && NOTIFICATIONS_ENABLED
 
+-(void)application:(UIApplication *)application didReceiveRemoteNotification:(NSDictionary *)userInfo fetchCompletionHandler:(void (^)(UIBackgroundFetchResult))completionHandler
+{
+	UE_LOG(LogTemp, Display, TEXT("Received RemoteNotification"));
+	completionHandler(UIBackgroundFetchResultNoData);
+}
+
 - (void)application:(UIApplication *)application didRegisterForRemoteNotificationsWithDeviceToken:(NSData *)deviceToken
 {
 	if (FApp::IsUnattended())
@@ -1776,6 +1782,7 @@ extern double GCStartTime;
 {
 	FString errorDescription([error description]);
 	
+	UE_LOG(LogTemp, Warning, TEXT("FAILED to register for remote notificaiton: %s"), *errorDescription);
     FFunctionGraphTask::CreateAndDispatchWhenReady([errorDescription]()
     {
 		FCoreDelegates::ApplicationFailedToRegisterForRemoteNotificationsDelegate.Broadcast(errorDescription);
diff --git a/Engine/Source/Runtime/Core/Private/IOS/IOSPlatformMisc.cpp b/Engine/Source/Runtime/Core/Private/IOS/IOSPlatformMisc.cpp
--- a/Engine/Source/Runtime/Core/Private/IOS/IOSPlatformMisc.cpp
+++ b/Engine/Source/Runtime/Core/Private/IOS/IOSPlatformMisc.cpp
@@ -1296,34 +1296,34 @@ void FIOSPlatformMisc::RegisterForRemoteNotifications()
 		return;
 	}
 
-    dispatch_async(dispatch_get_main_queue(), ^{
 #if !PLATFORM_TVOS && NOTIFICATIONS_ENABLED
+    dispatch_async(dispatch_get_main_queue(), ^{
 		UNUserNotificationCenter *Center = [UNUserNotificationCenter currentNotificationCenter];
 		[Center requestAuthorizationWithOptions:(UNAuthorizationOptionBadge | UNAuthorizationOptionSound | UNAuthorizationOptionAlert)
-							  completionHandler:^(BOOL granted, NSError * _Nullable error) {
+							  completionHandler:^(BOOL bGranted, NSError * _Nullable error) {
 								  if (error)
 								  {
 									  UE_LOG(LogIOS, Log, TEXT("Failed to register for notifications."));
 								  }
 								  else
 								  {
-									  int32 types = (int32)granted;
-                                      if (granted)
+									  int IsAllowed = (int32)bGranted;
+                                      if (bGranted)
                                       {
 										  dispatch_sync(dispatch_get_main_queue(), ^{
-											  UIApplication* application = [UIApplication sharedApplication];
-											  [application registerForRemoteNotifications];  
+											  UIApplication* Application = [UIApplication sharedApplication];
+											  [Application registerForRemoteNotifications];  
 										  });
                                       }
-									  FFunctionGraphTask::CreateAndDispatchWhenReady([types]()
+									  FFunctionGraphTask::CreateAndDispatchWhenReady([IsAllowed]()
 																					 {
-																						 FCoreDelegates::ApplicationRegisteredForUserNotificationsDelegate.Broadcast(types);
+																						 FCoreDelegates::ApplicationRegisteredForUserNotificationsDelegate.Broadcast(IsAllowed);
 																					 }, TStatId(), NULL, ENamedThreads::GameThread);
 									  
 								  }
 							  }];
-#endif
     });
+#endif
 }
 
 bool FIOSPlatformMisc::IsRegisteredForRemoteNotifications()
