Disable a few more incompatible plugins for Windows Arm64 & Arm64ec
- Some editor modules don't link properly on Arm64ec
- EOSShared added a dependency on SteamShared which isn't available for Windows Arm64

#rnx

diff --git a/Engine/Plugins/Online/EOSShared/Source/EOSShared/EOSShared.Build.cs b/Engine/Plugins/Online/EOSShared/Source/EOSShared/EOSShared.Build.cs
--- a/Engine/Plugins/Online/EOSShared/Source/EOSShared/EOSShared.Build.cs
+++ b/Engine/Plugins/Online/EOSShared/Source/EOSShared/EOSShared.Build.cs
@@ -44,7 +44,7 @@ public class EOSShared : ModuleRules
 		}
 
 		// On these platforms, the EOS SDK supports a native integration with Steam. 
-		bool bWithSteamIntegration = Target.Platform == UnrealTargetPlatform.Win64 || Target.Platform == UnrealTargetPlatform.Mac || Target.Platform == UnrealTargetPlatform.Linux;
+		bool bWithSteamIntegration = (Target.Platform == UnrealTargetPlatform.Win64 && Target.Architecture != UnrealArch.Arm64) || Target.Platform == UnrealTargetPlatform.Mac || Target.Platform == UnrealTargetPlatform.Linux;
 		PrivateDefinitions.Add("UE_WITH_EOS_STEAM_INTEGRATION=" + (bWithSteamIntegration ? "1" : "0"));
 		if (bWithSteamIntegration)
 		{
diff --git a/Engine/Source/Programs/UnrealBuildTool/Platform/Windows/UEBuildWindows.cs b/Engine/Source/Programs/UnrealBuildTool/Platform/Windows/UEBuildWindows.cs
--- a/Engine/Source/Programs/UnrealBuildTool/Platform/Windows/UEBuildWindows.cs
+++ b/Engine/Source/Programs/UnrealBuildTool/Platform/Windows/UEBuildWindows.cs
@@ -1219,8 +1219,9 @@ namespace UnrealBuildTool
 					{
 						Target.DisablePlugins.AddRange(new string[]
 						{
-							"Reflex",
-							"VirtualCamera", // WebRTC currently does not link properly
+							"Reflex", // nvapi
+							"ACLPlugin", // __vectorcall calling convention is not currently supported
+							"UbaController", // LNK2001 (EC Symbol) errors
 						});
 						Target.GlobalDefinitions.Add("UE_EXTERNAL_PROFILING_ENABLED=0");
 					}
@@ -1233,6 +1234,9 @@ namespace UnrealBuildTool
 							// steamapi
 							"OnlineSubsystemSteam",
 							"SteamSockets",
+							"SteamShared",
+							"SteamController",
+							"SocketSubsystemSteamIP",
 
 							// nvapi
 							"Reflex",
