Use parent's BundleID for app extensions base bundleId

diff --git a/Engine/Source/Programs/UnrealBuildTool/ProjectFiles/Xcode/XcodeProject.cs b/Engine/Source/Programs/UnrealBuildTool/ProjectFiles/Xcode/XcodeProject.cs
--- a/Engine/Source/Programs/UnrealBuildTool/ProjectFiles/Xcode/XcodeProject.cs
+++ b/Engine/Source/Programs/UnrealBuildTool/ProjectFiles/Xcode/XcodeProject.cs
@@ -2604,7 +2604,7 @@ namespace UnrealBuildTool.XcodeProjectXcconfig
 			Xcconfig.AppendLine("");
 			Xcconfig.AppendLine("// Project settings");
 			Xcconfig.AppendLine($"TARGETED_DEVICE_FAMILY = {SupportedDevices}");
-			Xcconfig.AppendLine($"PRODUCT_BUNDLE_IDENTIFIER = {SigningPrefix}.{UnrealData.ProductName}.{Name}");
+			Xcconfig.AppendLine($"PRODUCT_BUNDLE_IDENTIFIER = {BundleIdentifier}.{Name}");
 			Xcconfig.AppendLine($"{DeploymentTargetKey} = {DeploymentTarget}");
 
 			Xcconfig.AppendLine("");
