UnrealBuildTool: lld-link supports /RELEASE as of v16

#rnx

diff --git a/Engine/Source/Programs/UnrealBuildTool/Platform/Windows/VCToolChain.cs b/Engine/Source/Programs/UnrealBuildTool/Platform/Windows/VCToolChain.cs
--- a/Engine/Source/Programs/UnrealBuildTool/Platform/Windows/VCToolChain.cs
+++ b/Engine/Source/Programs/UnrealBuildTool/Platform/Windows/VCToolChain.cs
@@ -1629,12 +1629,8 @@ namespace UnrealBuildTool
 			//
 			if (LinkEnvironment.Configuration == CppConfiguration.Shipping)
 			{
-				if (!Target.WindowsPlatform.Compiler.IsClang() || !Target.WindowsPlatform.bAllowClangLinker)
-				{
-					// Generate an EXE checksum.
-					// Not supported by lld-link
-					Arguments.Add("/RELEASE");
-				}
+				// Generate an EXE checksum.
+				Arguments.Add("/RELEASE");
 			}
 
 			// Eliminate unreferenced symbols.
