UBT: Improve developer experience by stripping display level timing information from failed builds.
For failed builds the information is mostly irrelevant and makes it harder to spot the errors.
The timing information can still be found int he log file.
Always add one display line that shows the actual tool mode result for the build.
Add a new bShowTimeline setting to completely enable or disable the script exeuction timeline information (bShowCompilationTimes and bShowCPUUtilization already exists).

#rnx

diff --git a/Engine/Source/Programs/UnrealBuildTool/Executors/ImmediateActionQueue.cs b/Engine/Source/Programs/UnrealBuildTool/Executors/ImmediateActionQueue.cs
--- a/Engine/Source/Programs/UnrealBuildTool/Executors/ImmediateActionQueue.cs
+++ b/Engine/Source/Programs/UnrealBuildTool/Executors/ImmediateActionQueue.cs
@@ -1281,14 +1281,15 @@ namespace UnrealBuildTool
 			loggingTask?.Wait();
 			_flushToolOutput();
 
+			LogLevel LogLevel = Success ?  LogLevel.Information : LogLevel.Debug;
 			if (ShowCPUUtilization)
 			{
 				lock (_cpuUtilization)
 				{
 					if (_cpuUtilization.Count > 0)
 					{
-						Logger.LogInformation("");
-						Logger.LogInformation("Average CPU Utilization: {CPUPercentage}%", (int)(_cpuUtilization.Average()));
+						Logger.Log(LogLevel, "");
+						Logger.Log(LogLevel, "Average CPU Utilization: {CPUPercentage}%", (int)(_cpuUtilization.Average()));
 					}
 				}
 			}
@@ -1298,11 +1299,11 @@ namespace UnrealBuildTool
 				return;
 			}
 
-			Logger.LogInformation("");
+			Logger.Log(LogLevel, "");
 			if (ProcessGroup.TotalProcessorTime.Ticks > 0)
 			{
-				Logger.LogInformation("Total CPU Time: {TotalSeconds} s", ProcessGroup.TotalProcessorTime.TotalSeconds);
-				Logger.LogInformation("");
+				Logger.Log(LogLevel, "Total CPU Time: {TotalSeconds} s", ProcessGroup.TotalProcessorTime.TotalSeconds);
+				Logger.Log(LogLevel, "");
 			}
 
 			IEnumerable<int> CompletedActions = Enumerable.Range(0, Actions.Length)
@@ -1312,8 +1313,8 @@ namespace UnrealBuildTool
 
 			if (CompletedActions.Any())
 			{
-				Logger.LogInformation("Compilation Time Top {CompletedTaskCount}", CompletedActions.Count());
-				Logger.LogInformation("");
+				Logger.Log(LogLevel, "Compilation Time Top {CompletedTaskCount}", CompletedActions.Count());
+				Logger.Log(LogLevel, "");
 				foreach (int Index in CompletedActions)
 				{
 					IExternalAction Action = Actions[Index].Action.Inner;
@@ -1322,14 +1323,14 @@ namespace UnrealBuildTool
 					string Description = $"{(Action.CommandDescription ?? Action.CommandPath.GetFileNameWithoutExtension())} {Action.StatusDescription}".Trim();
 					if (Result.ProcessorTime.Ticks > 0)
 					{
-						Logger.LogInformation("{Description} [ Wall Time {ExecutionTime:0.00} s / CPU Time {ProcessorTime:0.00} s ]", Description, Result.ExecutionTime.TotalSeconds, Result.ProcessorTime.TotalSeconds);
+						Logger.Log(LogLevel, "{Description} [ Wall Time {ExecutionTime:0.00} s / CPU Time {ProcessorTime:0.00} s ]", Description, Result.ExecutionTime.TotalSeconds, Result.ProcessorTime.TotalSeconds);
 					}
 					else
 					{
-						Logger.LogInformation("{Description} [ Time {ExecutionTime:0.00} s ]", Description, Result.ExecutionTime.TotalSeconds);
+						Logger.Log(LogLevel, "{Description} [ Time {ExecutionTime:0.00} s ]", Description, Result.ExecutionTime.TotalSeconds);
 					}
 				}
-				Logger.LogInformation("");
+				Logger.Log(LogLevel, "");
 			}
 		}
 
@@ -1426,4 +1427,4 @@ namespace UnrealBuildTool
 			CancellationTokenSource.Cancel();
 		}
 	}
-}
\ No newline at end of file
+}
diff --git a/Engine/Source/Programs/UnrealBuildTool/UnrealBuildTool.cs b/Engine/Source/Programs/UnrealBuildTool/UnrealBuildTool.cs
--- a/Engine/Source/Programs/UnrealBuildTool/UnrealBuildTool.cs
+++ b/Engine/Source/Programs/UnrealBuildTool/UnrealBuildTool.cs
@@ -232,6 +232,13 @@ namespace UnrealBuildTool
 			[XmlConfigFile(Category = "Log")]
 			public int LogFileBackupCount = Log.LogFileBackupCount;
 
+			/// <summary>
+			/// If set and tool execution was successful, then display an unreal build tool script execution timeline summary.
+			/// If unset or the tool execution failed, print the same information silently to the log.
+			/// </summary>
+			[XmlConfigFile(Category = "BuildConfiguration")]
+			public bool bShowTimeline = Unreal.IsBuildMachine();
+
 			/// <summary>
 			/// If set TMP\TEMP will be overidden to this directory, each process will create a unique subdirectory in this folder.
 			/// </summary>
@@ -403,6 +410,7 @@ namespace UnrealBuildTool
                 Console.CancelKeyPress += CancelKeyPressAsync;
             }
 
+			bool bShowTimeline = false;
 			try
 			{
 				// Start capturing performance info
@@ -521,6 +529,7 @@ namespace UnrealBuildTool
 
 				Log.BackupLogFiles = Options.bBackupLogFiles;
 				Log.LogFileBackupCount = Options.LogFileBackupCount;
+				bShowTimeline = Options.bShowTimeline;
 
 				// Add the log writer if requested. When building a target, we'll create the writer for the default log file later.
 				if (Options.LogFileName != null)
@@ -684,10 +693,20 @@ namespace UnrealBuildTool
 
 				// Print out all the performance info
 				Timeline.Stop();
-				LogLevel TimelineLogLevel = ModeOptions.HasFlag(ToolModeOptions.ShowExecutionTime) ? LogLevel.Information : LogLevel.Debug;
+				bool bShowExecutionTime = ModeOptions.HasFlag(ToolModeOptions.ShowExecutionTime);
+				LogLevel ExecutionLogLevel = bShowExecutionTime ? LogLevel.Information : LogLevel.Debug;
+				LogLevel TimelineLogLevel = bShowExecutionTime && bShowTimeline && ApplicationResult == CompilationResult.Succeeded ?  LogLevel.Information : LogLevel.Debug;
 				Timeline.Print(TimeSpan.FromMilliseconds(100.0), TimeSpan.FromMilliseconds(200.0), TimelineLogLevel, Logger);
-				Logger.Log(TimelineLogLevel, (""));
-				Logger.Log(TimelineLogLevel, "Total execution time: {Time:0.00} seconds", Timeline.Elapsed.TotalSeconds);
+				Logger.Log(ExecutionLogLevel, "");
+				if (ApplicationResult == CompilationResult.Succeeded)
+				{
+					Logger.Log(ExecutionLogLevel, "Result: {ApplicationResult}", ApplicationResult);
+				}
+				else
+				{
+					Logger.Log(ExecutionLogLevel, "Result: Failed ({ApplicationResult})", ApplicationResult);
+				}
+				Logger.Log(ExecutionLogLevel, "Total execution time: {Time:0.00} seconds", Timeline.Elapsed.TotalSeconds);
 
 				// Make sure we flush the logs however we exit
 				Trace.Close();
