UnrealBuildTool: Set minimum pvs-studio version to 7.30 and clean up old version checks

#rnx

diff --git a/Engine/Source/Programs/UnrealBuildTool/Platform/Windows/PVSToolChain.cs b/Engine/Source/Programs/UnrealBuildTool/Platform/Windows/PVSToolChain.cs
--- a/Engine/Source/Programs/UnrealBuildTool/Platform/Windows/PVSToolChain.cs
+++ b/Engine/Source/Programs/UnrealBuildTool/Platform/Windows/PVSToolChain.cs
@@ -3,12 +3,9 @@
 using EpicGames.Core;
 using Microsoft.Extensions.Logging;
 using Newtonsoft.Json;
-using OpenTracing.Tag;
 using System;
 using System.Collections.Generic;
-using System.Collections.Immutable;
 using System.Diagnostics;
-using System.Diagnostics.CodeAnalysis;
 using System.IO;
 using System.Linq;
 using System.Text;
@@ -392,13 +389,6 @@ namespace UnrealBuildTool
 		[CommandLine("-PrintLevel")]
 		int _printLevel = 1;
 
-		/// <summary>
-		/// Version of the analyzers
-		/// </summary>
-		[CommandLine("-AnalyzerVersion", Required = true)]
-
-		Version? _analyzerVersion = null;
-
 		/// <summary>
 		/// If all ThirdParty code should be ignored
 		/// </summary>
@@ -410,8 +400,6 @@ namespace UnrealBuildTool
 
 		CppRootPaths _rootPaths = new();
 
-		bool isOldErrorInfo = false;
-
 		/// <summary>
 		/// Execute the command
 		/// </summary>
@@ -423,7 +411,7 @@ namespace UnrealBuildTool
 			arguments.ApplyTo(this);
 			arguments.CheckAllArgumentsUsed();
 
-			if (_inputFileList == null || _outputFile == null || _ignoredFile == null || _rootPathsFile == null || _analyzerVersion == null)
+			if (_inputFileList == null || _outputFile == null || _ignoredFile == null || _rootPathsFile == null)
 			{
 				throw new NullReferenceException();
 			}
@@ -436,9 +424,6 @@ namespace UnrealBuildTool
 			string[] ignoreFileLines = await FileReference.ReadAllLinesAsync(_ignoredFile);
 			_ignoredDirectories = ignoreFileLines.Select(x => x.Trim()).Where(x => x.Length > 0).Select(x => new DirectoryReference(x));
 
-			// TODO: Disallow old versions
-			isOldErrorInfo = _analyzerVersion.CompareTo(new Version("7.30")) < 0;
-
 			// Read the root paths file
 			_rootPaths = new(new BinaryArchiveReader(_rootPathsFile));
 
@@ -484,39 +469,6 @@ namespace UnrealBuildTool
 			return 0;
 		}
 
-		static readonly string[] errorInfoSplit = ["<#~>"];
-		static bool TryParseErrorInfo(string line, [NotNullWhen(true)] out PVSErrorInfo? errorInfo)
-		{
-			string[] tokens = line.Split(errorInfoSplit, StringSplitOptions.None);
-			if (tokens.Length >= 9)
-			{
-				string lineNumberStr = tokens[2];
-				string fileName = tokens[3];
-				string warningCode = tokens[5];
-				string warningMessage = tokens[6];
-				string falseAlarmStr = tokens[7];
-				string levelStr = tokens[8];
-
-				if (Int32.TryParse(lineNumberStr, out int LineNumber) && Boolean.TryParse(falseAlarmStr, out bool bFalseAlarm) && Int32.TryParse(levelStr, out int Level))
-				{
-					errorInfo = new PVSErrorInfo()
-					{
-						Code = warningCode,
-						Message = warningMessage,
-						FalseAlarm = bFalseAlarm,
-						Level = Level,
-						Positions = [new PVSPosition { File = fileName, Lines = [LineNumber] }],
-
-					};
-
-					return true;
-				}
-			}
-
-			errorInfo = null;
-			return false;
-		}
-
 		// Ignore anything in the IgnoredDirectories folders or ThirdParty if ignored
 		bool IsFileIgnored(FileReference? fileReference) => fileReference != null
 			&& ((_ignoreThirdParty && fileReference.FullName.Contains("ThirdParty", StringComparison.OrdinalIgnoreCase)) || (_ignoredDirectories.Any() && _ignoredDirectories.Any(fileReference.IsUnderDirectory)));
@@ -525,19 +477,7 @@ namespace UnrealBuildTool
 		{
 			try
 			{
-				PVSErrorInfo? errorInfo = null;
-				if (isOldErrorInfo)
-				{
-					if (!TryParseErrorInfo(line, out errorInfo))
-					{
-						throw new FormatException();
-					}
-				}
-				else
-				{
-					errorInfo = JsonConvert.DeserializeObject<PVSErrorInfo>(line) ?? throw new FormatException();
-				}
-
+				PVSErrorInfo errorInfo = JsonConvert.DeserializeObject<PVSErrorInfo>(line) ?? throw new FormatException();
 				FileReference? fileReference = errorInfo.Positions.FirstOrDefault()?.UpdateFilePath(_rootPaths);
 				return IsFileIgnored(fileReference) ? null : errorInfo;
 			}
@@ -598,7 +538,10 @@ namespace UnrealBuildTool
 		readonly UnrealTargetPlatform Platform;
 		readonly Version AnalyzerVersion;
 
-		string OutputFileExtension => AnalyzerVersion.CompareTo(new Version("7.30")) >= 0 ? ".PVS-Studio.log" : ".pvslog";
+		static readonly Version _minAnalyzerVersion = new Version("7.30");
+        static readonly Version _analysisPathsSkipVersion = new Version("7.34");
+
+        const string OutputFileExtension = ".PVS-Studio.log";
 
 		public PVSToolChain(ReadOnlyTargetRules Target, VCToolChain InInnerToolchain, ILogger Logger)
 			: base(Logger)
@@ -622,6 +565,12 @@ namespace UnrealBuildTool
 			}
 
 			AnalyzerVersion = GetAnalyzerVersion(AnalyzerFile);
+
+			if (AnalyzerVersion < _minAnalyzerVersion)
+			{
+				throw new BuildLogEventException("PVS-Studio version {Version} is older than the minimum supported version {MinVersion}", AnalyzerVersion, _minAnalyzerVersion);
+			}
+
 			Settings = Target.WindowsPlatform.PVS;
 			ApplicationSettings = Settings.ApplicationSettings;
 
@@ -835,7 +784,7 @@ namespace UnrealBuildTool
 
 			WindowsCompiler windowsCompiler = Target.WindowsPlatform.Compiler;
 			bool isVisualCppCompiler = windowsCompiler.IsMSVC();
-			if (AnalyzerVersion.CompareTo(new Version("7.07")) >= 0 && !shouldCompileAsC)
+			if (!shouldCompileAsC)
 			{
 				VersionNumber compilerVersion = Target.WindowsPlatform.Environment.CompilerVersion;
 				string languageStandardForCfg = GetLangStandForCfgFile(compileEnvironment.CppStandard, compilerVersion);
@@ -873,20 +822,12 @@ namespace UnrealBuildTool
 				configFileContents.AppendFormat($"errors-off={error}\n");
 			}
 
-			int timeout = (int)(Settings.AnalysisTimeoutFlag == AnalysisTimeoutFlags.No_timeout ? 0 : Settings.AnalysisTimeoutFlag);
+			int timeout = Settings.AnalysisTimeoutFlag == AnalysisTimeoutFlags.No_timeout ? 0 : (int)Settings.AnalysisTimeoutFlag;
 			configFileContents.AppendFormat("timeout={0}\n", timeout);
+			configFileContents.Append("silent-exit-code-mode=yes\n");
+			configFileContents.Append("new-output-format=yes\n");
 
-			if (AnalyzerVersion.CompareTo(new Version("7.28")) >= 0)
-			{
-				configFileContents.Append("silent-exit-code-mode=yes\n");
-			}
-
-			if (AnalyzerVersion.CompareTo(new Version("7.30")) >= 0)
-			{
-				configFileContents.Append("new-output-format=yes\n");
-			}
-
-			if (AnalyzerVersion.CompareTo(new Version("7.34")) >= 0)
+			if (AnalyzerVersion.CompareTo(_analysisPathsSkipVersion) >= 0)
 			{
 				if (Target.bStaticAnalyzerProjectOnly)
 				{
@@ -913,6 +854,7 @@ namespace UnrealBuildTool
 			analyzeAction.StatusDescription = baseFileName;
 			analyzeAction.WorkingDirectory = Unreal.EngineSourceDirectory;
 			analyzeAction.CommandPath = AnalyzerFile;
+			analyzeAction.CommandVersion = AnalyzerVersion.ToString();
 
 			List<string> arguments =
 			[
@@ -1022,7 +964,7 @@ namespace UnrealBuildTool
 				FileReference.WriteAllBytesIfDifferent(rootPathsItem.Location, stream.ToArray());
 			}
 
-			string arguments = $"-Input=\"{inputFileListItem.Location}\" -Output=\"{outputFile}\" -Ignored=\"{ignoredFileListItem.Location}\" -RootPaths=\"{rootPathsItem.Location}\" -PrintLevel={target.StaticAnalyzerPVSPrintLevel} -AnalyzerVersion={AnalyzerVersion}";
+			string arguments = $"-Input=\"{inputFileListItem.Location}\" -Output=\"{outputFile}\" -Ignored=\"{ignoredFileListItem.Location}\" -RootPaths=\"{rootPathsItem.Location}\" -PrintLevel={target.StaticAnalyzerPVSPrintLevel}";
 
 			Action finalizeAction = makefileBuilder.CreateRecursiveAction<PVSGatherMode>(ActionType.PostBuildStep, arguments);
 			finalizeAction.CommandDescription = "Process PVS-Studio Results";
