[UBA]
* Fixed some paths that were not virtualized when building mac with vfs
* Enabled so mac linking is going through uba

diff --git a/Engine/Source/Programs/UnrealBuildTool/Executors/UnrealBuildAccelerator/UBAExecutor.cs b/Engine/Source/Programs/UnrealBuildTool/Executors/UnrealBuildAccelerator/UBAExecutor.cs
--- a/Engine/Source/Programs/UnrealBuildTool/Executors/UnrealBuildAccelerator/UBAExecutor.cs
+++ b/Engine/Source/Programs/UnrealBuildTool/Executors/UnrealBuildAccelerator/UBAExecutor.cs
@@ -668,10 +668,8 @@ namespace UnrealBuildTool
 			// the shell results in dyld errors about no matching architecture.
 			// The shell is used to run various commands during a build like copy/ditto.
 			// So for these actions we need to make sure UBA is not used.
-			// Linking is similarly currently not working in Uba on Mac
 			bool bIsShellAction = action.CommandPath == BuildHostPlatform.Current.Shell;
-			bool bIsLinkAction = action.ActionType == ActionType.Link;
-			return (bIsShellAction || bIsLinkAction);
+			return bIsShellAction;
 		}
 
 		/// <summary>
diff --git a/Engine/Source/Programs/UnrealBuildTool/Platform/Mac/MacToolChain.cs b/Engine/Source/Programs/UnrealBuildTool/Platform/Mac/MacToolChain.cs
--- a/Engine/Source/Programs/UnrealBuildTool/Platform/Mac/MacToolChain.cs
+++ b/Engine/Source/Programs/UnrealBuildTool/Platform/Mac/MacToolChain.cs
@@ -232,7 +232,7 @@ namespace UnrealBuildTool
 			base.GetLinkArguments_Global(LinkEnvironment, Arguments);
 			// Pass through architecture and OS info		
 			Arguments.Add(FormatArchitectureArg(LinkEnvironment.Architectures));
-			Arguments.Add(String.Format("-isysroot \"{0}\"", Settings.GetSDKPath(LinkEnvironment.Architecture)));
+			Arguments.Add(String.Format("-isysroot \"{0}\"", NormalizeCommandLinePath(Settings.GetSDKPath(LinkEnvironment.Architecture), LinkEnvironment.RootPaths)));
 			Arguments.Add($"-target {Settings.GetTargetTuple(LinkEnvironment.Architecture, Target!.Platform, Target!.Type)}");
 			Arguments.Add("-dead_strip");
 
@@ -692,6 +692,9 @@ namespace UnrealBuildTool
 			LinkAction.CommandDescription = (bBuildImportLibraryOnly ? "LinkStub" : "Link") + $" [{Arch}]";
 			LinkAction.CommandVersion = Info.ClangVersionString;
 			LinkAction.bProducesImportLibrary = bBuildImportLibraryOnly || LinkEnvironment.bIsBuildingDLL;
+			LinkAction.bCanExecuteInUBA = true;
+			LinkAction.CacheBucket = GetCacheBucket(Target, null);
+			LinkAction.ArtifactMode = ArtifactMode.Enabled;
 
 			string EngineAPIVersion = LoadEngineAPIVersion();
 			string EngineDisplayVersion = LoadEngineDisplayVersion(true);
@@ -825,7 +828,7 @@ namespace UnrealBuildTool
 						}
 
 						LinkAction.PrerequisiteItems.Add(LinkedLib);
-						LinkCommand += $" \"{LinkedLib.FullName}\"";
+						LinkCommand += $" \"{NormalizeCommandLinePath(LinkedLib, LinkEnvironment.RootPaths)}\"";
 					}
 					else if (AdditionalLibrary.Contains(".framework/"))
 					{
@@ -921,7 +924,7 @@ namespace UnrealBuildTool
 			List<string> InputFileNames = new List<string>();
 			foreach (FileItem InputFile in InputFiles)
 			{
-				InputFileNames.Add(String.Format("\"{0}\"", NormalizeCommandLinePath(InputFile)));
+				InputFileNames.Add(String.Format("\"{0}\"", NormalizeCommandLinePath(InputFile, LinkEnvironment.RootPaths)));
 				LinkAction.PrerequisiteItems.Add(InputFile);
 			}
 
@@ -931,7 +934,7 @@ namespace UnrealBuildTool
 			}
 
 			// Add the output file to the command-line.
-			LinkCommand += String.Format(" -o \"{0}\"", OutputFile.AbsolutePath);
+			LinkCommand += String.Format(" -o \"{0}\"", NormalizeCommandLinePath(OutputFile, LinkEnvironment.RootPaths));
 
 			// Add the additional arguments specified by the environment.
 			LinkCommand += LinkEnvironment.AdditionalArguments;
diff --git a/Engine/Source/Programs/UnrealBuildTool/ToolChain/AppleToolChain.cs b/Engine/Source/Programs/UnrealBuildTool/ToolChain/AppleToolChain.cs
--- a/Engine/Source/Programs/UnrealBuildTool/ToolChain/AppleToolChain.cs
+++ b/Engine/Source/Programs/UnrealBuildTool/ToolChain/AppleToolChain.cs
@@ -317,10 +317,6 @@ namespace UnrealBuildTool
 			foreach (FileItem SourceFile in InputFiles)
 			{
 				Action CompileAction = CompileCPPFile(CompileEnvironment, SourceFile, OutputDir, ModuleName, Graph, GlobalArguments, Result);
-				if (CompileEnvironment.PrecompiledHeaderAction == PrecompiledHeaderAction.Create && !CompileEnvironment.RootPaths.bUseVfs)
-				{
-					CompileAction.bCanExecuteInUBA = false;
-				}
 				CompileAction.PrerequisiteItems.UnionWith(FrameworkTokenFiles);
 			}
 			return Result;
@@ -824,7 +820,7 @@ namespace UnrealBuildTool
 				Arguments.Add(" -ld_classic");
 			}
 
-			Arguments.Add($"-L{ApplePlatformSDK.GetToolchainDirectory()}/../lib/swift/{Settings.PlatformDirName}");
+			Arguments.Add($"-L{NormalizeCommandLinePath(ApplePlatformSDK.GetToolchainDirectory(), LinkEnvironment.RootPaths)}/../lib/swift/{Settings.PlatformDirName}");
 			Arguments.Add("-L/usr/lib/swift");
 		}
 
