UnrealBuildTool: Ensure adaptive module actions always have generated ispc header paths added

#rnx

diff --git a/Engine/Source/Programs/UnrealBuildTool/Configuration/UEBuildModuleCPP.cs b/Engine/Source/Programs/UnrealBuildTool/Configuration/UEBuildModuleCPP.cs
--- a/Engine/Source/Programs/UnrealBuildTool/Configuration/UEBuildModuleCPP.cs
+++ b/Engine/Source/Programs/UnrealBuildTool/Configuration/UEBuildModuleCPP.cs
@@ -571,15 +571,15 @@ namespace UnrealBuildTool
 			}
 
 			// Generate ISPC headers before normal c++ actions but after pch since we don't want pch actions to wait for ispc (ispc headers are not allowed in pch)
-			IEnumerable<FileItem> ISPCHeaders = Enumerable.Empty<FileItem>();
 			if (InputFiles.ISPCFiles.Count > 0)
 			{
-				ISPCHeaders = ToolChain.GenerateAllISPCHeaders(CppCompileEnvironment, InputFiles.ISPCFiles, IntermediateDirectory, Graph).GeneratedHeaderFiles;
+				IEnumerable<FileItem> ISPCHeaders = ToolChain.GenerateAllISPCHeaders(CppCompileEnvironment, InputFiles.ISPCFiles, IntermediateDirectory, Graph).GeneratedHeaderFiles;
 				CppCompileEnvironment.AdditionalPrerequisites.AddRange(ISPCHeaders);
+				ModuleCompileEnvironment.AdditionalPrerequisites.AddRange(ISPCHeaders);
 				CppCompileEnvironment.UserIncludePaths.Add(IntermediateDirectory);
+				ModuleCompileEnvironment.UserIncludePaths.Add(IntermediateDirectory);
 			}
 
-
 			// Write all the definitions to a separate file
 			CreateHeaderForDefinitions(CppCompileEnvironment, IntermediateDirectory, null, Graph);
 
@@ -914,8 +914,6 @@ namespace UnrealBuildTool
 			if (InputFiles.ISPCFiles.Count > 0)
 			{
 				CppCompileEnvironment IspcCompileEnvironment = new(ModuleCompileEnvironment);
-				IspcCompileEnvironment.AdditionalPrerequisites.AddRange(ISPCHeaders);
-				IspcCompileEnvironment.UserIncludePaths.Add(IntermediateDirectory);
 				LinkInputFiles.AddRange(ToolChain.CompileAllISPCFiles(IspcCompileEnvironment, InputFiles.ISPCFiles, IntermediateDirectory, Graph).ObjectFiles);
 			}
 
@@ -1057,12 +1055,6 @@ namespace UnrealBuildTool
 				if (HeaderFileItems.Count > 0)
 				{
 					CppCompileEnvironment HeaderCompileEnvironment = new CppCompileEnvironment(ModuleCompileEnvironment);
-					if (InputFiles.ISPCFiles.Count > 0)
-					{
-						HeaderCompileEnvironment.AdditionalPrerequisites.AddRange(ISPCHeaders);
-						HeaderCompileEnvironment.UserIncludePaths.Add(IntermediateDirectory);
-					}
-
 					DirectoryReference HeaderIntermediateDirectory = DirectoryReference.Combine(IntermediateDirectory, "H");
 
 					CreateHeaderForDefinitions(HeaderCompileEnvironment, HeaderIntermediateDirectory, "h", Graph);
