[UBT]
* Enabled uba for link actions on native linux

diff --git a/Engine/Source/Programs/UnrealBuildTool/Platform/Linux/LinuxToolChain.cs b/Engine/Source/Programs/UnrealBuildTool/Platform/Linux/LinuxToolChain.cs
--- a/Engine/Source/Programs/UnrealBuildTool/Platform/Linux/LinuxToolChain.cs
+++ b/Engine/Source/Programs/UnrealBuildTool/Platform/Linux/LinuxToolChain.cs
@@ -1091,15 +1091,12 @@ namespace UnrealBuildTool
 				LinkAction.bCanExecuteRemotely = false;
 			}
 
-			if (!OperatingSystem.IsWindows())
-			{
-				LinkAction.bCanExecuteInUBA = false; // Linker on native linux uses vfork/exec which is not handled in uba right now
-			}
-
 			// because the logic choosing between lld and ld is somewhat messy atm (lld fails to link .DSO due to bugs), make the name of the linker clear
 			LinkAction.CommandDescription += (LinkCommandString.Contains("-fuse-ld=lld")) ? " (lld)" : " (ld)";
 			LinkAction.CommandVersion = Info.ClangVersionString;
 			LinkAction.StatusDescription = Path.GetFileName(OutputFile.AbsolutePath);
+			LinkAction.CacheBucket = GetCacheBucket(TargetRules, null);
+			LinkAction.ArtifactMode = ArtifactMode.Enabled;
 
 			// Add the output file to the command-line.
 			LinkCommandString += String.Format(" -o \"{0}\"", OutputFile.AbsolutePath);
@@ -1445,7 +1442,8 @@ namespace UnrealBuildTool
 				RelinkAction.StatusDescription = LinkAction.StatusDescription;
 				RelinkAction.CommandDescription = "Relink";
 				RelinkAction.bCanExecuteRemotely = false;
-				RelinkAction.bCanExecuteInUBA = OperatingSystem.IsWindows(); // Linker on native linux uses vfork/exec which is not handled in uba right now
+				RelinkAction.CacheBucket = GetCacheBucket(TargetRules, null);
+				RelinkAction.ArtifactMode = ArtifactMode.Enabled;
 				RelinkAction.ProducedItems.Add(RealOutputFile);
 
 				foreach (FileItem Dependency in EngineAndGameLibrariesFiles)
