Unshelved from pending changelist '41071542':

[UBT]
* Changed so path used in apple tool chain does not contain ..

diff --git a/Engine/Source/Programs/UnrealBuildTool/ToolChain/AppleToolChain.cs b/Engine/Source/Programs/UnrealBuildTool/ToolChain/AppleToolChain.cs
--- a/Engine/Source/Programs/UnrealBuildTool/ToolChain/AppleToolChain.cs
+++ b/Engine/Source/Programs/UnrealBuildTool/ToolChain/AppleToolChain.cs
@@ -820,7 +820,8 @@ namespace UnrealBuildTool
 				Arguments.Add(" -ld_classic");
 			}
 
-			Arguments.Add($"-L{NormalizeCommandLinePath(ApplePlatformSDK.GetToolchainDirectory(), LinkEnvironment.RootPaths)}/../lib/swift/{Settings.PlatformDirName}");
+			DirectoryReference dir = DirectoryReference.Combine(ApplePlatformSDK.GetToolchainDirectory(), "../lib/swift", Settings.PlatformDirName);
+			Arguments.Add($"-L{dir}");
 			Arguments.Add("-L/usr/lib/swift");
 		}
 
