UnrealBuildTool: Fix vfs path in VCToolChain.GenerateTypeLibraryHeader

#rnx

diff --git a/Engine/Source/Programs/UnrealBuildTool/Platform/Windows/VCToolChain.cs b/Engine/Source/Programs/UnrealBuildTool/Platform/Windows/VCToolChain.cs
--- a/Engine/Source/Programs/UnrealBuildTool/Platform/Windows/VCToolChain.cs
+++ b/Engine/Source/Programs/UnrealBuildTool/Platform/Windows/VCToolChain.cs
@@ -2601,16 +2601,16 @@ namespace UnrealBuildTool
 				AddSystemIncludePath(Arguments, IncludePath, EnvVars.ToolChain, RootPaths);
 			}
 
-			FileItem ResponseFile = Graph.CreateIntermediateTextFile(OutputFile.ChangeExtension(ResponseExt), String.Join(" ", Arguments));
+			FileItem ResponseFile = Graph.CreateIntermediateTextFile(OutputFile.ChangeExtension(ResponseExt), Arguments);
 
 			// Build the command for touching the output file to update the time stamp
-			string OutfileFullName = NormalizeCommandLinePath(OutputFile).Replace('/', '\\'); // Must be backward slash
+			string OutfileFullName = NormalizeCommandLinePath(OutputFile, RootPaths).Replace('/', '\\'); // Must be backward slash
 			string TouchActionCommand = $"copy /b \"{OutfileFullName}\" + NUL \"{OutfileFullName}\">nul";
 
 			// Build the batch file
 			StringBuilder BatchFileContents = new();
 			BatchFileContents.AppendLine("@echo off");
-			BatchFileContents.AppendLine($"\"{NormalizeCommandLinePath(EnvVars.ToolchainCompilerPath)}\" @\"{NormalizeCommandLinePath(ResponseFile)}\"");
+			BatchFileContents.AppendLine($"\"{NormalizeCommandLinePath(EnvVars.ToolchainCompilerPath, RootPaths)}\" @\"{NormalizeCommandLinePath(ResponseFile, RootPaths)}\"");
 			BatchFileContents.AppendLine(TouchActionCommand);
 
 			FileItem BatchFile = Graph.CreateIntermediateTextFile(OutputFile.ChangeExtension(".bat"), BatchFileContents.ToString());
