UnrealBuildTool: Ensure .tli files have their timestamps updated as well when generating type headers

#rnx

diff --git a/Engine/Source/Programs/UnrealBuildTool/Platform/Windows/VCToolChain.cs b/Engine/Source/Programs/UnrealBuildTool/Platform/Windows/VCToolChain.cs
--- a/Engine/Source/Programs/UnrealBuildTool/Platform/Windows/VCToolChain.cs
+++ b/Engine/Source/Programs/UnrealBuildTool/Platform/Windows/VCToolChain.cs
@@ -2559,13 +2559,15 @@ namespace UnrealBuildTool
 		{
 			CppRootPaths RootPaths = CompileEnvironment.RootPaths;
 
+			FileItem TypeLibraryFile = FileItem.GetItemByPath(TypeLibrary.FileName);
+
 			// Create the input file
 			StringBuilder Contents = new StringBuilder();
 			Contents.AppendLine("#include <windows.h>");
 			Contents.AppendLine("#include <unknwn.h>");
 			Contents.AppendLine();
 
-			Contents.AppendFormat("#import \"{0}\"", NormalizeCommandLinePath(FileReference.FromString(TypeLibrary.FileName), RootPaths).Replace('\\', '/'));
+			Contents.AppendFormat("#import \"{0}\"", NormalizeCommandLinePath(TypeLibraryFile, RootPaths).Replace('\\', '/'));
 			if (!String.IsNullOrEmpty(TypeLibrary.Attributes))
 			{
 				Contents.Append(' ');
@@ -2603,15 +2605,18 @@ namespace UnrealBuildTool
 
 			FileItem ResponseFile = Graph.CreateIntermediateTextFile(OutputFile.ChangeExtension(ResponseExt), Arguments);
 
-			// Build the command for touching the output file to update the time stamp
-			string OutfileFullName = NormalizeCommandLinePath(OutputFile, RootPaths).Replace('/', '\\'); // Must be backward slash
-			string TouchActionCommand = $"copy /b \"{OutfileFullName}\" + NUL \"{OutfileFullName}\">nul";
+			// Build the command for touching the output file(s) to update the time stamp
+			string OutfileTlhFullName = NormalizeCommandLinePath(OutputFile, RootPaths).Replace('/', '\\'); // Must be backward slash
+			string TouchTlhActionCommand = $"if exist \"{OutfileTlhFullName}\" ( copy /b \"{OutfileTlhFullName}\" + NUL \"{OutfileTlhFullName}\">nul )";
+			string OutfileTliFullName = NormalizeCommandLinePath(OutputFile.ChangeExtension(".tli"), RootPaths).Replace('/', '\\'); // Must be backward slash
+			string TouchTliActionCommand = $"if exist \"{OutfileTliFullName}\" ( copy /b \"{OutfileTliFullName}\" + NUL \"{OutfileTliFullName}\">nul )";
 
 			// Build the batch file
 			StringBuilder BatchFileContents = new();
 			BatchFileContents.AppendLine("@echo off");
 			BatchFileContents.AppendLine($"\"{NormalizeCommandLinePath(EnvVars.ToolchainCompilerPath, RootPaths)}\" @\"{NormalizeCommandLinePath(ResponseFile, RootPaths)}\"");
-			BatchFileContents.AppendLine(TouchActionCommand);
+			BatchFileContents.AppendLine(TouchTlhActionCommand);
+			BatchFileContents.AppendLine(TouchTliActionCommand);
 
 			FileItem BatchFile = Graph.CreateIntermediateTextFile(OutputFile.ChangeExtension(".bat"), BatchFileContents.ToString());
 
@@ -2619,6 +2624,7 @@ namespace UnrealBuildTool
 			Action CompileAction = Graph.CreateAction(ActionType.Compile);
 			CompileAction.CommandDescription = "GenerateTLH";
 			CompileAction.PrerequisiteItems.Add(FileItem.GetItemByFileReference(EnvVars.ToolchainCompilerPath));
+			CompileAction.PrerequisiteItems.Add(TypeLibraryFile);
 			CompileAction.PrerequisiteItems.Add(InputFile);
 			CompileAction.PrerequisiteItems.Add(ResponseFile);
 			CompileAction.PrerequisiteItems.Add(BatchFile);
@@ -2627,6 +2633,7 @@ namespace UnrealBuildTool
 			if (OutputHeader != null)
 				CompileAction.ProducedItems.Add(FileItem.GetItemByFileReference(OutputHeader));
 			CompileAction.DeleteItems.Add(FileItem.GetItemByFileReference(OutputFile));
+			CompileAction.DeleteItems.Add(FileItem.GetItemByFileReference(OutputFile.ChangeExtension(".tli"))); // May not be created, depending on Attributes
 			CompileAction.StatusDescription = TypeLibrary.Header;
 			CompileAction.WorkingDirectory = Unreal.EngineSourceDirectory;
 			CompileAction.CommandPath = BuildHostPlatform.Current.Shell;
