[UBT]
* Enabled modules for apple toolchain

diff --git a/Engine/Source/Programs/UnrealBuildTool/ToolChain/AppleToolChain.cs b/Engine/Source/Programs/UnrealBuildTool/ToolChain/AppleToolChain.cs
--- a/Engine/Source/Programs/UnrealBuildTool/ToolChain/AppleToolChain.cs
+++ b/Engine/Source/Programs/UnrealBuildTool/ToolChain/AppleToolChain.cs
@@ -455,6 +455,11 @@ namespace UnrealBuildTool
 			Arguments.Add("-x objective-c++");
 			GetCppStandardCompileArgument(CompileEnvironment, Arguments);
 			Arguments.Add("-stdlib=libc++");
+			Arguments.Add("-fmodules");
+			Arguments.Add("-fno-implicit-modules");
+			Arguments.Add("-fimplicit-module-maps");
+			Arguments.Add("-Wno-module-import-in-extern-c");
+			//Arguments.Add("-fcxx-modules"); // For some reason this does not work
 		}
 
 		protected override void GetCompileArguments_MM(CppCompileEnvironment CompileEnvironment, List<string> Arguments)
