UBT: Optimize TelemetryService.FlushEvents.
BlankProgram null build down from 2.4s to 1.4s.
Wait for all endpoint flushes once instead of flushing each one individually - FlushEvents down from 1.1s to 0.65s.
Send an async dummy head request when adding an endpoint in order to initialize the connection early - FlushEvents down from 0.65s to 0.15s.
Fix threading issue in FlushEventsAsync that could loose events.

#rnx

diff --git a/Engine/Source/Programs/Shared/EpicGames.Core/Telemetry/DataRouterTelemetryService.cs b/Engine/Source/Programs/Shared/EpicGames.Core/Telemetry/DataRouterTelemetryService.cs
--- a/Engine/Source/Programs/Shared/EpicGames.Core/Telemetry/DataRouterTelemetryService.cs
+++ b/Engine/Source/Programs/Shared/EpicGames.Core/Telemetry/DataRouterTelemetryService.cs
@@ -353,6 +353,8 @@ namespace EpicGames.Core.Telemetry
 
 		private readonly Timer _autoFlushTimer = new Timer();
 
+		private Task<HttpResponseMessage> InitialHeadRequest;
+
 		/// <summary>
 		/// Length of an empty data router post when serialized to json
 		/// </summary>
@@ -409,6 +411,7 @@ namespace EpicGames.Core.Telemetry
 			HttpClient = inHttpClient;
 			BaseAddress = inBaseAddress;
 			SetHttpClientHeaders();
+			InitialHeadRequest = HeadRequestAsync();
 
 			_autoFlushTimer.Interval = AutoFlushIntervalMilliseconds;
 			_autoFlushTimer.Elapsed += AutoFlushEventsAsync;
@@ -510,6 +513,16 @@ namespace EpicGames.Core.Telemetry
 			return await HttpClient.SendAsync(postRequest);
 		}
 
+		/// <summary>
+		/// Async function to issue a single async head request
+		/// </summary>
+		/// <returns>Http responce</returns>
+		private async Task<HttpResponseMessage> HeadRequestAsync()
+		{
+			using HttpRequestMessage headRequest = new HttpRequestMessage(HttpMethod.Head, $"{BaseAddress}");
+			return await Task.Run(async () => await HttpClient.SendAsync(headRequest));
+		}
+
 		/// <summary>
 		/// Async function to flush events
 		/// </summary>
@@ -517,7 +530,8 @@ namespace EpicGames.Core.Telemetry
 		public virtual async Task FlushEventsAsync()
 		{
 			ConcurrentDictionary<string, DataRouterPost> posts = new ConcurrentDictionary<string, DataRouterPost>();
-			foreach (DataRouterEvent eventData in EventsQueue)
+			DataRouterEvent? eventData;
+			while (EventsQueue.TryDequeue(out eventData))
 			{
 				string requestUri = eventData.GetUrlParameters();
 				if (!posts.ContainsKey(requestUri))
@@ -528,11 +542,9 @@ namespace EpicGames.Core.Telemetry
 					}
 				}
 
-				posts[requestUri].Events.Add(eventData);
+				posts[requestUri].Events.Add((DataRouterEvent)eventData);
 			}
 
-			EventsQueue.Clear();
-
 			if (IsDryRun)
 			{
 				foreach (KeyValuePair<string, DataRouterPost> postData in posts)
diff --git a/Engine/Source/Programs/UnrealBuildTool/System/Telemetry.cs b/Engine/Source/Programs/UnrealBuildTool/System/Telemetry.cs
--- a/Engine/Source/Programs/UnrealBuildTool/System/Telemetry.cs
+++ b/Engine/Source/Programs/UnrealBuildTool/System/Telemetry.cs
@@ -6,6 +6,7 @@ using System.Collections.Generic;
 using System.Linq;
 using System.Net.Http;
 using System.Text.Json.Serialization;
+using System.Threading.Tasks;
 using EpicGames.Core;
 using EpicGames.Core.Telemetry;
 using UnrealBuildBase;
@@ -188,29 +189,24 @@ namespace UnrealBuildTool
 
 	internal class TelemetryEndpoint
 	{
-		readonly Lazy<DataRouterTelemetryService> _service;
-		readonly Lazy<HttpClient> _httpClient;
+		readonly DataRouterTelemetryService _service;
+		readonly HttpClient _httpClient;
 		readonly Uri _baseAddress;
 		readonly string _appId;
 		readonly Lazy<TelemetryMetadata> _metadata;
 
-		public TelemetryEndpoint(Lazy<HttpClient> httpClient, Uri baseAddress, string appId, Lazy<TelemetryMetadata> metadata)
+		public TelemetryEndpoint(HttpClient httpClient, Uri baseAddress, string appId, Lazy<TelemetryMetadata> metadata)
 		{
 			_httpClient = httpClient;
 			_baseAddress = baseAddress;
 			_appId = appId;
 			_metadata = metadata;
-			_service = new(() =>
-			{
-				DataRouterTelemetryService service = new(_httpClient.Value, _baseAddress);
-				//service.IsDryRun = false;
-				return service;
-			});
+			_service = new(_httpClient, _baseAddress);
 		}
 
-		public void FlushEvents() => _service.Value.FlushEvents();
+		public async Task FlushEventsAsync() => await _service.FlushEventsAsync();
 
-		public void RecordEvent(TelemetryEvent eventData) => _service.Value.RecordEvent(eventData.CloneWithMetadata(_appId, _metadata.Value));
+		public void RecordEvent(TelemetryEvent eventData) => _service.RecordEvent(eventData.CloneWithMetadata(_appId, _metadata.Value));
 	}
 
 	internal readonly struct TelemetryMetadata
@@ -387,7 +383,7 @@ namespace UnrealBuildTool
 			{
 				Uri baseAddress = new($"{apiServerET}{apiEndpointET}");
 				Tuple<Uri, string> key = new(baseAddress, apiKeyET);
-				_endpoints.GetOrAdd(key, _ => new TelemetryEndpoint(_httpClient, baseAddress, apiKeyET, _metadata));
+				_endpoints.GetOrAdd(key, new TelemetryEndpoint(_httpClient.Value, baseAddress, apiKeyET, _metadata));
 				return true;
 			}
 
@@ -417,10 +413,13 @@ namespace UnrealBuildTool
 		/// <inheritdoc/>
 		public void FlushEvents()
 		{
+			List<Task> flushTasks = new List<Task>();
+
 			foreach (KeyValuePair<Tuple<Uri, string>, TelemetryEndpoint> endpoint in _endpoints)
 			{
-				endpoint.Value.FlushEvents();
+				flushTasks.Add(endpoint.Value.FlushEventsAsync());
 			}
+			Task.WaitAll(flushTasks.ToArray());
 		}
 		#endregion
 
@@ -454,4 +453,4 @@ namespace UnrealBuildTool
 		}
 		#endregion
 	}
-}
\ No newline at end of file
+}
