[UBA]
* Always populate helper from xcode if it exists

diff --git a/Engine/Source/Programs/UnrealBuildTool/Executors/UnrealBuildAccelerator/UBAAgentCoordinatorHorde.cs b/Engine/Source/Programs/UnrealBuildTool/Executors/UnrealBuildAccelerator/UBAAgentCoordinatorHorde.cs
--- a/Engine/Source/Programs/UnrealBuildTool/Executors/UnrealBuildAccelerator/UBAAgentCoordinatorHorde.cs
+++ b/Engine/Source/Programs/UnrealBuildTool/Executors/UnrealBuildAccelerator/UBAAgentCoordinatorHorde.cs
@@ -570,7 +570,7 @@ namespace UnrealBuildTool
 						//arguments.Add("-NoStore");
 						//arguments.Add("-KillRandom"); // For debugging
 
-						if (OperatingSystem.IsMacOS() && EpicGames.UBA.Utils.DisallowedPaths.Any())
+						if (OperatingSystem.IsMacOS())// && EpicGames.UBA.Utils.DisallowedPaths.Any())
 						{
 							// we need to populate the cas with all known xcodes so we can serve all the ones we have installed
 							string xcodeVersion = Utils.RunLocalProcessAndReturnStdOut("/bin/sh", "-c '/usr/bin/defaults read $(xcode-select -p)/../version.plist ProductBuildVersion");
