Fixed plugin descriptor reading some boolean values with the `b` prefix.

diff --git a/Engine/Source/Programs/UnrealBuildTool/System/PluginDescriptor.cs b/Engine/Source/Programs/UnrealBuildTool/System/PluginDescriptor.cs
--- a/Engine/Source/Programs/UnrealBuildTool/System/PluginDescriptor.cs
+++ b/Engine/Source/Programs/UnrealBuildTool/System/PluginDescriptor.cs
@@ -349,13 +349,13 @@ namespace UnrealBuildTool
 			}
 
 			bool bPluginEnableSceneGraph;
-			if (RawObject.TryGetBoolField("bEnableSceneGraph", out bPluginEnableSceneGraph))
+			if (RawObject.TryGetBoolField("EnableSceneGraph", out bPluginEnableSceneGraph))
 			{
 				bEnableSceneGraph = bPluginEnableSceneGraph;
 			}
 
 			bool bPluginEnableVerseAssetReflection;
-			if (RawObject.TryGetBoolField("bEnableVerseAssetReflection", out bPluginEnableVerseAssetReflection))
+			if (RawObject.TryGetBoolField("EnableVerseAssetReflection", out bPluginEnableVerseAssetReflection))
 			{
 				bEnableVerseAssetReflection = bPluginEnableVerseAssetReflection;
 			}
