Speed up xcode project files generation

#rnx

diff --git a/Engine/Source/Programs/UnrealBuildTool/ProjectFiles/ProjectFileGenerator.cs b/Engine/Source/Programs/UnrealBuildTool/ProjectFiles/ProjectFileGenerator.cs
--- a/Engine/Source/Programs/UnrealBuildTool/ProjectFiles/ProjectFileGenerator.cs
+++ b/Engine/Source/Programs/UnrealBuildTool/ProjectFiles/ProjectFileGenerator.cs
@@ -459,6 +459,17 @@ namespace UnrealBuildTool
 
 		readonly ProjectDescriptor? OnlyGameProjectDescriptor;
 
+		/// <summary>
+		/// Cache for GlobalCompileEnvironment, to reduce GPF time, currently only useful on Mac
+		/// </summary>
+
+		protected ConcurrentDictionary<string, CppCompileEnvironment> GlobalCompileEnvironmentCache = new();
+
+		protected string MakeGlobalCompileEnvironmentCacheKey(FileReference? UnrealProjectFilePath, string TargetName, UnrealTargetPlatform Platform, UnrealTargetConfiguration Config, UnrealArchitectures Archs)
+		{
+			return (UnrealProjectFilePath != null ? UnrealProjectFilePath.FullName : "") + "|" + TargetName + "|" + Platform.ToString() + "|" + Config.ToString() + "|" + Archs.ToString();
+		}
+
 		/// <summary>
 		/// File extension for project files we'll be generating (e.g. ".vcxproj")
 		/// </summary>
@@ -2265,6 +2276,7 @@ namespace UnrealBuildTool
 
 				ConcurrentBag<UEBuildTarget> BuildTargets = [];
 				ConcurrentDictionary<UEBuildTarget, ProjectFile> BuildTargetToProjectFile = new();
+				ConcurrentDictionary<UEBuildTarget, string> BuildTargetToCompileEnvironmentCacheKey = new();
 				using ITimelineEvent BuildTargetsTimer = Timeline.ScopeEvent("Creating Build Targets", Logger);
 				using (ProgressWriter Progress = new ProgressWriter("Creating Build Targets", true, Logger))
 				{
@@ -2299,6 +2311,7 @@ namespace UnrealBuildTool
 								UEBuildTarget Target = UEBuildTarget.Create(TargetDesc, true, false, bUsePrecompiled, UnrealIntermediateEnvironment.GenerateProjectFiles, Logger);
 								BuildTargets.Add(Target);
 								BuildTargetToProjectFile.TryAdd(Target, TargetProjectFile);
+								BuildTargetToCompileEnvironmentCacheKey.TryAdd(Target, MakeGlobalCompileEnvironmentCacheKey(CurTarget.UnrealProjectFilePath, CurTarget.Name, IntellisensePlatform, UnrealTargetConfiguration.Development, DefaultArchitecture));
 							}
 						}
 						catch (Exception Ex)
@@ -2333,8 +2346,17 @@ namespace UnrealBuildTool
 						// If the project generator just cares about the result of UEBuildTarget.Create, skip generating the compile environments.
 						if (ShouldGenerateIntelliSenseCompileEnvironments() && BuildTargetToProjectFile.TryGetValue(Target, out ProjectFile? TargetProjectFile))
 						{
-							// Generate a compile environment for each module in the binary
-							CppCompileEnvironment GlobalCompileEnvironment = Target.CreateCompileEnvironmentForProjectFiles(Logger);
+							// Generate a compile environment for each module in the binary, but read from cache if possible
+							CppCompileEnvironment GlobalCompileEnvironment;
+							if (!GlobalCompileEnvironmentCache.IsEmpty && BuildTargetToCompileEnvironmentCacheKey.TryGetValue(Target, out string? CacheKey) && GlobalCompileEnvironmentCache.ContainsKey(CacheKey))
+							{
+								GlobalCompileEnvironment = GlobalCompileEnvironmentCache[CacheKey];
+							}
+							else
+							{
+								GlobalCompileEnvironment = Target.CreateCompileEnvironmentForProjectFiles(Logger);
+							}
+							
 							foreach (UEBuildBinary Binary in Target.Binaries)
 							{
 								CppCompileEnvironment BinaryCompileEnvironment = Binary.CreateBinaryCompileEnvironment(GlobalCompileEnvironment);
@@ -3465,7 +3487,7 @@ namespace UnrealBuildTool
 		}
 
 		/// <summary>
-		/// Allows the generator to XXXXXXXxXXXXXXXxXXXXXXXxXXXXXXXx
+		/// Allows the generator to get additional info as needed, currently only used by Xcode projects for finding frameworks/dylibs/bundles
 		/// </summary>
 		/// <param name="PlatformProjectGenerators"></param>
 		/// <param name="Targets"></param>
diff --git a/Engine/Source/Programs/UnrealBuildTool/ProjectFiles/Xcode/XcodeProjectFileGenerator.cs b/Engine/Source/Programs/UnrealBuildTool/ProjectFiles/Xcode/XcodeProjectFileGenerator.cs
--- a/Engine/Source/Programs/UnrealBuildTool/ProjectFiles/Xcode/XcodeProjectFileGenerator.cs
+++ b/Engine/Source/Programs/UnrealBuildTool/ProjectFiles/Xcode/XcodeProjectFileGenerator.cs
@@ -1,6 +1,7 @@
 // Copyright Epic Games, Inc. All Rights Reserved.
 
 using System;
+using System.Collections.Concurrent;
 using System.Collections.Generic;
 using System.IO;
 using System.Linq;
@@ -525,6 +526,7 @@ namespace UnrealBuildTool
 
 		protected override void AddAdditionalNativeTargetInformation(PlatformProjectGeneratorCollection PlatformProjectGenerators, List<Tuple<ProjectFile, ProjectTarget>> Targets, ILogger Logger)
 		{
+			using ITimelineEvent Timer = Timeline.ScopeEvent("Adding additional Xcode project info", Logger);
 			DateTime MainStart = DateTime.UtcNow;
 
 			ParallelOptions Options = new ParallelOptions
@@ -561,7 +563,7 @@ namespace UnrealBuildTool
 				{
 					UnrealArch Arch = UnrealArch.Arm64;
 					TargetDescriptor TargetDesc = new TargetDescriptor(CurTarget.UnrealProjectFilePath, CurTarget.Name, Platform, UnrealTargetConfiguration.Development,
-						new UnrealArchitectures(Arch), new CommandLineArguments(["-skipclangvalidation"]));
+						new UnrealArchitectures(Arch), null);
 					TargetDesc.IntermediateEnvironment = UnrealIntermediateEnvironment.GenerateProjectFiles;
 					DateTime Start = DateTime.UtcNow;
 
@@ -573,8 +575,10 @@ namespace UnrealBuildTool
 						List<UEBuildFramework> Frameworks = new();
 						List<UEBuildBundleResource> Bundles = new();
 						List<ModuleRules.RuntimeDependency> Dylibs = new();
-						// Generate a compile environment for each module in the binary
+						// Generate a compile environment for each module in the binary, and cache the one per Target, since that is the slowest one
 						CppCompileEnvironment GlobalCompileEnvironment = Target.CreateCompileEnvironmentForProjectFiles(Logger);
+						GlobalCompileEnvironmentCache.TryAdd(MakeGlobalCompileEnvironmentCacheKey(CurTarget.UnrealProjectFilePath, CurTarget.Name, Platform, UnrealTargetConfiguration.Development, new UnrealArchitectures(Arch)), GlobalCompileEnvironment);
+
 						foreach (UEBuildBinary Binary in Target.Binaries)
 						{
 							CppCompileEnvironment BinaryCompileEnvironment = Binary.CreateBinaryCompileEnvironment(GlobalCompileEnvironment);
@@ -622,7 +626,7 @@ namespace UnrealBuildTool
 					Logger.LogDebug("GettingNativeInfo [{Project} / {Platform}] {TimeMs}ms", TargetProjectFile.ProjectFilePath.GetFileNameWithoutAnyExtensions(), Platform, (DateTime.UtcNow - Start).TotalMilliseconds);
 				}
 			});
-			Logger.LogInformation("GettingNativeInfo {TimeMs}ms overall", (DateTime.UtcNow - MainStart).TotalMilliseconds);
+			Timer.Finish();
 		}
 	}
 }
