Expose MicrosoftPlatformSDK.TryGetToolChainDir in WindowsExports

#rnx

diff --git a/Engine/Source/Programs/UnrealBuildTool/Platform/Windows/WindowsExports.cs b/Engine/Source/Programs/UnrealBuildTool/Platform/Windows/WindowsExports.cs
--- a/Engine/Source/Programs/UnrealBuildTool/Platform/Windows/WindowsExports.cs
+++ b/Engine/Source/Programs/UnrealBuildTool/Platform/Windows/WindowsExports.cs
@@ -105,5 +105,26 @@ namespace UnrealBuildTool
 		{
 			MicrosoftPlatformSDK.EnumerateSdkRootDirs(RootDirs, Logger);
 		}
+
+		/// <summary>
+		/// Determines the directory containing the MSVC toolchain
+		/// </summary>
+		/// <param name="Compiler">Major version of the compiler to use</param>
+		/// <param name="CompilerVersion">The minimum compiler version to use</param>
+		/// <param name="Architecture">Architecture that is required</param>
+		/// <param name="Logger">Logger for output</param>
+		/// <param name="OutToolChainVersion">Receives the chosen toolchain version</param>
+		/// <param name="OutToolChainDir">Receives the directory containing the toolchain</param>
+		/// <param name="OutRedistDir">Receives the optional directory containing redistributable components</param>
+		/// <returns>True if the toolchain directory was found correctly</returns>
+		[SupportedOSPlatform("windows")]
+		public static bool TryGetToolChainDir(WindowsCompiler Compiler, string? CompilerVersion, UnrealArch Architecture, ILogger Logger, [NotNullWhen(true)] out VersionNumber? OutToolChainVersion, [NotNullWhen(true)] out DirectoryReference? OutToolChainDir, out DirectoryReference? OutRedistDir)
+		{
+			OutToolChainVersion = null;
+			OutToolChainDir = null;
+			OutRedistDir = null;
+			return MicrosoftPlatformSDK.TryGetToolChainDir(Compiler, CompilerVersion, Architecture, Logger, out OutToolChainVersion, out OutToolChainDir, out OutRedistDir);
+		}
+
 	}
 }
