ApplyWarnings Default values for target and regular ctor refactor to reduce boilerplate

diff --git a/Engine/Source/Programs/Shared/UnrealBuildTool.Tests/CppCompileWarningsTests.cs b/Engine/Source/Programs/Shared/UnrealBuildTool.Tests/CppCompileWarningsTests.cs
--- a/Engine/Source/Programs/Shared/UnrealBuildTool.Tests/CppCompileWarningsTests.cs
+++ b/Engine/Source/Programs/Shared/UnrealBuildTool.Tests/CppCompileWarningsTests.cs
@@ -417,7 +417,7 @@ namespace UnrealBuildTool.Tests
 					Assert.IsFalse(args.Any());
 
 					args = baseWarnings.GenerateWarningCommandLineArgs(validCompileEnvironment, typeof(ClangToolChain), newVersionNumber);
-					Assert.IsTrue(args.Count() == 1);
+					Assert.IsTrue(args.Count() == 2);
 					Assert.IsTrue(args.Contains("-Wno-unused-value"));
 
 					baseWarnings.UnusedValueWarningLevel = WarningLevel.Default;
@@ -590,10 +590,10 @@ namespace UnrealBuildTool.Tests
 					Assert.IsFalse(args.Any());
 
 					args = baseWarnings.GenerateWarningCommandLineArgs(validCompileEnvironment, typeof(ClangToolChain), newVersionNumber, StaticAnalyzer.None);
-					Assert.IsTrue(args.Count() == 1 && args.Contains("-Wno-unused-value"));
+					Assert.IsTrue(args.Count() == 2 && args.Contains("-Wno-unused-value"));
 
 					args = baseWarnings.GenerateWarningCommandLineArgs(validCompileEnvironment, typeof(ClangToolChain), newVersionNumber, StaticAnalyzer.Clang);
-					Assert.IsTrue(args.Count() == 1 && args.Contains("-Wno-unused-value"));
+					Assert.IsTrue(args.Count() == 2 && args.Contains("-Wno-unused-value"));
 				}
 
 				// VCClangCompilerFilter
diff --git a/Engine/Source/Programs/Shared/UnrealBuildTool.Tests/WarningLevelDefaultAttributeTests.cs b/Engine/Source/Programs/Shared/UnrealBuildTool.Tests/WarningLevelDefaultAttributeTests.cs
new file mode 100644
--- /dev/null
+++ b/Engine/Source/Programs/Shared/UnrealBuildTool.Tests/WarningLevelDefaultAttributeTests.cs
@@ -0,0 +1,120 @@
+// Copyright Epic Games, Inc. All Rights Reserved.
+
+using System.Collections.Generic;
+using Microsoft.VisualStudio.TestTools.UnitTesting;
+using UnrealBuildTool.Configuration.CompileWarnings;
+
+namespace UnrealBuildTool.Tests
+{
+	[TestClass]
+	public class WarningLevelDefaultAttributeTests
+	{
+		[TestMethod]
+		public void TestWindowsCompletelyDefined()
+		{
+			List<WarningLevelDefaultAttribute> attributes = new List<WarningLevelDefaultAttribute>
+			{
+				new WarningLevelDefaultAttribute( WarningLevel.Off, BuildSettingsVersion.V1, BuildSettingsVersion.V2),
+				new WarningLevelDefaultAttribute(WarningLevel.Warning, BuildSettingsVersion.V3, BuildSettingsVersion.V4),
+				new WarningLevelDefaultAttribute(WarningLevel.Error, BuildSettingsVersion.V5, BuildSettingsVersion.Latest)
+			};
+
+			List<WarningLevelDefaultAttribute> merged = WarningLevelDefaultAttribute.EnsureWarningLevelDefaultBounds(attributes);
+
+			Assert.AreEqual(3, merged.Count);
+			Assert.AreEqual(BuildSettingsVersion.V1, merged[0].MinVersion);
+			Assert.AreEqual(BuildSettingsVersion.V2, merged[0].MaxVersion);
+			Assert.AreEqual(WarningLevel.Off, merged[0].DefaultLevel);
+
+			Assert.AreEqual(BuildSettingsVersion.V3, merged[1].MinVersion);
+			Assert.AreEqual(BuildSettingsVersion.V4, merged[1].MaxVersion);
+			Assert.AreEqual(WarningLevel.Warning, merged[1].DefaultLevel);
+
+			Assert.AreEqual(BuildSettingsVersion.V5, merged[2].MinVersion);
+			Assert.AreEqual(BuildSettingsVersion.Latest, merged[2].MaxVersion);
+			Assert.AreEqual(WarningLevel.Error, merged[2].DefaultLevel);
+		}
+
+		[TestMethod]
+		public void TestGapResolveToLowerBound()
+		{
+			List<WarningLevelDefaultAttribute> attributes = new List<WarningLevelDefaultAttribute>
+			{
+				new WarningLevelDefaultAttribute(WarningLevel.Warning, BuildSettingsVersion.V1, BuildSettingsVersion.V2),
+				new WarningLevelDefaultAttribute(WarningLevel.Error, BuildSettingsVersion.V4, BuildSettingsVersion.V5),
+			};
+
+			List<WarningLevelDefaultAttribute> merged = WarningLevelDefaultAttribute.EnsureWarningLevelDefaultBounds(attributes);
+
+			Assert.AreEqual(2, merged.Count);
+			Assert.AreEqual(BuildSettingsVersion.V1, merged[0].MinVersion);
+			Assert.AreEqual(BuildSettingsVersion.V3, merged[0].MaxVersion);
+			Assert.AreEqual(WarningLevel.Warning, merged[0].DefaultLevel);
+
+			Assert.AreEqual(BuildSettingsVersion.V4, merged[1].MinVersion);
+			Assert.AreEqual(BuildSettingsVersion.V5, merged[1].MaxVersion);
+			Assert.AreEqual(WarningLevel.Error, merged[1].DefaultLevel);
+		}
+
+		[TestMethod]
+		public void TestLowerUpperOverlapResizeBottomUpward()
+		{
+			List<WarningLevelDefaultAttribute> attributes = new List<WarningLevelDefaultAttribute>
+			{
+				new WarningLevelDefaultAttribute(WarningLevel.Warning, BuildSettingsVersion.V1, BuildSettingsVersion.V3),
+				new WarningLevelDefaultAttribute(WarningLevel.Error, BuildSettingsVersion.V2, BuildSettingsVersion.V4)
+			};
+
+			List<WarningLevelDefaultAttribute> merged = WarningLevelDefaultAttribute.EnsureWarningLevelDefaultBounds(attributes);
+
+			Assert.AreEqual(2, merged.Count);
+			Assert.AreEqual(BuildSettingsVersion.V1, merged[0].MinVersion);
+			Assert.AreEqual(BuildSettingsVersion.V3, merged[0].MaxVersion);
+			Assert.AreEqual(WarningLevel.Warning, merged[0].DefaultLevel);
+
+			Assert.AreEqual(BuildSettingsVersion.V4, merged[1].MinVersion);
+			Assert.AreEqual(BuildSettingsVersion.V4, merged[1].MaxVersion);
+			Assert.AreEqual(WarningLevel.Error, merged[1].DefaultLevel);
+		}
+
+		[TestMethod]
+		public void TestAllRangeWide()
+		{
+			List<WarningLevelDefaultAttribute> attributes = new List<WarningLevelDefaultAttribute>
+			{
+				new WarningLevelDefaultAttribute( WarningLevel.Warning, BuildSettingsVersion.V1, BuildSettingsVersion.V2),
+				new WarningLevelDefaultAttribute(WarningLevel.Error, BuildSettingsVersion.V1, BuildSettingsVersion.V4)
+			};
+
+			List<WarningLevelDefaultAttribute> merged = WarningLevelDefaultAttribute.EnsureWarningLevelDefaultBounds(attributes);
+
+			Assert.AreEqual(2, merged.Count);
+
+			Assert.AreEqual(BuildSettingsVersion.V1, merged[0].MinVersion);
+			Assert.AreEqual(BuildSettingsVersion.V2, merged[0].MaxVersion);
+			Assert.AreEqual(WarningLevel.Warning, merged[0].DefaultLevel);
+
+			Assert.AreEqual(BuildSettingsVersion.V3, merged[1].MinVersion);
+			Assert.AreEqual(BuildSettingsVersion.V4, merged[1].MaxVersion);
+			Assert.AreEqual(WarningLevel.Error, merged[1].DefaultLevel);
+		}
+
+		[TestMethod]
+		public void TestResolveLowerBound()
+		{
+			List<WarningLevelDefaultAttribute> attributes = new List<WarningLevelDefaultAttribute>
+			{
+				new WarningLevelDefaultAttribute(WarningLevel.Warning, BuildSettingsVersion.V1, BuildSettingsVersion.V2),
+				new WarningLevelDefaultAttribute(WarningLevel.Error, BuildSettingsVersion.V3, BuildSettingsVersion.V4)
+			};
+
+			Assert.AreEqual(WarningLevel.Warning, WarningLevelDefaultAttribute.ResolveWarningLevelDefault(attributes, BuildSettingsVersion.V1));
+			Assert.AreEqual(WarningLevel.Warning, WarningLevelDefaultAttribute.ResolveWarningLevelDefault(attributes, BuildSettingsVersion.V2));
+
+			Assert.AreEqual(WarningLevel.Error, WarningLevelDefaultAttribute.ResolveWarningLevelDefault(attributes, BuildSettingsVersion.V3));
+			Assert.AreEqual(WarningLevel.Error, WarningLevelDefaultAttribute.ResolveWarningLevelDefault(attributes, BuildSettingsVersion.V4));
+
+			Assert.AreEqual(WarningLevel.Default, WarningLevelDefaultAttribute.ResolveWarningLevelDefault(attributes, BuildSettingsVersion.V5));
+		}
+	}
+}
\ No newline at end of file
diff --git a/Engine/Source/Programs/UnrealBuildTool/Configuration/CompileWarnings/WarningLevelDefaultAttribute.cs b/Engine/Source/Programs/UnrealBuildTool/Configuration/CompileWarnings/WarningLevelDefaultAttribute.cs
new file mode 100644
--- /dev/null
+++ b/Engine/Source/Programs/UnrealBuildTool/Configuration/CompileWarnings/WarningLevelDefaultAttribute.cs
@@ -0,0 +1,148 @@
+// Copyright Epic Games, Inc. All Rights Reserved.
+
+using System;
+using System.Collections.Generic;
+using System.Linq;
+using EpicGames.Core;
+using Microsoft.Extensions.Logging;
+
+namespace UnrealBuildTool.Configuration.CompileWarnings
+{
+	/// <summary>
+	/// Flag used to communicate which context the <see cref="WarningLevelDefaultAttribute"/> should be applied.
+	/// </summary>
+	internal enum InitializationContext
+	{
+		/// <summary>
+		/// Apply the <see cref="WarningLevelDefaultAttribute"/> in all contexts.
+		/// </summary>
+		Any,
+		/// <summary>
+		/// Apply the <see cref="WarningLevelDefaultAttribute"/> exclusively in the <see cref="CppCompileWarnings.ApplyTargetDefaults(CppCompileWarnings, Boolean)"/> context.
+		/// </summary>
+		Target,
+		/// <summary>
+		/// Apply the <see cref="WarningLevelDefaultAttribute"/> exclusively in the <see cref="CppCompileWarnings.ApplyDefaults"/> context.
+		/// </summary>
+		Constructor
+	}
+
+	/// <summary>
+	/// Attribute used to set the default values of a <see cref="WarningLevel"/> under the context of <see cref="CppCompileWarnings.ApplyDefaults"/> context and <see cref="CppCompileWarnings.ApplyTargetDefaults(CppCompileWarnings, Boolean)"/> context.
+	/// </summary>
+	[AttributeUsage(AttributeTargets.Field | AttributeTargets.Property, AllowMultiple = true)]
+	internal sealed class WarningLevelDefaultAttribute : Attribute
+	{
+		private static readonly ILogger? s_logger = Log.Logger;
+
+		/// <summary>
+		/// The lower bound (inclusive) <see cref="BuildSettingsVersion"/> of which to apply this <see cref="DefaultLevel"/> to.
+		/// </summary>
+		public BuildSettingsVersion MinVersion { get; }
+
+		/// <summary>
+		/// The upper bound (inclusive) <see cref="BuildSettingsVersion"/> of which to apply this <see cref="DefaultLevel"/> to.
+		/// </summary>
+		public BuildSettingsVersion MaxVersion { get; }
+
+		/// <summary>
+		/// The default warning level.
+		/// </summary>
+		public WarningLevel DefaultLevel { get; }
+
+		/// <summary>
+		/// The context in which this attribute should be applied.
+		/// </summary>
+		public InitializationContext Context { get; }
+
+		/// <summary>
+		/// Constructs a ApplyWarningsDefaultValueAttribute.
+		/// </summary>
+		/// <param name="minVersion">Minimum version of the bound to apply this setting to. Inclusive.</param>
+		/// <param name="maxVersion">Maximum version of the bound to apply this setting to. Inclusive.</param>
+		/// <param name="defaultLevel">The default warning level to apply for this bound.</param>
+		/// <param name="context">The context in which to apply this default.</param>
+		public WarningLevelDefaultAttribute(WarningLevel defaultLevel = WarningLevel.Default, BuildSettingsVersion minVersion = BuildSettingsVersion.V1, BuildSettingsVersion maxVersion = BuildSettingsVersion.Latest, InitializationContext context = InitializationContext.Any)
+		{
+			MinVersion = minVersion;
+			MaxVersion = maxVersion;
+			DefaultLevel = defaultLevel;
+			Context = context;
+		}
+
+		/// <summary>
+		/// Resolves the set of <see cref="WarningLevelDefaultAttribute"/> to a <see cref="WarningLevel"/> provided the current <see cref="BuildSettingsVersion"/>.
+		/// </summary>
+		/// <param name="unsortedDefaultValueAttributes">The set of attributes to consider.</param>
+		/// <param name="activeVersion">The current build settings version.</param>
+		/// <returns>The corresponding <see cref="WarningLevel"/> of which to apply to the property with the default value attributes.</returns>
+		/// <remarks>Will invoke <see cref="EnsureWarningLevelDefaultBounds(IList{WarningLevelDefaultAttribute})"/> on the input <see cref="WarningLevelDefaultAttribute"/>.</remarks>
+		internal static WarningLevel ResolveWarningLevelDefault(IList<WarningLevelDefaultAttribute> unsortedDefaultValueAttributes, BuildSettingsVersion activeVersion)
+		{
+			List<WarningLevelDefaultAttribute> sortedAndMerged = EnsureWarningLevelDefaultBounds(unsortedDefaultValueAttributes);
+
+			WarningLevel returnWarningLevel = WarningLevel.Default;
+
+			foreach (WarningLevelDefaultAttribute attr in sortedAndMerged)
+			{
+				// If we find our appropriate range, we early out.
+				if (attr.MinVersion <= activeVersion && activeVersion <= attr.MaxVersion)
+				{
+					returnWarningLevel = attr.DefaultLevel;
+					break;
+				}
+			}
+
+			return returnWarningLevel;
+		}
+
+		/// <summary>
+		/// Ensures the attributes collection represents a contiguous, non-overlapping range.
+		/// </summary>
+		/// <param name="attributes">The list of attributes to verify.</param>
+		/// <returns>A list of attributes that has no overlaps, and is contiguous.</returns>
+		internal static List<WarningLevelDefaultAttribute> EnsureWarningLevelDefaultBounds(IList<WarningLevelDefaultAttribute> attributes)
+		{
+			List<WarningLevelDefaultAttribute> sorted = attributes.OrderBy(a => a.MinVersion).ToList();
+			List<WarningLevelDefaultAttribute> merged = new List<WarningLevelDefaultAttribute>(sorted.Count);
+
+			for (int i = 0; i < sorted.Count - 1; i++)
+			{
+				WarningLevelDefaultAttribute current = sorted[i];
+				WarningLevelDefaultAttribute next = sorted[i + 1];
+
+				if (current.MaxVersion < next.MinVersion - 1)
+				{
+					s_logger?.LogWarning("Malformed ApplyWarningsDefaultValueAttribute collection; taking corrective action to address gap in range (CurrentMax: {CurrentMax} NextMin: {NextMin}-1).", next.MinVersion, current.MaxVersion);
+
+					// Extend current range to cover the gap up to the next MinVersion,using the old build settings version standard.
+					current = new WarningLevelDefaultAttribute(
+						current.DefaultLevel,
+						current.MinVersion,
+						next.MinVersion - 1
+					);
+				}
+				else if (next.MinVersion <= current.MaxVersion)
+				{
+					s_logger?.LogWarning("Malformed ApplyWarningsDefaultValueAttribute collection; taking corrective action to address overlap in range (NextMin: {NextMin} <= CurrentMax: {CurrentMax}).", next.MinVersion, current.MaxVersion);
+
+					// Reduce next range to be more constrained, leaving the larger current range at the old build settings version standard.
+					next = new WarningLevelDefaultAttribute(
+						next.DefaultLevel,
+						current.MaxVersion + 1,
+						next.MaxVersion
+					);
+
+					// Flatten the old value with the updated one.
+					sorted[i + 1] = next;
+				}
+
+				merged.Add(current);
+			}
+
+			merged.Add(sorted.Last());
+
+			return merged;
+		}
+	}
+}
\ No newline at end of file
diff --git a/Engine/Source/Programs/UnrealBuildTool/Configuration/CppCompileWarnings.cs b/Engine/Source/Programs/UnrealBuildTool/Configuration/CppCompileWarnings.cs
--- a/Engine/Source/Programs/UnrealBuildTool/Configuration/CppCompileWarnings.cs
+++ b/Engine/Source/Programs/UnrealBuildTool/Configuration/CppCompileWarnings.cs
@@ -6,6 +6,7 @@ using System;
 using System.Collections.Generic;
 using System.Linq;
 using System.Reflection;
+using UnrealBuildTool.Configuration.CompileWarnings;
 
 namespace UnrealBuildTool
 {
@@ -17,13 +18,28 @@ namespace UnrealBuildTool
 		/// <param name="inSourceWarningLevel">The source warning level to override.</param>
 		/// <param name="appliedWarningLevel">The warning level to apply to the source.</param>
 		/// <param name ="overwriteSourceValue">Flag to overwrite the source warning with new value. If false, will only apply setting if value is unset.</param>
-		/// <remarks>If <paramref name="appliedWarningLevel"/> with value of null or <see cref="WarningLevel.Default"/> will be ignored as an override.</remarks>
+		/// <remarks>If <paramref name="appliedWarningLevel"/> has a value of null or <see cref="WarningLevel.Default"/>, it will be ignored as an override.</remarks>
 		internal static void ApplyWarning(this ref WarningLevel inSourceWarningLevel, WarningLevel? appliedWarningLevel, bool overwriteSourceValue = true)
 		{
-			if ((appliedWarningLevel.HasValue && appliedWarningLevel.Value != WarningLevel.Default) && (overwriteSourceValue || (!overwriteSourceValue && inSourceWarningLevel == WarningLevel.Default)))
+			inSourceWarningLevel = GetAppliedWarningLevel(inSourceWarningLevel, appliedWarningLevel, overwriteSourceValue);
+		}
+
+		/// <summary>
+		/// Gets the applied warning level on top of another,  given the new <paramref name="appliedWarningLevel"/> has a value, and it is not the <see cref="WarningLevel.Default"/>.
+		/// </summary>
+		/// <param name="sourceWarningLevel">The source warning level to override.</param>
+		/// <param name="appliedWarningLevel">The warning level to apply to the source.</param>
+		/// <param name ="overwriteSourceValue">Flag to overwrite the source warning with new value. If false, will only apply setting if value is unset.</param>
+		/// <returns>The applied warning level to use.</returns>
+		/// <remarks>If <paramref name="appliedWarningLevel"/> has a value of null or <see cref="WarningLevel.Default"/>, it will be ignored as an override.</remarks>
+		internal static WarningLevel GetAppliedWarningLevel(WarningLevel sourceWarningLevel, WarningLevel? appliedWarningLevel, bool overwriteSourceValue)
+		{
+			if ((appliedWarningLevel.HasValue && appliedWarningLevel.Value != WarningLevel.Default) && (overwriteSourceValue || (!overwriteSourceValue && sourceWarningLevel == WarningLevel.Default)))
 			{
-				inSourceWarningLevel = appliedWarningLevel.Value;
+				sourceWarningLevel = appliedWarningLevel.Value;
 			}
+
+			return sourceWarningLevel;
 		}
 	}
 
@@ -43,9 +59,10 @@ namespace UnrealBuildTool
 		 *		private WarningLevel _NEW_WARNING_LEVEL_NAME;
 		 * 2. Update the ReadOnlyCppCompileWarnings class to include a forwarding property:
 		 *		public WarningLevel NEW_WARNING_LEVEL_NAME => _inner.NEW_WARNING_LEVEL_NAME;
-		 * 3. Set your default WarningLevel within appropraite ApplyDefaults(...) setting
-		 *		For Target defaults				- internal void ApplyTargetDefaults()
-		 *			Note: This is the most appropriatel place for BuildSettings version specific defaults
+		 * 3. Set your default WarningLevel by using WarningLevelDefault attribute (if WarningLevelDefault is not applied to the warning,it will be Default)
+		 *		For Target defaults - specify it's initialization context to Target 
+		 *		For regular constructor - specify it's initialization context to Constructor
+		 *		For both contexts - specify it's initialization context to Any
 		 * 4. Access the new property through the (TargetRules|ModuleRules|CppCompileEnvironment).CppCompileWarnings as needed
 		 * 5. (Optional) For compiler line support, utilize the ApplyWarningsAttribute system to participate in invocations to GenerateWarningCommandLineAgs.
 		 * 
@@ -70,6 +87,9 @@ namespace UnrealBuildTool
 		[RequiresUniqueBuildEnvironment]
 		[WarningsMSCV(["/wd4456", "/wd4458", "/wd4459"], null, ["/we4456", "/we4458", "/we4459"])]
 		[ShadowVariableWarningsClangToolChain(["-Wno-shadow"], ["-Wshadow", "-Wno-error=shadow"], ["-Wshadow"])]
+		[WarningLevelDefault(WarningLevel.Warning, BuildSettingsVersion.V1, BuildSettingsVersion.Latest, InitializationContext.Constructor)]
+		[WarningLevelDefault(WarningLevel.Warning, BuildSettingsVersion.V1, BuildSettingsVersion.V1, InitializationContext.Target)]
+		[WarningLevelDefault(WarningLevel.Error, BuildSettingsVersion.V2, BuildSettingsVersion.Latest, InitializationContext.Target)]
 		public WarningLevel ShadowVariableWarningLevel
 		{
 			get => ResolveWarning(ParentCppCompileWarnings?.ShadowVariableWarningLevel, _shadowVariableWarningLevel);
@@ -105,6 +125,7 @@ namespace UnrealBuildTool
 			["-Wno-float-conversion", "-Wno-implicit-float-conversion", "-Wno-implicit-int-conversion", "-Wno-c++11-narrowing"],
 			["-Wfloat-conversion", "-Wno-error=float-conversion", "-Wimplicit-float-conversion", "-Wno-error=implicit-float-conversion", "-Wimplicit-int-conversion", "-Wno-error=implicit-int-conversion", "-Wc++11-narrowing", "-Wno-error=c++11-narrowing"],
 			["-Wfloat-conversion", "-Wimplicit-float-conversion", "-Wimplicit-int-conversion", "-Wc++11-narrowing"])]
+		[WarningLevelDefault(WarningLevel.Off)]
 		public WarningLevel UnsafeTypeCastWarningLevel
 		{
 			get => ResolveWarning(ParentCppCompileWarnings?.UnsafeTypeCastWarningLevel, _unsafeTypeCastWarningLevel);
@@ -126,6 +147,7 @@ namespace UnrealBuildTool
 		[RequiresUniqueBuildEnvironment]
 		[UndefinedIdentifierWarningsVCToolChain(null, ["/w44668"], ["/we4668"])]
 		[WarningsClangToolChain(null, ["-Wundef", "-Wno-error=undef"], ["-Wundef"])]
+		[WarningLevelDefault(WarningLevel.Off)]
 		public WarningLevel UndefinedIdentifierWarningLevel
 		{
 			get => ResolveWarning(ParentCppCompileWarnings?.UndefinedIdentifierWarningLevel, _undefinedIdentifierWarningLevel);
@@ -146,6 +168,7 @@ namespace UnrealBuildTool
 		[RequiresUniqueBuildEnvironment]
 		[WarningsVCToolChain(null, ["/w44061"], ["/we4061"])]
 		[WarningsClangToolChain(null, ["-Wswitch-enum", "-Wno-error=switch-enum"], ["-Wswitch-enum"])]
+		[WarningLevelDefault(WarningLevel.Off)]
 		public WarningLevel SwitchUnhandledEnumeratorWarningLevel
 		{
 			get => ResolveWarning(ParentCppCompileWarnings?.SwitchUnhandledEnumeratorWarningLevel, _switchUnhandledEnumeratorWarningLevel);
@@ -161,6 +184,7 @@ namespace UnrealBuildTool
 		///		https://clang.llvm.org/docs/DiagnosticsReference.html#wenum-enum-conversion
 		/// </summary>
 		[WarningsClangToolChain(["-Wno-enum-enum-conversion"])]
+		[WarningLevelDefault(WarningLevel.Off)]
 		public WarningLevel EnumEnumConversionWarningLevel
 		{
 			get => ResolveWarning(ParentCppCompileWarnings?.EnumEnumConversionWarningLevel, _enumEnumConversionWarningLevel);
@@ -174,6 +198,7 @@ namespace UnrealBuildTool
 		///		https://clang.llvm.org/docs/DiagnosticsReference.html#wenum-float-conversion
 		/// </summary>
 		[WarningsClangToolChain(["-Wno-enum-float-conversion"])]
+		[WarningLevelDefault(WarningLevel.Off)]
 		public WarningLevel EnumFloatConversionWarningLevel
 		{
 			get => ResolveWarning(ParentCppCompileWarnings?.EnumFloatConversionWarningLevel, _enumFloatConversionWarningLevel);
@@ -187,6 +212,7 @@ namespace UnrealBuildTool
 		///		https://clang.llvm.org/docs/DiagnosticsReference.html#wambiguous-reversed-operator
 		/// </summary>
 		[WarningsClangToolChain(["-Wno-ambiguous-reversed-operator"], null, null, [nameof(FilterID.Cpp20Min)])]
+		[WarningLevelDefault(WarningLevel.Off)]
 		public WarningLevel AmbiguousReversedOperatorWarningLevel
 		{
 			get => ResolveWarning(ParentCppCompileWarnings?.AmbiguousReversedOperatorWarningLevel, _ambiguousReversedOperatorWarningLevel);
@@ -200,6 +226,7 @@ namespace UnrealBuildTool
 		///		https://clang.llvm.org/docs/DiagnosticsReference.html#wdeprecated-anon-enum-enum-conversion
 		/// </summary>
 		[WarningsClangToolChain(["-Wno-deprecated-anon-enum-enum-conversion"], null, null, [nameof(FilterID.Cpp20Min)])]
+		[WarningLevelDefault(WarningLevel.Off)]
 		public WarningLevel DeprecatedAnonEnumEnumConversionWarningLevel
 		{
 			get => ResolveWarning(ParentCppCompileWarnings?.DeprecatedAnonEnumEnumConversionWarningLevel, _deprecatedAnonEnumEnumConversionWarningLevel);
@@ -213,6 +240,7 @@ namespace UnrealBuildTool
 		///		https://clang.llvm.org/docs/DiagnosticsReference.html#wdeprecated-volatile
 		/// </summary>
 		[WarningsClangToolChain(["-Wno-deprecated-volatile"], null, null, [nameof(FilterID.Cpp20Min)])]
+		[WarningLevelDefault(WarningLevel.Off)]
 		public WarningLevel DeprecatedVolatileWarningLevel
 		{
 			get => ResolveWarning(ParentCppCompileWarnings?.DeprecatedVolatileWarningLevel, _deprecatedVolatileWarningLevel);
@@ -226,6 +254,7 @@ namespace UnrealBuildTool
 		///		https://clang.llvm.org/docs/DiagnosticsReference.html#wordered-compare-function-pointers
 		/// </summary>
 		[WarningsClangToolChain(["-Wno-ordered-compare-function-pointers"])]
+		[WarningLevelDefault(WarningLevel.Off)]
 		public WarningLevel OrderedCompareFunctionPointers
 		{
 			get => ResolveWarning(ParentCppCompileWarnings?.OrderedCompareFunctionPointers, _orderedCompareFunctionPointers);
@@ -239,6 +268,7 @@ namespace UnrealBuildTool
 		///		https://clang.llvm.org/docs/DiagnosticsReference.html#wbitwise-instead-of-logical
 		/// </summary>
 		[WarningsClangToolChain(["-Wno-bitwise-instead-of-logical"], null, null, [nameof(FilterID.Version14Min)])]
+		[WarningLevelDefault(WarningLevel.Off)]
 		public WarningLevel BitwiseInsteadOfLogicalWarningLevel
 		{
 			get => ResolveWarning(ParentCppCompileWarnings?.BitwiseInsteadOfLogicalWarningLevel, _bitwiseInsteadOfLogicalWarningLevel);
@@ -252,6 +282,7 @@ namespace UnrealBuildTool
 		///		https://clang.llvm.org/docs/DiagnosticsReference.html#wdeprecated-copy
 		/// </summary>
 		[WarningsClangToolChain(["-Wno-deprecated-copy"], null, null, [nameof(FilterID.Version16Min)])]
+		[WarningLevelDefault(WarningLevel.Off)]
 		public WarningLevel DeprecatedCopyWarningLevel
 		{
 			get => ResolveWarning(ParentCppCompileWarnings?.DeprecatedCopyWarningLevel, _deprecatedCopyWarningLevel);
@@ -265,6 +296,7 @@ namespace UnrealBuildTool
 		///		https://clang.llvm.org/docs/DiagnosticsReference.html#wdeprecated-copy-with-user-provided-copy
 		/// </summary>
 		[WarningsClangToolChain(["-Wno-deprecated-copy-with-user-provided-copy"], null, null, [nameof(FilterID.Version16Min)])]
+		[WarningLevelDefault(WarningLevel.Off)]
 		public WarningLevel DeprecatedCopyWithUserProvidedCopyWarningLevel
 		{
 			get => ResolveWarning(ParentCppCompileWarnings?.DeprecatedCopyWithUserProvidedCopyWarningLevel, _deprecatedCopyWithUserProvidedCopyWarningLevel);
@@ -278,6 +310,7 @@ namespace UnrealBuildTool
 		///		https://clang.llvm.org/docs/DiagnosticsReference.html#winvalid-unevaluated-string
 		/// </summary>
 		[WarningsClangToolChain(["-Wno-invalid-unevaluated-string"], null, null, [nameof(SpecializedFilters.AndroidNDKR26ToolChainVersionExclusion)])]
+		[WarningLevelDefault(WarningLevel.Off)]
 		public WarningLevel InvalidUnevaluatedStringWarningLevel
 		{
 			get => ResolveWarning(ParentCppCompileWarnings?.InvalidUnevaluatedStringWarningLevel, _invalidUnevaluatedStringWarningLevel);
@@ -296,6 +329,7 @@ namespace UnrealBuildTool
 		/// We tend to use NAN just as a bit pattern, and not rely on it in calculations, so it should be reasonably safe to disable.
 		/// </remarks>
 		[WarningsClangToolChain(["-Wno-nan-infinity-disabled"], null, null, [nameof(FilterID.Version18Min)])]
+		[WarningLevelDefault(WarningLevel.Off)]
 		public WarningLevel NaNInfinityDisabledWarningLevel
 		{
 			get => ResolveWarning(ParentCppCompileWarnings?.NaNInfinityDisabledWarningLevel, _nanInfinityDisabledWarningLevel);
@@ -309,6 +343,7 @@ namespace UnrealBuildTool
 		///		https://clang.llvm.org/docs/DiagnosticsReference.html#wextra-qualification
 		/// </summary>
 		[WarningsClangToolChain(["-Wno-extra-qualification"], null, null, [nameof(FilterID.Version19Min)])]
+		[WarningLevelDefault(WarningLevel.Off)]
 		public WarningLevel LevelExtraQualificationWarningLevel
 		{
 			get => ResolveWarning(ParentCppCompileWarnings?.LevelExtraQualificationWarningLevel, _levelExtraQualificationWarningLevel);
@@ -322,6 +357,7 @@ namespace UnrealBuildTool
 		///		https://clang.llvm.org/docs/DiagnosticsReference.html#wcast-function-type-mismatch
 		/// </summary>
 		[WarningsClangToolChain(["-Wno-cast-function-type-mismatch"], null, null, [nameof(SpecializedFilters.AndroidNDKR28ToolChainVersionExclusion)])]
+		[WarningLevelDefault(WarningLevel.Off)]
 		public WarningLevel CastFunctionTypeMistmatchWarningLevel
 		{
 			get => ResolveWarning(ParentCppCompileWarnings?.CastFunctionTypeMistmatchWarningLevel, _castFunctionTypeMistmatchWarningLevel);
@@ -335,6 +371,7 @@ namespace UnrealBuildTool
 		///		https://clang.llvm.org/docs/DiagnosticsReference.html#wmissing-template-arg-list-after-template-kw
 		/// </summary>
 		[WarningsClangToolChain(["-Wno-missing-template-arg-list-after-template-kw"], null, null, [nameof(SpecializedFilters.AndroidNDKR28ToolChainVersionExclusion)])]
+		[WarningLevelDefault(WarningLevel.Off)]
 		public WarningLevel MissingTemplateArgListAfterTemplateWarningLevel
 		{
 			get => ResolveWarning(ParentCppCompileWarnings?.MissingTemplateArgListAfterTemplateWarningLevel, _missingTemplateArgListAfterTemplateWarningLevel);
@@ -349,6 +386,7 @@ namespace UnrealBuildTool
 		/// </summary>
 		/// <remarks>We use this feature to allow static FNames.</remarks>
 		[WarningsClangToolChain(["-Wno-gnu-string-literal-operator-template"])]
+		[WarningLevelDefault(WarningLevel.Off)]
 		public WarningLevel GNUStringLiteralOperatorTemplateWarningLevel
 		{
 			get => ResolveWarning(ParentCppCompileWarnings?.GNUStringLiteralOperatorTemplateWarningLevel, _gnuStringLiteralOperatorTemplateWarningLevel);
@@ -362,6 +400,7 @@ namespace UnrealBuildTool
 		///		https://clang.llvm.org/docs/DiagnosticsReference.html#winconsistent-missing-override
 		/// </summary>
 		[WarningsClangToolChain(["-Wno-inconsistent-missing-override"])]
+		[WarningLevelDefault(WarningLevel.Off)]
 		public WarningLevel InconsistentMissingOverrideWarningLevel
 		{
 			get => ResolveWarning(ParentCppCompileWarnings?.InconsistentMissingOverrideWarningLevel, _inconsistentMissingOverrideWarningLevel);
@@ -376,6 +415,7 @@ namespace UnrealBuildTool
 		/// </summary>
 		/// <remarks>Needed to suppress warnings about using offsetof on non-POD types.</remarks>
 		[WarningsClangToolChain(["-Wno-invalid-offsetof"])]
+		[WarningLevelDefault(WarningLevel.Off)]
 		public WarningLevel InvalidOffsetWarningLevel
 		{
 			get => ResolveWarning(ParentCppCompileWarnings?.InvalidOffsetWarningLevel, _invalidOffsetOfWarningLevel);
@@ -390,6 +430,7 @@ namespace UnrealBuildTool
 		/// </summary>
 		/// <remarks>Fhis hides the "enumeration value 'XXXXX' not handled in switch [-Wswitch]" warnings - we should maybe remove this at some point and add UE_LOG(, Fatal, ) to default cases.</remarks>
 		[WarningsClangToolChain(["-Wno-switch"])]
+		[WarningLevelDefault(WarningLevel.Off)]
 		public WarningLevel SwitchWarningLevel
 		{
 			get => ResolveWarning(ParentCppCompileWarnings?.SwitchWarningLevel, _switchWarningLevel);
@@ -404,6 +445,7 @@ namespace UnrealBuildTool
 		/// </summary>
 		/// <remarks>This hides the "warning : comparison of unsigned expression &lt; 0 is always false" type warnings due to constant comparisons, which are possible with template arguments/</remarks>
 		[WarningsClangToolChain(["-Wno-tautological-compare"])]
+		[WarningLevelDefault(WarningLevel.Off)]
 		public WarningLevel TautologicalCompareWarningLevel
 		{
 			get => ResolveWarning(ParentCppCompileWarnings?.TautologicalCompareWarningLevel, _tautologicalCompareWarningLevel);
@@ -418,6 +460,7 @@ namespace UnrealBuildTool
 		/// </summary>
 		/// <remarks>Slate triggers this (with its optimize on/off pragmas).</remarks>
 		[WarningsClangToolChain(["-Wno-unknown-pragmas"])]
+		[WarningLevelDefault(WarningLevel.Off)]
 		public WarningLevel UnknownPragmasWarningLevel
 		{
 			get => ResolveWarning(ParentCppCompileWarnings?.UnknownPragmasWarningLevel, _unknownPragmasWarningLevel);
@@ -443,6 +486,7 @@ namespace UnrealBuildTool
 		/// unused-private-field: This will prevent the issue of warnings for unused private variables. MultichannelTcpSocket.h triggers this, possibly more.
 		/// </remarks>
 		[WarningsClangToolChain(["-Wno-unused-but-set-variable", "-Wno-unused-but-set-parameter", "-Wno-unused-function", "-Wno-unused-lambda-capture", "-Wno-unused-local-typedef", "-Wno-unused-private-field", "-Wno-unused-variable"], null, null)]
+		[WarningLevelDefault(WarningLevel.Off)]
 		public WarningLevel UnusedWarningLevel
 		{
 			get => ResolveWarning(ParentCppCompileWarnings?.UnusedWarningLevel, _unusedWarningLevel);
@@ -457,6 +501,7 @@ namespace UnrealBuildTool
 		/// </summary>
 		/// <remarks>Not a good warning to be disabling.</remarks>
 		[WarningsClangToolChain(["-Wno-undefined-var-template"])]
+		[WarningLevelDefault(WarningLevel.Off)]
 		public WarningLevel UndefinedVarTemplateWarningLevel
 		{
 			get => ResolveWarning(ParentCppCompileWarnings?.UndefinedVarTemplateWarningLevel, _undefinedVarTemplateWarningLevel);
@@ -472,6 +517,7 @@ namespace UnrealBuildTool
 		/// </summary>
 		/// <remarks> Clang emits warnings for each compiled object file that doesn't have a matching entry in the profile data. This can happen when the profile data is older than the binaries we're compiling. Disable these warnings.</remarks>
 		[WarningsClangToolChain(["-Wno-profile-instr-out-of-date", "-Wno-profile-instr-unprofiled"], null, null, [nameof(StandardFilters.PGOOptimizedFilter)])]
+		[WarningLevelDefault(WarningLevel.Off)]
 		public WarningLevel ProfileInstructWarningLevel
 		{
 			get => ResolveWarning(ParentCppCompileWarnings?.ProfileInstructWarningLevel, _profileInstructWarningLevel);
@@ -485,6 +531,7 @@ namespace UnrealBuildTool
 		/// </summary>
 		/// <remarks>Apparently there can be hashing conflicts with PGO which can result in: 'Function control flow change detected (hash mismatch)' warnings. </remarks>
 		[WarningsClangToolChain(["-Wno-backend-plugin"], null, null, [nameof(StandardFilters.PGOOptimizedFilter)])]
+		[WarningLevelDefault(WarningLevel.Off)]
 		public WarningLevel BackendPluginWarningLevel
 		{
 			get => ResolveWarning(ParentCppCompileWarnings?.BackendPluginWarningLevel, _backendPluginWarningLevel);
@@ -502,6 +549,7 @@ namespace UnrealBuildTool
 		/// -Wunused-result - we will always re-enable this as it's controlled by the same set. This will be moved out to a separate property. It is on by default, so it only must be re-enabled in the disable context.
 		/// </remarks>
 		[UnusedValueClangToolChain(["-Wno-unused-value", "-Wunused-result"])]
+		[WarningLevelDefault(WarningLevel.Off)]
 		public WarningLevel UnusedValueWarningLevel
 		{
 			get => ResolveWarning(ParentCppCompileWarnings?.UnusedValueWarningLevel, _unusedValueWarningLevel);
@@ -516,6 +564,7 @@ namespace UnrealBuildTool
 		/// </summary>
 		/// <remarks>-Wimplicit-int-conversion also controls -Wshorten-64-to-32 as of clang 19.</remarks>
 		[WarningsClangToolChain(["-Wno-shorten-64-to-32"], null, null, [nameof(FilterID.Version19Min)])]
+		[WarningLevelDefault(WarningLevel.Off)]
 		public WarningLevel Shorten64To32WarningLevel
 		{
 			get => ResolveWarning(ParentCppCompileWarnings?.Shorten64To32WarningLevel, _shorten64To32WarningLevel);
@@ -530,6 +579,7 @@ namespace UnrealBuildTool
 		/// </summary>
 		/// <remarks>The code base contains lots of places where we do "extern template X_API class ..." and we want to keep doing that to reduce compile times.</remarks>
 		[WarningsClangToolChain(["-Wno-dllexport-explicit-instantiation-decl"])]
+		[WarningLevelDefault(WarningLevel.Off)]
 		public WarningLevel DllExportExplicitInstantiationDeclWarningLevel
 		{
 			get => ResolveWarning(ParentCppCompileWarnings?.DllExportExplicitInstantiationDeclWarningLevel, _dllExportExplicitInstantiationDeclWarningLevel);
@@ -548,6 +598,7 @@ namespace UnrealBuildTool
 		/// </summary>
 		/// <remarks>Allow Microsoft-specific syntax to slide, even though it may be non-standard.  Needed for Windows headers.</remarks>
 		[WarningsVCClang(["-Wno-microsoft"])]
+		[WarningLevelDefault(WarningLevel.Off)]
 		public WarningLevel MicrosoftGroupWarningLevel
 		{
 			get => ResolveWarning(ParentCppCompileWarnings?.MicrosoftGroupWarningLevel, _microsoftGroupWarningLevel);
@@ -562,6 +613,7 @@ namespace UnrealBuildTool
 		/// </summary>
 		/// <remarks> Hack due to how we have our 'DummyPCH' wrappers setup when using unity builds.  This warning should not be disabled!</remarks>
 		[WarningsVCClang(["-Wno-msvc-include"])]
+		[WarningLevelDefault(WarningLevel.Off)]
 		public WarningLevel MSVCIncludeWarningLevel
 		{
 			get => ResolveWarning(ParentCppCompileWarnings?.MSVCIncludeWarningLevel, _msvcIncludeWarningLevel);
@@ -579,6 +631,7 @@ namespace UnrealBuildTool
 		/// @TODO clang: push/pop this in  Pre/PostWindowsApi.h headers instead?
 		/// </remarks>
 		[WarningsVCClang(["-Wno-pragma-pack"])]
+		[WarningLevelDefault(WarningLevel.Off)]
 		public WarningLevel PragmaPackWarningLevel
 		{
 			get => ResolveWarning(ParentCppCompileWarnings?.PragmaPackWarningLevel, _pragmaPackWarningLevel);
@@ -593,6 +646,7 @@ namespace UnrealBuildTool
 		/// </summary>
 		/// <remarks>@todo clang: We declare operator new as inline.  Clang doesn't seem to like that.</remarks>
 		[WarningsVCClang(["-Wno-inline-new-delete"])]
+		[WarningLevelDefault(WarningLevel.Off)]
 		public WarningLevel InlineNewDeleteWarningLevel
 		{
 			get => ResolveWarning(ParentCppCompileWarnings?.InlineNewDeleteWarningLevel, _inlineNewDeleteWarningLevel);
@@ -607,6 +661,7 @@ namespace UnrealBuildTool
 		/// </summary>
 
 		[WarningsVCClang(["-Wno-implicit-exception-spec-mismatch"])]
+		[WarningLevelDefault(WarningLevel.Off)]
 		public WarningLevel ImplicitExceptionSpecMismatchWarningLevel
 		{
 			get => ResolveWarning(ParentCppCompileWarnings?.ImplicitExceptionSpecMismatchWarningLevel, _implicitExceptionSpecMismatchWarningLevel);
@@ -621,6 +676,7 @@ namespace UnrealBuildTool
 		/// </summary>
 		/// <remarks>Sometimes we compare 'this' pointers against nullptr, which Clang warns about by default.</remarks>
 		[WarningsVCClang(["-Wno-undefined-bool-conversion"])]
+		[WarningLevelDefault(WarningLevel.Off)]
 		public WarningLevel UndefinedBoolConversionWarningLevel
 		{
 			get => ResolveWarning(ParentCppCompileWarnings?.UndefinedBoolConversionWarningLevel, _undefinedBoolConversionWarningLevel);
@@ -635,6 +691,7 @@ namespace UnrealBuildTool
 		/// </summary>
 
 		[WarningsVCClang(["-Wno-deprecated-writable-strings"])]
+		[WarningLevelDefault(WarningLevel.Off)]
 		public WarningLevel DeprecatedWriteableStringsWarningLevel
 		{
 			get => ResolveWarning(ParentCppCompileWarnings?.DeprecatedWriteableStringsWarningLevel, _deprecatedWriteableStringsWarningLevel);
@@ -648,6 +705,7 @@ namespace UnrealBuildTool
 		///		https://clang.llvm.org/docs/DiagnosticsReference.html#wdeprecated-register
 		/// </summary>
 		[WarningsVCClang(["-Wno-deprecated-register"])]
+		[WarningLevelDefault(WarningLevel.Off)]
 		public WarningLevel DeprecatedRegisterWarningLevel
 		{
 			get => ResolveWarning(ParentCppCompileWarnings?.DeprecatedRegisterWarningLevel, _deprecatedRegisterWarningLevel);
@@ -663,6 +721,7 @@ namespace UnrealBuildTool
 		/// </summary>
 		/// <remarks>Disable needed for external headers we shan't change.</remarks>
 		[WarningsVCClang(["-Wno-logical-op-parentheses"])]
+		[WarningLevelDefault(WarningLevel.Off)]
 		public WarningLevel LogicalOpParenthesesWarningLevel
 		{
 			get => ResolveWarning(ParentCppCompileWarnings?.LogicalOpParenthesesWarningLevel, _logicalOpParenthesesWarningLevel);
@@ -677,6 +736,7 @@ namespace UnrealBuildTool
 		/// </summary>
 		/// <remarks>Disable needed for external headers we shan't change.</remarks>
 		[WarningsVCClang(["-Wno-null-arithmetic"])]
+		[WarningLevelDefault(WarningLevel.Off)]
 		public WarningLevel NullArithmeticWarningLevel
 		{
 			get => ResolveWarning(ParentCppCompileWarnings?.NullArithmeticWarningLevel, _nullArithmeticWarningLevel);
@@ -691,6 +751,7 @@ namespace UnrealBuildTool
 		/// </summary>
 		/// <remarks>Disable needed for PhysX.</remarks>
 		[WarningsVCClang(["-Wno-return-type-c-linkage"])]
+		[WarningLevelDefault(WarningLevel.Off)]
 		public WarningLevel ReturnTypeCLinkageWarningLevel
 		{
 			get => ResolveWarning(ParentCppCompileWarnings?.ReturnTypeCLinkageWarningLevel, _returnTypeCLinkageWarningLevel);
@@ -705,6 +766,7 @@ namespace UnrealBuildTool
 		/// </summary>
 		/// <remarks>Disable needed for nvtesslib.</remarks>
 		[WarningsVCClang(["-Wno-ignored-attributes"])]
+		[WarningLevelDefault(WarningLevel.Off)]
 		public WarningLevel IgnoredAttributesWarningLevel
 		{
 			get => ResolveWarning(ParentCppCompileWarnings?.IgnoredAttributesWarningLevel, _ignoredAttributesWarningLevel);
@@ -718,6 +780,7 @@ namespace UnrealBuildTool
 		///		https://clang.llvm.org/docs/DiagnosticsReference.html#wuninitialized
 		/// </summary>
 		[WarningsVCClang(["-Wno-uninitialized"])]
+		[WarningLevelDefault(WarningLevel.Off)]
 		public WarningLevel UninitializedWarningLevel
 		{
 			get => ResolveWarning(ParentCppCompileWarnings?.UninitializedWarningLevel, _uninitializedWarningLevel);
@@ -732,6 +795,7 @@ namespace UnrealBuildTool
 		/// </summary>
 		/// <remarks>Disable needed for external headers we shan't change.</remarks>
 		[WarningsVCClang(["-Wno-return-type"])]
+		[WarningLevelDefault(WarningLevel.Off)]
 		public WarningLevel ReturnTypeWarningLevel
 		{
 			get => ResolveWarning(ParentCppCompileWarnings?.ReturnTypeWarningLevel, _returnTypeWarningLevel);
@@ -746,6 +810,7 @@ namespace UnrealBuildTool
 		/// </summary>
 		/// <remarks>Unused function parameter. A lot are named 'bUnused'...</remarks>
 		[WarningsVCClang(["-Wno-unused-parameter"])]
+		[WarningLevelDefault(WarningLevel.Off)]
 		public WarningLevel UnusedParameterWarningLevel
 		{
 			get => ResolveWarning(ParentCppCompileWarnings?.UnusedParameterWarningLevel, _unusedParameterWarningLevel);
@@ -760,6 +825,7 @@ namespace UnrealBuildTool
 		/// </summary>
 		/// <remarks>const ignored when returning by value e.g. 'const int foo() { return 4; }'.</remarks>
 		[WarningsVCClang(["-Wno-ignored-qualifiers"])]
+		[WarningLevelDefault(WarningLevel.Off)]
 		public WarningLevel IgnoredQualifiersWarningLevel
 		{
 			get => ResolveWarning(ParentCppCompileWarnings?.IgnoredQualifiersWarningLevel, _ignoredQualifiersWarningLevel);
@@ -774,6 +840,7 @@ namespace UnrealBuildTool
 		/// </summary>
 		/// <remarks>Usage of 'defined(X)' in a macro definition. Gives different results under MSVC.</remarks>
 		[WarningsVCClang(["-Wno-expansion-to-defined"])]
+		[WarningLevelDefault(WarningLevel.Off)]
 		public WarningLevel ExpansionToDefined
 		{
 			get => ResolveWarning(ParentCppCompileWarnings?.ExpansionToDefined, _expansionToDefined);
@@ -788,6 +855,7 @@ namespace UnrealBuildTool
 		/// </summary>
 		/// <remarks>Signed/unsigned comparison - millions of these.</remarks>
 		[WarningsVCClang(["-Wno-sign-compare"])]
+		[WarningLevelDefault(WarningLevel.Off)]
 		public WarningLevel SignCompareWarningLevel
 		{
 			get => ResolveWarning(ParentCppCompileWarnings?.SignCompareWarningLevel, _signCompareWarningLevel);
@@ -802,6 +870,7 @@ namespace UnrealBuildTool
 		/// </summary>
 		/// <remarks>Excessive warning, generated when you initialize with MyStruct A = {0};</remarks>
 		[WarningsVCClang(["-Wno-missing-field-initializers"])]
+		[WarningLevelDefault(WarningLevel.Off)]
 		public WarningLevel MissingFieldInitializersWarningLevel
 		{
 			get => ResolveWarning(ParentCppCompileWarnings?.MissingFieldInitializersWarningLevel, _missingFieldInitializersWarningLevel);
@@ -815,6 +884,7 @@ namespace UnrealBuildTool
 		///		https://clang.llvm.org/docs/DiagnosticsReference.html#wnonportable-include-path
 		/// </summary>
 		[WarningsVCClang(["-Wno-nonportable-include-path"])]
+		[WarningLevelDefault(WarningLevel.Off)]
 		public WarningLevel NonPortableIncludePathWarningLevel
 		{
 			get => ResolveWarning(ParentCppCompileWarnings?.NonPortableIncludePathWarningLevel, _nonPortableIncludePathWarningLevel);
@@ -830,6 +900,7 @@ namespace UnrealBuildTool
 		///		https://learn.microsoft.com/en-us/cpp/preprocessor/token-pasting-operator-hash-hash?view=msvc-170
 		/// </summary>
 		[WarningsVCClang(["-Wno-invalid-token-paste"])]
+		[WarningLevelDefault(WarningLevel.Off)]
 		public WarningLevel InvalidTokenPasteWarningLevel
 		{
 			get => ResolveWarning(ParentCppCompileWarnings?.InvalidTokenPasteWarningLevel, _invalidTokenPasteWarningLevel);
@@ -843,6 +914,7 @@ namespace UnrealBuildTool
 		///		https://clang.llvm.org/docs/DiagnosticsReference.html#wnull-pointer-arithmetic
 		/// </summary>
 		[WarningsVCClang(["-Wno-null-pointer-arithmetic"])]
+		[WarningLevelDefault(WarningLevel.Off)]
 		public WarningLevel NullPointerArithmeticWarningLevel
 		{
 			get => ResolveWarning(ParentCppCompileWarnings?.NullPointerArithmeticWarningLevel, _nullPointerArithmeticWarningLevel);
@@ -857,6 +929,7 @@ namespace UnrealBuildTool
 		/// </summary>
 		/// <remarks>Triggered by || of two template-derived values inside a static_assert.</remarks>
 		[WarningsVCClang(["-Wno-constant-logical-operand"])]
+		[WarningLevelDefault(WarningLevel.Off)]
 		public WarningLevel ConstantLogicalOperandWarningLevel
 		{
 			get => ResolveWarning(ParentCppCompileWarnings?.ConstantLogicalOperandWarningLevel, _constantLogicalOperandWarningLevel);
@@ -870,6 +943,7 @@ namespace UnrealBuildTool
 		///		https://clang.llvm.org/docs/DiagnosticsReference.html#wbitfield-enum-conversion
 		/// </summary>
 		[WarningsVCClang(["-Wno-bitfield-enum-conversion"])]
+		[WarningLevelDefault(WarningLevel.Off)]
 		public WarningLevel BitfieldEnumConversion
 		{
 			get => ResolveWarning(ParentCppCompileWarnings?.BitfieldEnumConversion, _bitfieldEnumConversion);
@@ -883,6 +957,7 @@ namespace UnrealBuildTool
 		///		https://clang.llvm.org/docs/DiagnosticsReference.html#wnull-pointer-subtraction
 		/// </summary>
 		[WarningsVCClang(["-Wno-null-pointer-subtraction"])]
+		[WarningLevelDefault(WarningLevel.Off)]
 		public WarningLevel NullPointerSubtractionWarningLevel
 		{
 			get => ResolveWarning(ParentCppCompileWarnings?.NullPointerSubtractionWarningLevel, _nullPointerSubtractionWarningLevel);
@@ -896,6 +971,7 @@ namespace UnrealBuildTool
 		///		https://clang.llvm.org/docs/DiagnosticsReference.html#wdangling
 		/// </summary>	
 		[WarningsVCClang(["-Wno-dangling"])]
+		[WarningLevelDefault(WarningLevel.Off)]
 		public WarningLevel DanglingWarningLevel
 		{
 			get => ResolveWarning(ParentCppCompileWarnings?.DanglingWarningLevel, _danglingWarningLevel);
@@ -919,6 +995,7 @@ namespace UnrealBuildTool
 		/// In order to ensure no change of behaviour for MSVC toolchain within the clang compilation context, we will keep this in, and deprecate with assignments to the clang setting.
 		/// </remarks>
 		[WarningsVCClang(["-Wno-switch-enum"])]
+		[WarningLevelDefault(WarningLevel.Off)]
 		public WarningLevel MSVCSwitchEnumWarningLevel
 		{
 			get => ResolveWarning(ParentCppCompileWarnings?.MSVCSwitchEnumWarningLevel, _msvcSwitchEnumWarningLevel);
@@ -938,6 +1015,7 @@ namespace UnrealBuildTool
 		/// -Wunused-result - we will always re-enable this as it's controlled by the same set. This will be moved out to a separate property. It is on by default, so it only must be re-enabled in the disable context.
 		/// </remarks>
 		[WarningsVCClang(["-Wno-unused-value", "-Wunused-result"])]
+		[WarningLevelDefault(WarningLevel.Off)]
 		public WarningLevel MSVCUnusedValueWarningLevel
 		{
 			get => ResolveWarning(ParentCppCompileWarnings?.MSVCUnusedValueWarningLevel, _msvcUnusedValueWarningLevel);
@@ -956,6 +1034,7 @@ namespace UnrealBuildTool
 		/// In order to ensure no change of behaviour for MSVC toolchain within the clang compilation context, we will keep this in, and deprecate with assignments to the clang setting.
 		/// </remarks>
 		[WarningsVCClang(["-Wno-deprecated-declarations"])]
+		[WarningLevelDefault(WarningLevel.Off)]
 		public WarningLevel MSVCDeprecationWarningLevel
 		{
 			get => ResolveWarning(ParentCppCompileWarnings?.MSVCDeprecationWarningLevel, _msvcDeprecationWarningLevel);
@@ -976,6 +1055,7 @@ namespace UnrealBuildTool
 		/// </summary>
 		/// <remarks>Disable needed for wxWidgets.</remarks>
 		[WarningsClangToolChain(["-Wno-deprecated-declarations"], null, ["-Werror=deprecated-declarations"])]
+		[WarningLevelDefault(WarningLevel.Warning)]
 		public WarningLevel DeprecationWarningLevel
 		{
 			get => ResolveWarning(ParentCppCompileWarnings?.DeprecationWarningLevel, _deprecationWarningLevel);
@@ -989,6 +1069,7 @@ namespace UnrealBuildTool
 		///		https://clang.llvm.org/docs/DiagnosticsReference.html#wdate-time
 		/// </summary>
 		[WarningsClangToolChain(null, ["-Wdate-time", "-Wno-error=date-time"], ["-Wdate-time"], [nameof(StandardFilters.DeterministcFlagSetFilter)])]
+		[WarningLevelDefault(WarningLevel.Off)]
 		public WarningLevel DeterministicWarningLevel
 		{
 			get => ResolveWarning(ParentCppCompileWarnings?.DeterministicWarningLevel, _deterministicWarningLevel);
@@ -1007,6 +1088,7 @@ namespace UnrealBuildTool
 		/// </summary>
 		[XmlConfigFile(Category = "BuildConfiguration")]
 		[CommandLine("-PCHPerformanceIssueWarningLevel=")]
+		[WarningLevelDefault(WarningLevel.Off)]
 		public WarningLevel PCHPerformanceIssueWarningLevel
 		{
 			get => ResolveWarning(ParentCppCompileWarnings?.PCHPerformanceIssueWarningLevel, _pchPerformanceIssueWarningLevel);
@@ -1023,6 +1105,7 @@ namespace UnrealBuildTool
 		/// </summary>
 		[XmlConfigFile(Category = "BuildConfiguration")]
 		[CommandLine("-ModuleUnsupportedWarningLevel=")]
+		[WarningLevelDefault(WarningLevel.Warning)]
 		public WarningLevel ModuleUnsupportedWarningLevel
 		{
 			get => ResolveWarning(ParentCppCompileWarnings?.ModuleUnsupportedWarningLevel, _moduleUnsupportedWarningLevel);
@@ -1035,6 +1118,7 @@ namespace UnrealBuildTool
 		/// </summary>
 		[XmlConfigFile(Category = "BuildConfiguration")]
 		[CommandLine("-PluginModuleUnsupportedWarningLevel=")]
+		[WarningLevelDefault(WarningLevel.Off)]
 		public WarningLevel PluginModuleUnsupportedWarningLevel
 		{
 			get => ResolveWarning(ParentCppCompileWarnings?.PluginModuleUnsupportedWarningLevel, _pluginModuleUnsupportedWarningLevel);
@@ -1047,6 +1131,7 @@ namespace UnrealBuildTool
 		/// </summary>
 		[XmlConfigFile(Category = "BuildConfiguration")]
 		[CommandLine("-ModuleIncludePathWarningLevel=")]
+		[WarningLevelDefault(WarningLevel.Off)]
 		public WarningLevel ModuleIncludePathWarningLevel
 		{
 			get => ResolveWarning(ParentCppCompileWarnings?.ModuleIncludePathWarningLevel, _moduleIncludePathWarningLevel);
@@ -1059,6 +1144,7 @@ namespace UnrealBuildTool
 		/// </summary>
 		[XmlConfigFile(Category = "BuildConfiguration")]
 		[CommandLine("-ModuleIncludePrivateWarningLevel=")]
+		[WarningLevelDefault(WarningLevel.Off)]
 		public WarningLevel ModuleIncludePrivateWarningLevel
 		{
 			get => ResolveWarning(ParentCppCompileWarnings?.ModuleIncludePrivateWarningLevel, _moduleIncludePrivateWarningLevel);
@@ -1071,16 +1157,18 @@ namespace UnrealBuildTool
 		/// </summary>
 		[XmlConfigFile(Category = "BuildConfiguration")]
 		[CommandLine("-ModuleIncludeSubdirectoryWarningLevel=")]
+		[WarningLevelDefault(WarningLevel.Off)]
 		public WarningLevel ModuleIncludeSubdirectoryWarningLevel
 		{
 			get => ResolveWarning(ParentCppCompileWarnings?.ModuleIncludeSubdirectoryWarningLevel, _moduleIncludeSubdirectoryWarningLevel);
-			set => _moduleIncludePrivateWarningLevel = value;
+			set => _moduleIncludeSubdirectoryWarningLevel = value;
 		}
 		private WarningLevel _moduleIncludeSubdirectoryWarningLevel;
 
 		/// <summary>
 		/// How to treat conflicts when a disabled plugin is being enabled by another plugin referencing it
 		/// </summary>
+		[WarningLevelDefault(WarningLevel.Default)]
 		public WarningLevel DisablePluginsConflictWarningLevel
 		{
 			get => ResolveWarning(ParentCppCompileWarnings?.DisablePluginsConflictWarningLevel, _disablePluginsConflictWarningLevel);
@@ -1091,6 +1179,7 @@ namespace UnrealBuildTool
 		/// <summary>
 		/// Enable warnings for when there are .gen.cpp files that could be inlined in a matching handwritten cpp file
 		/// </summary>
+		[WarningLevelDefault(WarningLevel.Off)]
 		public WarningLevel NonInlinedGenCppWarningLevel
 		{
 			get => ResolveWarning(ParentCppCompileWarnings?.NonInlinedGenCppWarningLevel, _nonInlinedGenCppWarningLevel);
@@ -1199,8 +1288,7 @@ namespace UnrealBuildTool
 
 		#endregion -- Public API --
 
-		private static IEnumerable<PropertyInfo>? s_propertiesWithAttribute = null;
-		private static IEnumerable<FieldInfo>? s_fieldsWithAttribute = null;
+		private static IEnumerable<MemberInfo>? s_membersWithApplyWarningsAttribute = null;
 
 		/// <summary>
 		/// Generates the compiler warnings command line arguments for the requested <see cref="UEToolChain"/> type, and version.
@@ -1224,92 +1312,30 @@ namespace UnrealBuildTool
 
 			List<string> warningArguments = new List<string>();
 
-			// Process properties
+			s_membersWithApplyWarningsAttribute ??= GetMembersWithAttributes(typeof(CppCompileWarnings), typeof(ApplyWarningsAttribute), typeof(WarningLevel));
+			foreach (MemberInfo memberInfo in s_membersWithApplyWarningsAttribute)
 			{
-				s_propertiesWithAttribute ??= typeof(CppCompileWarnings)
-					.GetProperties(BindingFlags.Public | BindingFlags.NonPublic | BindingFlags.Instance | BindingFlags.Static)
-					.Where(p => p.IsDefined(typeof(ApplyWarningsAttribute), false));
+				ApplyWarningsAttribute? applyWarningsAttribute = null;
+				IEnumerable<ApplyWarningsAttribute?> applyWarningsAttributes = memberInfo.GetCustomAttributes<ApplyWarningsAttribute>().Where(x => x.CanApplyToContext(toolChainContext));
 
-				foreach (PropertyInfo property in s_propertiesWithAttribute)
+				if (applyWarningsAttributes.Count() > 1)
 				{
-					if (property.PropertyType != typeof(WarningLevel))
-					{
-						_logger?.LogWarning("Attempting to resolve warnings, and encountered invalid property ({PropertyName}). Skipping this property.", property.Name);
-						continue;
-					}
-
-					ApplyWarningsAttribute? applyWarningsAttribute = null;
-					IEnumerable<ApplyWarningsAttribute?> applyWarningsAttributes = property.GetCustomAttributes<ApplyWarningsAttribute>().Where(x => x.CanApplyToContext(toolChainContext));
-
-					if (applyWarningsAttributes.Count() > 1)
-					{
-						_logger?.LogWarning("Ambiguous ApplyWarningsAttributes on the property ({PropertyName}), for the provided tool-chain ({ToolChainName}). Skipping this property.", property.Name, toolChainContext._toolChainType.Name);
-						continue;
-					}
-
-					applyWarningsAttribute = applyWarningsAttributes.FirstOrDefault();
-
-					if (applyWarningsAttribute != null)
-					{
-						object? instanceBoxedValue = property.GetValue(this);
-
-						if (instanceBoxedValue == null)
-						{
-							_logger?.LogWarning("Invalid boxed value returned from attribute. Skipping this property.");
-							continue;
-						}
-
-						WarningLevel instanceWarningLevel = (WarningLevel)instanceBoxedValue;
-						applyWarningsAttribute.ApplyWarningsToArguments(instanceWarningLevel, toolChainContext, warningArguments);
-					}
+					_logger?.LogWarning("Ambiguous ApplyWarningsAttributes on the property ({PropertyName}), for the provided tool-chain ({ToolChainName}). Skipping this property.", memberInfo.Name, toolChainContext._toolChainType.Name);
+					continue;
 				}
-			}
 
-			// Process fields
-			{
-				s_fieldsWithAttribute ??= typeof(CppCompileWarnings)
-					.GetFields(BindingFlags.Public | BindingFlags.NonPublic | BindingFlags.Instance | BindingFlags.Static)
-					.Where(p => p.IsDefined(typeof(ApplyWarningsAttribute), false));
+				applyWarningsAttribute = applyWarningsAttributes.FirstOrDefault();
 
-				foreach (FieldInfo field in s_fieldsWithAttribute)
+				if (applyWarningsAttribute != null)
 				{
-					if (field.FieldType != typeof(WarningLevel))
-					{
-						_logger?.LogWarning("Attempting to resolve warnings, and encountered invalid field ({FieldName}). Skipping this field.", field.Name);
-						continue;
-					}
-
-					ApplyWarningsAttribute? applyWarningsAttribute = null;
-					IEnumerable<ApplyWarningsAttribute?> applyWarningsAttributes = field.GetCustomAttributes<ApplyWarningsAttribute>().Where(x => x.CanApplyToContext(toolChainContext));
-
-					if (applyWarningsAttributes.Count() > 1)
-					{
-						_logger?.LogWarning("Ambiguous ApplyWarningsAttributes on the field ({FieldName}), for the provided tool-chain ({ToolChainName}). Skipping this field.", field.Name, toolChainContext._toolChainType.Name);
-						continue;
-					}
-
-					applyWarningsAttribute = applyWarningsAttributes.FirstOrDefault();
-
-					if (applyWarningsAttribute != null)
-					{
-						object? instanceBoxedValue = field.GetValue(this);
-
-						if (instanceBoxedValue == null)
-						{
-							_logger?.LogWarning("Invalid boxed value returned from attribute. Skipping this field.");
-							continue;
-						}
-
-						WarningLevel instanceWarningLevel = (WarningLevel)instanceBoxedValue;
-						applyWarningsAttribute.ApplyWarningsToArguments(instanceWarningLevel, toolChainContext, warningArguments);
-					}
+					WarningLevel instanceWarningLevel = GetCurrentWarningLevel(memberInfo, this);
+					applyWarningsAttribute.ApplyWarningsToArguments(instanceWarningLevel, toolChainContext, warningArguments);
 				}
 			}
 
 			return warningArguments;
 		}
 
-
 		#region -- Internal API --
 
 		/// <summary>
@@ -1317,95 +1343,28 @@ namespace UnrealBuildTool
 		/// </summary>
 		/// <param name="cppCompileWarnings">The CppCompileWarnings object to apply target defaults to.</param>
 		/// <param name="overwriteSourceValue">Whether to overwrite warning setting values, or set if default.</param>
+		/// <remarks>
+		/// To have items automatically participate in default setting, use <see cref="WarningLevelDefaultAttribute"/>, with <see cref="InitializationContext.Any"/> or <see cref="InitializationContext.Target"/>.
+		/// </remarks>
 		internal static void ApplyTargetDefaults(CppCompileWarnings cppCompileWarnings, bool overwriteSourceValue = false)
 		{
 			cppCompileWarnings._logger?.LogTrace("Applying default C++ Compiler warnings for the Target({TargetName}) context.", cppCompileWarnings.BuildContextProvider?._targetRulesPrivate?.Name);
 
-			// Compiler Warning Settings
-			cppCompileWarnings._shadowVariableWarningLevel.ApplyWarning((cppCompileWarnings.BuildContextProvider?._buildContext.GetBuildSettings() >= BuildSettingsVersion.V2) ? WarningLevel.Error : WarningLevel.Warning, overwriteSourceValue);
-			cppCompileWarnings._unsafeTypeCastWarningLevel.ApplyWarning(WarningLevel.Off, overwriteSourceValue);
-			cppCompileWarnings._undefinedIdentifierWarningLevel.ApplyWarning(WarningLevel.Error, overwriteSourceValue);
-			cppCompileWarnings._switchUnhandledEnumeratorWarningLevel.ApplyWarning(WarningLevel.Off, overwriteSourceValue); // Also disabled in Clang VC Toolchain
-
-			// Always clang disabled
-			cppCompileWarnings._enumEnumConversionWarningLevel.ApplyWarning(WarningLevel.Off, overwriteSourceValue);
-			cppCompileWarnings._enumFloatConversionWarningLevel.ApplyWarning(WarningLevel.Off, overwriteSourceValue);
-			cppCompileWarnings._ambiguousReversedOperatorWarningLevel.ApplyWarning(WarningLevel.Off, overwriteSourceValue);
-			cppCompileWarnings._deprecatedAnonEnumEnumConversionWarningLevel.ApplyWarning(WarningLevel.Off, overwriteSourceValue);
-			cppCompileWarnings._deprecatedVolatileWarningLevel.ApplyWarning(WarningLevel.Off, overwriteSourceValue);
-			cppCompileWarnings._orderedCompareFunctionPointers.ApplyWarning(WarningLevel.Off, overwriteSourceValue);
-			cppCompileWarnings._unusedWarningLevel.ApplyWarning(WarningLevel.Off, overwriteSourceValue);
-			cppCompileWarnings._unusedValueWarningLevel.ApplyWarning(WarningLevel.Off, overwriteSourceValue); // Also disabled in Clang VC Toolchain, but mirrored with _msvcUnusedValueWarningLevel
-			cppCompileWarnings._unknownPragmasWarningLevel.ApplyWarning(WarningLevel.Off, overwriteSourceValue);
-			cppCompileWarnings._tautologicalCompareWarningLevel.ApplyWarning(WarningLevel.Off, overwriteSourceValue);
-			cppCompileWarnings._switchWarningLevel.ApplyWarning(WarningLevel.Off, overwriteSourceValue);
-			cppCompileWarnings._invalidOffsetOfWarningLevel.ApplyWarning(WarningLevel.Off, overwriteSourceValue);
-			cppCompileWarnings._inconsistentMissingOverrideWarningLevel.ApplyWarning(WarningLevel.Off, overwriteSourceValue);
-			cppCompileWarnings._gnuStringLiteralOperatorTemplateWarningLevel.ApplyWarning(WarningLevel.Off, overwriteSourceValue);
-			cppCompileWarnings._undefinedVarTemplateWarningLevel.ApplyWarning(WarningLevel.Off, overwriteSourceValue);
-			cppCompileWarnings._bitwiseInsteadOfLogicalWarningLevel.ApplyWarning(WarningLevel.Off, overwriteSourceValue);
-			cppCompileWarnings._deprecatedCopyWithUserProvidedCopyWarningLevel.ApplyWarning(WarningLevel.Off, overwriteSourceValue); // Also disabled in Clang VC Toolchain
-			cppCompileWarnings._deprecatedCopyWarningLevel.ApplyWarning(WarningLevel.Off, overwriteSourceValue);
-			cppCompileWarnings._invalidUnevaluatedStringWarningLevel.ApplyWarning(WarningLevel.Off, overwriteSourceValue);
-			cppCompileWarnings._nanInfinityDisabledWarningLevel.ApplyWarning(WarningLevel.Off, overwriteSourceValue);
-			cppCompileWarnings._castFunctionTypeMistmatchWarningLevel.ApplyWarning(WarningLevel.Off, overwriteSourceValue);
-			cppCompileWarnings._missingTemplateArgListAfterTemplateWarningLevel.ApplyWarning(WarningLevel.Off, overwriteSourceValue);
-			cppCompileWarnings._levelExtraQualificationWarningLevel.ApplyWarning(WarningLevel.Off, overwriteSourceValue);
-			cppCompileWarnings._profileInstructWarningLevel.ApplyWarning(WarningLevel.Off, overwriteSourceValue);
-			cppCompileWarnings._backendPluginWarningLevel.ApplyWarning(WarningLevel.Off, overwriteSourceValue);
-			cppCompileWarnings._shorten64To32WarningLevel.ApplyWarning(WarningLevel.Off, overwriteSourceValue);
-			cppCompileWarnings._dllExportExplicitInstantiationDeclWarningLevel.ApplyWarning(WarningLevel.Off, overwriteSourceValue);
-
-			// Always clang vc disabled
-			cppCompileWarnings._microsoftGroupWarningLevel.ApplyWarning(WarningLevel.Off, overwriteSourceValue);
-			cppCompileWarnings._msvcIncludeWarningLevel.ApplyWarning(WarningLevel.Off, overwriteSourceValue);
-			cppCompileWarnings._pragmaPackWarningLevel.ApplyWarning(WarningLevel.Off, overwriteSourceValue);
-			cppCompileWarnings._inlineNewDeleteWarningLevel.ApplyWarning(WarningLevel.Off, overwriteSourceValue);
-			cppCompileWarnings._implicitExceptionSpecMismatchWarningLevel.ApplyWarning(WarningLevel.Off, overwriteSourceValue);
-			cppCompileWarnings._undefinedBoolConversionWarningLevel.ApplyWarning(WarningLevel.Off, overwriteSourceValue);
-			cppCompileWarnings._deprecatedWriteableStringsWarningLevel.ApplyWarning(WarningLevel.Off, overwriteSourceValue);
-			cppCompileWarnings._deprecatedRegisterWarningLevel.ApplyWarning(WarningLevel.Off, overwriteSourceValue);
-			cppCompileWarnings._logicalOpParenthesesWarningLevel.ApplyWarning(WarningLevel.Off, overwriteSourceValue);
-			cppCompileWarnings._nullArithmeticWarningLevel.ApplyWarning(WarningLevel.Off, overwriteSourceValue);
-			cppCompileWarnings._returnTypeCLinkageWarningLevel.ApplyWarning(WarningLevel.Off, overwriteSourceValue);
-			cppCompileWarnings._ignoredAttributesWarningLevel.ApplyWarning(WarningLevel.Off, overwriteSourceValue);
-			cppCompileWarnings._uninitializedWarningLevel.ApplyWarning(WarningLevel.Off, overwriteSourceValue);
-			cppCompileWarnings._returnTypeWarningLevel.ApplyWarning(WarningLevel.Off, overwriteSourceValue);
-			cppCompileWarnings._unusedParameterWarningLevel.ApplyWarning(WarningLevel.Off, overwriteSourceValue);
-			cppCompileWarnings._ignoredQualifiersWarningLevel.ApplyWarning(WarningLevel.Off, overwriteSourceValue);
-			cppCompileWarnings._expansionToDefined.ApplyWarning(WarningLevel.Off, overwriteSourceValue);
-			cppCompileWarnings._signCompareWarningLevel.ApplyWarning(WarningLevel.Off, overwriteSourceValue);
-			cppCompileWarnings._missingFieldInitializersWarningLevel.ApplyWarning(WarningLevel.Off, overwriteSourceValue);
-			cppCompileWarnings._nonPortableIncludePathWarningLevel.ApplyWarning(WarningLevel.Off, overwriteSourceValue);
-			cppCompileWarnings._invalidTokenPasteWarningLevel.ApplyWarning(WarningLevel.Off, overwriteSourceValue);
-			cppCompileWarnings._nullPointerArithmeticWarningLevel.ApplyWarning(WarningLevel.Off, overwriteSourceValue);
-			cppCompileWarnings._constantLogicalOperandWarningLevel.ApplyWarning(WarningLevel.Off, overwriteSourceValue);
-			cppCompileWarnings._bitfieldEnumConversion.ApplyWarning(WarningLevel.Off, overwriteSourceValue);
-			cppCompileWarnings._nullPointerSubtractionWarningLevel.ApplyWarning(WarningLevel.Off, overwriteSourceValue);
-			cppCompileWarnings._danglingWarningLevel.ApplyWarning(WarningLevel.Off, overwriteSourceValue);
-
-			// Special case duplicates, for MSVC guaranteed disable.
-			cppCompileWarnings._msvcUnusedValueWarningLevel.ApplyWarning(WarningLevel.Off, overwriteSourceValue);
-			cppCompileWarnings._msvcDeprecationWarningLevel.ApplyWarning(WarningLevel.Off, overwriteSourceValue);
-			cppCompileWarnings._msvcSwitchEnumWarningLevel.ApplyWarning(WarningLevel.Off, overwriteSourceValue);
-
-			// CppCompileEnvironment refactors
-			cppCompileWarnings._deprecationWarningLevel.ApplyWarning(WarningLevel.Warning, overwriteSourceValue); // Also disabled in Clang VC Toolchain
-			cppCompileWarnings._deterministicWarningLevel.ApplyWarning(WarningLevel.Off, overwriteSourceValue);
-
-			// Performance Warning Settings
-			cppCompileWarnings._pchPerformanceIssueWarningLevel.ApplyWarning(WarningLevel.Off, overwriteSourceValue);
-
-			// Rules Validation Warning Settings
-			cppCompileWarnings._moduleUnsupportedWarningLevel.ApplyWarning(WarningLevel.Warning, overwriteSourceValue);
-			cppCompileWarnings._pluginModuleUnsupportedWarningLevel.ApplyWarning(WarningLevel.Off, overwriteSourceValue);
-
-			cppCompileWarnings._moduleIncludePathWarningLevel.ApplyWarning(WarningLevel.Off, overwriteSourceValue);
-			cppCompileWarnings._moduleIncludePrivateWarningLevel.ApplyWarning(WarningLevel.Off, overwriteSourceValue);
-			cppCompileWarnings._moduleIncludeSubdirectoryWarningLevel.ApplyWarning(WarningLevel.Off, overwriteSourceValue);
-
-			cppCompileWarnings._disablePluginsConflictWarningLevel.ApplyWarning(WarningLevel.Default, overwriteSourceValue);
-			cppCompileWarnings._nonInlinedGenCppWarningLevel.ApplyWarning(WarningLevel.Off, false);
+			BuildSettingsVersion buildSettingsVersion = cppCompileWarnings.BuildContextProvider != null ? cppCompileWarnings.BuildContextProvider._buildContext.GetBuildSettings() : BuildSettingsVersion.V1;
+			IEnumerable<MemberInfo> memberInfos = GetMembersWithAttributes(typeof(CppCompileWarnings), typeof(WarningLevelDefaultAttribute), typeof(WarningLevel));
+
+			foreach (MemberInfo mi in memberInfos)
+			{
+				IList<WarningLevelDefaultAttribute> attributes = mi.GetCustomAttributes<WarningLevelDefaultAttribute>().Where(prop => prop.Context == InitializationContext.Any || prop.Context == InitializationContext.Target).ToList();
+				WarningLevel defaultWarningLevel = WarningLevelDefaultAttribute.ResolveWarningLevelDefault(attributes, buildSettingsVersion);
+
+				// Obtain instance value, and see if we need to preserve the TargetRules' constructor setting, or apply the default.
+				WarningLevel currentWarningLevel = GetCurrentWarningLevel(mi, cppCompileWarnings);
+
+				currentWarningLevel = WarningLevelExtensions.GetAppliedWarningLevel(currentWarningLevel, defaultWarningLevel, overwriteSourceValue);
+
+				SetWarningLevel(mi, cppCompileWarnings, currentWarningLevel);
+			}
 		}
 
 		#endregion
@@ -1431,98 +1390,68 @@ namespace UnrealBuildTool
 		/// <summary>
 		/// Default settings for the no-arg constructor <see cref="CppCompileWarnings()"/>.
 		/// </summary>
-		/// <remarks>Defaults for the default compile environment state at <see cref="BuildSettingsVersion.V5"/> from <see cref="CppCompileEnvironment"/>.</remarks>
+		/// <remarks>
+		/// To have items automatically participate in default setting, use <see cref="WarningLevelDefaultAttribute"/>, with <see cref="InitializationContext.Any"/> or <see cref="InitializationContext.Constructor"/>.
+		/// </remarks>
 		private void ApplyDefaults()
 		{
 			_logger?.LogTrace("Applying default C++ Compiler warnings for the default construction context.");
 
-			// Compiler Warning Settings
-			_shadowVariableWarningLevel = WarningLevel.Warning;
-			_unsafeTypeCastWarningLevel = WarningLevel.Off;
-			_undefinedIdentifierWarningLevel = WarningLevel.Off;
-			_switchUnhandledEnumeratorWarningLevel = WarningLevel.Off; // Also disabled in Clang VC Toolchain
-
-			// Always clang disabled
-			_enumEnumConversionWarningLevel = WarningLevel.Off;
-			_enumFloatConversionWarningLevel = WarningLevel.Off;
-			_ambiguousReversedOperatorWarningLevel = WarningLevel.Off;
-			_deprecatedAnonEnumEnumConversionWarningLevel = WarningLevel.Off;
-			_deprecatedVolatileWarningLevel = WarningLevel.Off;
-			_orderedCompareFunctionPointers = WarningLevel.Off;
-			_unusedWarningLevel = WarningLevel.Off;
-			_unusedValueWarningLevel = WarningLevel.Off; // Also disabled in Clang VC Toolchain
-			_unknownPragmasWarningLevel = WarningLevel.Off;
-			_tautologicalCompareWarningLevel = WarningLevel.Off;
-			_switchWarningLevel = WarningLevel.Off;
-			_invalidOffsetOfWarningLevel = WarningLevel.Off;
-			_inconsistentMissingOverrideWarningLevel = WarningLevel.Off;
-			_gnuStringLiteralOperatorTemplateWarningLevel = WarningLevel.Off;
-			_undefinedVarTemplateWarningLevel = WarningLevel.Off;
-			_bitwiseInsteadOfLogicalWarningLevel = WarningLevel.Off;
-			_deprecatedCopyWithUserProvidedCopyWarningLevel = WarningLevel.Off; // Also disabled in Clang VC Toolchain
-			_deprecatedCopyWarningLevel = WarningLevel.Off;
-			_invalidUnevaluatedStringWarningLevel = WarningLevel.Off;
-			_nanInfinityDisabledWarningLevel = WarningLevel.Off;
-			_castFunctionTypeMistmatchWarningLevel = WarningLevel.Off;
-			_missingTemplateArgListAfterTemplateWarningLevel = WarningLevel.Off;
-			_levelExtraQualificationWarningLevel = WarningLevel.Off;
-			_profileInstructWarningLevel = WarningLevel.Off;
-			_backendPluginWarningLevel = WarningLevel.Off;
-			_shorten64To32WarningLevel = WarningLevel.Off;
-			_dllExportExplicitInstantiationDeclWarningLevel = WarningLevel.Off;
-
-			// Always clang vc disables
-			_microsoftGroupWarningLevel = WarningLevel.Off;
-			_msvcIncludeWarningLevel = WarningLevel.Off;
-			_pragmaPackWarningLevel = WarningLevel.Off;
-			_inlineNewDeleteWarningLevel = WarningLevel.Off;
-			_implicitExceptionSpecMismatchWarningLevel = WarningLevel.Off;
-			_undefinedBoolConversionWarningLevel = WarningLevel.Off;
-			_deprecatedWriteableStringsWarningLevel = WarningLevel.Off;
-			_deprecatedRegisterWarningLevel = WarningLevel.Off;
-			_logicalOpParenthesesWarningLevel = WarningLevel.Off;
-			_nullArithmeticWarningLevel = WarningLevel.Off;
-			_returnTypeCLinkageWarningLevel = WarningLevel.Off;
-			_ignoredAttributesWarningLevel = WarningLevel.Off;
-			_uninitializedWarningLevel = WarningLevel.Off;
-			_returnTypeWarningLevel = WarningLevel.Off;
-			_unusedParameterWarningLevel = WarningLevel.Off;
-			_ignoredQualifiersWarningLevel = WarningLevel.Off;
-			_expansionToDefined = WarningLevel.Off;
-			_signCompareWarningLevel = WarningLevel.Off;
-			_missingFieldInitializersWarningLevel = WarningLevel.Off;
-			_nonPortableIncludePathWarningLevel = WarningLevel.Off;
-			_invalidTokenPasteWarningLevel = WarningLevel.Off;
-			_nullPointerArithmeticWarningLevel = WarningLevel.Off;
-			_constantLogicalOperandWarningLevel = WarningLevel.Off;
-			_bitfieldEnumConversion = WarningLevel.Off;
-			_nullPointerSubtractionWarningLevel = WarningLevel.Off;
-			_danglingWarningLevel = WarningLevel.Off;
-
-			// Special case duplicates, for MSVC guaranteed disable.
-			_msvcUnusedValueWarningLevel = WarningLevel.Off;
-			_msvcDeprecationWarningLevel = WarningLevel.Off;
-			_msvcSwitchEnumWarningLevel = WarningLevel.Off;
-
-			// CppCompileEnvironment refactors
-			_deprecationWarningLevel = WarningLevel.Warning;
-			_deterministicWarningLevel = WarningLevel.Off;
-
-			// Performance Warning Settings
-			_pchPerformanceIssueWarningLevel = WarningLevel.Off;
-
-			// Rules Validation Warning Settings
-			_moduleUnsupportedWarningLevel = WarningLevel.Warning;
-			_pluginModuleUnsupportedWarningLevel = WarningLevel.Off;
-
-			_moduleIncludePathWarningLevel = WarningLevel.Off;
-			_moduleIncludePrivateWarningLevel = WarningLevel.Off;
-			_moduleIncludeSubdirectoryWarningLevel = WarningLevel.Off;
-
-			_disablePluginsConflictWarningLevel = WarningLevel.Default;
-			_nonInlinedGenCppWarningLevel = WarningLevel.Off;
+			BuildSettingsVersion buildSettingsVersion = BuildContextProvider != null ? BuildContextProvider._buildContext.GetBuildSettings() : BuildSettingsVersion.V1;
+
+			IEnumerable<MemberInfo> memberInfos = GetMembersWithAttributes(typeof(CppCompileWarnings), typeof(WarningLevelDefaultAttribute), typeof(WarningLevel));
+
+			foreach (MemberInfo mi in memberInfos)
+			{
+				IList<WarningLevelDefaultAttribute> attributes = mi.GetCustomAttributes<WarningLevelDefaultAttribute>().Where(prop => prop.Context == InitializationContext.Any || prop.Context == InitializationContext.Constructor).ToList();
+				WarningLevel defaultWarningLevel = WarningLevelDefaultAttribute.ResolveWarningLevelDefault(attributes, buildSettingsVersion);
+
+				SetWarningLevel(mi, this, defaultWarningLevel);
+			}
+		}
+
+		#region -- Reflection helpers --
+
+		private static IEnumerable<MemberInfo> GetMembersWithAttributes(Type type, Type attributeType, Type targetType)
+		{
+			IEnumerable<MemberInfo> properties = type.GetProperties(BindingFlags.Instance | BindingFlags.Public | BindingFlags.NonPublic)
+				.Where(prop => Attribute.IsDefined(prop, attributeType) && prop.PropertyType == targetType)
+				.Select(prop => (MemberInfo)prop);
+
+			IEnumerable<MemberInfo> fields = type.GetFields(BindingFlags.Instance | BindingFlags.Public | BindingFlags.NonPublic)
+				.Where(field => Attribute.IsDefined(field, attributeType) && field.FieldType == targetType)
+				.Select(field => (MemberInfo)field);
+
+			return properties.Concat(fields);
+		}
+
+		private static void SetWarningLevel(MemberInfo member, object instance, WarningLevel value)
+		{
+			switch (member)
+			{
+				case PropertyInfo pi:
+					pi.SetValue(instance, value);
+					break;
+				case FieldInfo fi:
+					fi.SetValue(instance, value);
+					break;
+			}
 		}
 
+		private static WarningLevel GetCurrentWarningLevel(MemberInfo member, object instance)
+		{
+			object? rawValue = member switch
+			{
+				PropertyInfo pi => pi.GetValue(instance),
+				FieldInfo fi => fi.GetValue(instance),
+				_ => null
+			};
+
+			return rawValue != null ? (WarningLevel)rawValue : WarningLevel.Default;
+		}
+
+		#endregion
+
 		#endregion -- Private API --
 
 		#region -- Internal & Private Instance State --
