Optimized compilation flags for Linux Server ARM builds. Added option to optimize for AWS Graviton2 and Graviton3 CPUs

diff --git a/Engine/Plugins/Mutable/Source/MutableRuntime/Private/MuR/OpImageMipmap.cpp b/Engine/Plugins/Mutable/Source/MutableRuntime/Private/MuR/OpImageMipmap.cpp
--- a/Engine/Plugins/Mutable/Source/MutableRuntime/Private/MuR/OpImageMipmap.cpp
+++ b/Engine/Plugins/Mutable/Source/MutableRuntime/Private/MuR/OpImageMipmap.cpp
@@ -43,6 +43,7 @@ namespace OpImageMipmapInternal
 	}
 
 	template<int32 PIXEL_SIZE>
+	UE_NO_PROFILE_SVE_BUG
 	inline void GenerateNextMipmapUint8SimpleAverage(const uint8* SourceData, uint8* DestData, FIntVector2 SourceSize)
 	{
 		check(SourceSize[0] > 1 || SourceSize[1] > 1);
diff --git a/Engine/Source/Programs/UnrealBuildTool/Configuration/ReadOnlyTargetRules.cs b/Engine/Source/Programs/UnrealBuildTool/Configuration/ReadOnlyTargetRules.cs
--- a/Engine/Source/Programs/UnrealBuildTool/Configuration/ReadOnlyTargetRules.cs
+++ b/Engine/Source/Programs/UnrealBuildTool/Configuration/ReadOnlyTargetRules.cs
@@ -558,6 +558,8 @@ namespace UnrealBuildTool
 
 		public MinimumCpuArchitectureX64 MinCpuArchX64 => Inner.MinCpuArchX64;
 
+		public Arm64TargetCpuType MinArm64CpuTarget => Inner.MinArm64CpuTarget;
+
 		internal bool bNoManifestChanges => Inner.bNoManifestChanges;
 
 		public string? BuildVersion => Inner.BuildVersion;
diff --git a/Engine/Source/Programs/UnrealBuildTool/Configuration/TargetRules.cs b/Engine/Source/Programs/UnrealBuildTool/Configuration/TargetRules.cs
--- a/Engine/Source/Programs/UnrealBuildTool/Configuration/TargetRules.cs
+++ b/Engine/Source/Programs/UnrealBuildTool/Configuration/TargetRules.cs
@@ -424,7 +424,7 @@ namespace UnrealBuildTool
 		Binned3,
 
 		/// <summary>
-		/// Default for Desktop\Editor targets that allows allocator selection to be overriden with a command line
+		/// Default for Desktop\Editor targets that allows allocator selection to be overridden with a command line
 		/// </summary>
 		None,
 	}
@@ -2602,6 +2602,16 @@ namespace UnrealBuildTool
 		[XmlConfigFile(Category = "BuildConfiguration")]
 		public MinimumCpuArchitectureX64 MinCpuArchX64 { get; set; } = MinimumCpuArchitectureX64.Default;
 
+		/// <summary>
+		/// Select a minimum ARM64 target CPU. This will change what ARMv8(9)-A ISA + extension code will be compiled for.
+		/// Changing it from default will guarantee not to work on all ARM64 CPUs.
+		/// Currently supported by clang compiler only
+		/// </summary>
+		[RequiresUniqueBuildEnvironment]
+		[CommandLine("-MinArm64CpuTarget")]
+		[XmlConfigFile(Category = "BuildConfiguration")]
+		public Arm64TargetCpuType MinArm64CpuTarget { get; set; } = Arm64TargetCpuType.Default;
+
 		/// <summary>
 		/// Do not allow manifest changes when building this target. Used to cause earlier errors when building multiple targets with a shared build environment.
 		/// </summary>
diff --git a/Engine/Source/Programs/UnrealBuildTool/Configuration/UEBuildTarget.cs b/Engine/Source/Programs/UnrealBuildTool/Configuration/UEBuildTarget.cs
--- a/Engine/Source/Programs/UnrealBuildTool/Configuration/UEBuildTarget.cs
+++ b/Engine/Source/Programs/UnrealBuildTool/Configuration/UEBuildTarget.cs
@@ -5296,6 +5296,7 @@ namespace UnrealBuildTool
 			GlobalCompileEnvironment.CppStandard = Rules.CppStandard;
 			GlobalCompileEnvironment.CStandard = Rules.CStandard;
 			GlobalCompileEnvironment.MinCpuArchX64 = Rules.MinCpuArchX64;
+			GlobalCompileEnvironment.MinArm64CpuTarget = Rules.MinArm64CpuTarget;
 			GlobalCompileEnvironment.AdditionalArguments = Rules.AdditionalCompilerArguments ?? String.Empty;
 			GlobalCompileEnvironment.bDeterministic = Rules.bDeterministic;
 			GlobalCompileEnvironment.CrashDiagnosticDirectory = Rules.CrashDiagnosticDirectory;
diff --git a/Engine/Source/Programs/UnrealBuildTool/Platform/Linux/LinuxToolChain.cs b/Engine/Source/Programs/UnrealBuildTool/Platform/Linux/LinuxToolChain.cs
--- a/Engine/Source/Programs/UnrealBuildTool/Platform/Linux/LinuxToolChain.cs
+++ b/Engine/Source/Programs/UnrealBuildTool/Platform/Linux/LinuxToolChain.cs
@@ -383,14 +383,18 @@ namespace UnrealBuildTool
 					else if (CompileEnvironment.OptimizationLevel == OptimizationMode.SizeAndSpeed)
 					{
 						Arguments.Add("-Os");
-						if (Architecture == UnrealArch.Arm64)
-						{
-							Arguments.Add("-moutline");
-						}
 					}
 					else
 					{
-						Arguments.Add("-O3");
+						if (CompileEnvironment.bPGOOptimize || CompileEnvironment.bPGOProfile)
+						{
+							// -Os tends to be both faster and smaller than -O3 when PGO is enabled
+							Arguments.Add("-Os");
+						}
+						else
+						{
+							Arguments.Add("-O3");
+						}
 					}
 				}
 			}
@@ -488,8 +492,6 @@ namespace UnrealBuildTool
 		/// <inheritdoc/>
 		protected override void GetCompileArguments_Global(CppCompileEnvironment CompileEnvironment, List<string> Arguments)
 		{
-			base.GetCompileArguments_Global(CompileEnvironment, Arguments);
-
 			// build up the commandline common to C and C++
 
 			// These aren't supported on Linux at this time
@@ -573,6 +575,8 @@ namespace UnrealBuildTool
 				Arguments.Add($"-target {CompileEnvironment.Architecture.LinuxName}");        // Set target triple
 				Arguments.Add($"--sysroot=\"{NormalizeCommandLinePath(LinuxInfo.BaseLinuxPath!, CompileEnvironment.RootPaths)}\"");
 			}
+
+			base.GetCompileArguments_Global(CompileEnvironment, Arguments);
 		}
 
 		/// <inheritdoc/>
@@ -831,6 +835,7 @@ namespace UnrealBuildTool
 			if (LinkEnvironment.Configuration == CppConfiguration.Shipping)
 			{
 				Arguments.Add("-Wl,--icf=all"); // Enables ICF (Identical Code Folding). [all, safe] safe == fold functions that can be proven not to have their address taken.
+				Arguments.Add("-Wl,-O3");
 			}
 		}
 
@@ -887,7 +892,7 @@ namespace UnrealBuildTool
 				Logger.LogInformation("Compressing debug files");
 			}
 
-			Logger.LogInformation("Targeted minimum CPU architecture: {0}", (CompileEnvironment.Architecture == UnrealArch.X64) ? CompileEnvironment.MinCpuArchX64 : "default");
+			Logger.LogInformation("Targeted minimum CPU architecture: {0}", (CompileEnvironment.Architecture == UnrealArch.X64) ? CompileEnvironment.MinCpuArchX64 : CompileEnvironment.MinArm64CpuTarget);
 
 			if (CompileEnvironment.bPGOOptimize)
 			{
diff --git a/Engine/Source/Programs/UnrealBuildTool/Platform/Linux/UEBuildLinux.cs b/Engine/Source/Programs/UnrealBuildTool/Platform/Linux/UEBuildLinux.cs
--- a/Engine/Source/Programs/UnrealBuildTool/Platform/Linux/UEBuildLinux.cs
+++ b/Engine/Source/Programs/UnrealBuildTool/Platform/Linux/UEBuildLinux.cs
@@ -198,12 +198,7 @@ namespace UnrealBuildTool
 		public bool IsLTOEnabled(ReadOnlyTargetRules Target)
 		{
 			// Force LTO on if using PGO, or if specified in the target rules.
-			if (Target.bAllowLTCG || Target.bPGOOptimize || Target.bPGOProfile)
-			{
-				return true;
-			}
-
-			return false;
+			return Target.bAllowLTCG || Target.bPGOOptimize || Target.bPGOProfile;
 		}
 
 		public override void ValidateTarget(TargetRules Target)
@@ -672,7 +667,8 @@ namespace UnrealBuildTool
 			{
 				Options |= ClangToolChainOptions.EnableLinkTimeOptimization;
 
-				if (Target.bPreferThinLTO)
+				// Prefer Thin LTO for PGO builds, even if it's not explicitly specified
+				if (Target.bPreferThinLTO || Target.bPGOOptimize || Target.bPGOProfile)
 				{
 					Options |= ClangToolChainOptions.EnableThinLTO;
 				}
diff --git a/Engine/Source/Programs/UnrealBuildTool/System/CppCompileEnvironment.cs b/Engine/Source/Programs/UnrealBuildTool/System/CppCompileEnvironment.cs
--- a/Engine/Source/Programs/UnrealBuildTool/System/CppCompileEnvironment.cs
+++ b/Engine/Source/Programs/UnrealBuildTool/System/CppCompileEnvironment.cs
@@ -127,11 +127,32 @@ namespace UnrealBuildTool
 		AVX512,
 
 		/// <summary>
-		/// Use the default minimum architecure
+		/// Use the default minimum architecture
 		/// </summary>
 		Default = None,
 	}
 
+	/// <summary>
+	/// Target ARM64 CPU for codegen. This defines different ARM ISA and extension targets
+	/// </summary>
+	public enum Arm64TargetCpuType
+	{
+		/// <summary>
+		/// Default ARMv8.0-A ISA target CPU, supported by every 64 bit ARM CPU
+		/// </summary>
+		Default,
+
+		/// <summary>
+		/// AWS Graviton2 CPU - ARM Neoverse N1 uArch with ARMv8.2-A ISA + FP16 + rcpc + dotprod + crypto extensions
+		/// </summary>
+		Graviton2,
+
+		/// <summary>
+		/// AWS Graviton3 CPU - ARM Neoverse V1 uArch with AMRv8.4-A ISA + SVE + bf16 + rng + crypto extensions
+		/// </summary>
+		Graviton3,
+	}
+
 	/// <summary>
 	/// The optimization level that may be compilation targets for C# files.
 	/// </summary>
@@ -694,6 +715,13 @@ namespace UnrealBuildTool
 		/// </summary>
 		public MinimumCpuArchitectureX64 MinCpuArchX64 = MinimumCpuArchitectureX64.Default;
 
+		/// <summary>
+		/// Select a minimum ARM64 target CPU. This will change what ARMv8(9)-A ISA + extension code will be compiled for.
+		/// Changing it from default will guarantee not to work on all ARM64 CPUs.
+		/// Currently supported by clang compiler only
+		/// </summary>
+		public Arm64TargetCpuType MinArm64CpuTarget = Arm64TargetCpuType.Default;
+
 		/// <summary>
 		/// The amount of the stack usage to report static analysis warnings.
 		/// </summary>
@@ -857,6 +885,7 @@ namespace UnrealBuildTool
 			CppStandard = Other.CppStandard;
 			CStandard = Other.CStandard;
 			MinCpuArchX64 = Other.MinCpuArchX64;
+			MinArm64CpuTarget = Other.MinArm64CpuTarget;
 			bEnableCoroutines = Other.bEnableCoroutines;
 			IncludeOrderVersion = Other.IncludeOrderVersion;
 			bDeterministic = Other.bDeterministic;
diff --git a/Engine/Source/Programs/UnrealBuildTool/ToolChain/ClangToolChain.cs b/Engine/Source/Programs/UnrealBuildTool/ToolChain/ClangToolChain.cs
--- a/Engine/Source/Programs/UnrealBuildTool/ToolChain/ClangToolChain.cs
+++ b/Engine/Source/Programs/UnrealBuildTool/ToolChain/ClangToolChain.cs
@@ -777,6 +777,19 @@ namespace UnrealBuildTool
 			if (CompileEnvironment.Architecture == UnrealArch.Arm64)
 			{
 				Arguments.Add("-fsigned-char");
+
+				switch (CompileEnvironment.MinArm64CpuTarget)
+				{
+					case Arm64TargetCpuType.Graviton2:
+						Arguments.Add("-march=armv8.2-a+fp16+rcpc+dotprod+crypto");
+						Arguments.Add("-mtune=neoverse-n1");
+						break;
+
+					case Arm64TargetCpuType.Graviton3:
+						Arguments.Add("-mcpu=neoverse-512tvb+bf16+rng+crypto");
+						Arguments.Add("-mtune=neoverse-v1");
+						break;
+				}
 			}
 		}
 
diff --git a/Engine/Source/Runtime/Core/Public/HAL/Platform.h b/Engine/Source/Runtime/Core/Public/HAL/Platform.h
--- a/Engine/Source/Runtime/Core/Public/HAL/Platform.h
+++ b/Engine/Source/Runtime/Core/Public/HAL/Platform.h
@@ -1004,6 +1004,17 @@ int32 tchar_main(int32 ArgC, TCHAR* ArgV[])
 	#define UE_FORCE_CRASH() UE_FORCE_CRASH_AT_OFFSET(3)
 #endif
 
+#ifndef UE_WORKAROUND_CLANG_SVE_LTO_BUGS
+#	define UE_WORKAROUND_CLANG_SVE_LTO_BUGS 0
+#endif
+
+// workaround for LTO generating invalid SVE code that crashes clang
+#if UE_WORKAROUND_CLANG_SVE_LTO_BUGS
+#	define UE_NO_PROFILE_SVE_BUG UE_NO_PROFILE_ATTRIBUTE
+#else
+#	define UE_NO_PROFILE_SVE_BUG
+#endif
+
 //--------------------------------------------------------------------------------------------
 // POD types refactor for porting old code:
 //
