Add MinimalAPI for FStruct and UEnum
Add a module option to not export all FStructs and/or UEnums by default

[RN] Add a module option to not export all FStructs and/or UEnums by default

diff --git a/Engine/Source/Programs/Shared/EpicGames.Core/UHTTypes.cs b/Engine/Source/Programs/Shared/EpicGames.Core/UHTTypes.cs
--- a/Engine/Source/Programs/Shared/EpicGames.Core/UHTTypes.cs
+++ b/Engine/Source/Programs/Shared/EpicGames.Core/UHTTypes.cs
@@ -126,6 +126,8 @@ namespace EpicGames.Core
 			public UHTVerseScope VerseScope { get; set; } = UHTVerseScope.PublicAPI;
 			public bool HasVerse { get; set; } = false;
 			public string VerseMountPoint { get; set; } = "";
+			public bool AlwaysExportStructs { get; set; } = true;
+			public bool AlwaysExportEnums { get; set; } = true;
 
 			public override string ToString()
 			{
diff --git a/Engine/Source/Programs/Shared/EpicGames.UHT/Exporters/CodeGen/UhtHeaderCodeGeneratorCppFile.cs b/Engine/Source/Programs/Shared/EpicGames.UHT/Exporters/CodeGen/UhtHeaderCodeGeneratorCppFile.cs
--- a/Engine/Source/Programs/Shared/EpicGames.UHT/Exporters/CodeGen/UhtHeaderCodeGeneratorCppFile.cs
+++ b/Engine/Source/Programs/Shared/EpicGames.UHT/Exporters/CodeGen/UhtHeaderCodeGeneratorCppFile.cs
@@ -440,7 +440,14 @@ namespace EpicGames.UHT.Exporters.CodeGen
 
 			builder.Append("}\r\n");
 
-			builder.Append("template<> ").Append(Module.Api).Append("UEnum* StaticEnum<").Append(enumObj.CppType).Append(">()\r\n");
+			if (Module.Module.AlwaysExportEnums)
+			{
+				builder.Append("template<> ").Append(Module.Api).Append("UEnum* StaticEnum<").Append(enumObj.CppType).Append(">()\r\n");
+			}
+			else
+			{
+				builder.Append("template<> ").Append("UEnum* StaticEnum<").Append(enumObj.CppType).Append(">()\r\n");
+			}
 			builder.Append("{\r\n");
 			builder.Append("\treturn ").Append(enumObj.SourceName).Append("_StaticEnum();\r\n");
 			builder.Append("}\r\n");
diff --git a/Engine/Source/Programs/Shared/EpicGames.UHT/Exporters/CodeGen/UhtHeaderCodeGeneratorHFile.cs b/Engine/Source/Programs/Shared/EpicGames.UHT/Exporters/CodeGen/UhtHeaderCodeGeneratorHFile.cs
--- a/Engine/Source/Programs/Shared/EpicGames.UHT/Exporters/CodeGen/UhtHeaderCodeGeneratorHFile.cs
+++ b/Engine/Source/Programs/Shared/EpicGames.UHT/Exporters/CodeGen/UhtHeaderCodeGeneratorHFile.cs
@@ -229,7 +229,8 @@ namespace EpicGames.UHT.Exporters.CodeGen
 				}
 
 				builder.Append('\t');
-				if (!scriptStruct.ScriptStructFlags.HasAnyFlags(EStructFlags.RequiredAPI))
+				if (!scriptStruct.ScriptStructFlags.HasAnyFlags(EStructFlags.RequiredAPI) && 
+				    (Module.Module.AlwaysExportStructs || scriptStruct.ScriptStructExportFlags.HasAnyFlags(UhtScriptStructExportFlags.MinimalAPI)))
 				{
 					builder.Append(Module.Api);
 				}
@@ -451,7 +452,8 @@ namespace EpicGames.UHT.Exporters.CodeGen
 				{
 					builder.Append("\tfriend struct Z_Construct_UScriptStruct_").Append(scriptStruct.SourceName).Append("_Statics; \\\r\n");
 					builder.Append('\t');
-					if (!scriptStruct.ScriptStructFlags.HasAnyFlags(EStructFlags.RequiredAPI))
+					if (!scriptStruct.ScriptStructFlags.HasAnyFlags(EStructFlags.RequiredAPI) && 
+					    (Module.Module.AlwaysExportStructs || scriptStruct.ScriptStructExportFlags.HasAnyFlags(UhtScriptStructExportFlags.MinimalAPI)))
 					{
 						builder.Append(Module.Api);
 					}
@@ -472,7 +474,8 @@ namespace EpicGames.UHT.Exporters.CodeGen
 							.Append(", ")
 							.Append(fastArrayProperty.SourceName)
 							.Append(", ");
-						if (!scriptStruct.ScriptStructFlags.HasAnyFlags(EStructFlags.RequiredAPI))
+						if (!scriptStruct.ScriptStructFlags.HasAnyFlags(EStructFlags.RequiredAPI) && 
+						    (Module.Module.AlwaysExportStructs || scriptStruct.ScriptStructExportFlags.HasAnyFlags(UhtScriptStructExportFlags.MinimalAPI)))
 						{
 							builder.Append(Module.Api);
 						}
@@ -710,7 +713,14 @@ namespace EpicGames.UHT.Exporters.CodeGen
 			if (enumObj.CppForm == UhtEnumCppForm.EnumClass || enumObj.UnderlyingType != UhtEnumUnderlyingType.Unspecified)
 			{	
 				// Forward declare the StaticEnum<> specialization for enum classes
-				builder.Append("template<> ").Append(Module.Api).Append("UEnum* StaticEnum<").Append(enumObj.CppType).Append(">();\r\n");
+				if (Module.Module.AlwaysExportEnums || enumObj.EnumExportFlags.HasAnyFlags(UhtEnumExportFlags.MinimalAPI))
+				{
+					builder.Append("template<> ").Append(Module.Api).Append("UEnum* StaticEnum<").Append(enumObj.CppType).Append(">();\r\n");
+				}
+				else
+				{
+					builder.Append("template<> ").Append("UEnum* StaticEnum<").Append(enumObj.CppType).Append(">();\r\n");
+				}
 			}
 
 			// Verse declarations
diff --git a/Engine/Source/Programs/Shared/EpicGames.UHT/Specifiers/UhtEnumSpecifiers.cs b/Engine/Source/Programs/Shared/EpicGames.UHT/Specifiers/UhtEnumSpecifiers.cs
--- a/Engine/Source/Programs/Shared/EpicGames.UHT/Specifiers/UhtEnumSpecifiers.cs
+++ b/Engine/Source/Programs/Shared/EpicGames.UHT/Specifiers/UhtEnumSpecifiers.cs
@@ -22,5 +22,12 @@ namespace EpicGames.UHT.Parsers
 			UhtEnum enumObj = (UhtEnum)specifierContext.Type;
 			enumObj.EnumFlags |= EEnumFlags.Flags;
 		}
+		
+		[UhtSpecifier(Extends = UhtTableNames.Enum, ValueType = UhtSpecifierValueType.Legacy)]
+		private static void MinimalAPISpecifier(UhtSpecifierContext specifierContext)
+		{
+			UhtEnum enumObj = (UhtEnum)specifierContext.Type;
+			enumObj.EnumExportFlags |= UhtEnumExportFlags.MinimalAPI;
+		}
 	}
 }
diff --git a/Engine/Source/Programs/Shared/EpicGames.UHT/Specifiers/UhtScriptStructSpecifiers.cs b/Engine/Source/Programs/Shared/EpicGames.UHT/Specifiers/UhtScriptStructSpecifiers.cs
--- a/Engine/Source/Programs/Shared/EpicGames.UHT/Specifiers/UhtScriptStructSpecifiers.cs
+++ b/Engine/Source/Programs/Shared/EpicGames.UHT/Specifiers/UhtScriptStructSpecifiers.cs
@@ -81,5 +81,12 @@ namespace EpicGames.UHT.Parsers
 			UhtScriptStruct scriptStruct = (UhtScriptStruct)specifierContext.Type;
 			scriptStruct.ScriptStructExportFlags |= UhtScriptStructExportFlags.IsCoreType;
 		}
+		
+		[UhtSpecifier(Extends = UhtTableNames.ScriptStruct, ValueType = UhtSpecifierValueType.Legacy)]
+		private static void MinimalAPISpecifier(UhtSpecifierContext specifierContext)
+		{
+			UhtScriptStruct scriptStruct = (UhtScriptStruct)specifierContext.Type;
+			scriptStruct.ScriptStructExportFlags |= UhtScriptStructExportFlags.MinimalAPI;
+		}
 	}
 }
diff --git a/Engine/Source/Programs/Shared/EpicGames.UHT/Types/UhtEnum.cs b/Engine/Source/Programs/Shared/EpicGames.UHT/Types/UhtEnum.cs
--- a/Engine/Source/Programs/Shared/EpicGames.UHT/Types/UhtEnum.cs
+++ b/Engine/Source/Programs/Shared/EpicGames.UHT/Types/UhtEnum.cs
@@ -103,6 +103,66 @@ namespace EpicGames.UHT.Types
 		/// </summary>
 		public long Value { get; set; }
 	}
+	/// <summary>
+	/// Series of flags not part of the engine's enum flags that affect code generation or verification
+	/// </summary>
+	[Flags]
+	public enum UhtEnumExportFlags : int
+	{
+
+		/// <summary>
+		/// No export flags
+		/// </summary>
+		None = 0,
+	
+		/// <summary>
+		/// Indicates that the StaticEnum function should be exported to other modules
+		/// </summary>
+		MinimalAPI = 1 << 0,
+	}
+
+	/// <summary>
+	/// Helper methods for testing flags.  These methods perform better than the generic HasFlag which hits
+	/// the GC and stalls.
+	/// </summary>
+	public static class UhtEnumExportFlagsExtensions
+	{
+
+		/// <summary>
+		/// Test to see if any of the specified flags are set
+		/// </summary>
+		/// <param name="inFlags">Current flags</param>
+		/// <param name="testFlags">Flags to test for</param>
+		/// <returns>True if any of the flags are set</returns>
+		public static bool HasAnyFlags(this UhtEnumExportFlags inFlags, UhtEnumExportFlags testFlags)
+		{
+			return (inFlags & testFlags) != 0;
+		}
+
+		/// <summary>
+		/// Test to see if all of the specified flags are set
+		/// </summary>
+		/// <param name="inFlags">Current flags</param>
+		/// <param name="testFlags">Flags to test for</param>
+		/// <returns>True if all the flags are set</returns>
+		public static bool HasAllFlags(this UhtEnumExportFlags inFlags, UhtEnumExportFlags testFlags)
+		{
+			return (inFlags & testFlags) == testFlags;
+		}
+
+		/// <summary>
+		/// Test to see if a specific set of flags have a specific value.
+		/// </summary>
+		/// <param name="inFlags">Current flags</param>
+		/// <param name="testFlags">Flags to test for</param>
+		/// <param name="matchFlags">Expected value of the tested flags</param>
+		/// <returns>True if the given flags have a specific value.</returns>
+		public static bool HasExactFlags(this UhtEnumExportFlags inFlags, UhtEnumExportFlags testFlags, UhtEnumExportFlags matchFlags)
+		{
+			return (inFlags & testFlags) == matchFlags;
+		}
+	}
+
 
 	/// <summary>
 	/// Represents a UEnum
@@ -115,6 +175,12 @@ namespace EpicGames.UHT.Types
 		/// </summary>
 		[JsonConverter(typeof(JsonStringEnumConverter))]
 		public EEnumFlags EnumFlags { get; set; } = EEnumFlags.None;
+		
+		/// <summary>
+		/// UHT only script struct flags
+		/// </summary>
+		[JsonConverter(typeof(JsonStringEnumConverter))]
+		public UhtEnumExportFlags EnumExportFlags { get; set; } = UhtEnumExportFlags.None;
 
 		/// <summary>
 		/// C++ form of the enumeration
diff --git a/Engine/Source/Programs/Shared/EpicGames.UHT/Types/UhtScriptStruct.cs b/Engine/Source/Programs/Shared/EpicGames.UHT/Types/UhtScriptStruct.cs
--- a/Engine/Source/Programs/Shared/EpicGames.UHT/Types/UhtScriptStruct.cs
+++ b/Engine/Source/Programs/Shared/EpicGames.UHT/Types/UhtScriptStruct.cs
@@ -613,6 +613,11 @@ namespace EpicGames.UHT.Types
 		/// "Parametric" specifier present
 		/// </summary>
 		IsVerseParametric = 1 << 4,
+		
+		/// <summary>
+		/// Indicates that the autogenerated functions of this struct will be DLL exported/imported
+		/// </summary>
+		MinimalAPI = 1 << 5,
 	}
 
 	/// <summary>
@@ -670,7 +675,7 @@ namespace EpicGames.UHT.Types
 		public EStructFlags ScriptStructFlags { get; set; } = EStructFlags.NoFlags;
 
 		/// <summary>
-		/// UHT only script struct falgs
+		/// UHT only script struct flags
 		/// </summary>
 		[JsonConverter(typeof(JsonStringEnumConverter))]
 		public UhtScriptStructExportFlags ScriptStructExportFlags { get; set; } = UhtScriptStructExportFlags.None;
diff --git a/Engine/Source/Programs/UnrealBuildTool/Configuration/ModuleRules.cs b/Engine/Source/Programs/UnrealBuildTool/Configuration/ModuleRules.cs
--- a/Engine/Source/Programs/UnrealBuildTool/Configuration/ModuleRules.cs
+++ b/Engine/Source/Programs/UnrealBuildTool/Configuration/ModuleRules.cs
@@ -1248,7 +1248,17 @@ namespace UnrealBuildTool
 		/// Disables AutoRTFM instrumentation to this module only when AutoRTFMCompiler is enabled
 		/// </summary>
 		public bool bDisableAutoRTFMInstrumentation { get; set; }
+		
+		/// <summary>
+		/// If true, autogenerated functions of USTRUCTS are always exported
+		/// </summary>
+		public bool bAlwaysExportStructs = true;
 
+		/// <summary>
+		/// If true, autogenerated functions of UENUMS are always exported
+		/// </summary>
+		public bool bAlwaysExportEnums = true;
+		
 		/// <summary>
 		/// List of modules names (no path needed) with header files that our module's public headers needs access to, but we don't need to "import" or link against.
 		/// </summary>
diff --git a/Engine/Source/Programs/UnrealBuildTool/System/ExternalExecution.cs b/Engine/Source/Programs/UnrealBuildTool/System/ExternalExecution.cs
--- a/Engine/Source/Programs/UnrealBuildTool/System/ExternalExecution.cs
+++ b/Engine/Source/Programs/UnrealBuildTool/System/ExternalExecution.cs
@@ -177,7 +177,17 @@ namespace UnrealBuildTool
 		/// If true, verse scripts were found
 		/// </summary>
 		public bool HasVerse;
+		
+		/// <summary>
+		/// If true, autogenerated functions of USTRUCTS are always exported
+		/// </summary>
+		public bool AlwaysExportStructs;
 
+		/// <summary>
+		/// If true, autogenerated functions of UENUMS are always exported
+		/// </summary>
+		public bool AlwaysExportEnums;
+		
 		public UHTModuleInfo(
 			string ModuleName, 
 			FileReference ModuleRulesFile, 
@@ -191,7 +201,9 @@ namespace UnrealBuildTool
 			string VersePath,
 			VerseScope VerseScope,
 			bool HasVerse,
-			string VerseMountPoint)
+			string VerseMountPoint,
+			bool AlwaysExportStructs,
+			bool AlwaysExportEnums)
 		{
 			this.ModuleName = ModuleName;
 			this.ModuleRulesFile = ModuleRulesFile;
@@ -211,6 +223,8 @@ namespace UnrealBuildTool
 			this.VerseScope = VerseScope;
 			this.HasVerse = HasVerse;
 			this.VerseMountPoint = VerseMountPoint;
+			this.AlwaysExportStructs = AlwaysExportStructs;
+			this.AlwaysExportEnums = AlwaysExportEnums;
 		}
 
 		public UHTModuleInfo(BinaryArchiveReader Reader)
@@ -234,6 +248,8 @@ namespace UnrealBuildTool
 			VerseScope = (VerseScope)Reader.ReadInt();
 			HasVerse = Reader.ReadBool();
 			VerseMountPoint = Reader.ReadString()!;
+			AlwaysExportStructs = Reader.ReadBool();
+			AlwaysExportEnums = Reader.ReadBool();
 		}
 
 		public void Write(BinaryArchiveWriter Writer)
@@ -257,6 +273,8 @@ namespace UnrealBuildTool
 			Writer.WriteInt((int)VerseScope);
 			Writer.WriteBool(HasVerse);
 			Writer.WriteString(VerseMountPoint);
+			Writer.WriteBool(AlwaysExportStructs);
+			Writer.WriteBool(AlwaysExportEnums);
 		}
 
 		public override string ToString()
@@ -477,7 +495,9 @@ namespace UnrealBuildTool
 						Module.Rules.VersePath ?? "",
 						Module.Rules.VerseScope,
 						Module.bHasVerse,
-						Module.bHasVerse ? GetVerseMountPointForModule(Module) : string.Empty);
+						Module.bHasVerse ? GetVerseMountPointForModule(Module) : string.Empty,
+						Module.Rules.bAlwaysExportStructs,
+						Module.Rules.bAlwaysExportEnums);
 					ModuleInfoArray[Idx] = Info;
 
 					Queue.Enqueue(() => SetupUObjectModule(Info, ExcludedFolders, MetadataCache, Queue));
@@ -1052,6 +1072,8 @@ namespace UnrealBuildTool
 						VerseScope = (UHTVerseScope)Enum.Parse(typeof(UHTVerseScope), UObjectModule.VerseScope.ToString()),
 						HasVerse = UObjectModule.HasVerse,
 						VerseMountPoint = UObjectModule.VerseMountPoint,
+						AlwaysExportStructs = UObjectModule.AlwaysExportStructs,
+						AlwaysExportEnums = UObjectModule.AlwaysExportEnums,
 					}); ;
 			}
 
