[UBT]
* Fixed android rsp paths not being virtualized when vfs is enabled

diff --git a/Engine/Source/Programs/UnrealBuildTool/Platform/Android/AndroidToolChain.cs b/Engine/Source/Programs/UnrealBuildTool/Platform/Android/AndroidToolChain.cs
--- a/Engine/Source/Programs/UnrealBuildTool/Platform/Android/AndroidToolChain.cs
+++ b/Engine/Source/Programs/UnrealBuildTool/Platform/Android/AndroidToolChain.cs
@@ -123,6 +123,9 @@ namespace UnrealBuildTool
 		{
 		}
 
+		DirectoryReference GCCToolchainPath;
+		DirectoryReference SysrootPath;
+
 		public AndroidToolChain(FileReference? InProjectFile, ClangToolChainOptions ToolchainOptions, ILogger InLogger)
 			: base(ToolchainOptions, InLogger)
 		{
@@ -222,12 +225,12 @@ namespace UnrealBuildTool
 			// NDK setup (enforce minimum API level)
 			int NDKApiLevel64Int = GetNdkApiLevelInt(MinimumNDKAPILevel);
 
-			string GCCToolchainPath = Path.Combine(NDKPath, @"toolchains/llvm", ArchitecturePath);
-			string SysrootPath = Path.Combine(NDKPath, @"toolchains/llvm", ArchitecturePath, "sysroot").Replace("\\", "/");
+			GCCToolchainPath = DirectoryReference.FromString(Path.Combine(NDKPath, @"toolchains/llvm", ArchitecturePath))!;
+			SysrootPath = DirectoryReference.FromString(Path.Combine(NDKPath, @"toolchains/llvm", ArchitecturePath, "sysroot").Replace("\\", "/"))!;
 
 			// toolchain params (note: use ANDROID=1 same as we define it)
-			ToolchainLinkParamsArm64 = " --target=aarch64-none-linux-android" + NDKApiLevel64Int + " --gcc-toolchain=\"" + GCCToolchainPath + "\" --sysroot=\"" + SysrootPath + "\" -DANDROID=1";
-			ToolchainLinkParamsx64 = " --target=x86_64-none-linux-android" + NDKApiLevel64Int + " --gcc-toolchain=\"" + GCCToolchainPath + "\" --sysroot=\"" + SysrootPath + "\" -DANDROID=1";
+			ToolchainLinkParamsArm64 = " --target=aarch64-none-linux-android" + NDKApiLevel64Int + " -DANDROID=1";
+			ToolchainLinkParamsx64 = " --target=x86_64-none-linux-android" + NDKApiLevel64Int + " -DANDROID=1";
 
 			ToolchainParamsArm64 = ToolchainLinkParamsArm64;
 			ToolchainParamsx64 = ToolchainLinkParamsx64;
@@ -242,6 +245,15 @@ namespace UnrealBuildTool
 			ReadElfPath = Path.Combine(NDKPath, @"toolchains/llvm", ArchitecturePath, @"bin/llvm-readelf" + ExeExtension);
 		}
 
+		IEnumerable<string> AdditionalPaths(CppRootPaths Roots)
+		{
+			return new[]
+			{
+				$"--gcc-toolchain=\"{NormalizeCommandLinePath(GCCToolchainPath, Roots)}\"",
+				$"--sysroot=\"{NormalizeCommandLinePath(SysrootPath, Roots)}\"",
+			};
+		}
+
 		protected override ClangToolChainInfo GetToolChainInfo()
 		{
 			return new ClangToolChainInfo(DirectoryReference.FromString(AndroidPlatformSDK.GetNDKRoot()), FileReference.FromString(ClangPath)!, FileReference.FromString(ArPathArm64)!, Logger);
@@ -551,6 +563,8 @@ namespace UnrealBuildTool
 		{
 			base.GetCompileArguments_Architecture(CompileEnvironment, Arguments);
 
+			Arguments.AddRange(AdditionalPaths(CompileEnvironment.RootPaths));
+
 			if (CompileEnvironment.Architecture == UnrealArch.Arm64)
 			{
 				Arguments.Add(ToolchainParamsArm64);
@@ -708,7 +722,7 @@ namespace UnrealBuildTool
 
 		protected virtual string GetLinkArguments(LinkEnvironment LinkEnvironment, UnrealArch Architecture)
 		{
-			string Result = "";
+			string Result = string.Join(' ', AdditionalPaths(LinkEnvironment.RootPaths));
 
 			//Result += " -nostdlib";
 			Result += " -static-libstdc++";
