[WinArm64] Tweak host/local architecture detection so that arm64 and arm64ec are both considered 'local' when generating projects on arm64.

#rnx

diff --git a/Engine/Source/Programs/UnrealBuildTool/Platform/Windows/WindowsProjectGenerator.cs b/Engine/Source/Programs/UnrealBuildTool/Platform/Windows/WindowsProjectGenerator.cs
--- a/Engine/Source/Programs/UnrealBuildTool/Platform/Windows/WindowsProjectGenerator.cs
+++ b/Engine/Source/Programs/UnrealBuildTool/Platform/Windows/WindowsProjectGenerator.cs
@@ -116,7 +116,7 @@ namespace UnrealBuildTool
 		{
 			StringBuilder VCUserFileContent = new StringBuilder();
 
-			string LocalOrRemoteString = InVSSettings.Architecture == null || InVSSettings.Architecture.Value == UnrealArch.Host.Value
+			string LocalOrRemoteString = InVSSettings.Architecture == null || (InVSSettings.Architecture.Value.bIsX64 == UnrealArch.Host.Value.bIsX64)
 				? "Local" : "Remote";
 
 			VCUserFileContent.AppendLine("  <PropertyGroup {0}>", InConditionString);
