UnrealBuildTool: Disable hotreload by default for program targets

#rnx

diff --git a/Engine/Source/Programs/UnrealBuildTool/System/HotReload.cs b/Engine/Source/Programs/UnrealBuildTool/System/HotReload.cs
--- a/Engine/Source/Programs/UnrealBuildTool/System/HotReload.cs
+++ b/Engine/Source/Programs/UnrealBuildTool/System/HotReload.cs
@@ -176,7 +176,11 @@ namespace UnrealBuildTool
 		/// <returns>Default hotreload mode</returns>
 		public static HotReloadMode GetDefaultMode(TargetDescriptor TargetDescriptor, TargetMakefile Makefile, BuildConfiguration BuildConfiguration, ILogger Logger)
 		{
-			if (TargetDescriptor.HotReloadModuleNameToSuffix.Count > 0 && TargetDescriptor.ForeignPlugin == null)
+			if (Makefile.TargetType == TargetType.Program)
+			{
+				return HotReloadMode.Disabled;
+			}
+			else if (TargetDescriptor.HotReloadModuleNameToSuffix.Count > 0 && TargetDescriptor.ForeignPlugin == null)
 			{
 				return HotReloadMode.FromEditor;
 			}
@@ -189,10 +193,7 @@ namespace UnrealBuildTool
 				Logger.LogWarning("Live coding session active. Actions will be limited to compilation of specified files.  Output will be sent to a temporary location.");
 				return HotReloadMode.LiveCodingPassThrough;
 			}
-			else
-			{
-				return HotReloadMode.Disabled;
-			}
+			return HotReloadMode.Disabled;
 		}
 
 		/// <summary>
