Add stub intent game acitivity to quickly check if any app on the device uses zen streaming via "pm query-activities --brief -a com.epicgames.unreal.Zen.StreamingEnabled"

diff --git a/Engine/Source/Programs/UnrealBuildTool/Platform/Android/UEDeployAndroid.cs b/Engine/Source/Programs/UnrealBuildTool/Platform/Android/UEDeployAndroid.cs
--- a/Engine/Source/Programs/UnrealBuildTool/Platform/Android/UEDeployAndroid.cs
+++ b/Engine/Source/Programs/UnrealBuildTool/Platform/Android/UEDeployAndroid.cs
@@ -2033,7 +2033,7 @@ namespace UnrealBuildTool
 		}
 
 		private string GetAllBuildSettings(AndroidToolChain ToolChain, UnrealPluginLanguage UPL, bool bForDistribution, bool bMakeSeparateApks, bool bPackageDataInsideApk,
-			bool bDisableVerifyOBBOnStartUp, bool bUseExternalFilesDir, bool bDontBundleLibrariesInAPK, string TemplatesHashCode)
+			bool bDisableVerifyOBBOnStartUp, bool bUseExternalFilesDir, bool bDontBundleLibrariesInAPK, bool bUsesZenStreaming, string TemplatesHashCode)
 		{
 			// make the settings string - this will be char by char compared against last time
 			StringBuilder CurrentSettings = new StringBuilder();
@@ -2050,6 +2050,7 @@ namespace UnrealBuildTool
 			CurrentSettings.AppendLine(String.Format("UPLHashCode={0}", UPLHashCode));
 			CurrentSettings.AppendLine(String.Format("TemplatesHashCode={0}", TemplatesHashCode));
 			CurrentSettings.AppendLine(String.Format("bDontBundleLibrariesInAPK={0}", bDontBundleLibrariesInAPK));
+			CurrentSettings.AppendLine(String.Format("bUsesZenStreaming={0}", bUsesZenStreaming));
 
 			// all AndroidRuntimeSettings ini settings in here
 			ConfigHierarchy Ini = GetConfigCacheIni(ConfigHierarchyType.Engine);
@@ -2635,7 +2636,7 @@ namespace UnrealBuildTool
 			return EngineMajorVersion + "." + EngineMinorVersion + "." + EnginePatchVersion;
 		}
 
-		private string GenerateManifest(AndroidToolChain ToolChain, string ProjectName, TargetType InTargetType, string EngineDirectory, bool bIsForDistribution, bool bPackageDataInsideApk, string GameBuildFilesPath, bool bHasOBBFiles, bool bDisableVerifyOBBOnStartUp, UnrealArch UnrealArch, string CookFlavor, bool bUseExternalFilesDir, string Configuration, int SDKLevelInt, bool bIsEmbedded, bool bEnableBundle)
+		private string GenerateManifest(AndroidToolChain ToolChain, string ProjectName, TargetType InTargetType, string EngineDirectory, bool bIsForDistribution, bool bPackageDataInsideApk, string GameBuildFilesPath, bool bHasOBBFiles, bool bDisableVerifyOBBOnStartUp, UnrealArch UnrealArch, string CookFlavor, bool bUseExternalFilesDir, string Configuration, int SDKLevelInt, bool bIsEmbedded, bool bEnableBundle, bool bUsesZenStreaming)
 		{
 			// Read the engine version
 			string EngineVersion = ReadEngineVersion();
@@ -2943,11 +2944,19 @@ namespace UnrealBuildTool
 			}
 			Text.AppendLine(String.Format("\t\t          android:debuggable=\"{0}\">", bIsForDistribution ? "false" : "true"));
 			Text.AppendLine("\t\t\t<meta-data android:name=\"android.app.lib_name\" android:value=\"Unreal\"/>");
-			if (!bShowLaunchImage && !bASISEnabled)
+			bool bNeedMainIntent = !bShowLaunchImage && !bASISEnabled;
+			if (bNeedMainIntent || bUsesZenStreaming)
 			{
 				Text.AppendLine("\t\t\t<intent-filter>");
-				Text.AppendLine("\t\t\t\t<action android:name=\"android.intent.action.MAIN\" />");
-				Text.AppendLine(String.Format("\t\t\t\t<category android:name=\"android.intent.category.LAUNCHER\" />"));
+				if (bNeedMainIntent)
+				{
+					Text.AppendLine("\t\t\t\t<action android:name=\"android.intent.action.MAIN\" />");
+					Text.AppendLine(String.Format("\t\t\t\t<category android:name=\"android.intent.category.LAUNCHER\" />"));
+				}
+				if (bUsesZenStreaming)
+				{
+					Text.AppendLine("\t\t\t\t<action android:name=\"com.epicgames.unreal.Zen.StreamingEnabled\" />");
+				}
 				Text.AppendLine("\t\t\t</intent-filter>");
 			}
 			if (!String.IsNullOrEmpty(ExtraActivitySettings))
@@ -4531,9 +4540,12 @@ namespace UnrealBuildTool
 
 			bool bDontBundleLibrariesInAPK = GetDontBundleLibrariesInAPK(Configuration, bIsArchive, bIsFromUAT, bVerbose: !bIsFromUAT);
 
+			string ZenProjectStoreFileName = Path.Combine(Path.GetDirectoryName(ObbFileLocation)!, Path.GetFileNameWithoutExtension(ObbFileLocation), "ue.projectstore");
+			bool bUsesZenStreaming = File.Exists(ZenProjectStoreFileName);
+
 			// check to see if any "meta information" is newer than last time we build
 			string TemplatesHashCode = GenerateTemplatesHashCode(EngineDirectory);
-			string CurrentBuildSettings = GetAllBuildSettings(ToolChain, UPL!, bForDistribution, bMakeSeparateApks, bPackageDataInsideApk, bDisableVerifyOBBOnStartUp, bUseExternalFilesDir, bDontBundleLibrariesInAPK, TemplatesHashCode);
+			string CurrentBuildSettings = GetAllBuildSettings(ToolChain, UPL!, bForDistribution, bMakeSeparateApks, bPackageDataInsideApk, bDisableVerifyOBBOnStartUp, bUseExternalFilesDir, bDontBundleLibrariesInAPK, bUsesZenStreaming, TemplatesHashCode);
 			string BuildSettingsCacheFile = Path.Combine(IntermediateAndroidPath, "UEBuildSettings.txt");
 
 			// do we match previous build settings?
@@ -4646,7 +4658,7 @@ namespace UnrealBuildTool
 				// LINQ query is subsequently used multiple times, so must copy results
 				BuildList = (from Arch in Architectures.Architectures
 							let CompileSDKLevelForArch = NDKArchToCompileSDK[GetNDKArch(Arch)]
-							let manifest = GenerateManifest(ToolChain, ProjectName, InTargetType, EngineDirectory, bForDistribution, bPackageDataInsideApk, GameBuildFilesPath, bRequiresOBB, bDisableVerifyOBBOnStartUp, Arch, CookFlavor, bUseExternalFilesDir, Configuration.ToString(), CompileSDKLevelForArch, bSkipGradleBuild, bEnableBundle)
+							let manifest = GenerateManifest(ToolChain, ProjectName, InTargetType, EngineDirectory, bForDistribution, bPackageDataInsideApk, GameBuildFilesPath, bRequiresOBB, bDisableVerifyOBBOnStartUp, Arch, CookFlavor, bUseExternalFilesDir, Configuration.ToString(), CompileSDKLevelForArch, bSkipGradleBuild, bEnableBundle, bUsesZenStreaming)
 							select Tuple.Create(Arch, manifest)).ToArray();
 			}
 			else
@@ -4655,7 +4667,7 @@ namespace UnrealBuildTool
 				BuildList = (from Arch in Architectures.Architectures
 							let CompileSDKLevelForArch = NDKArchToCompileSDK[GetNDKArch(Arch)]
 							let manifestFile = Path.Combine(IntermediateAndroidPath, Arch + "_AndroidManifest.xml")
-							let manifest = GenerateManifest(ToolChain, ProjectName, InTargetType, EngineDirectory, bForDistribution, bPackageDataInsideApk, GameBuildFilesPath, bRequiresOBB, bDisableVerifyOBBOnStartUp, Arch, CookFlavor, bUseExternalFilesDir, Configuration.ToString(), CompileSDKLevelForArch, bSkipGradleBuild, bEnableBundle)
+							let manifest = GenerateManifest(ToolChain, ProjectName, InTargetType, EngineDirectory, bForDistribution, bPackageDataInsideApk, GameBuildFilesPath, bRequiresOBB, bDisableVerifyOBBOnStartUp, Arch, CookFlavor, bUseExternalFilesDir, Configuration.ToString(), CompileSDKLevelForArch, bSkipGradleBuild, bEnableBundle, bUsesZenStreaming)
 							let OldManifest = File.Exists(manifestFile) ? File.ReadAllText(manifestFile) : ""
 							where manifest != OldManifest
 							select Tuple.Create(Arch, manifest)).ToArray();
