Support package names containing the 'build' substring when building AABs. Having build in the package name would cause Java files not to be deployed to AAB grandle directory resulting in failed builds. Update 'build' file path filtering to ingnore java source root.

diff --git a/Engine/Source/Programs/UnrealBuildTool/Platform/Android/UEDeployAndroid.cs b/Engine/Source/Programs/UnrealBuildTool/Platform/Android/UEDeployAndroid.cs
--- a/Engine/Source/Programs/UnrealBuildTool/Platform/Android/UEDeployAndroid.cs
+++ b/Engine/Source/Programs/UnrealBuildTool/Platform/Android/UEDeployAndroid.cs
@@ -5431,6 +5431,7 @@ popd
 				Directory.CreateDirectory(UnrealGradleDest);
 
 				string ABIFilter = "";
+				string[] JavaSrcRootComponents = { "app", "src", "main" };
 
 				// loop through and merge the different architecture gradle directories
 				foreach (Tuple<UnrealArch, string> build in BuildList)
@@ -5460,12 +5461,34 @@ popd
 
 						// skip the build directories
 						string Workname = Filename.Replace("\\", "/");
-						string DirectoryName = Path.GetDirectoryName(Filename)!.Substring(UnrealBuildGradlePathLength);
-						if (DirectoryName.Contains("build") || Workname.Contains("/."))
+						if (Workname.Contains("/."))
 						{
 							continue;
 						}
 
+						string DirectoryName = Path.GetDirectoryName(Filename)!.Substring(UnrealBuildGradlePathLength);
+						string[] DirectoryNameComponents = DirectoryName.Split(Path.DirectorySeparatorChar, StringSplitOptions.RemoveEmptyEntries);
+						int BuildDirIndex = Array.IndexOf(DirectoryNameComponents, "build");
+						// if a build directory exists
+						if (BuildDirIndex >= 0)
+						{
+							// get the directory component index where app/src/main (JavaSrcRootComponents) begins
+							int JavaSrcRootIndex = Array.IndexOf(DirectoryNameComponents, JavaSrcRootComponents[0]);
+							for (int c = 1; c < JavaSrcRootComponents.Length && JavaSrcRootIndex >= 0; ++c)
+							{
+								int NextComponentIndex = Array.IndexOf(DirectoryNameComponents, JavaSrcRootComponents[c]);
+								if (NextComponentIndex != JavaSrcRootIndex + c)
+								{
+									JavaSrcRootIndex = -1;
+								}
+							}
+							// if the said build directory is not beneath the Java src root
+							if (JavaSrcRootIndex == -1 || BuildDirIndex < JavaSrcRootIndex)
+							{
+								continue;
+							}
+						}
+
 						// if destination doesn't exist, just copy it
 						if (!File.Exists(DestFilename))
 						{
@@ -5686,7 +5709,7 @@ popd
 
 					// update manifest to use versionCode properly
 					string BaseStoreVersion = GetStoreVersion(null).ToString();
-					string ManifestFilename = Path.Combine(UnrealBuildGradlePath, "app", "src", "main", "AndroidManifest.xml");
+					string ManifestFilename = Path.Combine(UnrealBuildGradlePath, Path.Combine(JavaSrcRootComponents), "AndroidManifest.xml");
 					string[] ManifestContents = File.ReadAllLines(ManifestFilename);
 					for (int Index = 0; Index < ManifestContents.Length; Index++)
 					{
