UnrealBuildTool: Special consideration if someone really wants to disable /Zc:preprocessor for Cpp20 msvc compiles

#rnx

diff --git a/Engine/Source/Programs/UnrealBuildTool/Platform/Windows/VCToolChain.cs b/Engine/Source/Programs/UnrealBuildTool/Platform/Windows/VCToolChain.cs
--- a/Engine/Source/Programs/UnrealBuildTool/Platform/Windows/VCToolChain.cs
+++ b/Engine/Source/Programs/UnrealBuildTool/Platform/Windows/VCToolChain.cs
@@ -1285,7 +1285,7 @@ namespace UnrealBuildTool
 
 			if (CompileEnvironment.CppStandard >= CppStandardVersion.Cpp20)
 			{
-				if (Target.WindowsPlatform.Compiler.IsMSVC())
+				if (Target.WindowsPlatform.Compiler.IsMSVC() && !CompileEnvironment.AdditionalArguments.Contains("/Zc:preprocessor-", StringComparison.Ordinal))
 				{
 					Arguments.Add("/Zc:preprocessor");
 				}
