UnrealBuildTool: Invalidate makefile if PVS-Studio Settings.xml is updated

#rnx

diff --git a/Engine/Source/Programs/UnrealBuildTool/Platform/Windows/PVSToolChain.cs b/Engine/Source/Programs/UnrealBuildTool/Platform/Windows/PVSToolChain.cs
--- a/Engine/Source/Programs/UnrealBuildTool/Platform/Windows/PVSToolChain.cs
+++ b/Engine/Source/Programs/UnrealBuildTool/Platform/Windows/PVSToolChain.cs
@@ -610,6 +610,12 @@ namespace UnrealBuildTool
 			Lines.Add(String.Format("Using PVS-Studio {0} at {1} with analysis mode {2} ({3})", AnalyzerVersion, AnalyzerFile, (uint)settings.ModeFlags, settings.ModeFlags.ToString()));
 		}
 
+		public override void GetExternalDependencies(HashSet<FileItem> ExternalDependencies)
+		{
+			InnerToolChain.GetExternalDependencies(ExternalDependencies);
+			ExternalDependencies.Add(FileItem.GetItemByFileReference(FileReference.Combine(new DirectoryReference(Environment.GetFolderPath(Environment.SpecialFolder.ApplicationData)), "PVS-Studio", "Settings.xml")));
+		}
+
 		public override void SetUpGlobalEnvironment(ReadOnlyTargetRules Target, CppCompileEnvironment GlobalCompileEnvironment, LinkEnvironment GlobalLinkEnvironment)
 		{
 			base.SetUpGlobalEnvironment(Target, GlobalCompileEnvironment, GlobalLinkEnvironment);
