UnrealBuildTool: Add compiler paths and Platform_SDK.json paths to external dependencies

#rnx

diff --git a/Engine/Source/Programs/UnrealBuildTool/Platform/Windows/PVSToolChain.cs b/Engine/Source/Programs/UnrealBuildTool/Platform/Windows/PVSToolChain.cs
--- a/Engine/Source/Programs/UnrealBuildTool/Platform/Windows/PVSToolChain.cs
+++ b/Engine/Source/Programs/UnrealBuildTool/Platform/Windows/PVSToolChain.cs
@@ -614,6 +614,7 @@ namespace UnrealBuildTool
 		{
 			InnerToolChain.GetExternalDependencies(ExternalDependencies);
 			ExternalDependencies.Add(FileItem.GetItemByFileReference(FileReference.Combine(new DirectoryReference(Environment.GetFolderPath(Environment.SpecialFolder.ApplicationData)), "PVS-Studio", "Settings.xml")));
+			ExternalDependencies.Add(FileItem.GetItemByFileReference(AnalyzerFile));
 		}
 
 		public override void SetUpGlobalEnvironment(ReadOnlyTargetRules Target, CppCompileEnvironment GlobalCompileEnvironment, LinkEnvironment GlobalLinkEnvironment)
diff --git a/Engine/Source/Programs/UnrealBuildTool/Platform/Windows/VCToolChain.cs b/Engine/Source/Programs/UnrealBuildTool/Platform/Windows/VCToolChain.cs
--- a/Engine/Source/Programs/UnrealBuildTool/Platform/Windows/VCToolChain.cs
+++ b/Engine/Source/Programs/UnrealBuildTool/Platform/Windows/VCToolChain.cs
@@ -174,8 +174,15 @@ namespace UnrealBuildTool
 
 		public override void GetExternalDependencies(HashSet<FileItem> ExternalDependencies)
 		{
+			base.GetExternalDependencies(ExternalDependencies);
+
 			ExternalDependencies.Add(FileItem.GetItemByFileReference(EnvVars.CompilerPath));
 			ExternalDependencies.Add(FileItem.GetItemByFileReference(EnvVars.LinkerPath));
+			ExternalDependencies.Add(FileItem.GetItemByFileReference(EnvVars.ResourceCompilerPath));
+			if (EnvVars.Compiler != EnvVars.ToolChain)
+			{
+				ExternalDependencies.Add(FileItem.GetItemByFileReference(EnvVars.ToolchainCompilerPath));
+			}
 		}
 
 		public static void AddDefinition(List<string> Arguments, string Definition)
diff --git a/Engine/Source/Programs/UnrealBuildTool/ToolChain/ClangToolChain.cs b/Engine/Source/Programs/UnrealBuildTool/ToolChain/ClangToolChain.cs
--- a/Engine/Source/Programs/UnrealBuildTool/ToolChain/ClangToolChain.cs
+++ b/Engine/Source/Programs/UnrealBuildTool/ToolChain/ClangToolChain.cs
@@ -231,6 +231,13 @@ namespace UnrealBuildTool
 		/// </summary>
 		protected abstract ClangToolChainInfo GetToolChainInfo();
 
+		public override void GetExternalDependencies(HashSet<FileItem> ExternalDependencies)
+		{
+			base.GetExternalDependencies(ExternalDependencies);
+			ExternalDependencies.Add(FileItem.GetItemByFileReference(LazyInfo.Value.Clang));
+			ExternalDependencies.Add(FileItem.GetItemByFileReference(LazyInfo.Value.Archiver));
+		}
+
 		public override FileReference? GetCppCompilerPath()
 		{
 			return LazyInfo.Value.Clang;
diff --git a/Engine/Source/Programs/UnrealBuildTool/ToolChain/ISPCToolChain.cs b/Engine/Source/Programs/UnrealBuildTool/ToolChain/ISPCToolChain.cs
--- a/Engine/Source/Programs/UnrealBuildTool/ToolChain/ISPCToolChain.cs
+++ b/Engine/Source/Programs/UnrealBuildTool/ToolChain/ISPCToolChain.cs
@@ -1,13 +1,14 @@
 // Copyright Epic Games, Inc. All Rights Reserved.
 
+using EpicGames.Core;
+using Microsoft.CodeAnalysis;
+using Microsoft.Extensions.Logging;
 using System;
 using System.Collections.Generic;
 using System.IO;
 using System.Linq;
 using System.Runtime.InteropServices;
 using System.Text;
-using EpicGames.Core;
-using Microsoft.Extensions.Logging;
 using UnrealBuildBase;
 
 namespace UnrealBuildTool
@@ -32,6 +33,33 @@ namespace UnrealBuildTool
 			TargetRules = Target;
 		}
 
+		public override void GetExternalDependencies(HashSet<FileItem> ExternalDependencies)
+		{
+			base.GetExternalDependencies(ExternalDependencies);
+
+			ExternalDependencies.Add(FileItem.GetItemByPath(GetISPCHostCompilerPath(BuildHostPlatform.Current.Platform)));
+			string? ISPCHostBytecodeCompilerPath = GetISPCHostBytecodeCompilerPath(BuildHostPlatform.Current.Platform);
+			if (ISPCHostBytecodeCompilerPath != null)
+			{
+				ExternalDependencies.Add(FileItem.GetItemByPath(ISPCHostBytecodeCompilerPath));
+			}
+
+			if (TargetRules != null)
+			{
+				string Platform = TargetRules.Platform == UnrealTargetPlatform.Win64 ? "Windows" : TargetRules.Platform.ToString();
+				FileItem PlatformConfig = FileItem.GetItemByFileReference(FileReference.Combine(Unreal.RootDirectory, "Platforms", Platform, "Config", $"{Platform}_SDK.json"));
+				if (PlatformConfig.Exists)
+				{
+					ExternalDependencies.Add(PlatformConfig);
+				}
+				PlatformConfig = FileItem.GetItemByFileReference(FileReference.Combine(Unreal.RootDirectory, "Config", Platform, $"{Platform}_SDK.json"));
+				if (PlatformConfig.Exists)
+				{
+					ExternalDependencies.Add(PlatformConfig);
+				}
+			}
+		}
+
 		/// <summary>
 		/// Get CPU Instruction set targets for ISPC.
 		/// </summary>
