UnrealBuildTool: Move NormalizeCommandLinePath to UEToolChain

#rnx

diff --git a/Engine/Source/Programs/UnrealBuildTool/Configuration/UEBuildTarget.cs b/Engine/Source/Programs/UnrealBuildTool/Configuration/UEBuildTarget.cs
--- a/Engine/Source/Programs/UnrealBuildTool/Configuration/UEBuildTarget.cs
+++ b/Engine/Source/Programs/UnrealBuildTool/Configuration/UEBuildTarget.cs
@@ -3045,15 +3045,16 @@ namespace UnrealBuildTool
 			{
 				HashSet<FileItem> otherObjects = new(allObjects.Except(toStrip));
 				FileItem stripRsp = FileItem.GetItemByFileReference(FileReference.Combine(IntermediateDirectory, $"{Name}.strip.rsp"));
-				List<string> arguments = new();
-				arguments.AddRange(toStrip.Select(path => $"/S:{path.Location.ChangeExtension(".exi")}").OrderBy(x => x));
-				arguments.AddRange(otherObjects.Select(path => $"/D:{path.Location.ChangeExtension(".exi")}").OrderBy(x => x));
-
 				string Ext = TargetToolChain.GetExtraLinkFileExtension();
 				FileReference extraObj = FileReference.Combine(IntermediateDirectory, $"{Name}.extra.{Ext}");
-				arguments.Add($"/O:{extraObj}");
-				arguments.Add($"/T:{Platform}");
-				arguments.Add($"/M:{Name}");
+				IEnumerable<string> arguments =
+				[
+					.. toStrip.Select(path => $"/S:{path.Location.ChangeExtension(".exi")}").OrderBy(x => x),
+					.. otherObjects.Select(path => $"/D:{path.Location.ChangeExtension(".exi")}").OrderBy(x => x),
+					$"/O:{extraObj}",
+					$"/T:{Platform}",
+					$"/M:{Name}",
+				];
 
 				MakefileBuilder.CreateIntermediateTextFile(stripRsp, arguments);
 
@@ -3063,12 +3064,13 @@ namespace UnrealBuildTool
 				stripAction.PrerequisiteItems.UnionWith(allObjects);
 				stripAction.ProducedItems.Add(FileItem.GetItemByFileReference(extraObj));
 				stripAction.CommandPath = commandPath;
-				stripAction.CommandArguments = "@" + stripRsp.FullName;
+				stripAction.CommandArguments = $"@\"{TargetToolChain.NormalizeCommandLinePath(stripRsp.Location, CompileEnvironment.RootPaths)}\"";
 				stripAction.CommandDescription = "GenerateExports";
 				stripAction.WorkingDirectory = Unreal.EngineSourceDirectory;
 				stripAction.StatusDescription = Name;
 				stripAction.bCanExecuteRemotely = false; // Can run remotely but is so fast so it would be a waste
 				stripAction.ArtifactMode = ArtifactMode.Enabled;
+				stripAction.RootPaths = CompileEnvironment.RootPaths;
 
 				Makefile.OutputItems.AddRange(stripAction.ProducedItems);
 			}
diff --git a/Engine/Source/Programs/UnrealBuildTool/Platform/Windows/VCToolChain.cs b/Engine/Source/Programs/UnrealBuildTool/Platform/Windows/VCToolChain.cs
--- a/Engine/Source/Programs/UnrealBuildTool/Platform/Windows/VCToolChain.cs
+++ b/Engine/Source/Programs/UnrealBuildTool/Platform/Windows/VCToolChain.cs
@@ -217,31 +217,18 @@ namespace UnrealBuildTool
 			}
 		}
 
-		public static new string NormalizeCommandLinePath(FileSystemReference Reference)
-		{
+		public static new string NormalizeCommandLinePath(FileSystemReference Reference) =>
 			// Try to use a relative path to shorten command line length and to enable remote distribution where absolute paths are not desired
-			if (Reference.IsUnderDirectory(Unreal.EngineDirectory))
-			{
-				return Reference.MakeRelativeTo(Unreal.EngineSourceDirectory).Replace('\\', '/');
-			}
+			Reference.IsUnderDirectory(Unreal.EngineDirectory)
+				? Reference.MakeRelativeTo(Unreal.EngineSourceDirectory).Replace('\\', '/')
+				: Reference.FullName.Replace('\\', '/');
 
-			return Reference.FullName.Replace('\\', '/');
-		}
-
-		public static new string NormalizeCommandLinePath(FileItem Item)
-		{
-			return NormalizeCommandLinePath(Item.Location);
-		}
-
-		public static new string NormalizeCommandLinePath(FileSystemReference Reference, CppRootPaths RootPaths)
-		{
-			return RootPaths.GetVfsOverlayPath(Reference, out string? vfsPath) ? vfsPath : NormalizeCommandLinePath(Reference);
-		}
+		public static new string NormalizeCommandLinePath(FileItem Item) => NormalizeCommandLinePath(Item.Location);
+		
+		public static new string NormalizeCommandLinePath(FileSystemReference Reference, CppRootPaths RootPaths) =>
+			RootPaths.GetVfsOverlayPath(Reference, out string? vfsPath) ? vfsPath : NormalizeCommandLinePath(Reference);
 
-		public static new string NormalizeCommandLinePath(FileItem Item, CppRootPaths RootPaths)
-		{
-			return NormalizeCommandLinePath(Item.Location, RootPaths);
-		}
+		public static new string NormalizeCommandLinePath(FileItem Item, CppRootPaths RootPaths) => NormalizeCommandLinePath(Item.Location, RootPaths);
 
 		public static void AddResponseFile(List<string> Arguments, FileItem ResponseFile, CppRootPaths RootPaths)
 		{
diff --git a/Engine/Source/Programs/UnrealBuildTool/ToolChain/ISPCToolChain.cs b/Engine/Source/Programs/UnrealBuildTool/ToolChain/ISPCToolChain.cs
--- a/Engine/Source/Programs/UnrealBuildTool/ToolChain/ISPCToolChain.cs
+++ b/Engine/Source/Programs/UnrealBuildTool/ToolChain/ISPCToolChain.cs
@@ -356,58 +356,6 @@ namespace UnrealBuildTool
 			return Escaped;
 		}
 
-		/// <summary>
-		/// Normalize a path for use in a command line, making it relative to Engine/Source if under the root directory
-		/// </summary>
-		/// <param name="Reference">The FileSystemReference to normalize</param>
-		/// <returns>Normalized path as a string</returns>
-		protected static string NormalizeCommandLinePath(FileSystemReference Reference)
-		{
-			string path = Reference.FullName;
-			// Try to use a relative path to shorten command line length.
-			if (Reference.IsUnderDirectory(Unreal.RootDirectory))
-			{
-				path = Reference.MakeRelativeTo(Unreal.EngineSourceDirectory);
-			}
-			if (Path.DirectorySeparatorChar == '\\')
-			{
-				path = path.Replace("\\", "/");
-			}
-			return path;
-		}
-
-		/// <summary>
-		/// Normalize a path for use in a command line, making it relative if under the Root Directory
-		/// </summary>
-		/// <param name="Item">The FileItem to normalize</param>
-		/// <returns>Normalized path as a string</returns>
-		protected static string NormalizeCommandLinePath(FileItem Item)
-		{
-			return NormalizeCommandLinePath(Item.Location);
-		}
-
-		/// <summary>
-		/// Normalize a path for use in a command line, making it relative to Engine/Source if under the root directory
-		/// </summary>
-		/// <param name="Reference">The FileSystemReference to normalize</param>
-		/// <param name="RootPaths">The root paths to use for normalization</param>
-		/// <returns>Normalized path as a string</returns>
-		protected virtual string NormalizeCommandLinePath(FileSystemReference Reference, CppRootPaths RootPaths)
-		{
-			return RootPaths.GetVfsOverlayPath(Reference, out string? vfsPath) ? vfsPath : NormalizeCommandLinePath(Reference);
-		}
-
-		/// <summary>
-		/// Normalize a path for use in a command line, making it relative if under the Root Directory
-		/// </summary>
-		/// <param name="Item">The FileItem to normalize</param>
-		/// <param name="RootPaths">The root paths to use for normalization</param>
-		/// <returns>Normalized path as a string</returns>
-		protected virtual string NormalizeCommandLinePath(FileItem Item, CppRootPaths RootPaths)
-		{
-			return NormalizeCommandLinePath(Item.Location, RootPaths);
-		}
-
 		/// <summary>
 		/// Simple hash function for strings (can be chained)
 		/// </summary>
diff --git a/Engine/Source/Programs/UnrealBuildTool/ToolChain/IWYUToolChain.cs b/Engine/Source/Programs/UnrealBuildTool/ToolChain/IWYUToolChain.cs
--- a/Engine/Source/Programs/UnrealBuildTool/ToolChain/IWYUToolChain.cs
+++ b/Engine/Source/Programs/UnrealBuildTool/ToolChain/IWYUToolChain.cs
@@ -137,13 +137,13 @@ namespace UnrealBuildTool
         // ".." than there are directories to reduce. This silently fails in odd ways
 
         /// <inheritdoc/>
-        protected override string NormalizeCommandLinePath(FileSystemReference Reference, CppRootPaths RootPath)
+        public override string NormalizeCommandLinePath(FileSystemReference Reference, CppRootPaths RootPath)
 		{
 			return Reference.FullName.Replace('\\', '/');
 		}
 
-        /// <inheritdoc/>
-		protected override string NormalizeCommandLinePath(FileItem Item, CppRootPaths RootPath)
+		/// <inheritdoc/>
+		public override string NormalizeCommandLinePath(FileItem Item, CppRootPaths RootPath)
 		{
 			return NormalizeCommandLinePath(Item.Location, RootPath);
 		}
diff --git a/Engine/Source/Programs/UnrealBuildTool/ToolChain/UEToolChain.cs b/Engine/Source/Programs/UnrealBuildTool/ToolChain/UEToolChain.cs
--- a/Engine/Source/Programs/UnrealBuildTool/ToolChain/UEToolChain.cs
+++ b/Engine/Source/Programs/UnrealBuildTool/ToolChain/UEToolChain.cs
@@ -1,11 +1,12 @@
 // Copyright Epic Games, Inc. All Rights Reserved.
 
+using EpicGames.Core;
+using Microsoft.Extensions.Logging;
 using System;
 using System.Collections.Generic;
+using System.IO;
 using System.Linq;
 using System.Text.RegularExpressions;
-using EpicGames.Core;
-using Microsoft.Extensions.Logging;
 using UnrealBuildBase;
 
 namespace UnrealBuildTool
@@ -34,6 +35,50 @@ namespace UnrealBuildTool
 		{
 		}
 
+		/// <summary>
+		/// Normalize a path for use in a command line, making it relative to Engine/Source if under the root directory
+		/// </summary>
+		/// <param name="Reference">The FileSystemReference to normalize</param>
+		/// <returns>Normalized path as a string</returns>
+		protected static string NormalizeCommandLinePath(FileSystemReference Reference)
+		{
+			string path = Reference.FullName;
+			// Try to use a relative path to shorten command line length.
+			if (Reference.IsUnderDirectory(Unreal.RootDirectory))
+			{
+				path = Reference.MakeRelativeTo(Unreal.EngineSourceDirectory);
+			}
+			if (Path.DirectorySeparatorChar == '\\')
+			{
+				path = path.Replace("\\", "/");
+			}
+			return path;
+		}
+
+		/// <summary>
+		/// Normalize a path for use in a command line, making it relative if under the Root Directory
+		/// </summary>
+		/// <param name="Item">The FileItem to normalize</param>
+		/// <returns>Normalized path as a string</returns>
+		protected static string NormalizeCommandLinePath(FileItem Item) => NormalizeCommandLinePath(Item.Location);
+
+		/// <summary>
+		/// Normalize a path for use in a command line, making it relative to Engine/Source if under the root directory
+		/// </summary>
+		/// <param name="Reference">The FileSystemReference to normalize</param>
+		/// <param name="RootPaths">The root paths to use for normalization</param>
+		/// <returns>Normalized path as a string</returns>
+		public virtual string NormalizeCommandLinePath(FileSystemReference Reference, CppRootPaths RootPaths) => 
+			RootPaths.GetVfsOverlayPath(Reference, out string? vfsPath) ? vfsPath : NormalizeCommandLinePath(Reference);
+
+		/// <summary>
+		/// Normalize a path for use in a command line, making it relative if under the Root Directory
+		/// </summary>
+		/// <param name="Item">The FileItem to normalize</param>
+		/// <param name="RootPaths">The root paths to use for normalization</param>
+		/// <returns>Normalized path as a string</returns>
+		public virtual string NormalizeCommandLinePath(FileItem Item, CppRootPaths RootPaths) => NormalizeCommandLinePath(Item.Location, RootPaths);
+
 		public static DirectoryReference GetModuleInterfaceDir(DirectoryReference OutputDir)
 		{
 			return DirectoryReference.Combine(OutputDir, "Ifc");
