Allow separate signing profile for iOS app extensions. ie: IOSProvisioningProfile_NotifService=(FilePath="<path to .mobileprovision>")
- update the mobile provisioning profiles path based on Xcode version used (Xcode16+ vs <Xcode16) and platform

diff --git a/Engine/Source/Programs/AutomationTool/IOS/IOSPlatform.Automation.cs b/Engine/Source/Programs/AutomationTool/IOS/IOSPlatform.Automation.cs
--- a/Engine/Source/Programs/AutomationTool/IOS/IOSPlatform.Automation.cs
+++ b/Engine/Source/Programs/AutomationTool/IOS/IOSPlatform.Automation.cs
@@ -1129,10 +1129,12 @@ Console.WriteLine("target receipt executable path: {0}", InExecutablePath);
 			if (!string.IsNullOrEmpty(Provision))
 			{
 				// read the provision to get the UUID
-				if (File.Exists(Environment.GetEnvironmentVariable("HOME") + "/Library/MobileDevice/Provisioning Profiles/" + Provision))
+				DirectoryReference MobileProvisionDir = AppleExports.GetProvisionDirectory();
+				FileReference MobileProvisionFile = FileReference.Combine(MobileProvisionDir, Provision);
+				if (File.Exists(MobileProvisionFile.FullName))
 				{
 					string UUID = "";
-					string AllText = File.ReadAllText(Environment.GetEnvironmentVariable("HOME") + "/Library/MobileDevice/Provisioning Profiles/" + Provision);
+					string AllText = File.ReadAllText(MobileProvisionDir.FullName);
 					int idx = AllText.IndexOf("<key>UUID</key>");
 					if (idx > 0)
 					{
@@ -2379,16 +2381,7 @@ Console.WriteLine("target receipt executable path: {0}", InExecutablePath);
 		string AppName = SC.IsCodeBasedProject ? SC.StageExecutables[0].Split("-".ToCharArray())[0] : "UnrealGame";
 
 		// mobile provisioning file
-		DirectoryReference MobileProvisionDir;
-		if (BuildHostPlatform.Current.Platform == UnrealTargetPlatform.Mac)
-		{
-			MobileProvisionDir = DirectoryReference.Combine(new DirectoryReference(Environment.GetEnvironmentVariable("HOME")), "Library", "MobileDevice", "Provisioning Profiles");
-		}
-		else
-		{
-			MobileProvisionDir = DirectoryReference.Combine(DirectoryReference.GetSpecialFolder(Environment.SpecialFolder.LocalApplicationData), "Apple Computer", "MobileDevice", "Provisioning Profiles");
-		}
-
+		DirectoryReference MobileProvisionDir = AppleExports.GetProvisionDirectory();
 		FileReference MobileProvisionFile = null;
 		
 		if(MobileProvisionDir != null && Params.Provision != null)
diff --git a/Engine/Source/Programs/UnrealBuildTool/Platform/Apple/AppleExports.cs b/Engine/Source/Programs/UnrealBuildTool/Platform/Apple/AppleExports.cs
--- a/Engine/Source/Programs/UnrealBuildTool/Platform/Apple/AppleExports.cs
+++ b/Engine/Source/Programs/UnrealBuildTool/Platform/Apple/AppleExports.cs
@@ -19,6 +19,8 @@ namespace UnrealBuildTool
 		private static bool bForceLegacyXcode = Environment.CommandLine.Contains("-legacyxcode", StringComparison.OrdinalIgnoreCase);
 		private static bool bNoEntitlements = Environment.CommandLine.Contains("-noEntitlements", StringComparison.OrdinalIgnoreCase);
 
+		private static DirectoryReference? _MobileProvisionDir;
+
 		/// <summary>
 		///  True if we should force disable entitlements based on "-noEntitlements"
 		/// </summary>
@@ -317,7 +319,7 @@ namespace UnrealBuildTool
 		}
 
 		/// <summary>
-		/// FInds the latest .xcarchive for a given Target name (the .xcarchive will start with this name then have a data appended)
+		/// Finds the latest .xcarchive for a given Target name (the .xcarchive will start with this name then have a data appended)
 		/// </summary>
 		/// <param name="TargetName">Name of the target to look for, this will be the prefix for the .xcarchive to search for</param>
 		/// <param name="SearchPath">If null, look in Xcode Archives library, otherwiwse, look here</param>
@@ -348,5 +350,62 @@ namespace UnrealBuildTool
 
 			return null;
 		}
+		
+		/// <summary>
+		/// Returns the location of the Mobile Provisioning directory. Handles Mac and Windows.
+		/// On mac, this value changed starting with Xcode16
+		/// </summary>
+		/// <returns>The location of the Mobile Provisioning directory</returns>
+		public static DirectoryReference GetProvisionDirectory()
+		{
+			if (!String.IsNullOrEmpty(_MobileProvisionDir?.FullName))
+			{
+				return _MobileProvisionDir!;
+			}
+			
+			if (OperatingSystem.IsMacOS())
+			{
+				// Default to the Xcode15 and less location
+				string MobileProvisionDir = "/Library/MobileDevice/Provisioning Profiles/";
+				
+				// Find out the version number in Xcode.app/Contents/Info.plist
+				int ExitCode;
+				string XcodeVersion = Utils.RunLocalProcessAndReturnStdOut("/bin/sh",
+					$"-c 'plutil -extract CFBundleShortVersionString raw $(xcode-select -p)/../Info.plist'", out ExitCode);
+				if (ExitCode == 0 && XcodeVersion != null)
+				{
+					// parse it into Major/Minor version numbers
+					int Major = 15;
+					int Minor = 0;
+					try
+					{
+						string[] Tokens = XcodeVersion!.Split(".".ToCharArray());
+						if (Tokens.Length >= 2)
+						{
+							Major = Int32.Parse(Tokens[0]);
+							Minor = Int32.Parse(Tokens[1]);
+						}
+					}
+					catch (Exception)
+					{
+					}
+
+					if (Major >= 16)
+					{
+						MobileProvisionDir = "/Library/Developer/Xcode/UserData/Provisioning Profiles/";
+					}
+				}
+
+				_MobileProvisionDir = new DirectoryReference(Environment.GetEnvironmentVariable("HOME") + MobileProvisionDir);
+				Console.WriteLine(" Setting Mobile Provision Profile Dir: {0}", _MobileProvisionDir?.FullName ?? "Not set");
+				return _MobileProvisionDir!;
+			}
+			else
+			{
+				_MobileProvisionDir = new DirectoryReference(Environment.GetFolderPath(Environment.SpecialFolder.LocalApplicationData) + "/Apple Computer/MobileDevice/Provisioning Profiles/");
+				Console.WriteLine(" Setting Mobile Provision Profile Dir: {0}", _MobileProvisionDir?.FullName ?? "Not set");
+				return _MobileProvisionDir!;
+			}
+		}
 	}
 }
diff --git a/Engine/Source/Programs/UnrealBuildTool/Platform/IOS/IOSCodeSigning.cs b/Engine/Source/Programs/UnrealBuildTool/Platform/IOS/IOSCodeSigning.cs
--- a/Engine/Source/Programs/UnrealBuildTool/Platform/IOS/IOSCodeSigning.cs
+++ b/Engine/Source/Programs/UnrealBuildTool/Platform/IOS/IOSCodeSigning.cs
@@ -42,20 +42,7 @@ namespace UnrealBuildTool
 		/// <summary>
 		/// The shared provision library directory (on PC)
 		/// </summary>
-		public static string ProvisionDirectory
-		{
-			get
-			{
-				if (OperatingSystem.IsMacOS())
-				{
-					return Environment.GetEnvironmentVariable("HOME") + "/Library/MobileDevice/Provisioning Profiles/";
-				}
-				else
-				{
-					return Environment.GetFolderPath(Environment.SpecialFolder.LocalApplicationData) + "/Apple Computer/MobileDevice/Provisioning Profiles/";
-				}
-			}
-		}
+		public static readonly string ProvisionDirectory = AppleExports.GetProvisionDirectory().FullName;
 
 		/// <summary>
 		/// Get the name from the Certificate
@@ -197,9 +184,9 @@ namespace UnrealBuildTool
 
 			List<FileReference> Provisions = new();
 
-			if (!Directory.Exists(CodeSigningConfig.ProvisionDirectory))
+			if (!Directory.Exists(ProvisionDirectory))
 			{
-				_logger?.LogError("Could not find provision directory '{0}'.", CodeSigningConfig.ProvisionDirectory);
+				_logger?.LogError("Could not find provision directory '{0}'.", ProvisionDirectory);
 				MatchedProvision = null;
 				return Provisions;
 			}
@@ -212,7 +199,7 @@ namespace UnrealBuildTool
 
 			Dictionary<string, MobileProvision> ProvisionLibrary = new Dictionary<string, MobileProvision>();
 
-			foreach (string ProvisionFile in Directory.EnumerateFiles(CodeSigningConfig.ProvisionDirectory, "*.mobileprovision"))
+			foreach (string ProvisionFile in Directory.EnumerateFiles(ProvisionDirectory, "*.mobileprovision"))
 			{
 				MobileProvision Provision = MobileProvisionParser.ParseFile(ProvisionFile);
 				ProvisionLibrary.Add(ProvisionFile, Provision);
@@ -805,10 +792,11 @@ namespace UnrealBuildTool
 		public static void CacheMobileProvisions()
 		{
 			_logger?.LogInformation("Caching provisions");
-			if (!Directory.Exists(CodeSigningConfig.ProvisionDirectory))
+			string ProvisionDirectory = AppleCodeSign.ProvisionDirectory;
+			if (!Directory.Exists(ProvisionDirectory))
 			{
-				_logger?.LogInformation("Provision Folder {0} doesn't exist, creating..", CodeSigningConfig.ProvisionDirectory);
-				Directory.CreateDirectory(CodeSigningConfig.ProvisionDirectory);
+				_logger?.LogInformation("Provision Folder {0} doesn't exist, creating..", ProvisionDirectory);
+				Directory.CreateDirectory(ProvisionDirectory);
 			}
 
 			List<string> ProvisionCopySrcDirectories = new List<string>();
@@ -842,7 +830,7 @@ namespace UnrealBuildTool
 					_logger?.LogInformation("Finding provisions in {0}", ProvisionCopySrcDirectory);
 					foreach (string Provision in Directory.EnumerateFiles(ProvisionCopySrcDirectory, "*.mobileprovision", SearchOption.AllDirectories))
 					{
-						string TargetFile = CodeSigningConfig.ProvisionDirectory + Path.GetFileName(Provision);
+						string TargetFile = ProvisionDirectory + Path.GetFileName(Provision);
 						if (!File.Exists(TargetFile) || File.GetLastWriteTime(TargetFile) < File.GetLastWriteTime(Provision))
 						{
 							FileInfo DestFileInfo;
@@ -874,14 +862,14 @@ namespace UnrealBuildTool
 		/// </summary>
 		public static void CleanMobileProvisions()
 		{
-			if (!Directory.Exists(CodeSigningConfig.ProvisionDirectory))
+			if (!Directory.Exists(AppleCodeSign.ProvisionDirectory))
 			{
-				_logger?.LogInformation("Provision Folder {0} doesn't exist, nothing to do.", CodeSigningConfig.ProvisionDirectory);
+				_logger?.LogInformation("Provision Folder {0} doesn't exist, nothing to do.", AppleCodeSign.ProvisionDirectory);
 			}
 			else
 			{
-				_logger?.LogInformation("Cleaning out contents of  Provision Folder {0}", CodeSigningConfig.ProvisionDirectory);
-				foreach (string Provision in Directory.GetFiles(CodeSigningConfig.ProvisionDirectory))
+				_logger?.LogInformation("Cleaning out contents of  Provision Folder {0}", AppleCodeSign.ProvisionDirectory);
+				foreach (string Provision in Directory.GetFiles(AppleCodeSign.ProvisionDirectory))
 				{
 					File.Delete(Provision);
 				}
@@ -906,16 +894,17 @@ namespace UnrealBuildTool
 			}
 
 			// ensure the provision directory exists
-			if (!Directory.Exists(CodeSigningConfig.ProvisionDirectory))
+			string ProvisionDirectory = AppleCodeSign.ProvisionDirectory;
+			if (!Directory.Exists(ProvisionDirectory))
 			{
-				Directory.CreateDirectory(CodeSigningConfig.ProvisionDirectory);
+				Directory.CreateDirectory(ProvisionDirectory);
 			}
 
 			if (CodeSigningConfig.bProvision)
 			{
-				if (File.Exists(CodeSigningConfig.ProvisionDirectory + "/" + CodeSigningConfig.Provision))
+				if (File.Exists(ProvisionDirectory + "/" + CodeSigningConfig.Provision))
 				{
-					return CodeSigningConfig.ProvisionDirectory + "/" + CodeSigningConfig.Provision;
+					return ProvisionDirectory + "/" + CodeSigningConfig.Provision;
 				}
 			}
 
@@ -925,25 +914,25 @@ namespace UnrealBuildTool
 
 			// cache the provision library
 			Dictionary<string, MobileProvision> ProvisionLibrary = new Dictionary<string, MobileProvision>();
-			foreach (string Provision in Directory.EnumerateFiles(CodeSigningConfig.ProvisionDirectory, "*.mobileprovision"))
+			foreach (string Provision in Directory.EnumerateFiles(ProvisionDirectory, "*.mobileprovision"))
 			{
 				MobileProvision? p = MobileProvisionParser.ParseFile(Provision);
 				if (p != null)
 				{
 					ProvisionLibrary.Add(Provision, p);
-					if (p.FileName != null && p.UUID != null && p.FileName.Contains(p.UUID) && !File.Exists(Path.Combine(CodeSigningConfig.ProvisionDirectory, "UE4_" + p.UUID + ".mobileprovision")))
+					if (p.FileName != null && p.UUID != null && p.FileName.Contains(p.UUID) && !File.Exists(Path.Combine(ProvisionDirectory, "UE4_" + p.UUID + ".mobileprovision")))
 					{
-						File.Copy(Provision, Path.Combine(CodeSigningConfig.ProvisionDirectory, "UE4_" + p.UUID + ".mobileprovision"));
-						p = MobileProvisionParser.ParseFile(Path.Combine(CodeSigningConfig.ProvisionDirectory, "UE4_" + p.UUID + ".mobileprovision"));
+						File.Copy(Provision, Path.Combine(ProvisionDirectory, "UE4_" + p.UUID + ".mobileprovision"));
+						p = MobileProvisionParser.ParseFile(Path.Combine(ProvisionDirectory, "UE4_" + p.UUID + ".mobileprovision"));
 						if (p != null)
 						{
-							ProvisionLibrary.Add(Path.Combine(CodeSigningConfig.ProvisionDirectory, "UE4_" + p.UUID + ".mobileprovision"), p);
+							ProvisionLibrary.Add(Path.Combine(ProvisionDirectory, "UE4_" + p.UUID + ".mobileprovision"), p);
 						}
 					}
 				}
 			}
 
-			_logger?.LogInformation("Searching for mobile provisions that match the game '{0}' (distribution: {3}) with CFBundleIdentifier='{1}' in '{2}'", GameName, CFBundleIdentifier, CodeSigningConfig.ProvisionDirectory, CodeSigningConfig.bForDistribution);
+			_logger?.LogInformation("Searching for mobile provisions that match the game '{0}' (distribution: {3}) with CFBundleIdentifier='{1}' in '{2}'", GameName, CFBundleIdentifier, ProvisionDirectory, CodeSigningConfig.bForDistribution);
 
 			// first sort all profiles so we look at newer ones first.
 			IEnumerable<string> ProfileKeys = ProvisionLibrary.Select(KV => KV.Key)
@@ -1578,24 +1567,6 @@ namespace UnrealBuildTool
 			}
 		}
 
-		/// <summary>
-		/// The shared provision library directory (on PC)
-		/// </summary>
-		public static string ProvisionDirectory
-		{
-			get
-			{
-				if (Environment.OSVersion.Platform == PlatformID.MacOSX || Environment.OSVersion.Platform == PlatformID.Unix)
-				{
-					return Environment.GetEnvironmentVariable("HOME") + "/Library/MobileDevice/Provisioning Profiles/";
-				}
-				else
-				{
-					return Environment.GetFolderPath(Environment.SpecialFolder.LocalApplicationData) + "/Apple Computer/MobileDevice/Provisioning Profiles/";
-				}
-			}
-		}
-
 		/// <summary>
 		/// The shared (Engine) build directory (on PC)
 		/// </summary>
diff --git a/Engine/Source/Programs/UnrealBuildTool/Platform/IOS/UEDeployIOS.cs b/Engine/Source/Programs/UnrealBuildTool/Platform/IOS/UEDeployIOS.cs
--- a/Engine/Source/Programs/UnrealBuildTool/Platform/IOS/UEDeployIOS.cs
+++ b/Engine/Source/Programs/UnrealBuildTool/Platform/IOS/UEDeployIOS.cs
@@ -29,6 +29,8 @@ namespace UnrealBuildTool
 		}
 		bool bForDistribution = false;
 
+		private static readonly DirectoryReference MobileProvisionDirRef = AppleExports.GetProvisionDirectory();
+
 		protected class VersionUtilities
 		{
 			public static string? BuildDirectory
@@ -1082,7 +1084,7 @@ namespace UnrealBuildTool
 					throw new DirectoryNotFoundException(String.Format("Provision Directory {0} not found.", ProvisionDir), null);
 				}
 
-				string LocalProvisionFolder = Path.Combine(Environment.GetFolderPath(Environment.SpecialFolder.Personal), "Library/MobileDevice/Provisioning Profiles");
+				string LocalProvisionFolder = MobileProvisionDirRef.FullName;
 				if (!Directory.Exists(LocalProvisionFolder))
 				{
 					Logger.LogDebug("Local Provision Folder {LocalProvisionFolder} not found, attempting to create...", LocalProvisionFolder);
@@ -1171,7 +1173,7 @@ namespace UnrealBuildTool
 			Directory.CreateDirectory(PayloadDirectory);
 			Directory.CreateDirectory(AppDirectory);
 			Directory.CreateDirectory(BuildDirectory);
-			Directory.CreateDirectory(Environment.GetEnvironmentVariable("HOME") + "/Library/MobileDevice/Provisioning Profiles");
+			Directory.CreateDirectory(MobileProvisionDirRef.FullName);
 
 			// create the entitlements file
 
@@ -1216,14 +1218,16 @@ namespace UnrealBuildTool
 			}
 			if (File.Exists(ProvisionWithPrefix))
 			{
-				Directory.CreateDirectory(Environment.GetEnvironmentVariable("HOME") + "/Library/MobileDevice/Provisioning Profiles/");
-				if (File.Exists(Environment.GetEnvironmentVariable("HOME") + "/Library/MobileDevice/Provisioning Profiles/" + ProjectProvision))
+				Directory.CreateDirectory(MobileProvisionDirRef.FullName);
+
+				string ProjectProvisionPath = DirectoryReference.Combine(MobileProvisionDirRef, ProjectProvision).FullName;
+				if (File.Exists(ProjectProvisionPath))
 				{
-					DestFileInfo = new FileInfo(Environment.GetEnvironmentVariable("HOME") + "/Library/MobileDevice/Provisioning Profiles/" + ProjectProvision);
+					DestFileInfo = new FileInfo(ProjectProvisionPath);
 					DestFileInfo.Attributes &= ~FileAttributes.ReadOnly;
 				}
-				File.Copy(ProvisionWithPrefix, Environment.GetEnvironmentVariable("HOME") + "/Library/MobileDevice/Provisioning Profiles/" + ProjectProvision, true);
-				DestFileInfo = new FileInfo(Environment.GetEnvironmentVariable("HOME") + "/Library/MobileDevice/Provisioning Profiles/" + ProjectProvision);
+				File.Copy(ProvisionWithPrefix, ProjectProvisionPath, true);
+				DestFileInfo = new FileInfo(ProjectProvisionPath);
 				DestFileInfo.Attributes &= ~FileAttributes.ReadOnly;
 			}
 			if (!File.Exists(ProvisionWithPrefix) || Unreal.IsBuildMachine())
@@ -1258,13 +1262,16 @@ namespace UnrealBuildTool
 			}
 			if (File.Exists(ProvisionWithPrefix))
 			{
-				if (File.Exists(Environment.GetEnvironmentVariable("HOME") + "/Library/MobileDevice/Provisioning Profiles/" + InProjectName + "_Distro.mobileprovision"))
+				Directory.CreateDirectory(MobileProvisionDirRef.FullName);
+				
+				string InProjectProvisionPath = DirectoryReference.Combine(MobileProvisionDirRef, InProjectName, "_Distro.mobileprovision").FullName;
+				if (File.Exists(InProjectProvisionPath))
 				{
-					DestFileInfo = new FileInfo(Environment.GetEnvironmentVariable("HOME") + "/Library/MobileDevice/Provisioning Profiles/" + InProjectName + "_Distro.mobileprovision");
+					DestFileInfo = new FileInfo(InProjectProvisionPath);
 					DestFileInfo.Attributes &= ~FileAttributes.ReadOnly;
 				}
-				File.Copy(ProvisionWithPrefix, Environment.GetEnvironmentVariable("HOME") + "/Library/MobileDevice/Provisioning Profiles/" + InProjectName + "_Distro.mobileprovision", true);
-				DestFileInfo = new FileInfo(Environment.GetEnvironmentVariable("HOME") + "/Library/MobileDevice/Provisioning Profiles/" + InProjectName + "_Distro.mobileprovision");
+				File.Copy(ProvisionWithPrefix, InProjectProvisionPath, true);
+				DestFileInfo = new FileInfo(InProjectProvisionPath);
 				DestFileInfo.Attributes &= ~FileAttributes.ReadOnly;
 			}
 
diff --git a/Engine/Source/Programs/UnrealBuildTool/Platform/Mac/ApplePlatformSDK.cs b/Engine/Source/Programs/UnrealBuildTool/Platform/Mac/ApplePlatformSDK.cs
--- a/Engine/Source/Programs/UnrealBuildTool/Platform/Mac/ApplePlatformSDK.cs
+++ b/Engine/Source/Programs/UnrealBuildTool/Platform/Mac/ApplePlatformSDK.cs
@@ -31,7 +31,7 @@ namespace UnrealBuildTool
 		// this shouldn't get called if there's no SDK installed
 		public static DirectoryReference DeveloperDir => _DeveloperDir.Value!;
 
-		// the vrsion of Xcode that is installed, not any of the SDKs inside
+		// the version of Xcode that is installed, not any of the SDKs inside
 		public static Lazy<string?> InstalledXcodeVersion => new(() => _GetInstalledSDKVersion());
 
 		public static DirectoryReference GetToolchainDirectory()
diff --git a/Engine/Source/Programs/UnrealBuildTool/ProjectFiles/Xcode/XcodeProject.cs b/Engine/Source/Programs/UnrealBuildTool/ProjectFiles/Xcode/XcodeProject.cs
--- a/Engine/Source/Programs/UnrealBuildTool/ProjectFiles/Xcode/XcodeProject.cs
+++ b/Engine/Source/Programs/UnrealBuildTool/ProjectFiles/Xcode/XcodeProject.cs
@@ -1307,7 +1307,7 @@ namespace UnrealBuildTool.XcodeProjectXcconfig
 				{
 					FileReference ProfileFile = AppleExports.ConvertFilePath(UnrealData.UProjectFileLocation?.Directory, ProvisionSetting);
 
-					// get mobile provision UUID, either directly or from the probision
+					// get mobile provision UUID, either directly or from the provision
 					if (ProfileFile.GetExtension() == ".mobileprovision")
 					{
 						// security will read the provision and dump out the text (plist) bits of the profile, and plutil will extract the UUID to stdout
@@ -1315,15 +1315,14 @@ namespace UnrealBuildTool.XcodeProjectXcconfig
 
 						// make sure it's living in user's library
 						// (note: i couldn't find a way to install it without opening Xcode, so copying it works well enough)
-						DirectoryReference UserDir = new DirectoryReference(Environment.GetFolderPath(Environment.SpecialFolder.UserProfile));
-						FileReference InstalledProfileFile = FileReference.Combine(UserDir, "Library", "MobileDevice", "Provisioning Profiles", $"{UUID}.mobileprovision");
+						FileReference InstalledProfileFile = FileReference.Combine(AppleExports.GetProvisionDirectory(), $"{UUID}.mobileprovision");
 
 						// copy if needed
 						if (!FileReference.Exists(InstalledProfileFile))
 						{
 							DirectoryReference.CreateDirectory(InstalledProfileFile.Directory);
 							FileReference.Copy(ProfileFile, InstalledProfileFile);
-							Logger.LogInformation("Copying project's provision '{SourceProvision}' to your libary: '{TargetProvision}'", ProfileFile, InstalledProfileFile);
+							Logger.LogInformation("Copying project's provision '{SourceProvision}' to your library: '{TargetProvision}'", ProfileFile, InstalledProfileFile);
 						}
 
 						SigningName = UUID;
@@ -1403,7 +1402,6 @@ namespace UnrealBuildTool.XcodeProjectXcconfig
 							XcodeCopyAppExtFilesBuildPhase EmbedAppExt = new XcodeCopyAppExtFilesBuildPhase(Project.FileCollection);
 							
 							// Get the source files
-							// XcodeFileCollection AppExtFiles = new XcodeFileCollection();
 							foreach (FileSystemInfo FSI in DI.EnumerateFileSystemInfos())
 							{
 								if (FSI.Name.StartsWith(".") || FSI.Name.EndsWith(".plist"))
@@ -1415,7 +1413,6 @@ namespace UnrealBuildTool.XcodeProjectXcconfig
 								FileReference FileRef = new FileReference(FSI.FullName);
 								XcodeSourceFile SourceFile = new XcodeSourceFile(FileRef, FileRef.Directory, XcodeProjectFileGenerator.MakeXcodeGuid());
 								Project.FileCollection.ProcessFile(SourceFile, true, false, "Source");
-
 							}
 							
 							Project.FileCollection.AddFileReference(ExtensionGuid, Extension, "explicitFileType", "wrapper.app-extension", "BUILT_PRODUCTS_DIR", "Products");
@@ -2460,11 +2457,17 @@ namespace UnrealBuildTool.XcodeProjectXcconfig
 			PlatformIni.TryGetValue("/Script/MacTargetPlatform.XcodeProjectSettings", "bMacSignToRunLocally", out bMacSignToRunLocally);
 			PlatformIni.TryGetValue("/Script/MacTargetPlatform.XcodeProjectSettings", "CodeSigningTeam", out SigningTeam);
 			PlatformIni.TryGetValue("/Script/MacTargetPlatform.XcodeProjectSettings", "CodeSigningPrefix", out SigningPrefix);
-			PlatformIni.TryGetValue("/Script/MacTargetPlatform.XcodeProjectSettings", $"{Platform}ProvisioningProfile", out ProvisioningProfile);
+			PlatformIni.TryGetValue("/Script/MacTargetPlatform.XcodeProjectSettings", $"{Platform}ProvisioningProfile_{TargetName}", out ProvisioningProfile);
 			PlatformIni.TryGetValue("/Script/MacTargetPlatform.XcodeProjectSettings", $"{Platform}SigningIdentity", out SigningIdentity);
 			PlatformIni.TryGetValue("/Script/MacTargetPlatform.XcodeProjectSettings", "AppCategory", out AppCategory);
 			PlatformIni.TryGetValue("/Script/MacTargetPlatform.XcodeProjectSettings", "ApplicationDisplayName", out ApplicationDisplayName);
 
+			// If there's no specific App Extension signing profile, default to the app's profile
+			if (String.IsNullOrEmpty(ProvisioningProfile))
+			{
+				PlatformIni.TryGetValue("/Script/MacTargetPlatform.XcodeProjectSettings", $"{Platform}ProvisioningProfile", out ProvisioningProfile);
+			}
+			
 			if (Platform == UnrealTargetPlatform.Mac)
 			{
 				// editor vs game metadata
