[UBA]
* Improved cancel logic when pressing ctrl-c

diff --git a/Engine/Source/Programs/UnrealBuildTool/Executors/UnrealBuildAccelerator/UBAExecutor.cs b/Engine/Source/Programs/UnrealBuildTool/Executors/UnrealBuildAccelerator/UBAExecutor.cs
--- a/Engine/Source/Programs/UnrealBuildTool/Executors/UnrealBuildAccelerator/UBAExecutor.cs
+++ b/Engine/Source/Programs/UnrealBuildTool/Executors/UnrealBuildAccelerator/UBAExecutor.cs
@@ -134,6 +134,11 @@ namespace UnrealBuildTool
 				targetDescriptor.AdditionalArguments.ApplyTo(UBAConfig);
 			}
 
+			if (UBAConfig.bUseCrypto)
+			{
+				Crypto = CreateCrypto();
+			}
+
 			_agentCoordinators.AddRange(UBAAgentCoordinatorHorde.Init(logger, this, targetDescriptors, UBAConfig, ref _remoteConnectionMode));
 			_threadedLogger = new ThreadedLogger(logger);
 		}
@@ -287,8 +292,8 @@ namespace UnrealBuildTool
 
 		private void ActionQueueCanceled(IStorageServer? ubaStorage)
 		{
-			Server?.StopServer(); // Make sure all remove processes are returned. We can't have any callbacks after this
 			_bIsCancelled = true;
+			Server?.StopServer(); // Make sure all remove processes are returned. We can't have any callbacks after this
 			_session?.CancelAll(); // Cancel all processes native side
 			ubaStorage?.SaveCasTable();
 
@@ -336,11 +341,6 @@ namespace UnrealBuildTool
 
 			logger = _threadedLogger;
 
-			if (UBAConfig.bUseCrypto)
-			{
-				Crypto = CreateCrypto();
-			}
-
 			if (!String.IsNullOrEmpty(UBAConfig.RootDir))
 			{
 				_rootDirRef = DirectoryReference.FromString(UBAConfig.RootDir);
@@ -912,7 +912,7 @@ namespace UnrealBuildTool
 				using (IProcess process = _session!.RunProcess(startInfo, false, null, enableDetour))
 				{
 					Interlocked.Add(ref _localProcessedActions, 1);
-					if (UBAConfig.bAllowRetry && (process.ExitCode != 0 && UBAConfig.bForcedRetry || (process.ExitCode >= 9000 && process.ExitCode < 10000)))
+					if (!_bIsCancelled && (UBAConfig.bAllowRetry && (process.ExitCode != 0 && UBAConfig.bForcedRetry || (process.ExitCode >= 9000 && process.ExitCode < 10000))))
 					{
 						if (UBAConfig.bStoreObjFilesCompressed)
 						{
