[UBT]
* Disabled some warnings when building with header units. Will re-enable them once we have things working

diff --git a/Engine/Source/Programs/UnrealBuildTool/Platform/Windows/VCToolChain.cs b/Engine/Source/Programs/UnrealBuildTool/Platform/Windows/VCToolChain.cs
--- a/Engine/Source/Programs/UnrealBuildTool/Platform/Windows/VCToolChain.cs
+++ b/Engine/Source/Programs/UnrealBuildTool/Platform/Windows/VCToolChain.cs
@@ -1186,6 +1186,10 @@ namespace UnrealBuildTool
 				Arguments.Add("/wd4251"); // 'type' : class 'type1' needs to have dll-interface to be used by clients of class 'type2'
 				Arguments.Add("/wd4702"); // unreachable code
 				Arguments.Add("/wd4180"); // qualifier applied to function type has no meaning; ignored
+				Arguments.Add("/wd4996"); // deprecation
+				Arguments.Add("/wd4200"); // nonstandard extension used: zero - sized array in struct/union
+				Arguments.Add("/wd4714"); // marked as __forceinline not inlined
+				Arguments.Add("/wd5321"); // nonstandard extension used: encoding '\xE5' as a multi-byte utf-8 character
 			}
 
 
