UnrealBuildTool: Fully obsolete C++14 and prevent building a module if set. C++17 is still allowed but unsupported

#rnx

diff --git a/Engine/Source/Programs/UnrealBuildTool/Configuration/CompileWarnings/ApplyWarningsFilterAttribute.cs b/Engine/Source/Programs/UnrealBuildTool/Configuration/CompileWarnings/ApplyWarningsFilterAttribute.cs
--- a/Engine/Source/Programs/UnrealBuildTool/Configuration/CompileWarnings/ApplyWarningsFilterAttribute.cs
+++ b/Engine/Source/Programs/UnrealBuildTool/Configuration/CompileWarnings/ApplyWarningsFilterAttribute.cs
@@ -201,7 +201,7 @@ namespace UnrealBuildTool
 		/// <param name="lowerBound">The lower bound of the <see cref="CppStandardVersion"/> to constrain this warning attribute to.</param>
 		/// <param name="upperBound">The upper bound of the <see cref="CppStandardVersion"/> to constrain this warning attribute to.</param>
 #pragma warning disable CA1019 // Define accessors for attribute arguments
-		public CppStandardConstrainedFilterAttribute(FilterID filterID, CppStandardVersion lowerBound = CppStandardVersion.Cpp14, CppStandardVersion upperBound = CppStandardVersion.Latest) : base(filterID)
+		public CppStandardConstrainedFilterAttribute(FilterID filterID, CppStandardVersion lowerBound = CppStandardVersion.Minimum, CppStandardVersion upperBound = CppStandardVersion.Latest) : base(filterID)
 #pragma warning restore CA1019 // Define accessors for attribute arguments
 		{
 			_lowerBound = lowerBound;
@@ -214,7 +214,7 @@ namespace UnrealBuildTool
 		/// <param name="filterName">The filter name that describes the instance of the ApplyWarningsFilter attribute.</param>
 		/// <param name="lowerBound">The lower bound of the <see cref="CppStandardVersion"/> to constrain this warning attribute to.</param>
 		/// <param name="upperBound">The upper bound of the <see cref="CppStandardVersion"/> to constrain this warning attribute to.</param>
-		public CppStandardConstrainedFilterAttribute(string filterName, CppStandardVersion lowerBound = CppStandardVersion.Cpp14, CppStandardVersion upperBound = CppStandardVersion.Latest) : base(filterName)
+		public CppStandardConstrainedFilterAttribute(string filterName, CppStandardVersion lowerBound = CppStandardVersion.Minimum, CppStandardVersion upperBound = CppStandardVersion.Latest) : base(filterName)
 		{
 			_lowerBound = lowerBound;
 			_upperBound = upperBound;
diff --git a/Engine/Source/Programs/UnrealBuildTool/Configuration/UEBuildModuleCPP.cs b/Engine/Source/Programs/UnrealBuildTool/Configuration/UEBuildModuleCPP.cs
--- a/Engine/Source/Programs/UnrealBuildTool/Configuration/UEBuildModuleCPP.cs
+++ b/Engine/Source/Programs/UnrealBuildTool/Configuration/UEBuildModuleCPP.cs
@@ -1709,7 +1709,6 @@ namespace UnrealBuildTool
 			// Always add the cpp standard into the filename to prevent variants when the engine standard is mismatched with a project
 			switch (CompileEnvironment.CppStandard)
 			{
-				case CppStandardVersion.Cpp14: Variant += ".Cpp14"; break;
 				case CppStandardVersion.Cpp17: Variant += ".Cpp17"; break;
 				case CppStandardVersion.Cpp20: Variant += ".Cpp20"; break;
 				case CppStandardVersion.Latest: Variant += ".CppLatest"; break;
diff --git a/Engine/Source/Programs/UnrealBuildTool/Platform/Windows/VCToolChain.cs b/Engine/Source/Programs/UnrealBuildTool/Platform/Windows/VCToolChain.cs
--- a/Engine/Source/Programs/UnrealBuildTool/Platform/Windows/VCToolChain.cs
+++ b/Engine/Source/Programs/UnrealBuildTool/Platform/Windows/VCToolChain.cs
@@ -1262,12 +1262,8 @@ namespace UnrealBuildTool
 				Arguments.Add("/GR-");
 			}
 
-		
 			switch (CompileEnvironment.CppStandard)
 			{
-				case CppStandardVersion.Cpp14:
-					Arguments.Add("/std:c++14");
-					break;
 				case CppStandardVersion.Cpp17:
 					Arguments.Add("/std:c++17");
 					break;
diff --git a/Engine/Source/Programs/UnrealBuildTool/ProjectFiles/Rider/RiderProjectFile.cs b/Engine/Source/Programs/UnrealBuildTool/ProjectFiles/Rider/RiderProjectFile.cs
--- a/Engine/Source/Programs/UnrealBuildTool/ProjectFiles/Rider/RiderProjectFile.cs
+++ b/Engine/Source/Programs/UnrealBuildTool/ProjectFiles/Rider/RiderProjectFile.cs
@@ -691,9 +691,6 @@ namespace UnrealBuildTool
 				// Do not use version.ToString(). See: https://youtrack.jetbrains.com/issue/RIDER-68030
 				switch (version)
 				{
-					case CppStandardVersion.Cpp14:
-						Writer.WriteValue(Name, "Cpp14");
-						break;
 					case CppStandardVersion.Cpp17:
 						Writer.WriteValue(Name, "Cpp17");
 						break;
diff --git a/Engine/Source/Programs/UnrealBuildTool/ProjectFiles/VisualStudio/VCProject.cs b/Engine/Source/Programs/UnrealBuildTool/ProjectFiles/VisualStudio/VCProject.cs
--- a/Engine/Source/Programs/UnrealBuildTool/ProjectFiles/VisualStudio/VCProject.cs
+++ b/Engine/Source/Programs/UnrealBuildTool/ProjectFiles/VisualStudio/VCProject.cs
@@ -557,8 +557,6 @@ namespace UnrealBuildTool
 		{
 			switch (Version)
 			{
-				case CppStandardVersion.Cpp14:
-					return "/std:c++14";
 				case CppStandardVersion.Cpp17:
 					return "/std:c++17";
 				case CppStandardVersion.Cpp20:
diff --git a/Engine/Source/Programs/UnrealBuildTool/ProjectFiles/VisualStudioCode/VSCodeProjectFileGenerator.cs b/Engine/Source/Programs/UnrealBuildTool/ProjectFiles/VisualStudioCode/VSCodeProjectFileGenerator.cs
--- a/Engine/Source/Programs/UnrealBuildTool/ProjectFiles/VisualStudioCode/VSCodeProjectFileGenerator.cs
+++ b/Engine/Source/Programs/UnrealBuildTool/ProjectFiles/VisualStudioCode/VSCodeProjectFileGenerator.cs
@@ -708,10 +708,6 @@ namespace UnrealBuildTool
 
 			switch (CppStandard)
 			{
-				case CppStandardVersion.Cpp14:
-					OutFile.AddField("cStandard", "c11");
-					OutFile.AddField("cppStandard", "c++14");
-					break;
 				case CppStandardVersion.Cpp17:
 					OutFile.AddField("cStandard", "c17");
 					OutFile.AddField("cppStandard", "c++17");
diff --git a/Engine/Source/Programs/UnrealBuildTool/System/CppCompileEnvironment.cs b/Engine/Source/Programs/UnrealBuildTool/System/CppCompileEnvironment.cs
--- a/Engine/Source/Programs/UnrealBuildTool/System/CppCompileEnvironment.cs
+++ b/Engine/Source/Programs/UnrealBuildTool/System/CppCompileEnvironment.cs
@@ -28,17 +28,18 @@ namespace UnrealBuildTool
 	public enum CppStandardVersion
 	{
 		/// <summary>
-		/// Supports C++14. No longer supported, code may not compile with this version.
+		/// Obsolete C++14 standard. No longer allowed, code will not compile with this version.
 		/// </summary>
+		[Obsolete("C++14 is no longer allowed, code will not compile with this version")]
 		Cpp14,
 
 		/// <summary>
-		/// Supports C++17. No longer supported, code may not compile with this version.
+		/// C++17 standard. No longer supported, code may not compile with this version.
 		/// </summary>
 		Cpp17,
 
 		/// <summary>
-		/// Supports C++20
+		/// C++20 standard
 		/// </summary>
 		Cpp20,
 
@@ -56,6 +57,11 @@ namespace UnrealBuildTool
 		/// Use the default standard version for engine modules
 		/// </summary>
 		EngineDefault = Cpp20,
+
+		/// <summary>
+		/// The minimum allowed C++ standard version. Lower values than this will cause a build error.
+		/// </summary>
+		Minimum = Cpp17,
 	}
 
 	/// <summary>
diff --git a/Engine/Source/Programs/UnrealBuildTool/ToolChain/ClangToolChain.cs b/Engine/Source/Programs/UnrealBuildTool/ToolChain/ClangToolChain.cs
--- a/Engine/Source/Programs/UnrealBuildTool/ToolChain/ClangToolChain.cs
+++ b/Engine/Source/Programs/UnrealBuildTool/ToolChain/ClangToolChain.cs
@@ -381,9 +381,6 @@ namespace UnrealBuildTool
 			// https://clang.llvm.org/cxx_status.html
 			switch (CompileEnvironment.CppStandard)
 			{
-				case CppStandardVersion.Cpp14:
-					Arguments.Add("-std=c++14");
-					break;
 				case CppStandardVersion.Cpp17:
 					Arguments.Add("-std=c++17");
 					break;
