UnrealBuildTool: WindowsPlatform setting for struct member alignment

#rnx

diff --git a/Engine/Source/Programs/UnrealBuildTool/Platform/Windows/UEBuildWindows.cs b/Engine/Source/Programs/UnrealBuildTool/Platform/Windows/UEBuildWindows.cs
--- a/Engine/Source/Programs/UnrealBuildTool/Platform/Windows/UEBuildWindows.cs
+++ b/Engine/Source/Programs/UnrealBuildTool/Platform/Windows/UEBuildWindows.cs
@@ -605,6 +605,14 @@ namespace UnrealBuildTool
 		[CommandLine("-DisableVolatileMetadata")]
 		public bool bDisableVolatileMetadata { get; set; } = false;
 
+		/// <summary>
+		/// Controls how the members of a structure are packed into memory and specifies the same packing for all structures in a module.
+		/// See https://learn.microsoft.com/en-us/cpp/build/reference/zp-struct-member-alignment
+		/// </summary>
+		[RequiresUniqueBuildEnvironment]
+		[XmlConfigFile(Category = "WindowsPlatform")]
+		public int? StructMemberAlignment { get; set; } = 8;
+
 		/// <summary>
 		/// Whether to request the linker create a stripped pdb file as part of the build.
 		/// If enabled the full debug pdb will have the extension .full.pdb
@@ -997,6 +1005,8 @@ namespace UnrealBuildTool
 
 		public bool bDisableVolatileMetadata => Inner.bDisableVolatileMetadata;
 
+		public int? StructMemberAlignment => Inner.StructMemberAlignment;
+
 		public bool bStripPrivateSymbols => Inner.bStripPrivateSymbols;
 
 		public bool bNoLinkerDebugInfo => Inner.bNoLinkerDebugInfo;
diff --git a/Engine/Source/Programs/UnrealBuildTool/Platform/Windows/VCToolChain.cs b/Engine/Source/Programs/UnrealBuildTool/Platform/Windows/VCToolChain.cs
--- a/Engine/Source/Programs/UnrealBuildTool/Platform/Windows/VCToolChain.cs
+++ b/Engine/Source/Programs/UnrealBuildTool/Platform/Windows/VCToolChain.cs
@@ -1161,8 +1161,10 @@ namespace UnrealBuildTool
 				Arguments.Add("/Zo");
 			}
 
-			// Pack struct members on 8-byte boundaries.
-			Arguments.Add("/Zp8");
+			if (Target.WindowsPlatform.StructMemberAlignment != null)
+			{
+				Arguments.Add($"/Zp{Target.WindowsPlatform.StructMemberAlignment}");
+			}
 
 			if (CompileEnvironment.DefaultWarningLevel == WarningLevel.Error)
 			{
