UnrealBuildTool: Add support for UBSan for msvc platforms

diff --git a/Engine/Source/Programs/UnrealBuildTool/Platform/Windows/UEBuildWindows.cs b/Engine/Source/Programs/UnrealBuildTool/Platform/Windows/UEBuildWindows.cs
--- a/Engine/Source/Programs/UnrealBuildTool/Platform/Windows/UEBuildWindows.cs
+++ b/Engine/Source/Programs/UnrealBuildTool/Platform/Windows/UEBuildWindows.cs
@@ -487,12 +487,19 @@ namespace UnrealBuildTool
 		public string? ProductName;
 
 		/// <summary>
-		/// Enables address sanitizer (ASan). Only supported for Visual Studio 2019 16.7.0 and up.
+		/// Enables address sanitizer (ASan)
 		/// </summary>
 		[XmlConfigFile(Category = "BuildConfiguration", Name = "bEnableAddressSanitizer")]
 		[CommandLine("-EnableASan")]
 		public bool bEnableAddressSanitizer = false;
 
+		/// <summary>
+		/// Enables undefined behavior sanitizer (UBSan)
+		/// </summary>
+		[XmlConfigFile(Category = "BuildConfiguration", Name = "bEnableUndefinedBehaviorSanitizer")]
+		[CommandLine("-EnableUBSan")]
+		public bool bEnableUndefinedBehaviorSanitizer = false;
+		
 		/// <summary>
 		/// Enables instrumentation.
 		/// </summary>
@@ -501,7 +508,7 @@ namespace UnrealBuildTool
 		public bool bEnableInstrumentation = false;
 
 		/// <summary>
-		/// Enables LibFuzzer. Only supported for Visual Studio 2022 17.0.0 and up.
+		/// Enables LibFuzzer.
 		/// </summary>
 		[XmlConfigFile(Category = "BuildConfiguration", Name = "bEnableLibFuzzer")]
 		[CommandLine("-EnableLibFuzzer")]
@@ -973,6 +980,8 @@ namespace UnrealBuildTool
 
 		public bool bEnableAddressSanitizer => Inner.bEnableAddressSanitizer;
 
+		public bool bEnableUndefinedBehaviorSanitizer => Inner.bEnableUndefinedBehaviorSanitizer;
+
 		public bool bEnableInstrumentation => Inner.bEnableInstrumentation;
 
 		public bool bEnableLibFuzzer => Inner.bEnableLibFuzzer;
@@ -1302,6 +1311,15 @@ namespace UnrealBuildTool
 				}
 			}
 
+			if (Target.WindowsPlatform.bEnableUndefinedBehaviorSanitizer)
+			{
+				// UBSan currently requires a clang toolchain
+				if (!Target.WindowsPlatform.Compiler.IsClang())
+				{
+					Target.WindowsPlatform.Compiler = WindowsCompiler.Clang;
+				}
+			}
+
 			// Set the compiler version if necessary
 			if (Target.WindowsPlatform.Compiler == WindowsCompiler.Default)
 			{
diff --git a/Engine/Source/Programs/UnrealBuildTool/Platform/Windows/VCToolChain.cs b/Engine/Source/Programs/UnrealBuildTool/Platform/Windows/VCToolChain.cs
--- a/Engine/Source/Programs/UnrealBuildTool/Platform/Windows/VCToolChain.cs
+++ b/Engine/Source/Programs/UnrealBuildTool/Platform/Windows/VCToolChain.cs
@@ -780,6 +780,18 @@ namespace UnrealBuildTool
 				AddSystemIncludePath(Arguments, DirectoryReference.Combine(EnvVars.CompilerDir, "crt", "src"), Target.WindowsPlatform.Compiler, CompileEnvironment.RootPaths);
 			}
 
+			if (Target.WindowsPlatform.bEnableUndefinedBehaviorSanitizer)
+			{
+				if (Target.WindowsPlatform.Compiler.IsClang())
+				{
+					Arguments.Add("-fsanitize=undefined");
+				}
+				else
+				{
+					Arguments.Add("/fsanitize=undefined");
+				}
+			}
+
 			if (Target.WindowsPlatform.bEnableLibFuzzer)
 			{
 				if (Target.WindowsPlatform.Compiler.IsClang())
@@ -1718,6 +1730,29 @@ namespace UnrealBuildTool
 					LinkEnvironment.Libraries.Add(FileReference.Combine(ASanRuntimeDir, $"clang_rt.asan_dynamic_runtime_thunk-{ASanArchSuffix}.lib"));
 				}
 			}
+
+			if (Target.WindowsPlatform.bEnableUndefinedBehaviorSanitizer)
+			{
+				// With clang we seemingly need to explicitly pass the .lib's to link against for UBSan.
+				if (Target.WindowsPlatform.Compiler.IsClang())
+				{
+					DirectoryReference UBSanRuntimeDir;
+					string UBSanArchSuffix;
+					if (EnvVars.Architecture == UnrealArch.X64)
+					{
+						VersionNumber ClangVersion = Target.WindowsPlatform.Compiler == WindowsCompiler.Intel ? MicrosoftPlatformSDK.GetClangVersionForIntelCompiler(EnvVars.CompilerPath) : EnvVars.CompilerVersion;
+						UBSanRuntimeDir = DirectoryReference.Combine(EnvVars.CompilerDir, "lib", "clang", $"{ClangVersion.Components[0]}", "lib", "windows");
+						//UBSanRuntimeDir = DirectoryReference.Combine(EnvVars.ToolChainDir, "lib", "x64");
+						UBSanArchSuffix = "x86_64";
+					}
+					else
+					{
+						throw new BuildException("Unsupported build architecture for Undefined Behavior Sanitizer");
+					}
+					string UBSanRuntimeLib = $"clang_rt.ubsan_standalone-{UBSanArchSuffix}.lib";
+					LinkEnvironment.Libraries.Add(FileReference.Combine(UBSanRuntimeDir, UBSanRuntimeLib));
+				}
+			}
 		}
 
 		protected virtual void AppendLibArguments(LinkEnvironment LinkEnvironment, List<string> Arguments)
