UnrealBuildTool: Support compiling .cxx files

#rnx

diff --git a/Engine/Source/Programs/Shared/EpicGames.Horde/Issues/IssueChangeFilter.cs b/Engine/Source/Programs/Shared/EpicGames.Horde/Issues/IssueChangeFilter.cs
--- a/Engine/Source/Programs/Shared/EpicGames.Horde/Issues/IssueChangeFilter.cs
+++ b/Engine/Source/Programs/Shared/EpicGames.Horde/Issues/IssueChangeFilter.cs
@@ -27,6 +27,7 @@ namespace EpicGames.Horde.Issues
 		{
 			"*.c",
 			"*.cc",
+			"*.cxx",
 			"*.cpp",
 			"*.inl",
 			"*.m",
diff --git a/Engine/Source/Programs/UnrealBuildTool/Configuration/ModuleRules.cs b/Engine/Source/Programs/UnrealBuildTool/Configuration/ModuleRules.cs
--- a/Engine/Source/Programs/UnrealBuildTool/Configuration/ModuleRules.cs
+++ b/Engine/Source/Programs/UnrealBuildTool/Configuration/ModuleRules.cs
@@ -291,6 +291,11 @@ namespace UnrealBuildTool
 			/// </summary>
 			CC,
 
+			/// <summary>
+			/// C++ (alternate extension)
+			/// </summary>
+			CXX,
+
 			/// <summary>
 			/// Objective-C
 			/// </summary>
diff --git a/Engine/Source/Programs/UnrealBuildTool/Configuration/UEBuildModuleCPP.cs b/Engine/Source/Programs/UnrealBuildTool/Configuration/UEBuildModuleCPP.cs
--- a/Engine/Source/Programs/UnrealBuildTool/Configuration/UEBuildModuleCPP.cs
+++ b/Engine/Source/Programs/UnrealBuildTool/Configuration/UEBuildModuleCPP.cs
@@ -30,6 +30,7 @@ namespace UnrealBuildTool
 			public readonly List<FileItem> CPPFiles = new List<FileItem>();
 			public readonly List<FileItem> CFiles = new List<FileItem>();
 			public readonly List<FileItem> CCFiles = new List<FileItem>();
+			public readonly List<FileItem> CXXFiles = new List<FileItem>();
 			public readonly List<FileItem> MFiles = new List<FileItem>();
 			public readonly List<FileItem> MMFiles = new List<FileItem>();
 			public readonly List<FileItem> SwiftFiles = new List<FileItem>();
@@ -37,7 +38,7 @@ namespace UnrealBuildTool
 			public readonly List<FileItem> ISPCFiles = new List<FileItem>();
 
 			public IEnumerable<FileItem> AllHeaderFiles => HeaderFiles.Concat(ISPCHeaderFiles);
-			public IEnumerable<FileItem> AllCompiledFiles => IXXFiles.Concat(CPPFiles).Concat(CFiles).Concat(CCFiles).Concat(MFiles).Concat(MMFiles).Concat(SwiftFiles).Concat(RCFiles).Concat(ISPCFiles);
+			public IEnumerable<FileItem> AllCompiledFiles => IXXFiles.Concat(CPPFiles).Concat(CFiles).Concat(CCFiles).Concat(CXXFiles).Concat(MFiles).Concat(MMFiles).Concat(SwiftFiles).Concat(RCFiles).Concat(ISPCFiles);
 		}
 
 		/// <summary>
@@ -955,6 +956,25 @@ namespace UnrealBuildTool
 				}
 			}
 
+			// Compile CXX files
+			if (InputFiles.CXXFiles.Count > 0)
+			{
+				CppCompileEnvironment CXXCompileEnvironment = new CppCompileEnvironment(CppCompileEnvironment);
+				CXXCompileEnvironment.bUseUnity = bModuleUsesUnityBuild && Rules.UnityFileTypes.HasFlag(ModuleRules.UnityFileType.CXX);
+				CreateHeaderForDefinitions(CXXCompileEnvironment, IntermediateDirectory, "cxx", Graph);
+
+				if (CXXCompileEnvironment.bUseUnity)
+				{
+					Unity.GenerateUnitySource(Target, InputFiles.CXXFiles, InputFiles.HeaderFiles, CXXCompileEnvironment, WorkingSet, Rules.ShortName ?? Name, IntermediateDirectory, Graph, SourceFileToUnityFile,
+						out List<FileItem> NormalFiles, out List<FileItem> AdaptiveFiles, NumIncludedBytesPerUnityCPP);
+					LinkInputFiles.AddRange(CompileFilesWithToolChain(Target, ToolChain, CXXCompileEnvironment, CppCompileEnvironment, NormalFiles, AdaptiveFiles, Graph, Logger).ObjectFiles);
+				}
+				else
+				{
+					LinkInputFiles.AddRange(ToolChain.CompileAllCPPFiles(CXXCompileEnvironment, InputFiles.CXXFiles, IntermediateDirectory, Name, Graph).ObjectFiles);
+				}
+			}
+
 			// Compile M files. Do not use a PCH here, because a C++ PCH is not compatible with Objective-C source files.
 			if (InputFiles.MFiles.Count > 0)
 			{
@@ -2622,7 +2642,7 @@ namespace UnrealBuildTool
 		static string[] KnownInputFileExtensions =
 		{
 			".h", ".ipsh",
-			".cpp", ".ixx", ".c", ".cc",
+			".cpp", ".cxx", ".ixx", ".c", ".cc",
 			".m", ".mm", ".swift",
 			".rc", ".ispc",
 		};
@@ -2682,6 +2702,11 @@ namespace UnrealBuildTool
 					SourceFiles.Add(InputFile);
 					InputFiles.CCFiles.Add(InputFile);
 				}
+				else if (InputFile.HasExtension(".cxx"))
+				{
+					SourceFiles.Add(InputFile);
+					InputFiles.CXXFiles.Add(InputFile);
+				}
 				else if (InputFile.HasExtension(".m"))
 				{
 					SourceFiles.Add(InputFile);
diff --git a/Engine/Source/Programs/UnrealBuildTool/Modes/GenerateClangDatabase.cs b/Engine/Source/Programs/UnrealBuildTool/Modes/GenerateClangDatabase.cs
--- a/Engine/Source/Programs/UnrealBuildTool/Modes/GenerateClangDatabase.cs
+++ b/Engine/Source/Programs/UnrealBuildTool/Modes/GenerateClangDatabase.cs
@@ -114,7 +114,7 @@ namespace UnrealBuildTool
 					if (bExecCodeGenActions)
 					{
 						// Filter all the actions to execute
-						HashSet<FileItem> PrerequisiteItems = new HashSet<FileItem>(Makefile.Actions.SelectMany(x => x.ProducedItems).Where(x => x.HasExtension(".h") || x.HasExtension(".cpp") || x.HasExtension(".cc") || x.HasExtension(".c")));
+						HashSet<FileItem> PrerequisiteItems = new HashSet<FileItem>(Makefile.Actions.SelectMany(x => x.ProducedItems).Where(x => x.HasExtension(".h") || x.HasExtension(".cpp") || x.HasExtension(".cc") || x.HasExtension(".cxx") || x.HasExtension(".c")));
 						List<LinkedAction> PrerequisiteActions = ActionGraph.GatherPrerequisiteActions(Actions, PrerequisiteItems);
 
 						Utils.ExecuteCustomBuildSteps(Makefile.PreBuildScripts, Logger);
@@ -137,7 +137,7 @@ namespace UnrealBuildTool
 					{
 						foreach (IExternalAction Action in CompileActions)
 						{
-							FileItem? SourceFile = Action.PrerequisiteItems.FirstOrDefault(x => x.HasExtension(".cpp") || x.HasExtension(".cc") || x.HasExtension(".c")) ?? Action.PrerequisiteItems.FirstOrDefault(x => x.HasExtension(".h"));
+							FileItem? SourceFile = Action.PrerequisiteItems.FirstOrDefault(x => x.HasExtension(".cpp") || x.HasExtension(".cc") || x.HasExtension(".cxx") || x.HasExtension(".c")) ?? Action.PrerequisiteItems.FirstOrDefault(x => x.HasExtension(".h"));
 							FileItem? OutputFile = Action.ProducedItems.FirstOrDefault(x => x.HasExtension(".obj") | x.HasExtension(".o"));
 							if (SourceFile == null || OutputFile == null)
 							{
diff --git a/Engine/Source/Programs/UnrealBuildTool/ProjectFiles/Eddie/EddieProject.cs b/Engine/Source/Programs/UnrealBuildTool/ProjectFiles/Eddie/EddieProject.cs
--- a/Engine/Source/Programs/UnrealBuildTool/ProjectFiles/Eddie/EddieProject.cs
+++ b/Engine/Source/Programs/UnrealBuildTool/ProjectFiles/Eddie/EddieProject.cs
@@ -55,7 +55,7 @@ namespace UnrealBuildTool
 
 		private bool IsSourceCode(string Extension)
 		{
-			return Extension == ".c" || Extension == ".cc" || Extension == ".cpp" || Extension == ".m" || Extension == ".mm" || Extension == ".cs";
+			return Extension == ".c" || Extension == ".cc" || Extension == ".cxx" || Extension == ".cpp" || Extension == ".m" || Extension == ".mm" || Extension == ".cs";
 		}
 
 		public override SourceFile? AllocSourceFile(FileReference InitFilePath, DirectoryReference? InitProjectSubFolder)
diff --git a/Engine/Source/Programs/UnrealBuildTool/ProjectFiles/Xcode/XcodeProject.File.cs b/Engine/Source/Programs/UnrealBuildTool/ProjectFiles/Xcode/XcodeProject.File.cs
--- a/Engine/Source/Programs/UnrealBuildTool/ProjectFiles/Xcode/XcodeProject.File.cs
+++ b/Engine/Source/Programs/UnrealBuildTool/ProjectFiles/Xcode/XcodeProject.File.cs
@@ -137,6 +137,7 @@ namespace UnrealBuildTool.XcodeProjectXcconfig
 			".c",
 			".cc",
 			".cpp",
+			".cxx",
 			".m",
 			".mm",
 		};
@@ -147,6 +148,7 @@ namespace UnrealBuildTool.XcodeProjectXcconfig
 			{ ".m", "sourcecode.c.objc" },
 			{ ".cc", "sourcecode.cpp.objcpp" },
 			{ ".cpp", "sourcecode.cpp.objcpp" },
+			{ ".cxx", "sourcecode.cpp.objcpp" },
 			{ ".mm", "sourcecode.cpp.objcpp" },
 			{ ".h", "sourcecode.c.h" },
 			{ ".hpp", "sourcecode.cpp.h" },
diff --git a/Engine/Source/Programs/UnrealBuildTool/ProjectFiles/Xcode/XcodeProjectLegacy.cs b/Engine/Source/Programs/UnrealBuildTool/ProjectFiles/Xcode/XcodeProjectLegacy.cs
--- a/Engine/Source/Programs/UnrealBuildTool/ProjectFiles/Xcode/XcodeProjectLegacy.cs
+++ b/Engine/Source/Programs/UnrealBuildTool/ProjectFiles/Xcode/XcodeProjectLegacy.cs
@@ -162,6 +162,7 @@ namespace UnrealBuildTool.XcodeProjectLegacy
 					return "sourcecode.c.objc";
 				case ".cc":
 				case ".cpp":
+				case ".cxx":
 				case ".mm":
 					return "sourcecode.cpp.objcpp";
 				case ".h":
@@ -186,7 +187,7 @@ namespace UnrealBuildTool.XcodeProjectLegacy
 		/// </summary>
 		private bool IsSourceCode(string Extension)
 		{
-			return Extension == ".c" || Extension == ".cc" || Extension == ".cpp" || Extension == ".m" || Extension == ".mm";
+			return Extension == ".c" || Extension == ".cc" || Extension == ".cxx" || Extension == ".cpp" || Extension == ".m" || Extension == ".mm";
 		}
 
 		private bool ShouldIncludeFileInBuildPhaseSection(XcodeSourceFile SourceFile)
