UnrealBuildTool: Fix platform sdk file path for makefile invalidation

#rnx

diff --git a/Engine/Source/Programs/UnrealBuildTool/ToolChain/ISPCToolChain.cs b/Engine/Source/Programs/UnrealBuildTool/ToolChain/ISPCToolChain.cs
--- a/Engine/Source/Programs/UnrealBuildTool/ToolChain/ISPCToolChain.cs
+++ b/Engine/Source/Programs/UnrealBuildTool/ToolChain/ISPCToolChain.cs
@@ -47,12 +47,12 @@ namespace UnrealBuildTool
 			if (TargetRules != null)
 			{
 				string Platform = TargetRules.Platform == UnrealTargetPlatform.Win64 ? "Windows" : TargetRules.Platform.ToString();
-				FileItem PlatformConfig = FileItem.GetItemByFileReference(FileReference.Combine(Unreal.RootDirectory, "Platforms", Platform, "Config", $"{Platform}_SDK.json"));
+				FileItem PlatformConfig = FileItem.GetItemByFileReference(FileReference.Combine(Unreal.EngineDirectory, "Platforms", Platform, "Config", $"{Platform}_SDK.json"));
 				if (PlatformConfig.Exists)
 				{
 					ExternalDependencies.Add(PlatformConfig);
 				}
-				PlatformConfig = FileItem.GetItemByFileReference(FileReference.Combine(Unreal.RootDirectory, "Config", Platform, $"{Platform}_SDK.json"));
+				PlatformConfig = FileItem.GetItemByFileReference(FileReference.Combine(Unreal.EngineDirectory, "Config", Platform, $"{Platform}_SDK.json"));
 				if (PlatformConfig.Exists)
 				{
 					ExternalDependencies.Add(PlatformConfig);
