UnrealBuildTool: Fix nullreference exception with IWYUMode

#rnx

diff --git a/Engine/Source/Programs/UnrealBuildTool/ToolChain/ClangToolChain.cs b/Engine/Source/Programs/UnrealBuildTool/ToolChain/ClangToolChain.cs
--- a/Engine/Source/Programs/UnrealBuildTool/ToolChain/ClangToolChain.cs
+++ b/Engine/Source/Programs/UnrealBuildTool/ToolChain/ClangToolChain.cs
@@ -234,8 +234,8 @@ namespace UnrealBuildTool
 		public override void GetExternalDependencies(HashSet<FileItem> ExternalDependencies)
 		{
 			base.GetExternalDependencies(ExternalDependencies);
-			ExternalDependencies.Add(FileItem.GetItemByFileReference(LazyInfo.Value.Clang));
-			ExternalDependencies.Add(FileItem.GetItemByFileReference(LazyInfo.Value.Archiver));
+			ExternalDependencies.Add(FileItem.GetItemByFileReference(Info.Clang));
+			ExternalDependencies.Add(FileItem.GetItemByFileReference(Info.Archiver));
 		}
 
 		public override FileReference? GetCppCompilerPath()
diff --git a/Engine/Source/Programs/UnrealBuildTool/ToolChain/IWYUToolChain.cs b/Engine/Source/Programs/UnrealBuildTool/ToolChain/IWYUToolChain.cs
--- a/Engine/Source/Programs/UnrealBuildTool/ToolChain/IWYUToolChain.cs
+++ b/Engine/Source/Programs/UnrealBuildTool/ToolChain/IWYUToolChain.cs
@@ -59,7 +59,7 @@ namespace UnrealBuildTool
 			string DevPath = ""; //@"include-what-you-use-0.19\vs_projects\bin\RelWithDebInfo";
 
 			FileReference IWYUPath = FileReference.Combine(Unreal.EngineDirectory, RelativePathToIWYUDirectory, DevPath, @"include-what-you-use.exe");
-			return new ClangToolChainInfo(IWYUPath.Directory, IWYUPath, null!, Logger);
+			return new ClangToolChainInfo(IWYUPath.Directory, IWYUPath, IWYUPath, Logger);
 		}
 
 		protected override string GetFileNameFromExtension(string AbsolutePath, string Extension)
