[Mac]
* Changed link action to use rsp file

diff --git a/Engine/Source/Programs/UnrealBuildTool/Platform/Mac/MacToolChain.cs b/Engine/Source/Programs/UnrealBuildTool/Platform/Mac/MacToolChain.cs
--- a/Engine/Source/Programs/UnrealBuildTool/Platform/Mac/MacToolChain.cs
+++ b/Engine/Source/Programs/UnrealBuildTool/Platform/Mac/MacToolChain.cs
@@ -215,16 +215,14 @@ namespace UnrealBuildTool
 			}
 		}
 
-		string AddFrameworkToLinkCommand(string FrameworkName, string Arg = "-framework")
+		void AddFrameworkToLinkCommand(List<string> LinkArguments, string FrameworkName, string Arg = "-framework")
 		{
-			string Result = "";
 			if (FrameworkName.EndsWith(".framework"))
 			{
-				Result += " -F \"" + Path.GetDirectoryName(Path.GetFullPath(FrameworkName)) + "\"";
+				LinkArguments.Add("-F \"" + Path.GetDirectoryName(Path.GetFullPath(FrameworkName)) + "\"");
 				FrameworkName = Path.GetFileNameWithoutExtension(FrameworkName);
 			}
-			Result += " " + Arg + " \"" + FrameworkName + "\"";
-			return Result;
+			LinkArguments.Add(Arg + " \"" + FrameworkName + "\"");
 		}
 
 		protected override void GetLinkArguments_Global(LinkEnvironment LinkEnvironment, List<string> Arguments)
@@ -355,7 +353,7 @@ namespace UnrealBuildTool
 			return GameName;
 		}
 
-		private void AddLibraryPathToRPaths(string Library, string ExeAbsolutePath, ref List<string> RPaths, ref string LinkCommand, bool bIsBuildingAppBundle, ILogger Logger)
+		private void AddLibraryPathToRPaths(string Library, string ExeAbsolutePath, ref List<string> RPaths, List<string> LinkArguments, bool bIsBuildingAppBundle, ILogger Logger)
 		{
 			string LibraryFullPath = Path.GetFullPath(Library);
 			string LibraryDir = Path.GetDirectoryName(LibraryFullPath)!;
@@ -379,7 +377,7 @@ namespace UnrealBuildTool
 
 				if (bCanUseMultipleRPATHs)
 				{
-					LinkCommand += " -rpath \"@loader_path/" + RelativePath + "\"";
+					LinkArguments.Add("-rpath \"@loader_path/" + RelativePath + "\"");
 
 					// We are referencing plugins that need to be relative to the engine not the dylib we are building
 					// To be safe leave the previous relative loader_path in place and add a relative engine executable path
@@ -388,7 +386,7 @@ namespace UnrealBuildTool
 						int EngineDirStrIndex = LibraryDir.IndexOf("Engine");
 						if (EngineDirStrIndex >= 0)
 						{
-							LinkCommand += " -rpath \"@executable_path/../../../../../../" + LibraryDir.Substring(EngineDirStrIndex) + "\"";
+							LinkArguments.Add("-rpath \"@executable_path/../../../../../../" + LibraryDir.Substring(EngineDirStrIndex) + "\"");
 						}
 					}
 				}
@@ -415,7 +413,7 @@ namespace UnrealBuildTool
 					if (BundleLibraryDir != null)
 					{
 						string BundleRelativeDir = Utils.MakePathRelativeTo(BundleLibraryDir, ExeDir).Replace("\\", "/");
-						LinkCommand += " -rpath \"@loader_path/" + BundleRelativeDir + "\"";
+						LinkArguments.Add("-rpath \"@loader_path/" + BundleRelativeDir + "\"");
 					}
 					else
 					{
@@ -430,7 +428,7 @@ namespace UnrealBuildTool
 						string StagedRelativeDir = Utils.MakePathRelativeTo(StagedLibraryDir, ExeDir).Replace("\\", "/");
 						if (StagedRelativeDir != RelativePath)
 						{
-							LinkCommand += " -rpath \"@loader_path/" + StagedRelativeDir + "\"";
+							LinkArguments.Add("-rpath \"@loader_path/" + StagedRelativeDir + "\"");
 						}
 					}
 				}
@@ -468,7 +466,6 @@ namespace UnrealBuildTool
 			bool bIsBuildingLibrary = LinkEnvironment.bIsBuildingLibrary || bBuildImportLibraryOnly;
 			bool bIsBuildingAppBundle = !LinkEnvironment.bIsBuildingDLL && !LinkEnvironment.bIsBuildingLibrary && !LinkEnvironment.bIsBuildingConsoleApplication;
 			FileReference LinkerPath = bIsBuildingLibrary ? Info.Archiver : Info.Clang;
-			string LinkCommand = "";
 
 			if (LinkEnvironment.Architectures.bIsMultiArch)
 			{
@@ -481,7 +478,7 @@ namespace UnrealBuildTool
 						ArchEnvironment.OutputFilePaths = LinkEnvironment.OutputFilePaths.Select(x => new FileReference($"{LinkEnvironment.IntermediateDirectory}/{Path.GetFileName(x.FullName)}_{Arch}")).ToList();
 					}
 
-					PerArchOutputFiles.Add(LinkArchitectureFiles(ArchEnvironment, LinkEnvironment.OutputFilePath, bBuildImportLibraryOnly, Graph, out LinkCommand));
+					PerArchOutputFiles.Add(LinkArchitectureFiles(ArchEnvironment, LinkEnvironment.OutputFilePath, bBuildImportLibraryOnly, Graph));
 				}
 
 				if (bBuildImportLibraryOnly)
@@ -513,7 +510,7 @@ namespace UnrealBuildTool
 			}
 			else
 			{
-				OutputFiles.Add(LinkArchitectureFiles(LinkEnvironment, LinkEnvironment.OutputFilePath, bBuildImportLibraryOnly, Graph, out LinkCommand));
+				OutputFiles.Add(LinkArchitectureFiles(LinkEnvironment, LinkEnvironment.OutputFilePath, bBuildImportLibraryOnly, Graph));
 			}
 
 			if (!DirectoryReference.Exists(LinkEnvironment.IntermediateDirectory!))
@@ -671,7 +668,7 @@ namespace UnrealBuildTool
 			return FileItem.GetItemByPath(Path.Combine(IntermediateDirectory, "Stubs", LinkEnvironment.Architecture.ToString(), LinkEnvironment.Configuration.ToString(), Path.GetFileName(DylibPath)));
 		}
 
-		private FileItem LinkArchitectureFiles(LinkEnvironment LinkEnvironment, FileReference MultiArchOutputFile, bool bBuildImportLibraryOnly, IActionGraphBuilder Graph, out string LinkCommand)
+		private FileItem LinkArchitectureFiles(LinkEnvironment LinkEnvironment, FileReference MultiArchOutputFile, bool bBuildImportLibraryOnly, IActionGraphBuilder Graph)
 		{
 			List<FileItem> InputFiles = LinkEnvironment.InputFiles;
 
@@ -696,10 +693,6 @@ namespace UnrealBuildTool
 			LinkAction.CacheBucket = GetCacheBucket(Target, null);
 			LinkAction.ArtifactMode = ArtifactMode.Enabled;
 
-			string EngineAPIVersion = LoadEngineAPIVersion();
-			string EngineDisplayVersion = LoadEngineDisplayVersion(true);
-			string VersionArg = LinkEnvironment.bIsBuildingDLL ? " -current_version " + EngineAPIVersion + " -compatibility_version " + EngineDisplayVersion : "";
-
 			List<string> LinkArguments = new();
 			if (LinkEnvironment.bIsBuildingLibrary)
 			{
@@ -710,7 +703,11 @@ namespace UnrealBuildTool
 				GetLinkArguments_Global(LinkEnvironment, LinkArguments);
 			}
 
-			LinkCommand = VersionArg + " " + String.Join(' ', LinkArguments);
+			if (LinkEnvironment.bIsBuildingDLL)
+			{
+				LinkArguments.Add("-current_version " + LoadEngineAPIVersion());
+				LinkArguments.Add("-compatibility_version " + LoadEngineDisplayVersion(true));
+			}
 
 			// Tell the action that we're building an import library here and it should conditionally be
 			// ignored as a prerequisite for other actions
@@ -732,13 +729,13 @@ namespace UnrealBuildTool
 			string AbsolutePath = MultiArchOutputFile.FullName.Replace("\\", "/");
 			if (!LinkEnvironment.bIsBuildingLibrary)
 			{
-				LinkCommand += " -rpath @loader_path/ -rpath @executable_path/";
+				LinkArguments.Add("-rpath @loader_path/ -rpath @executable_path/");
 			}
 
 			bool bIsBuildingAppBundle = !LinkEnvironment.bIsBuildingDLL && !LinkEnvironment.bIsBuildingLibrary && !LinkEnvironment.bIsBuildingConsoleApplication;
 			if (bIsBuildingAppBundle)
 			{
-				LinkCommand += " -rpath @executable_path/../../../";
+				LinkArguments.Add("-rpath @executable_path/../../../");
 			}
 
 			List<string> RPaths = new List<string>();
@@ -746,8 +743,8 @@ namespace UnrealBuildTool
 			// static libs and stub dylibs don't need to look other libs
 			if (bBuildImportLibraryOnly)
 			{
-				LinkCommand += " -undefined dynamic_lookup";
-				LinkCommand += " -Wl,-no_fixup_chains";
+				LinkArguments.Add("-undefined dynamic_lookup");
+				LinkArguments.Add("-Wl,-no_fixup_chains");
 			}
 			else if (!LinkEnvironment.bIsBuildingLibrary)
 			{
@@ -763,7 +760,7 @@ namespace UnrealBuildTool
 
 					HashSet<DirectoryReference> UsedDestDirs = new();
 					HashSet<FileReference> UsedSourceLibs = new();
-					var ProcessLib = (FileReference SourceLib, FileReference DestLib, FileReference ReferencingFile, ref string LinkCommand) => 
+					var ProcessLib = (FileReference SourceLib, FileReference DestLib, FileReference ReferencingFile) => 
 					{
 						if (string.Equals(SourceLib.GetExtension(), ".dylib", StringComparison.InvariantCultureIgnoreCase))
 						{
@@ -786,11 +783,11 @@ namespace UnrealBuildTool
 							{
 								RPathEntry = StagedBundlePrefix + DestDir.MakeRelativeTo(Unreal.EngineDirectory.ParentDirectory!);
 							}
-							LinkCommand += $" -rpath \"@executable_path/{RPathEntry}\"";
+							LinkArguments.Add($"-rpath \"@executable_path/{RPathEntry}\"");
 
 							// handle unstaged paths
 							RPathEntry = UnstagedBundlePrefix + DestDir.MakeRelativeTo(ReferencingFile.Directory);
-							LinkCommand += $" -rpath \"@loader_path/{RPathEntry}\"";
+							LinkArguments.Add($"-rpath \"@loader_path/{RPathEntry}\"");
 						}
 					};
 
@@ -800,12 +797,12 @@ namespace UnrealBuildTool
 					{
 						FileReference DestPath = new(SrcDep.Path);
 						FileReference SourcePath = SrcDep.SourcePath == null ? DestPath : new FileReference(SrcDep.SourcePath);
-						ProcessLib(SourcePath, DestPath, MultiArchOutputFile, ref LinkCommand);
+						ProcessLib(SourcePath, DestPath, MultiArchOutputFile);
 					}
 					foreach (string AdditionalLibrary in AdditionalLibraries)
 					{
 						FileReference LibPath = new(AdditionalLibrary);
-						ProcessLib(LibPath, LibPath, MultiArchOutputFile, ref LinkCommand);
+						ProcessLib(LibPath, LibPath, MultiArchOutputFile);
 					}
 				}
 
@@ -828,34 +825,34 @@ namespace UnrealBuildTool
 						}
 
 						LinkAction.PrerequisiteItems.Add(LinkedLib);
-						LinkCommand += $" \"{NormalizeCommandLinePath(LinkedLib, LinkEnvironment.RootPaths)}\"";
+						LinkArguments.Add($"\"{NormalizeCommandLinePath(LinkedLib, LinkEnvironment.RootPaths)}\"");
 					}
 					else if (AdditionalLibrary.Contains(".framework/"))
 					{
-						LinkCommand += String.Format(" \"{0}\"", AdditionalLibrary);
+						LinkArguments.Add(String.Format("\"{0}\"", AdditionalLibrary));
 					}
 					else if (String.IsNullOrEmpty(Path.GetDirectoryName(AdditionalLibrary)) && String.IsNullOrEmpty(Path.GetExtension(AdditionalLibrary)))
 					{
-						LinkCommand += String.Format(" -l\"{0}\"", AdditionalLibrary);
+						LinkArguments.Add(String.Format("-l\"{0}\"", AdditionalLibrary));
 					}
 					else
 					{
-						LinkCommand += String.Format(" \"{0}\"", Path.GetFullPath(AdditionalLibrary));
+						LinkArguments.Add($"\"{NormalizeCommandLinePath(FileItem.GetItemByPath(AdditionalLibrary), LinkEnvironment.RootPaths)}\"");
 					}
 
 					if (!bCanUseSimpleRPaths)
 					{
-						AddLibraryPathToRPaths(AdditionalLibrary, AbsolutePath, ref RPaths, ref LinkCommand, bIsBuildingAppBundle, Logger);
+						AddLibraryPathToRPaths(AdditionalLibrary, AbsolutePath, ref RPaths, LinkArguments, bIsBuildingAppBundle, Logger);
 					}
 				}
 
 				foreach (string AdditionalLibrary in LinkEnvironment.DelayLoadDLLs)
 				{
-					LinkCommand += String.Format(" -weak_library \"{0}\"", Path.GetFullPath(AdditionalLibrary));
+					LinkArguments.Add(String.Format("-weak_library \"{0}\"", Path.GetFullPath(AdditionalLibrary)));
 
 					if (!bCanUseSimpleRPaths)
 					{
-						AddLibraryPathToRPaths(AdditionalLibrary, AbsolutePath, ref RPaths, ref LinkCommand, bIsBuildingAppBundle, Logger);
+						AddLibraryPathToRPaths(AdditionalLibrary, AbsolutePath, ref RPaths, LinkArguments, bIsBuildingAppBundle, Logger);
 					}
 				}
 
@@ -891,8 +888,8 @@ namespace UnrealBuildTool
 
 				foreach (KeyValuePair<string, bool> Framework in AllFrameworks)
 				{
-					LinkCommand += AddFrameworkToLinkCommand(Framework.Key, Framework.Value ? "-weak_framework" : "-framework");
-					AddLibraryPathToRPaths(Framework.Key, AbsolutePath, ref RPaths, ref LinkCommand, bIsBuildingAppBundle, Logger);
+					AddFrameworkToLinkCommand(LinkArguments, Framework.Key, Framework.Value ? "-weak_framework" : "-framework");
+					AddLibraryPathToRPaths(Framework.Key, AbsolutePath, ref RPaths, LinkArguments, bIsBuildingAppBundle, Logger);
 				}
 
 				// Write the MAP file to the output directory.
@@ -908,7 +905,7 @@ namespace UnrealBuildTool
 
 					FileReference MapFilePath = new FileReference(MapFileBaseName + ".map");
 					FileItem MapFile = FileItem.GetItemByFileReference(MapFilePath);
-					LinkCommand += String.Format(" -Wl,-map,\"{0}\"", MapFilePath);
+					LinkArguments.Add(String.Format("-Wl,-map,\"{0}\"", MapFilePath));
 					LinkAction.ProducedItems.Add(MapFile);
 				}
 			}
@@ -918,7 +915,7 @@ namespace UnrealBuildTool
 				// Add the output file to the command-line.
 				string? InstallName = LinkEnvironment.InstallName;
 				InstallName ??= String.Format("{0}/{1}", DylibsPath, Path.GetFileName(OutputFile.AbsolutePath).Replace($".dylib_{LinkEnvironment.Architecture}", ".dylib"));
-				LinkCommand += String.Format(" -install_name \"{0}\"", InstallName);
+				LinkArguments.Add(String.Format("-install_name \"{0}\"", InstallName));
 			}
 
 			List<string> InputFileNames = new List<string>();
@@ -930,16 +927,31 @@ namespace UnrealBuildTool
 
 			foreach (string Filename in InputFileNames)
 			{
-				LinkCommand += " " + Filename;
+				LinkArguments.Add(Filename);
 			}
 
 			// Add the output file to the command-line.
-			LinkCommand += String.Format(" -o \"{0}\"", NormalizeCommandLinePath(OutputFile, LinkEnvironment.RootPaths));
+			LinkArguments.Add(String.Format("-o \"{0}\"", NormalizeCommandLinePath(OutputFile, LinkEnvironment.RootPaths)));
 
 			// Add the additional arguments specified by the environment.
-			LinkCommand += LinkEnvironment.AdditionalArguments;
+			LinkArguments.Add(LinkEnvironment.AdditionalArguments);
+
+			FileReference ResponseFileName;
+			if (bBuildImportLibraryOnly)
+			{
+				ResponseFileName = FileReference.FromString(OutputFile.FullName + ResponseExt);
+			}
+			else
+			{
+				ResponseFileName = GetResponseFileName(LinkEnvironment, OutputFile);
+			}
+			
+			FileItem ResponseFileItem = Graph.CreateIntermediateTextFile(ResponseFileName, LinkArguments);
+			string ResponsFileArgument = GetResponseFileArgument(ResponseFileItem, LinkEnvironment.RootPaths);
+
+			LinkAction.PrerequisiteItems.Add(ResponseFileItem);
 
-			LinkAction.CommandArguments = LinkCommand;
+			LinkAction.CommandArguments = ResponsFileArgument;
 
 			// Only execute linking on the local Mac.
 			LinkAction.bCanExecuteRemotely = false;
diff --git a/Engine/Source/Programs/UnrealBuildTool/ToolChain/AppleToolChain.cs b/Engine/Source/Programs/UnrealBuildTool/ToolChain/AppleToolChain.cs
--- a/Engine/Source/Programs/UnrealBuildTool/ToolChain/AppleToolChain.cs
+++ b/Engine/Source/Programs/UnrealBuildTool/ToolChain/AppleToolChain.cs
@@ -841,7 +841,7 @@ namespace UnrealBuildTool
 			}
 
 			DirectoryReference dir = DirectoryReference.Combine(ApplePlatformSDK.GetToolchainDirectory(), "../lib/swift", Settings.PlatformDirName);
-			Arguments.Add($"-L{dir}");
+			Arguments.Add($"-L{NormalizeCommandLinePath(dir, LinkEnvironment.RootPaths)}");
 			Arguments.Add("-L/usr/lib/swift");
 		}
 
