Add "IsTestTarget" field to target details in QueryMode

diff --git a/Engine/Source/Programs/UnrealBuildTool/Modes/QueryMode.cs b/Engine/Source/Programs/UnrealBuildTool/Modes/QueryMode.cs
--- a/Engine/Source/Programs/UnrealBuildTool/Modes/QueryMode.cs
+++ b/Engine/Source/Programs/UnrealBuildTool/Modes/QueryMode.cs
@@ -493,6 +493,7 @@ namespace UnrealBuildTool
 					DirToModule = CurrentTargetIntellisenseInfo?.DirToModule.ToImmutableSortedDictionary(x => x.Key.ToString(), x => x.Value.Name),
 					ModuleToCompileSettings = CurrentTargetIntellisenseInfo?.ModuleToCompileSettings.ToImmutableSortedDictionary(x => x.Key.Name, x => x.Value),
 					LaunchSettings = CurrentLaunchSettings,
+					IsTestTarget = CurrentTarget.Rules.IsTestTarget
 				};
 
 				await WriteResultsAsync(Result, JsonOptions, Logger);
