[UBT]
* Fixed ios linker warning  by not adding swift additional directory for dylibs and then keep the path non-virtualized for the other

diff --git a/Engine/Source/Programs/UnrealBuildTool/ToolChain/AppleToolChain.cs b/Engine/Source/Programs/UnrealBuildTool/ToolChain/AppleToolChain.cs
--- a/Engine/Source/Programs/UnrealBuildTool/ToolChain/AppleToolChain.cs
+++ b/Engine/Source/Programs/UnrealBuildTool/ToolChain/AppleToolChain.cs
@@ -840,8 +840,14 @@ namespace UnrealBuildTool
 				Arguments.Add(" -ld_classic");
 			}
 
-			DirectoryReference dir = DirectoryReference.Combine(ApplePlatformSDK.GetToolchainDirectory(), "../lib/swift", Settings.PlatformDirName);
-			Arguments.Add($"-L{NormalizeCommandLinePath(dir, LinkEnvironment.RootPaths)}");
+			if (!LinkEnvironment.bIsBuildingDLL)
+			{
+				// Todo: This is currently only used when linking binary and swift files have not been tested with dylibs.
+				// The right fix is to move the vfs root one step up to be able to handle the ../lib/ path
+				DirectoryReference dir = DirectoryReference.Combine(ApplePlatformSDK.GetToolchainDirectory(), "../lib/swift", Settings.PlatformDirName);
+				Arguments.Add($"-L{dir}");
+			}
+
 			Arguments.Add("-L/usr/lib/swift");
 		}
 
