Added Support for multiple Android SDK installs for different engine versions. Also added small fix for gradle NDK_VERSION and NDKROOT potentially being out of sync.

#android

diff --git a/Engine/Source/Programs/AutomationTool/Android/AndroidPlatform.Automation.cs b/Engine/Source/Programs/AutomationTool/Android/AndroidPlatform.Automation.cs
--- a/Engine/Source/Programs/AutomationTool/Android/AndroidPlatform.Automation.cs
+++ b/Engine/Source/Programs/AutomationTool/Android/AndroidPlatform.Automation.cs
@@ -311,7 +311,7 @@ public class AndroidPlatform : Platform
 		if (OperatingSystem.IsLinux())
 		{
 			string UserHome = Environment.GetEnvironmentVariable("HOME");
-			string AndroidSdkPath = Path.Combine(UserHome, "Android", "Sdk");
+			string AndroidSdkPath = Path.Combine(UserHome, "Android", "sdk");
 
 			return AndroidSdkPath;
 		}
@@ -341,7 +341,7 @@ public class AndroidPlatform : Platform
 			}
 
 			string UserHome = Environment.GetEnvironmentVariable("HOME");
-			string AndroidSdkPath = Path.Combine(UserHome, "Library", "Android", "Sdk");
+			string AndroidSdkPath = Path.Combine(UserHome, "Library", "Android", "sdk");
 
 			return AndroidSdkPath;
 		}
@@ -349,7 +349,7 @@ public class AndroidPlatform : Platform
 		Debug.Assert(OperatingSystem.IsWindows());
 
 		string DefaultSdkDir = Path.Combine(Environment.GetFolderPath(Environment.SpecialFolder.LocalApplicationData), "Android", "Sdk");
-		string RegValue = Microsoft.Win32.Registry.GetValue("HKEY_LOCAL_MACHINE\\SOFTWARE\\Android", "SdkPath", null) as string;
+		string RegValue = Microsoft.Win32.Registry.GetValue("HKEY_LOCAL_MACHINE\\SOFTWARE\\Android Studio", "SdkPath", null) as string;
 		return RegValue == null ? DefaultSdkDir : RegValue;
 	}
 
diff --git a/Engine/Source/Programs/UnrealBuildTool/Platform/Android/AndroidPlatformSDK.cs b/Engine/Source/Programs/UnrealBuildTool/Platform/Android/AndroidPlatformSDK.cs
--- a/Engine/Source/Programs/UnrealBuildTool/Platform/Android/AndroidPlatformSDK.cs
+++ b/Engine/Source/Programs/UnrealBuildTool/Platform/Android/AndroidPlatformSDK.cs
@@ -2,6 +2,7 @@
 
 using System;
 using System.Collections.Generic;
+using System.Diagnostics;
 using System.IO;
 using System.Linq;
 using System.Text.RegularExpressions;
@@ -175,6 +176,92 @@ namespace UnrealBuildTool
 			return NDKToolchainVersion;
 		}
 
+		static string GetStudioDir()
+		{
+			if (OperatingSystem.IsLinux())
+			{
+				return Path.Combine(Environment.GetEnvironmentVariable("HOME")!, "android-studio");
+			}
+			else if (OperatingSystem.IsMacOS())
+			{
+				string Directory = Path.Combine(Environment.GetEnvironmentVariable("HOME")!, "Applications/Android Studio.app");
+				return System.IO.Directory.Exists(Directory) ? Directory : "/Applications/Android Studio.app";
+			}
+			else
+			{
+				Debug.Assert(OperatingSystem.IsWindows());
+
+				string? Directory = Microsoft.Win32.Registry.GetValue(@"HKEY_LOCAL_MACHINE\SOFTWARE\Android Studio", "Path", null) as string;
+				return !string.IsNullOrEmpty(Directory) ? Directory : Path.Combine(Environment.GetFolderPath(Environment.SpecialFolder.LocalApplicationData), @"Programs\Android Studio");
+			}
+		}
+		
+		static string GetStudioSdkDir()
+		{
+			if (OperatingSystem.IsLinux())
+			{
+				return Path.Combine(Environment.GetEnvironmentVariable("HOME")!, "Android/sdk");
+			}
+			else if (OperatingSystem.IsMacOS())
+			{
+				return Path.Combine(Environment.GetEnvironmentVariable("HOME")!, "Library/Android/sdk");
+			}
+			else
+			{
+				Debug.Assert(OperatingSystem.IsWindows());
+
+				string? Directory = Microsoft.Win32.Registry.GetValue(@"HKEY_LOCAL_MACHINE\SOFTWARE\Android Studio", "SdkPath", null) as string;
+				return !string.IsNullOrEmpty(Directory) ? Directory : Path.Combine(Environment.GetFolderPath(Environment.SpecialFolder.LocalApplicationData), @"Android\Sdk");
+			}
+		}
+
+		string NdkVersionToName(string[] Parts)
+		{
+			return Parts[1] == "0" ? $"r{Parts[0]}" : $"r{Parts[0]}{Convert.ToChar('a' + int.Parse(Parts[1]))}";
+		}
+		
+		public override string[] GetAllInstalledSDKVersions()
+		{
+			string StudioSdkDir = GetStudioSdkDir();
+			
+			if (Directory.Exists(Path.Combine(StudioSdkDir, "platforms", GetPlatformSpecificVersion("platforms"))))
+			{
+				return Directory.EnumerateDirectories(Path.Combine(StudioSdkDir, "ndk")).Select(Directory => NdkVersionToName(Path.GetFileName(Directory).Split("."))).Distinct().ToArray();
+			}
+			
+			return [];
+		}
+
+		protected override bool TryGetEnvironmentForManualSDK(string SelectedSDKVersion, out Dictionary<string, string>? EnvVarValues, out List<string>? PathAdds, out List<string>? PathRemoves)
+		{
+			PathAdds = null;
+			PathRemoves = null;
+			
+			string? CurrentSdkDir = Environment.GetEnvironmentVariable("ANDROID_HOME");
+			
+			if (string.IsNullOrEmpty(CurrentSdkDir) || !Directory.Exists(Path.Combine(CurrentSdkDir, "platforms", GetPlatformSpecificVersion("platforms"))) || GetInstalledSDKVersion() != SelectedSDKVersion)
+			{
+				string StudioDir = GetStudioDir();
+				string StudioSdkDir = GetStudioSdkDir();
+				string NdkBaseDir = Path.Combine(StudioSdkDir, "ndk");
+				string NdkDir = Path.Combine(NdkBaseDir, string.Join(".", Directory.EnumerateDirectories(NdkBaseDir).Select(Directory => Path.GetFileName(Directory).Split(".")).Where(Parts => NdkVersionToName(Parts) == SelectedSDKVersion).OrderByDescending(Parts => int.Parse(Parts[2])).First()));
+
+				EnvVarValues = new()
+				{
+					["ANDROID_HOME"] = StudioSdkDir,
+					["ANDROID_SDK_HOME"] = "",
+					["ANDROID_NDK_ROOT"] = NdkDir,
+					["NDK_ROOT"] = NdkDir,
+					["NDKROOT"] = NdkDir,
+					["JAVA_HOME"] = Path.Combine(StudioDir, int.Parse(GetVersionFromConfig("AutoSDKDirectory")![1..]) <= 25 ? "jre" : "jbr"),
+				};
+				return true;
+			}
+
+			EnvVarValues = null;
+			return false;
+		}
+
 		public override bool TryConvertVersionToInt(string? StringValue, out ulong OutValue, string? Hint)
 		{
 			// convert r<num>[letter] to hex
diff --git a/Engine/Source/Programs/UnrealBuildTool/Platform/Android/UEDeployAndroid.cs b/Engine/Source/Programs/UnrealBuildTool/Platform/Android/UEDeployAndroid.cs
--- a/Engine/Source/Programs/UnrealBuildTool/Platform/Android/UEDeployAndroid.cs
+++ b/Engine/Source/Programs/UnrealBuildTool/Platform/Android/UEDeployAndroid.cs
@@ -3939,6 +3939,7 @@ namespace UnrealBuildTool
 			int NDKVersionIndex = NDKPath.LastIndexOf("/");
 			string NDKVersion = NDKVersionIndex > 0 ? NDKPath.Substring(NDKVersionIndex + 1) : DEFAULT_NDK_VERSION;
 			GradleProperties.AppendLine(String.Format("NDK_VERSION={0}", NDKVersion));
+			GradleProperties.AppendLine(String.Format("NDK_PATH={0}", NDKPath));
 
 			if (DestApkName != null)
 			{
@@ -6559,7 +6560,7 @@ import java.util.Collection;
 				BuildGradleContent.AppendLine("apply plugin: 'com.android.library'");
 				BuildGradleContent.AppendLine("android {");
 				BuildGradleContent.AppendLine("\tndkVersion NDK_VERSION");
-				BuildGradleContent.AppendLine("\tndkPath System.getenv(\"NDKROOT\")");
+				BuildGradleContent.AppendLine("\tndkPath NDK_PATH");
 				BuildGradleContent.AppendLine("\tcompileSdk COMPILE_SDK_VERSION.toInteger()");
 				BuildGradleContent.AppendLine("\tbuildToolsVersion BUILD_TOOLS_VERSION");
 				BuildGradleContent.AppendLine("\tdefaultConfig {");
