[WinArm64] mark the SxS manifest as writable after copying & try to remove existing one first.

#rnx

diff --git a/Engine/Source/Programs/UnrealBuildTool/Platform/Windows/VCToolChain.cs b/Engine/Source/Programs/UnrealBuildTool/Platform/Windows/VCToolChain.cs
--- a/Engine/Source/Programs/UnrealBuildTool/Platform/Windows/VCToolChain.cs
+++ b/Engine/Source/Programs/UnrealBuildTool/Platform/Windows/VCToolChain.cs
@@ -2468,7 +2468,15 @@ namespace UnrealBuildTool
 				bool bSuccess = false;
 				try
 				{
+					if (FileReference.Exists(Manifest!))
+					{
+						FileReference.SetAttributes(Manifest!, FileAttributes.Normal);
+						FileReference.Delete(Manifest!);
+					}
+
 					FileReference.Copy(Template!, Manifest!, true);
+					FileReference.SetAttributes(Manifest!, FileAttributes.Normal);
+
 					bSuccess = WindowsExports.UpdateSxSManifest(Manifest!);
 				}
 				catch(Exception e)
