Fixed Android x64 so it's buildable

#android
#tests tested x64 and arm64 both build and run

diff --git a/Engine/Source/Programs/UnrealBuildTool/Platform/Android/UEDeployAndroid.cs b/Engine/Source/Programs/UnrealBuildTool/Platform/Android/UEDeployAndroid.cs
--- a/Engine/Source/Programs/UnrealBuildTool/Platform/Android/UEDeployAndroid.cs
+++ b/Engine/Source/Programs/UnrealBuildTool/Platform/Android/UEDeployAndroid.cs
@@ -124,7 +124,7 @@ namespace UnrealBuildTool
 		{
 			this.Architectures = Architectures;
 
-			List<string> NDKArches = Architectures.Architectures.Select(x => GetNDKArch(x)).ToList();
+			List<(string UEArch, string NativeArch)> Archs = Architectures.Architectures.Select(x => (x.ToString(), GetNDKArch(x))).ToList();
 
 			// check if certain plugins are enabled
 			ARCorePluginEnabled = false;
@@ -165,7 +165,7 @@ namespace UnrealBuildTool
 				}
 			}
 
-			UPL = new UnrealPluginLanguage(ProjectFile, inPluginExtraData, NDKArches, ManifestXMLNamespaces, UnrealTargetPlatform.Android, Logger);
+			UPL = new UnrealPluginLanguage(ProjectFile, inPluginExtraData, Archs, ManifestXMLNamespaces, UnrealTargetPlatform.Android, Logger);
 			UPLHashCode = UPL.GetUPLHash();
 			//			APL.SetTrace();
 		}
@@ -4349,16 +4349,8 @@ dependencies.implementation('androidx.tracing:tracing:1.2.0')
 			string GameBuildFilesPath_NR = Path.Combine(Path.Combine(ProjectDirectory, "Restricted", "NoRedist"), "Build", "Android");
 			string GameBuildFilesPath_LA = Path.Combine(Path.Combine(ProjectDirectory, "Restricted", "LimitedAccess"), "Build", "Android");
 
-			// get a list of unique NDK architectures enabled for build
-			List<string> NDKArches = new List<string>();
-			foreach (UnrealArch Arch in Architectures.Architectures)
-			{
-				string NDKArch = GetNDKArch(Arch);
-				if (!NDKArches.Contains(NDKArch))
-				{
-					NDKArches.Add(NDKArch);
-				}
-			}
+			// get a list of unique architectures enabled for build
+			List<(UnrealArch UEArch, string NativeArch)> Archs = Architectures.Architectures.Select(x => (x, GetNDKArch(x))).ToList();
 
 			// force create from scratch if on build machine
 			bool bCreateFromScratch = bIsBuildMachine;
@@ -4476,9 +4468,9 @@ dependencies.implementation('androidx.tracing:tracing:1.2.0')
 			});
 
 			// Sometimes old files get left behind if things change, so we'll do a clean up pass
-			foreach (string NDKArch in NDKArches)
+			foreach ((UnrealArch UnrealArch, string _) in Archs)
 			{
-				string UnrealBuildPath = Path.Combine(IntermediateAndroidPath, GetUnrealArch(NDKArch).ToString());
+				string UnrealBuildPath = Path.Combine(IntermediateAndroidPath, UnrealArch.ToString());
 
 				string CleanUpBaseDir = Path.Combine(ProjectDirectory, "Build", "Android", "src");
 				string ImmediateBaseDir = Path.Combine(UnrealBuildPath, "src");
@@ -4577,7 +4569,7 @@ dependencies.implementation('androidx.tracing:tracing:1.2.0')
 
 			// get application name and whether it changed, needing to force repackage
 			string? ApplicationDisplayName;
-			if (CheckApplicationName(Path.Combine(IntermediateAndroidPath, ArchRemapping[NDKArches[0]]), ProjectName, out ApplicationDisplayName))
+			if (CheckApplicationName(Path.Combine(IntermediateAndroidPath, ArchRemapping[Archs[0].NativeArch]), ProjectName, out ApplicationDisplayName))
 			{
 				bBuildSettingsMatch = false;
 				Logger.LogInformation("Application display name is different than last build, forcing repackage.");
@@ -4632,13 +4624,12 @@ dependencies.implementation('androidx.tracing:tracing:1.2.0')
 			// Initialize UPL contexts for each architecture enabled
 			UPL.SetGlobalContextVariable("AndroidPackageName", PackageName);
 			UPL.SetGlobalContextVariable("CookFlavor", CookFlavor);
-			UPL.Init(NDKArches, bForDistribution, EngineDirectory, IntermediateAndroidPath, ProjectDirectory, Configuration.ToString(), bSkipGradleBuild, bPerArchBuildDir: true, ArchRemapping: ArchRemapping);
+			UPL.Init(Archs.Select(x => x.NativeArch), bForDistribution, EngineDirectory, IntermediateAndroidPath, ProjectDirectory, Configuration.ToString(), bSkipGradleBuild, bPerArchBuildDir: true, ArchRemapping: ArchRemapping);
 			
 
 			Dictionary<string, int> NDKArchToCompileSDK = new Dictionary<string, int>();
-			foreach (UnrealArch Arch in Architectures.Architectures)
+			foreach ((UnrealArch _, string NDKArch) in Archs)
 			{
-				string NDKArch = GetNDKArch(Arch);
 				// check if any plugins want to increase the required compile SDK version
 				string CompileSDKVersion = SDKAPILevel.Replace("android-", "");
 				string CompileSDKMin = UPL.ProcessPluginNode(NDKArch, "minimumSDKAPI", "");
@@ -4702,10 +4693,8 @@ dependencies.implementation('androidx.tracing:tracing:1.2.0')
 			// Need to do stripping here because .apk will stay the same during iteration
 			if (bDontBundleLibrariesInAPK)
 			{
-				foreach (string NDKArch in NDKArches)
+				foreach ((UnrealArch Arch, string _) in Archs)
 				{
-					UnrealArch Arch = GetUnrealArch(NDKArch);
-
 					string SOName = AndroidToolChain.InlineArchName(OutputPath, Arch);
 					if (!File.Exists(SOName))
 					{
@@ -5050,7 +5039,7 @@ dependencies.implementation('androidx.tracing:tracing:1.2.0')
 				ParseFilterFile(Path.Combine(UnrealBuildPath, "GradleFilter.txt"));
 
 				//Generate Gradle AAR dependencies
-				GenerateGradleAARImports(EngineDirectory, UnrealBuildPath, NDKArches);
+				GenerateGradleAARImports(EngineDirectory, UnrealBuildPath, Archs.Select(x => x.NativeArch));
 
 				//Now validate GooglePlay app_id if enabled
 				ValidateGooglePlay(UnrealBuildPath);
@@ -5058,7 +5047,7 @@ dependencies.implementation('androidx.tracing:tracing:1.2.0')
 				//determine which orientation requirements this app has
 				bool bNeedLandscape = false;
 				bool bNeedPortrait = false;
-				DetermineScreenOrientationRequirements(NDKArches[0], out bNeedPortrait, out bNeedLandscape);
+				DetermineScreenOrientationRequirements(Archs[0].NativeArch, out bNeedPortrait, out bNeedLandscape);
 
 				//Now keep the splash screen images matching orientation requested
 				PickSplashScreenOrientation(CookFlavor, UnrealBuildPath, bNeedPortrait, bNeedLandscape);
@@ -5715,7 +5704,7 @@ popd
 
 				if (bCombinedBundleOK)
 				{
-					string NDKArch = NDKArches[0];
+					string NDKArch = Archs[0].NativeArch;
 
 					string UnrealBuildGradlePath = UnrealGradleDest;
 
@@ -6454,7 +6443,7 @@ import java.util.Collection;
 			}
 		}
 
-		private AndroidAARHandler CreateAARHandler(string EngineDir, string UnrealBuildPath, List<string> NDKArches, bool HandleDependencies = true)
+		private AndroidAARHandler CreateAARHandler(string EngineDir, string UnrealBuildPath, IEnumerable<string> NDKArches, bool HandleDependencies = true)
 		{
 			AndroidAARHandler AARHandler = new AndroidAARHandler();
 			string ImportList = "";
@@ -6681,7 +6670,7 @@ import java.util.Collection;
 			File.WriteAllText(ProjectsGradleFilename, ProjectDependencyContent.ToString());
 		}
 
-		private void GenerateGradleAARImports(string EngineDir, string UnrealBuildPath, List<string> NDKArches)
+		private void GenerateGradleAARImports(string EngineDir, string UnrealBuildPath, IEnumerable<string> NDKArches)
 		{
 			AndroidAARHandler AARHandler = CreateAARHandler(EngineDir, UnrealBuildPath, NDKArches, false);
 			StringBuilder AARImportsContent = new StringBuilder();
diff --git a/Engine/Source/Programs/UnrealBuildTool/Platform/IOS/UEDeployIOS.cs b/Engine/Source/Programs/UnrealBuildTool/Platform/IOS/UEDeployIOS.cs
--- a/Engine/Source/Programs/UnrealBuildTool/Platform/IOS/UEDeployIOS.cs
+++ b/Engine/Source/Programs/UnrealBuildTool/Platform/IOS/UEDeployIOS.cs
@@ -989,8 +989,7 @@ namespace UnrealBuildTool
 			// strip out the markup
 			GameName = GameName.Split("-".ToCharArray())[0];
 
-			List<string> ProjectArches = new List<string>();
-			ProjectArches.Add("None");
+			List<(string UEArch, string NativeArch)> ProjectArches = new() { ("None", "None") };
 
 			string BundlePath;
 
@@ -1011,7 +1010,7 @@ namespace UnrealBuildTool
 			UnrealPluginLanguage UPL = new UnrealPluginLanguage(ProjectFile, UPLScripts, ProjectArches, null, UnrealTargetPlatform.IOS, Logger);
 
 			// Passing in true for distribution is not ideal here but given the way that ios packaging happens and this call chain it seems unavoidable for now, maybe there is a way to correctly pass it in that I can't find?
-			UPL.Init(ProjectArches, true, RelativeEnginePath, BundlePath, ProjectDirectory, Config.ToString(), false);
+			UPL.Init(ProjectArches.Select(x => x.NativeArch), true, RelativeEnginePath, BundlePath, ProjectDirectory, Config.ToString(), false);
 
 			return GenerateIOSPList(ProjectFile, Config, ProjectDirectory, bIsUnrealGame, GameName, bIsClient, ProjectName, InEngineDir, AppDirectory, UPL, BundleID, bBuildAsFramework, Logger);
 		}
diff --git a/Engine/Source/Programs/UnrealBuildTool/System/UnrealPluginLanguage.cs b/Engine/Source/Programs/UnrealBuildTool/System/UnrealPluginLanguage.cs
--- a/Engine/Source/Programs/UnrealBuildTool/System/UnrealPluginLanguage.cs
+++ b/Engine/Source/Programs/UnrealBuildTool/System/UnrealPluginLanguage.cs
@@ -47,6 +47,7 @@ namespace UnrealBuildTool
 	 * The following variables are initialized automatically:
 	 * 
 	 *	$S(Output) = the output returned for evaluating the section (initialized to Input)
+	 *	$S(UnrealArch) = target Unreal arch (arm32, arm64, x86, x64)
 	 *	$S(Architecture) = target architecture (armeabi-armv7a, arm64-v8a, x86, x86_64)
 	 *	$S(PluginDir) = directory the XML file was loaded from
 	 *	$S(EngineDir) = engine directory
@@ -537,7 +538,7 @@ namespace UnrealBuildTool
 			public bool bTrace;
 			public bool bTraceFilename;
 
-			public UPLContext(string inFilename, string Architecture, string PluginDir)
+			public UPLContext(string inFilename, (string UEArch, string NativeArch) Arch, string PluginDir)
 			{
 				Filename = inFilename;
 
@@ -553,7 +554,8 @@ namespace UnrealBuildTool
 				StringVariables["PluginDir"] = PluginDir.Replace("\\", "/");
 				StringVariables["AbsPluginDir"] = Path.GetFullPath(PluginDir).Replace("\\", "/");
 
-				StringVariables["Architecture"] = Architecture;
+				StringVariables["UnrealArch"] = Arch.UEArch;
+				StringVariables["Architecture"] = Arch.NativeArch;
 
 				bTrace = false;
 				bTraceFilename = false;
@@ -567,7 +569,7 @@ namespace UnrealBuildTool
 		private ILogger Logger;
 
 		// First entry in InXMLNameSpaceAliases will be used as default namespace
-		public UnrealPluginLanguage(FileReference? InProjectFile, List<string> InXMLFiles, List<string> InArchitectures, List<AliasedXMLNamespace>? InXMLNameSpaceAliases, UnrealTargetPlatform InTargetPlatform, ILogger InLogger)
+		public UnrealPluginLanguage(FileReference? InProjectFile, List<string> InXMLFiles, List<(string UEArch, string NativeArch)> InArchs, List<AliasedXMLNamespace>? InXMLNameSpaceAliases, UnrealTargetPlatform InTargetPlatform, ILogger InLogger)
 		{
 			ProjectFile = InProjectFile;
 			Logger = InLogger;
@@ -575,7 +577,7 @@ namespace UnrealBuildTool
 			LastError = null;
 
 			Contexts = new Dictionary<string, UPLContext>();
-			GlobalContext = new UPLContext("", "", "");
+			GlobalContext = new UPLContext("", ("", ""), "");
 			ContextIndex = 0;
 			
 			if (InXMLNameSpaceAliases != null && InXMLNameSpaceAliases.Count > 0)
@@ -611,7 +613,7 @@ namespace UnrealBuildTool
 					try
 					{
 						XDocument MergeDoc = XDocument.Load(Filename, LoadOptions.SetLineInfo);
-						MergeXML(MergeDoc, Filename, PluginDir, InArchitectures);
+						MergeXML(MergeDoc, Filename, PluginDir, InArchs);
 					}
 					catch (Exception e)
 					{
@@ -642,7 +644,7 @@ namespace UnrealBuildTool
 			return ContentHash.MD5(XDoc.ToString()).ToString();
 		}
 
-		public bool MergeXML(XDocument MergeDoc, string Filename, string PluginDir, List<string> Architectures)
+		public bool MergeXML(XDocument MergeDoc, string Filename, string PluginDir, List<(string UEArch, string NativeArch)> Archs)
 		{
 			if (MergeDoc == null)
 			{
@@ -651,10 +653,10 @@ namespace UnrealBuildTool
 
 			// create a context for each architecture
 			ContextIndex++;
-			foreach (string Architecture in Architectures)
+			foreach ((string UEArch, string NativeArch) Arch in Archs)
 			{
-				UPLContext Context = new UPLContext(Filename, Architecture, PluginDir);
-				Contexts[Architecture + "_" + ContextIndex] = Context;
+				UPLContext Context = new UPLContext(Filename, Arch, PluginDir);
+				Contexts[Arch.NativeArch + "_" + ContextIndex] = Context;
 			}
 
 			// line numbers are lost in the merge so add them as an attribute for later tracing of warnings/errors
@@ -2632,7 +2634,7 @@ namespace UnrealBuildTool
 			return GlobalContext.StringVariables["Output"];
 		}
 
-		public void Init(List<string> Architectures, bool bDistribution, string EngineDirectory, string BuildDirectory, string ProjectDirectory, string Configuration, bool bIsEmbedded,
+		public void Init(IEnumerable<string> Architectures, bool bDistribution, string EngineDirectory, string BuildDirectory, string ProjectDirectory, string Configuration, bool bIsEmbedded,
 			bool bPerArchBuildDir = false, Dictionary<string, string>? ArchRemapping = null)
 		{
 			GlobalContext.BoolVariables["Distribution"] = bDistribution;
