UnrealBuildTool: Don't use changelist as Mac -current_version unless a formal or precompiled build is being made, to prevent unnecessary rebuilds

#rnx

diff --git a/Engine/Source/Programs/UnrealBuildTool/Platform/Mac/MacToolChain.cs b/Engine/Source/Programs/UnrealBuildTool/Platform/Mac/MacToolChain.cs
--- a/Engine/Source/Programs/UnrealBuildTool/Platform/Mac/MacToolChain.cs
+++ b/Engine/Source/Programs/UnrealBuildTool/Platform/Mac/MacToolChain.cs
@@ -309,7 +309,7 @@ namespace UnrealBuildTool
 			BuildVersion? Version;
 			if (BuildVersion.TryRead(BuildVersion.GetDefaultFileName(), out Version))
 			{
-				return String.Format("{0}.{1}.{2}", Version.MajorVersion, Version.MinorVersion, bIgnorePatchVersion ? 0 : Version.PatchVersion);
+				return $"{Version.MajorVersion}.{Version.MinorVersion}.{(bIgnorePatchVersion ? 0 : Version.PatchVersion)}";
 			}
 			else
 			{
@@ -324,15 +324,15 @@ namespace UnrealBuildTool
 
 		private string LoadEngineAPIVersion()
 		{
-			int CL = 0;
-
-			BuildVersion? Version;
-			if (BuildVersion.TryRead(BuildVersion.GetDefaultFileName(), out Version))
+			if (BuildVersion.TryRead(BuildVersion.GetDefaultFileName(), out BuildVersion? Version))
 			{
-				CL = Version.EffectiveCompatibleChangelist;
+				if ((Target?.bPrecompile == true || Version.IsPromotedBuild) && Version.EffectiveCompatibleChangelist > 0)
+				{
+					int CL = Version.EffectiveCompatibleChangelist;
+					return $"{CL / (100 * 100)}.{(CL / 100) % 100}.{CL % 100}";
+				}
 			}
-
-			return String.Format("{0}.{1}.{2}", CL / (100 * 100), (CL / 100) % 100, CL % 100);
+			return LoadEngineDisplayVersion(true);
 		}
 
 		private string GetGameNameFromExecutablePath(string ExePath)
