UnrealBuildTool: Suppress -Wnontrivial-memaccess for Clang 20 by default
Fix some typos

#rnx

diff --git a/Engine/Source/Programs/Shared/UnrealBuildTool.Tests/CppCompileWarningsTests.cs b/Engine/Source/Programs/Shared/UnrealBuildTool.Tests/CppCompileWarningsTests.cs
--- a/Engine/Source/Programs/Shared/UnrealBuildTool.Tests/CppCompileWarningsTests.cs
+++ b/Engine/Source/Programs/Shared/UnrealBuildTool.Tests/CppCompileWarningsTests.cs
@@ -511,7 +511,7 @@ namespace UnrealBuildTool.Tests
 				// AndroidNDKR28ToolChainVersionExclusion
 				{
 					testTargetRules.WindowsPlatform.Compiler = WindowsCompiler.Clang;
-					baseWarnings.CastFunctionTypeMistmatchWarningLevel = WarningLevel.Off;
+					baseWarnings.CastFunctionTypeMismatchWarningLevel = WarningLevel.Off;
 
 					VersionNumber newVersionNumber = new VersionNumber(19, 0, 0);
 					CppCompileEnvironment androidCompileEnvironment = new CppCompileEnvironment(UnrealTargetPlatform.Android, CppConfiguration.Development, new UnrealArchitectures(UnrealArch.X64), null!);
@@ -530,7 +530,7 @@ namespace UnrealBuildTool.Tests
 					Assert.IsTrue(args.Count() == 1);
 					Assert.IsTrue(args.Contains("-Wno-cast-function-type-mismatch"));
 
-					baseWarnings.CastFunctionTypeMistmatchWarningLevel = WarningLevel.Default;
+					baseWarnings.CastFunctionTypeMismatchWarningLevel = WarningLevel.Default;
 				}
 
 				// AndroidNDKR26ToolChainVersionExclusion
diff --git a/Engine/Source/Programs/UnrealBuildTool/Configuration/CompileWarnings/ApplyWarningsFilterAttribute.cs b/Engine/Source/Programs/UnrealBuildTool/Configuration/CompileWarnings/ApplyWarningsFilterAttribute.cs
--- a/Engine/Source/Programs/UnrealBuildTool/Configuration/CompileWarnings/ApplyWarningsFilterAttribute.cs
+++ b/Engine/Source/Programs/UnrealBuildTool/Configuration/CompileWarnings/ApplyWarningsFilterAttribute.cs
@@ -50,6 +50,11 @@ namespace UnrealBuildTool
 		/// </summary>
 		Version19Min,
 
+		/// <summary>
+		/// Filter ID for <see cref="ToolChainVersionConstrainedFilterAttribute"/>, specifically for <see cref="VersionNumber"/> 12.0 as a min, with no upper bound.
+		/// </summary>
+		Version20Min,
+
 		/// <summary>
 		/// Filter ID for <see cref="CppStandardConstrainedFilterAttribute"/>, specifically for <see cref="CppStandardVersion.Cpp20"/> as a minimum, with no upper bound.
 		/// </summary>
diff --git a/Engine/Source/Programs/UnrealBuildTool/Configuration/CompileWarnings/ApplyWarningsFilterRegistry.cs b/Engine/Source/Programs/UnrealBuildTool/Configuration/CompileWarnings/ApplyWarningsFilterRegistry.cs
--- a/Engine/Source/Programs/UnrealBuildTool/Configuration/CompileWarnings/ApplyWarningsFilterRegistry.cs
+++ b/Engine/Source/Programs/UnrealBuildTool/Configuration/CompileWarnings/ApplyWarningsFilterRegistry.cs
@@ -35,6 +35,7 @@ namespace UnrealBuildTool
 	[ToolChainVersionConstrainedFilter(FilterID.Version16Min, [16])]
 	[ToolChainVersionConstrainedFilter(FilterID.Version18Min, [18])]
 	[ToolChainVersionConstrainedFilter(FilterID.Version19Min, [19])]
+	[ToolChainVersionConstrainedFilter(FilterID.Version20Min, [20])]
 	[CppStandardConstrainedFilter(FilterID.Cpp20Min, CppStandardVersion.Cpp20)]
 	internal class ApplyWarningsFilterRegistry
 	{
diff --git a/Engine/Source/Programs/UnrealBuildTool/Configuration/CppCompileWarnings.cs b/Engine/Source/Programs/UnrealBuildTool/Configuration/CppCompileWarnings.cs
--- a/Engine/Source/Programs/UnrealBuildTool/Configuration/CppCompileWarnings.cs
+++ b/Engine/Source/Programs/UnrealBuildTool/Configuration/CppCompileWarnings.cs
@@ -249,7 +249,7 @@ namespace UnrealBuildTool
 		private WarningLevel _deprecatedVolatileWarningLevel;
 
 		/// <summary>
-		/// Indicates what warning/error level to treat compairing function pointers with relational operators (&lt;,&gt;, etc), which is undefined behaviour.
+		/// Indicates what warning/error level to treat compairing function pointers with relational operators (&lt;,&gt;, etc), which is undefined behavior.
 		/// Clang -
 		///		https://clang.llvm.org/docs/DiagnosticsReference.html#wordered-compare-function-pointers
 		/// </summary>
@@ -319,7 +319,7 @@ namespace UnrealBuildTool
 		private WarningLevel _invalidUnevaluatedStringWarningLevel;
 
 		/// <summary>
-		/// Indicates what warning/error level to treat infinity macros, and subsequent undefined behaviour.
+		/// Indicates what warning/error level to treat infinity macros, and subsequent undefined behavior.
 		/// Clang -
 		///		https://clang.llvm.org/docs/DiagnosticsReference.html#wnan-infinity-disabled
 		/// </summary>
@@ -352,18 +352,18 @@ namespace UnrealBuildTool
 		private WarningLevel _levelExtraQualificationWarningLevel;
 
 		/// <summary>
-		/// Indicates what warning/error level to treat function pointer casts to an incompatible function type, potentially leading to undefined behaviour.
+		/// Indicates what warning/error level to treat function pointer casts to an incompatible function type, potentially leading to undefined behavior.
 		/// Clang -
 		///		https://clang.llvm.org/docs/DiagnosticsReference.html#wcast-function-type-mismatch
 		/// </summary>
 		[WarningsClangToolChain(["-Wno-cast-function-type-mismatch"], null, null, [nameof(SpecializedFilters.AndroidNDKR28ToolChainVersionExclusion)])]
 		[WarningLevelDefault(WarningLevel.Off)]
-		public WarningLevel CastFunctionTypeMistmatchWarningLevel
+		public WarningLevel CastFunctionTypeMismatchWarningLevel
 		{
-			get => ResolveWarning(ParentCppCompileWarnings?.CastFunctionTypeMistmatchWarningLevel, _castFunctionTypeMistmatchWarningLevel);
-			set => _castFunctionTypeMistmatchWarningLevel = value;
+			get => ResolveWarning(ParentCppCompileWarnings?.CastFunctionTypeMismatchWarningLevel, _castFunctionTypeMismatchWarningLevel);
+			set => _castFunctionTypeMismatchWarningLevel = value;
 		}
-		private WarningLevel _castFunctionTypeMistmatchWarningLevel;
+		private WarningLevel _castFunctionTypeMismatchWarningLevel;
 
 		/// <summary>
 		/// Indicates what warning/error level to treat templates used without specifying the required template arguments, such as missing parameter packs (&lt;...&gt;) after the template keyword.
@@ -409,7 +409,7 @@ namespace UnrealBuildTool
 		private WarningLevel _inconsistentMissingOverrideWarningLevel;
 
 		/// <summary>
-		/// Indicates what warning/error level to treat incorrect usages of std::offsetof, such as with non-POD types, where behaviour is undefined.
+		/// Indicates what warning/error level to treat incorrect usages of std::offsetof, such as with non-POD types, where behavior is undefined.
 		/// Clang -
 		///		https://clang.llvm.org/docs/DiagnosticsReference.html#winvalid-offsetof
 		/// </summary>
@@ -428,7 +428,7 @@ namespace UnrealBuildTool
 		/// Clang -
 		///		https://clang.llvm.org/docs/DiagnosticsReference.html#wswitch
 		/// </summary>
-		/// <remarks>Fhis hides the "enumeration value 'XXXXX' not handled in switch [-Wswitch]" warnings - we should maybe remove this at some point and add UE_LOG(, Fatal, ) to default cases.</remarks>
+		/// <remarks>This hides the "enumeration value 'XXXXX' not handled in switch [-Wswitch]" warnings - we should maybe remove this at some point and add UE_LOG(, Fatal, ) to default cases.</remarks>
 		[WarningsClangToolChain(["-Wno-switch"])]
 		[WarningLevelDefault(WarningLevel.Off)]
 		public WarningLevel SwitchWarningLevel
@@ -587,6 +587,20 @@ namespace UnrealBuildTool
 		}
 		private WarningLevel _dllExportExplicitInstantiationDeclWarningLevel;
 
+		/// <summary>
+		/// Indicates what warning/error level to treat nontrivial memory access.
+		/// Clang -
+		///		https://clang.llvm.org/docs/DiagnosticsReference.html#wnontrivial-memaccess
+		/// </summary>
+		[WarningsClangToolChain(["-Wno-nontrivial-memaccess"], null, null, [nameof(FilterID.Version20Min)])]
+		[WarningLevelDefault(WarningLevel.Off)]
+		public WarningLevel NonTrivialMemAccessWarningLevel
+		{
+			get => ResolveWarning(ParentCppCompileWarnings?.NonTrivialMemAccessWarningLevel, _nonTrivialMemAccessWarningLevel);
+			set => _nonTrivialMemAccessWarningLevel = value;
+		}
+		private WarningLevel _nonTrivialMemAccessWarningLevel;
+
 		#endregion
 
 		#region -- Clang VC Compiler Disables --
@@ -730,7 +744,7 @@ namespace UnrealBuildTool
 		private WarningLevel _logicalOpParenthesesWarningLevel;
 
 		/// <summary>
-		/// Indicates what warning/error level to treat occurrences where arithmetic oeprations are performed on nullptr.
+		/// Indicates what warning/error level to treat occurrences where arithmetic operations are performed on nullptr.
 		/// Clang -
 		///		https://clang.llvm.org/docs/DiagnosticsReference.html#wnull-arithmetic
 		/// </summary>
@@ -992,7 +1006,7 @@ namespace UnrealBuildTool
 		/// <remarks>
 		/// This is a special case where we have a duplicate field for clang <see cref="CppCompileWarnings.SwitchUnhandledEnumeratorWarningLevel"/>.
 		/// The reason for duplication is that the clang field will only participate in the shipping configuration, whereas the MSVC variant has been disabled in all.
-		/// In order to ensure no change of behaviour for MSVC toolchain within the clang compilation context, we will keep this in, and deprecate with assignments to the clang setting.
+		/// In order to ensure no change of behavior for MSVC toolchain within the clang compilation context, we will keep this in, and deprecate with assignments to the clang setting.
 		/// </remarks>
 		[WarningsVCClang(["-Wno-switch-enum"])]
 		[WarningLevelDefault(WarningLevel.Off)]
@@ -1011,7 +1025,7 @@ namespace UnrealBuildTool
 		/// <remarks>
 		/// This is a special case where we have a duplicate field for clang <see cref="CppCompileWarnings.UnusedValueWarningLevel"/>.
 		/// The reason for duplication is that the clang field will only participate in the shipping configuration, whereas the MSVC variant has been disabled in all.
-		/// In order to ensure no change of behaviour for MSVC toolchain within the clang compilation context, we will keep this in, and deprecate with assignments to the clang setting.
+		/// In order to ensure no change of behavior for MSVC toolchain within the clang compilation context, we will keep this in, and deprecate with assignments to the clang setting.
 		/// -Wunused-result - we will always re-enable this as it's controlled by the same set. This will be moved out to a separate property. It is on by default, so it only must be re-enabled in the disable context.
 		/// </remarks>
 		[WarningsVCClang(["-Wno-unused-value", "-Wunused-result"])]
@@ -1031,7 +1045,7 @@ namespace UnrealBuildTool
 		/// <remarks>
 		/// This is a special case where we have a duplicate field for clang <see cref="CppCompileWarnings.DeprecationWarningLevel"/>.
 		/// The reason for duplication is that the clang field will be defaulted to WARNING in clang, but not for msvc which is disabled.
-		/// In order to ensure no change of behaviour for MSVC toolchain within the clang compilation context, we will keep this in, and deprecate with assignments to the clang setting.
+		/// In order to ensure no change of behavior for MSVC toolchain within the clang compilation context, we will keep this in, and deprecate with assignments to the clang setting.
 		/// </remarks>
 		[WarningsVCClang(["-Wno-deprecated-declarations"])]
 		[WarningLevelDefault(WarningLevel.Off)]
@@ -1303,14 +1317,14 @@ namespace UnrealBuildTool
 		{
 			if (!typeof(UEToolChain).IsAssignableFrom(toolChainType))
 			{
-				_logger?.LogWarning("Requested command line arugments of a toolchain type ({TypeName}) that is not a part of the UEToolChain hierarchy.", toolChainType.Name);
+				_logger?.LogWarning("Requested command line arguments of a toolchain type ({TypeName}) that is not a part of the UEToolChain hierarchy.", toolChainType.Name);
 
-				return new List<string>(0);
+				return [];
 			}
 
-			CompilerWarningsToolChainContext toolChainContext = new CompilerWarningsToolChainContext(compileEnvironment, BuildContextProvider, toolChainType, toolChainVersion, analyzer); //pull out into request?
+			CompilerWarningsToolChainContext toolChainContext = new(compileEnvironment, BuildContextProvider, toolChainType, toolChainVersion, analyzer); //pull out into request?
 
-			List<string> warningArguments = new List<string>();
+			List<string> warningArguments = [];
 
 			s_membersWithApplyWarningsAttribute ??= GetMembersWithAttributes(typeof(CppCompileWarnings), typeof(ApplyWarningsAttribute), typeof(WarningLevel));
 			foreach (MemberInfo memberInfo in s_membersWithApplyWarningsAttribute)
@@ -1355,7 +1369,7 @@ namespace UnrealBuildTool
 
 			foreach (MemberInfo mi in memberInfos)
 			{
-				IList<WarningLevelDefaultAttribute> attributes = mi.GetCustomAttributes<WarningLevelDefaultAttribute>().Where(prop => prop.Context == InitializationContext.Any || prop.Context == InitializationContext.Target).ToList();
+				IList<WarningLevelDefaultAttribute> attributes = [.. mi.GetCustomAttributes<WarningLevelDefaultAttribute>().Where(prop => prop.Context == InitializationContext.Any || prop.Context == InitializationContext.Target)];
 				WarningLevel defaultWarningLevel = WarningLevelDefaultAttribute.ResolveWarningLevelDefault(attributes, buildSettingsVersion);
 
 				// Obtain instance value, and see if we need to preserve the TargetRules' constructor setting, or apply the default.
@@ -1403,7 +1417,7 @@ namespace UnrealBuildTool
 
 			foreach (MemberInfo mi in memberInfos)
 			{
-				IList<WarningLevelDefaultAttribute> attributes = mi.GetCustomAttributes<WarningLevelDefaultAttribute>().Where(prop => prop.Context == InitializationContext.Any || prop.Context == InitializationContext.Constructor).ToList();
+				IList<WarningLevelDefaultAttribute> attributes = [.. mi.GetCustomAttributes<WarningLevelDefaultAttribute>().Where(prop => prop.Context == InitializationContext.Any || prop.Context == InitializationContext.Constructor)];
 				WarningLevel defaultWarningLevel = WarningLevelDefaultAttribute.ResolveWarningLevelDefault(attributes, buildSettingsVersion);
 
 				SetWarningLevel(mi, this, defaultWarningLevel);
diff --git a/Engine/Source/Programs/UnrealBuildTool/Configuration/ReadOnlyCppCompileWarnings.cs b/Engine/Source/Programs/UnrealBuildTool/Configuration/ReadOnlyCppCompileWarnings.cs
--- a/Engine/Source/Programs/UnrealBuildTool/Configuration/ReadOnlyCppCompileWarnings.cs
+++ b/Engine/Source/Programs/UnrealBuildTool/Configuration/ReadOnlyCppCompileWarnings.cs
@@ -66,13 +66,14 @@ namespace UnrealBuildTool
 		public WarningLevel DeprecatedCopyWarningLevel => _inner.DeprecatedCopyWarningLevel;
 		public WarningLevel InvalidUnevaluatedStringWarningLevel => _inner.InvalidUnevaluatedStringWarningLevel;
 		public WarningLevel NaNInfinityDisabledWarningLevel => _inner.NaNInfinityDisabledWarningLevel;
-		public WarningLevel CastFunctionTypeMistmatchWarningLevel => _inner.CastFunctionTypeMistmatchWarningLevel;
+		public WarningLevel CastFunctionTypeMismatchWarningLevel => _inner.CastFunctionTypeMismatchWarningLevel;
 		public WarningLevel MissingTemplateArgListAfterTemplateWarningLevel => _inner.MissingTemplateArgListAfterTemplateWarningLevel;
 		public WarningLevel LevelExtraQualificationWarningLevel => _inner.LevelExtraQualificationWarningLevel;
 		public WarningLevel ProfileInstructWarningLevel => _inner.ProfileInstructWarningLevel;
 		public WarningLevel BackendPluginWarningLevel => _inner.BackendPluginWarningLevel;
 		public WarningLevel Shorten64To32WarningLevel => _inner.Shorten64To32WarningLevel;
 		public WarningLevel DllExportExplicitInstantiationDeclWarningLevel => _inner.DllExportExplicitInstantiationDeclWarningLevel;
+		public WarningLevel NonTrivialMemAccessWarningLevel => _inner.NonTrivialMemAccessWarningLevel;
 		public WarningLevel DeprecationWarningLevel => _inner.DeprecationWarningLevel;
 		public WarningLevel DeterministicWarningLevel => _inner.DeterministicWarningLevel;
 
diff --git a/Engine/Source/Runtime/AutoRTFM/AutoRTFM.Build.cs b/Engine/Source/Runtime/AutoRTFM/AutoRTFM.Build.cs
--- a/Engine/Source/Runtime/AutoRTFM/AutoRTFM.Build.cs
+++ b/Engine/Source/Runtime/AutoRTFM/AutoRTFM.Build.cs
@@ -33,7 +33,7 @@ namespace UnrealBuildTool.Rules
 			CppCompileWarningSettings.InvalidUnevaluatedStringWarningLevel = WarningLevel.Error;
 			CppCompileWarningSettings.NaNInfinityDisabledWarningLevel = WarningLevel.Error;
 			CppCompileWarningSettings.LevelExtraQualificationWarningLevel = WarningLevel.Error;
-			CppCompileWarningSettings.CastFunctionTypeMistmatchWarningLevel = WarningLevel.Error;
+			CppCompileWarningSettings.CastFunctionTypeMismatchWarningLevel = WarningLevel.Error;
 			CppCompileWarningSettings.MissingTemplateArgListAfterTemplateWarningLevel = WarningLevel.Error;
 			CppCompileWarningSettings.GNUStringLiteralOperatorTemplateWarningLevel = WarningLevel.Error;
 			CppCompileWarningSettings.InconsistentMissingOverrideWarningLevel = WarningLevel.Error;
