Enable compression of debug symbols for Android and add a flag to toggle it

#android
#PR #11752

diff --git a/Engine/Config/BaseEngine.ini b/Engine/Config/BaseEngine.ini
--- a/Engine/Config/BaseEngine.ini
+++ b/Engine/Config/BaseEngine.ini
@@ -3280,6 +3280,7 @@ bEnableNewKeyboard=True
 bAndroidVoiceEnabled=false
 bRecordPermissionAtStartupEnabled=false
 bBuildWithHiddenSymbolVisibility=True
+bCompressDebugSymbols=True
 bSaveSymbols=false
 bAllowControllers=True
 bAllowIMU=False
diff --git a/Engine/Source/Programs/UnrealBuildTool/Platform/Android/AndroidToolChain.cs b/Engine/Source/Programs/UnrealBuildTool/Platform/Android/AndroidToolChain.cs
--- a/Engine/Source/Programs/UnrealBuildTool/Platform/Android/AndroidToolChain.cs
+++ b/Engine/Source/Programs/UnrealBuildTool/Platform/Android/AndroidToolChain.cs
@@ -277,6 +277,12 @@ namespace UnrealBuildTool
 			return Ini.GetBool("/Script/AndroidRuntimeSettings.AndroidRuntimeSettings", "bBuildWithHiddenSymbolVisibility", out bool bBuild) && bBuild;
 		}
 
+		private bool CompressDebugSymbols()
+		{
+			ConfigHierarchy Ini = ConfigCache.ReadHierarchy(ConfigHierarchyType.Engine, DirectoryReference.FromFile(ProjectFile), UnrealTargetPlatform.Android);
+			return Ini.GetBool("/Script/AndroidRuntimeSettings.AndroidRuntimeSettings", "bCompressDebugSymbols", out bool bBuild) && bBuild;
+		}
+		
 		private bool DisableFunctionDataSections()
 		{
 			ConfigHierarchy Ini = ConfigCache.ReadHierarchy(ConfigHierarchyType.Engine, DirectoryReference.FromFile(ProjectFile), UnrealTargetPlatform.Android);
@@ -592,6 +598,11 @@ namespace UnrealBuildTool
 				Arguments.Add("-g2");
 				Arguments.Add("-gdwarf-4");
 
+				if (CompressDebugSymbols())
+				{
+					Arguments.Add("-gz");
+				}
+
 				if (CompileEnvironment.bDebugLineTablesOnly)
 				{
 					Arguments.Add("-gline-tables-only");
@@ -739,6 +750,10 @@ namespace UnrealBuildTool
 			{
 				Result += " -Wl,--strip-debug";
 			}
+			else if (CompressDebugSymbols())
+			{
+				Result += " -gz";
+			}
 
 			if (Architecture == UnrealArch.X64)
 			{
diff --git a/Engine/Source/Runtime/Android/AndroidRuntimeSettings/Classes/AndroidRuntimeSettings.h b/Engine/Source/Runtime/Android/AndroidRuntimeSettings/Classes/AndroidRuntimeSettings.h
--- a/Engine/Source/Runtime/Android/AndroidRuntimeSettings/Classes/AndroidRuntimeSettings.h
+++ b/Engine/Source/Runtime/Android/AndroidRuntimeSettings/Classes/AndroidRuntimeSettings.h
@@ -474,6 +474,9 @@ public:
 	UPROPERTY(GlobalConfig, EditAnywhere, Category = AdvancedBuild, meta = (DisplayName = "Build with hidden symbol visibility in shipping config."))
 	bool bBuildWithHiddenSymbolVisibility;
 
+	UPROPERTY(GlobalConfig, EditAnywhere, Category = AdvancedBuild, meta = (DisplayName = "Enable compression of debug symbols"))
+	bool bCompressDebugSymbols;
+	
 	// Disables extra checks for buffer overflows, comes with perf improvement, but might make tracing stack corruptions in production harder. Note that _FORTIFY_SOURCE=2 is still enabled by the toolchain providing lightweight stack checks
 	UPROPERTY(GlobalConfig, EditAnywhere, Category = AdvancedBuild, meta = (DisplayName = "Disable extra checks for buffer overflows"))
 	bool bDisableStackProtector;
