Add FPass-Plugin support to UBT

This adds the functionality for a module to specify an LLVM FPass-Plugin to be included in the clang compile arguments and run druing the compilation.
To use, set the FPassPlugin value from a module's Build.CS file to point to the path of a valid fpass-plugin for the host compiler's architecture.

#Tests Included a path to an fpass plugin in a module's Build.CS file. Verified that the FPassPlugin was correctly applied and the engine could build/run.

diff --git a/Engine/Source/Programs/UnrealBuildTool/Configuration/ModuleRules.cs b/Engine/Source/Programs/UnrealBuildTool/Configuration/ModuleRules.cs
--- a/Engine/Source/Programs/UnrealBuildTool/Configuration/ModuleRules.cs
+++ b/Engine/Source/Programs/UnrealBuildTool/Configuration/ModuleRules.cs
@@ -1571,6 +1571,12 @@ namespace UnrealBuildTool
 		/// </summary>
 		public string? SwiftInteropHeader = null;
 
+		/// <summary>
+		/// A List of files to be used as a LLVM FPass-Plugins for all clang compilation units of the current module.
+		///    'clang -fpass-plugin=[FPassPlugin]' 
+		/// </summary>
+		public List<string> FPassPlugins = new();
+
 #nullable disable
 		/// <summary>
 		/// Constructor. For backwards compatibility while the parameterless constructor is being phased out, initialization which would happen here is done by 
diff --git a/Engine/Source/Programs/UnrealBuildTool/Configuration/UEBuildModuleCPP.cs b/Engine/Source/Programs/UnrealBuildTool/Configuration/UEBuildModuleCPP.cs
--- a/Engine/Source/Programs/UnrealBuildTool/Configuration/UEBuildModuleCPP.cs
+++ b/Engine/Source/Programs/UnrealBuildTool/Configuration/UEBuildModuleCPP.cs
@@ -2325,6 +2325,7 @@ namespace UnrealBuildTool
 			Result.bValidateInternalApi = Rules.bValidateInternalApi;
 			Result.bValidateExperimentalApi = Rules.bValidateExperimentalApi;
 			Result.bUseAutoRTFMCompiler = Target.bUseAutoRTFMCompiler;
+			Result.FPassPlugins = Rules.FPassPlugins;
 
 			CompileEnvironmentDebugInfoSettings(Target, Result);
 
diff --git a/Engine/Source/Programs/UnrealBuildTool/System/CppCompileEnvironment.cs b/Engine/Source/Programs/UnrealBuildTool/System/CppCompileEnvironment.cs
--- a/Engine/Source/Programs/UnrealBuildTool/System/CppCompileEnvironment.cs
+++ b/Engine/Source/Programs/UnrealBuildTool/System/CppCompileEnvironment.cs
@@ -779,6 +779,12 @@ namespace UnrealBuildTool
 		/// </summary>
 		public string? CrashDiagnosticDirectory;
 
+		/// <summary>
+		/// File to use as an LLVM FPass Plugin:
+		///    'clang -fpass-plugin=[FPassPlugin]' 
+		/// </summary>
+		public List<string> FPassPlugins = new();
+
 		/// <summary>
 		/// Default constructor.
 		/// </summary>
@@ -899,6 +905,7 @@ namespace UnrealBuildTool
 			bValidateFormatStrings = Other.bValidateFormatStrings;
 			bValidateInternalApi = Other.bValidateInternalApi;
 			bValidateExperimentalApi = Other.bValidateExperimentalApi;
+			FPassPlugins = Other.FPassPlugins;
 		}
 
 		public CppCompileEnvironment(CppCompileEnvironment Other, UnrealArch OverrideArchitecture)
diff --git a/Engine/Source/Programs/UnrealBuildTool/ToolChain/ClangToolChain.cs b/Engine/Source/Programs/UnrealBuildTool/ToolChain/ClangToolChain.cs
--- a/Engine/Source/Programs/UnrealBuildTool/ToolChain/ClangToolChain.cs
+++ b/Engine/Source/Programs/UnrealBuildTool/ToolChain/ClangToolChain.cs
@@ -1111,6 +1111,17 @@ namespace UnrealBuildTool
 				OutputFile = FileItem.GetItemByFileReference(FileReference.Combine(OutputDir, GetFileNameFromExtension(FileName, ObjectFileExtension)));
 				CompileResult.ObjectFiles.Add(OutputFile);
 			}
+	
+			// FPassPlugin
+			foreach (string FPassPlugin in CompileEnvironment.FPassPlugins)
+			{
+				FileReference? FPassPluginFile = FileReference.FromString(FPassPlugin);
+    			if (FPassPluginFile != null && FileReference.Exists(FPassPluginFile))
+    			{
+    				Arguments.Add($"-fpass-plugin=\"{NormalizeCommandLinePath(FPassPluginFile, CompileEnvironment.RootPaths)}\"");
+    				CompileAction.PrerequisiteItems.Add(FileItem.GetItemByFileReference(FPassPluginFile));
+    			}			
+			}
 
 			// Add the output file to the produced item list.
 			CompileAction.ProducedItems.Add(OutputFile);
