UnrealBuildTool: Track installed SDK info in the makefile, so if the available sdks change the makefile will be invalidated

#rnx

diff --git a/Engine/Source/Programs/UnrealBuildTool/Configuration/UEBuildTarget.cs b/Engine/Source/Programs/UnrealBuildTool/Configuration/UEBuildTarget.cs
--- a/Engine/Source/Programs/UnrealBuildTool/Configuration/UEBuildTarget.cs
+++ b/Engine/Source/Programs/UnrealBuildTool/Configuration/UEBuildTarget.cs
@@ -2989,6 +2989,13 @@ namespace UnrealBuildTool
 			// Add any toolchain dependencies
 			TargetToolChain.GetExternalDependencies(Makefile.ExternalDependencies);
 
+			// Add any relevent platform sdk info
+			IEnumerable<UnrealTargetPlatform> SdkPlatforms = Rules.Type == TargetType.Editor || Rules.Type == TargetType.Program ? UnrealTargetPlatform.GetValidPlatforms() : [BuildHostPlatform.Current.Platform, Rules.Platform];
+			foreach (UnrealTargetPlatform platform in SdkPlatforms.Distinct())
+			{
+				Makefile.PlatformSDKInfo.Add(platform.ToString(), UEBuildPlatformSDK.GetSDKForPlatform(platform.ToString())?.GetAllSDKInfo().ToString() ?? String.Empty);
+			}
+
 			// Write a header containing public definitions for this target
 			if (Rules.ExportPublicHeader != null)
 			{
diff --git a/Engine/Source/Programs/UnrealBuildTool/System/TargetMakefile.cs b/Engine/Source/Programs/UnrealBuildTool/System/TargetMakefile.cs
--- a/Engine/Source/Programs/UnrealBuildTool/System/TargetMakefile.cs
+++ b/Engine/Source/Programs/UnrealBuildTool/System/TargetMakefile.cs
@@ -348,11 +348,6 @@ namespace UnrealBuildTool
 		public bool bUseVerseBPVM = true;
 #endif
 
-		/// <summary>
-		/// List of config settings in generated config files
-		/// </summary>
-		public Dictionary<string, string> ConfigSettings = new Dictionary<string, string>();
-
 		/// <summary>
 		/// List of all plugin names. The makefile will be considered invalid if any of these changes, or new plugins are added.
 		/// </summary>
@@ -368,6 +363,11 @@ namespace UnrealBuildTool
 		/// </summary>
 		public HashSet<FileItem> InternalDependencies;
 
+		/// <summary>
+		/// Collection of available Platform SDK info which will cause the makefile to be invalidated if modified.
+		/// </summary>
+		public Dictionary<string, string> PlatformSDKInfo;
+
 		/// <summary>
 		/// Set of paths that are disallowed to be transferred with UBA to remote machines
 		/// </summary>
@@ -460,6 +460,7 @@ namespace UnrealBuildTool
 			PluginFiles = new HashSet<FileItem>();
 			ExternalDependencies = new HashSet<FileItem>();
 			InternalDependencies = new HashSet<FileItem>();
+			PlatformSDKInfo = new Dictionary<string, string>();
 
 			UbaDisallowedPaths = [.. EpicGames.UBA.Utils.DisallowedPaths.Select(DirectoryItem.GetItemByPath)];
 			UbaCrossArchitecturePaths = new(EpicGames.UBA.Utils.CrossArchitecturePaths.Select(x => new KeyValuePair<DirectoryItem, DirectoryItem>(DirectoryItem.GetItemByPath(x.Key), DirectoryItem.GetItemByPath(x.Value))));
@@ -513,17 +514,19 @@ namespace UnrealBuildTool
 			PluginFiles = Reader.ReadHashSet(() => Reader.ReadFileItem())!;
 			ExternalDependencies = Reader.ReadHashSet(() => Reader.ReadFileItem())!;
 			InternalDependencies = Reader.ReadHashSet(() => Reader.ReadFileItem())!;
+			PlatformSDKInfo = Reader.ReadDictionary(() => Reader.ReadString()!, () => Reader.ReadString()!)!;
+
 			UbaDisallowedPaths = Reader.ReadHashSet(() => Reader.ReadDirectoryItem())!;
 			UbaCrossArchitecturePaths = Reader.ReadDictionary(() => Reader.ReadDirectoryItem()!, () => Reader.ReadDirectoryItem())!;
 			UbaPathHashes = Reader.ReadDictionary(() => Reader.ReadDirectoryItem()!, () => Reader.ReadString())!;
 			MemoryPerActionGB = Reader.ReadDouble();
 
-			EpicGames.UBA.Utils.RegisterDisallowedPaths(UbaDisallowedPaths.Select(x => x.FullName).ToArray());
-			foreach (KeyValuePair<DirectoryItem, DirectoryItem> item in UbaCrossArchitecturePaths)
+			EpicGames.UBA.Utils.RegisterDisallowedPaths([.. UbaDisallowedPaths.Where(x => x.Exists).Select(x => x.FullName)]);
+			foreach (KeyValuePair<DirectoryItem, DirectoryItem> item in UbaCrossArchitecturePaths.Where(x => x.Key.Exists && x.Value.Exists))
 			{
 				EpicGames.UBA.Utils.RegisterCrossArchitecturePath(item.Key.FullName, item.Value.FullName);
 			}
-			foreach (KeyValuePair<DirectoryItem, string> item in UbaPathHashes)
+			foreach (KeyValuePair<DirectoryItem, string> item in UbaPathHashes.Where(x => x.Key.Exists))
 			{
 				EpicGames.UBA.Utils.RegisterPathHash(item.Key.FullName, item.Value);
 			}
@@ -585,6 +588,8 @@ namespace UnrealBuildTool
 			Writer.WriteHashSet(PluginFiles, x => Writer.WriteFileItem(x));
 			Writer.WriteHashSet(ExternalDependencies, x => Writer.WriteFileItem(x));
 			Writer.WriteHashSet(InternalDependencies, x => Writer.WriteFileItem(x));
+			Writer.WriteDictionary(PlatformSDKInfo, k => Writer.WriteString(k), v => Writer.WriteString(v));
+
 			Writer.WriteHashSet(UbaDisallowedPaths, Writer.WriteDirectoryItem);
 			Writer.WriteDictionary(UbaCrossArchitecturePaths, Writer.WriteDirectoryItem, Writer.WriteDirectoryItem);
 			Writer.WriteDictionary(UbaPathHashes, Writer.WriteDirectoryItem, Writer.WriteString);
@@ -917,6 +922,18 @@ namespace UnrealBuildTool
 					}
 				}
 
+				// Check if available Platform SDKs have changed
+				foreach (KeyValuePair<string, string> Pair in Makefile.PlatformSDKInfo)
+				{
+					string value = UEBuildPlatformSDK.GetSDKForPlatform(Pair.Key)?.GetAllSDKInfo().ToString() ?? String.Empty;
+					if (Pair.Value != value)
+					{
+						Logger.LogDebug("Platform SDK {Key} '{Value}' has been modified to '{NewValue}' since makefile was written.", Pair.Key, Pair.Value, !String.IsNullOrEmpty(value) ? value : "Missing");
+						ReasonNotLoaded = String.Format("Platform SDK {0} modified", Pair.Key);
+						return false;
+					}
+				}
+
 				IEnumerable<DirectoryItem> UbaPaths = [.. Makefile.UbaDisallowedPaths, .. Makefile.UbaPathHashes.Keys, .. Makefile.UbaCrossArchitecturePaths.Keys, .. Makefile.UbaCrossArchitecturePaths.Values];
 				foreach (DirectoryItem item in UbaPaths.Distinct())
 				{
