UnrealBuildTool: UBA horde settings via providers from Engine.ini
Remove separate local only UBA executor

diff --git a/Engine/Config/BaseEngine.ini b/Engine/Config/BaseEngine.ini
--- a/Engine/Config/BaseEngine.ini
+++ b/Engine/Config/BaseEngine.ini
@@ -3839,6 +3839,27 @@ SendOSData=true
 ; APIServerET="https://myownhordeserver.com/" ; Point Server URL to your own Horde server
 ; APIEndpointET="api/v1/telemetry/engine"
 
+; Example Horde UBA Distribution Settings - See Engine/Source/Programs/UnrealBuildTool/Executors/UnrealBuildAccelerator/UnrealBuildAcceleratorHordeConfig.cs
+;[UbaController]
+; DefaultProvider=Uba.Provider.Horde.Default ; The default provider if none are set
+; +BuildMachineProviders=Uba.Provider.Horde.BuildMachine ; Providers for build machines
+; +Providers=Uba.Provider.Horde.MyOwnHordeServer ; Example to add multiple horde providers
+; +Providers=Uba.Provider.Horde.MyOwnHordeServer2
+;[Uba.Provider.Horde] ; Default settings for all users. It's not recommended to automatically enable all users
+; Enabled=False ; Disable distribution by default
+;[Uba.Provider.Horde.BuildMachine]
+; ServerUrl="https://myownhordeserver.com/"
+; Pool=MyBuildMachineUbaPool
+; Enabled=BuildMachineOnly ; Enable distribution only for build machines
+;[Uba.Provider.Horde.MyOwnHordeServer]
+; ServerUrl="https://myownhordeserver.com/"
+; Pool=MyUbaPool
+; Enabled=True ; Enable distribution for this provider
+;[Uba.Provider.Horde.MyOwnHordeServer2] ; alternate provider
+; ServerUrl="https://myownhordeserver2.com/"
+; Pool=MyUbaPool2
+; Enabled=True
+
 [/Script/Engine.WorldPartitionSettings]
 EditorHashDefaultClass=/Script/CoreUObject.Class'/Script/Engine.WorldPartitionEditorSpatialHash'
 RuntimeHashDefaultClass=/Script/CoreUObject.Class'/Script/Engine.WorldPartitionRuntimeHashSet'
diff --git a/Engine/Source/Programs/UnrealBuildTool/Configuration/BuildConfiguration.cs b/Engine/Source/Programs/UnrealBuildTool/Configuration/BuildConfiguration.cs
--- a/Engine/Source/Programs/UnrealBuildTool/Configuration/BuildConfiguration.cs
+++ b/Engine/Source/Programs/UnrealBuildTool/Configuration/BuildConfiguration.cs
@@ -44,7 +44,7 @@ namespace UnrealBuildTool
 		/// Priority order for remote executors (XGE, SNDBS, FASTBuild, UBA)
 		/// </summary>
 		[XmlConfigFile]
-		public string[] RemoteExecutorPriority { get; set; } = new string[] { "XGE", "SNDBS", "FASTBuild", "UBA" };
+		public string[] RemoteExecutorPriority { get; set; } = ["XGE", "SNDBS", "FASTBuild", "UBA"];
 
 		/// <summary>
 		/// Whether the UnrealBuildAccelerator executor will be used.
@@ -52,20 +52,18 @@ namespace UnrealBuildTool
 		[XmlConfigFile]
 		[CommandLine("-UBA", Value = "true")]
 		[CommandLine("-NoUBA", Value = "false")]
-		public bool bAllowUBAExecutor { get; set; } = false;
+		public bool bAllowUBAExecutor { get; set; } = true;
 
 		/// <summary>
 		/// Whether the UnrealBuildAccelerator (local only) executor will be used.
 		/// </summary>
-		[XmlConfigFile]
-		[CommandLine("-UBALocal", Value = "true")]
-		[CommandLine("-NoUBALocal", Value = "false")]
+		[XmlConfigFile(Deprecated = true)]
+		[Obsolete("Replace with bAllowUBAExecutor")]
 		public bool bAllowUBALocalExecutor
 		{
-			get => bAllowUBAExecutor || _bAllowUBALocalExecutorPrivate;
-			set => _bAllowUBALocalExecutorPrivate = value;
+			get => bAllowUBAExecutor;
+			set => bAllowUBAExecutor = value;
 		}
-		private bool _bAllowUBALocalExecutorPrivate = true;
 
 		/// <summary>
 		/// Whether XGE may be used if available, default is true.
diff --git a/Engine/Source/Programs/UnrealBuildTool/Executors/UnrealBuildAccelerator/UBAAgentCoordinatorHorde.cs b/Engine/Source/Programs/UnrealBuildTool/Executors/UnrealBuildAccelerator/UBAAgentCoordinatorHorde.cs
--- a/Engine/Source/Programs/UnrealBuildTool/Executors/UnrealBuildAccelerator/UBAAgentCoordinatorHorde.cs
+++ b/Engine/Source/Programs/UnrealBuildTool/Executors/UnrealBuildAccelerator/UBAAgentCoordinatorHorde.cs
@@ -4,7 +4,6 @@ using System;
 using System.Collections.Generic;
 using System.Diagnostics;
 using System.Linq;
-using System.Net.Http;
 using System.Runtime.InteropServices;
 using System.Threading;
 using System.Threading.Tasks;
@@ -706,50 +705,32 @@ namespace UnrealBuildTool
 
 	class UBAAgentCoordinatorHorde : IUBAAgentCoordinator, IDisposable
 	{
-		public string? Server => _hordeConfig.HordeServer;
-		public bool Enabled => !_hordeConfig.bDisableHorde;
 		public string ConnectionModeString => (Enum.TryParse(_hordeConfig.HordeConnectionMode, true, out ConnectionMode cm) ? cm : ConnectionMode.Direct).ToString();
 
-		public static List<IUBAAgentCoordinator> Init(ILogger logger, UBAExecutor executor, IEnumerable<TargetDescriptor> targetDescriptors, UnrealBuildAcceleratorConfig ubaConfig, ref string remoteConnectionMode)
+		public static string ProviderPrefix = "Uba.Provider.Horde";
+
+		public static IEnumerable<IUBAAgentCoordinator> Init(ILogger logger, UBAExecutor executor, IEnumerable<TargetDescriptor> targetDescriptors, UnrealBuildAcceleratorConfig ubaConfig, ref string remoteConnectionMode)
 		{
-			List<IUBAAgentCoordinator> outCoordinators = new();
-			List<UBAAgentCoordinatorHorde> hordeAgentCoordinators = new();
-			foreach (TargetDescriptor targetDescriptor in targetDescriptors)
-			{
-				hordeAgentCoordinators.Add(new UBAAgentCoordinatorHorde(logger, ubaConfig, targetDescriptor.AdditionalArguments, targetDescriptor.ProjectFile?.Directory));
-			}
-			hordeAgentCoordinators.RemoveAll(x => !x.Enabled);
-			remoteConnectionMode = hordeAgentCoordinators.FirstOrDefault()?.ConnectionModeString ?? "Local";
-			outCoordinators.AddRange(hordeAgentCoordinators.DistinctBy(x => x.Server));
+			IEnumerable<string> providers = Unreal.IsBuildMachine()
+				? [.. ubaConfig.BuildMachineProviders, .. ubaConfig.IniBuildMachineProviders]
+				: [.. ubaConfig.Providers, .. ubaConfig.IniProviders];
+
 
-			// Some users want their extra local horde. This is the very simplified version with much fewer settings
+			if (!providers.Any())
 			{
-				UnrealBuildAcceleratorHordeConfig2 config2 = new();
-				XmlConfig.ApplyTo(config2);
-				if (!String.IsNullOrEmpty(config2.HordeServer))
-				{
-					UnrealBuildAcceleratorHordeConfig extraHordeConfig = new();
-					extraHordeConfig.HordeServer = config2.HordeServer;
-					extraHordeConfig.bDisableHorde = config2.bDisableHorde;
-					UBAAgentCoordinatorHorde extraHorde = new(logger, ubaConfig, extraHordeConfig);
-					if (extraHorde.Enabled)
-					{
-						outCoordinators.Add(extraHorde);
-					}
-				}
+				providers = [ProviderPrefix];
 			}
-			return outCoordinators;
-		}
 
-		public UBAAgentCoordinatorHorde(ILogger logger, UnrealBuildAcceleratorConfig ubaConfig, UnrealBuildAcceleratorHordeConfig hordeConfig)
-		{
-			_logger = logger;
-			_title = "Horde2";
-			_ubaConfig = ubaConfig;
-			_hordeConfig = hordeConfig;
+			IEnumerable<UBAAgentCoordinatorHorde> coordinators = providers.Distinct().Where(x => x.StartsWith(ProviderPrefix))
+				.SelectMany(provider => targetDescriptors.Select(targetDescriptor => new UBAAgentCoordinatorHorde(logger, ubaConfig, provider, targetDescriptor.AdditionalArguments, targetDescriptor.ProjectFile?.Directory)))
+				.Where(x => !x._hordeConfig.bDisableHorde && !String.IsNullOrWhiteSpace(x._hordeConfig.HordeServer))
+				.DistinctBy(x => HashCode.Combine(x._hordeConfig.HordeServer, x._hordeConfig.HordeCondition, x._hordeConfig.HordeCluster, x._hordeConfig.HordePool, x._hordeConfig.HordeConnectionMode));
+
+			remoteConnectionMode = coordinators.FirstOrDefault()?.ConnectionModeString ?? "Local";
+			return coordinators;
 		}
 
-		public UBAAgentCoordinatorHorde(ILogger logger, UnrealBuildAcceleratorConfig ubaConfig, CommandLineArguments? additionalArguments = null, DirectoryReference? projectDir = null)
+		public UBAAgentCoordinatorHorde(ILogger logger, UnrealBuildAcceleratorConfig ubaConfig, string provider, CommandLineArguments? additionalArguments = null, DirectoryReference? projectDir = null)
 		{
 			_logger = logger;
 			_title = "Horde";
@@ -757,11 +738,15 @@ namespace UnrealBuildTool
 
 			_hordeConfig = new();
 
-			ConfigCache.ReadSettings(projectDir, BuildHostPlatform.Current.Platform, _hordeConfig);
-			if (_hordeConfig.HordeEnabled?.Equals("False", StringComparison.OrdinalIgnoreCase) == true || (_hordeConfig.HordeEnabled?.Equals("BuildMachineOnly", StringComparison.OrdinalIgnoreCase) == true && !Unreal.IsBuildMachine()))
+			ConfigHierarchy engineIni = ConfigCache.ReadHierarchy(ConfigHierarchyType.Engine, projectDir, BuildHostPlatform.Current.Platform);
+			ConfigHierarchySection section = engineIni.FindSection(provider);
+			_hordeConfig.LoadConfigProvider(engineIni, provider);
+
+			if (provider == ProviderPrefix)
 			{
-				_hordeConfig.bDisableHorde = true;
+				XmlConfig.ApplyTo(_hordeConfig);
 			}
+
 			if (Unreal.IsBuildMachine())
 			{
 				string? hordeUrl = System.Environment.GetEnvironmentVariable("UE_HORDE_URL");
@@ -769,8 +754,13 @@ namespace UnrealBuildTool
 				{
 					_hordeConfig.HordeServer = hordeUrl;
 				}
-			}	
-			XmlConfig.ApplyTo(_hordeConfig);
+			}
+
+			if (_hordeConfig.HordeEnabled?.Equals("False", StringComparison.OrdinalIgnoreCase) == true || (_hordeConfig.HordeEnabled?.Equals("BuildMachineOnly", StringComparison.OrdinalIgnoreCase) == true && !Unreal.IsBuildMachine()))
+			{
+				_hordeConfig.bDisableHorde = true;
+			}
+
 			additionalArguments?.ApplyTo(_hordeConfig);
 
 			// Sentry is currently unsupported for non-Windows and non-x64
@@ -778,6 +768,15 @@ namespace UnrealBuildTool
 			{
 				_hordeConfig.UBASentryUrl = null;
 			}
+
+			// Normalize URLs and pools
+			_hordeConfig.HordeServer = _hordeConfig.HordeServer?.TrimEnd('/').ToLowerInvariant();
+			_hordeConfig.DefaultHordePool = _hordeConfig.DefaultHordePool?.ToLowerInvariant();
+			_hordeConfig.OverrideHordePool = _hordeConfig.OverrideHordePool?.ToLowerInvariant();
+			_hordeConfig.WindowsHordePool = _hordeConfig.WindowsHordePool?.ToLowerInvariant();
+			_hordeConfig.MacHordePool = _hordeConfig.MacHordePool?.ToLowerInvariant();
+			_hordeConfig.LinuxHordePool = _hordeConfig.LinuxHordePool?.ToLowerInvariant();
+			_hordeConfig.UBASentryUrl = _hordeConfig.UBASentryUrl?.TrimEnd('/').ToLowerInvariant();
 		}
 
 		public DirectoryReference? GetUBARootDir()
diff --git a/Engine/Source/Programs/UnrealBuildTool/Executors/UnrealBuildAccelerator/UBAExecutor.cs b/Engine/Source/Programs/UnrealBuildTool/Executors/UnrealBuildAccelerator/UBAExecutor.cs
--- a/Engine/Source/Programs/UnrealBuildTool/Executors/UnrealBuildAccelerator/UBAExecutor.cs
+++ b/Engine/Source/Programs/UnrealBuildTool/Executors/UnrealBuildAccelerator/UBAExecutor.cs
@@ -1,5 +1,9 @@
 // Copyright Epic Games, Inc. All Rights Reserved.
 
+using EpicGames.Core;
+using EpicGames.Horde.Projects;
+using EpicGames.UBA;
+using Microsoft.Extensions.Logging;
 using System;
 using System.Collections.Concurrent;
 using System.Collections.Generic;
@@ -10,9 +14,6 @@ using System.Runtime.InteropServices;
 using System.Runtime.Versioning;
 using System.Threading;
 using System.Threading.Tasks;
-using EpicGames.Core;
-using EpicGames.UBA;
-using Microsoft.Extensions.Logging;
 using UnrealBuildBase;
 using UnrealBuildTool.Artifacts;
 
@@ -76,7 +77,7 @@ namespace UnrealBuildTool
 			base.Dispose(disposing);
 		}
 
-		public override string Name => "Unreal Build Accelerator";
+		public override string Name => !UBAConfig.bDisableRemote && _agentCoordinators.Count > 0 ? "Unreal Build Accelerator" : "Unreal Build Accelerator local";
 
 		public static new bool IsAvailable()
 		{
@@ -126,6 +127,7 @@ namespace UnrealBuildTool
 
 			foreach (TargetDescriptor targetDescriptor in targetDescriptors)
 			{
+				ConfigCache.ReadSettings(targetDescriptor.ProjectFile?.Directory, BuildHostPlatform.Current.Platform, UBAConfig);
 				if (targetDescriptor.HotReloadMode != HotReloadMode.Disabled)
 				{
 					UBAConfig.CacheServer = null; // Disable cache server when using hot reload path (there would likely never be any cache hits anyway
@@ -317,6 +319,10 @@ namespace UnrealBuildTool
 			if (inputActions.Count() < NumParallelProcesses && !UBAConfig.bForceBuildAllRemote)
 			{
 				UBAConfig.bDisableRemote = true;
+			}
+
+			if (UBAConfig.bDisableRemote)
+			{
 				UBAConfig.Zone = "local";
 			}
 
@@ -1168,25 +1174,4 @@ namespace UnrealBuildTool
 			}
 		}
 	}
-
-	/// <summary>
-	/// UBAExecutor, but force local only compiles regardless of other settings
-	/// </summary>
-	class UBALocalExecutor : UBAExecutor
-	{
-		public override string Name => "Unreal Build Accelerator local";
-
-		public static new bool IsAvailable()
-		{
-			return EpicGames.UBA.Utils.IsAvailable();
-		}
-
-		public UBALocalExecutor(int maxLocalActions, bool bAllCores, bool bCompactOutput, Microsoft.Extensions.Logging.ILogger logger, IEnumerable<TargetDescriptor> targetDescriptors)
-			: base(maxLocalActions, bAllCores, bCompactOutput, logger, targetDescriptors)
-		{
-			UBAConfig.bDisableRemote = true;
-			UBAConfig.bForceBuildAllRemote = false;
-			UBAConfig.Zone = "local";
-		}
-	}
 }
diff --git a/Engine/Source/Programs/UnrealBuildTool/Executors/UnrealBuildAccelerator/UnrealBuildAcceleratorConfig.cs b/Engine/Source/Programs/UnrealBuildTool/Executors/UnrealBuildAccelerator/UnrealBuildAcceleratorConfig.cs
--- a/Engine/Source/Programs/UnrealBuildTool/Executors/UnrealBuildAccelerator/UnrealBuildAcceleratorConfig.cs
+++ b/Engine/Source/Programs/UnrealBuildTool/Executors/UnrealBuildAccelerator/UnrealBuildAcceleratorConfig.cs
@@ -1,7 +1,9 @@
 // Copyright Epic Games, Inc. All Rights Reserved.
 
-using System;
 using EpicGames.Core;
+using System;
+using System.Collections.Generic;
+using UnrealBuildBase;
 
 namespace UnrealBuildTool
 {
@@ -11,6 +13,36 @@ namespace UnrealBuildTool
 	[System.Diagnostics.CodeAnalysis.SuppressMessage("Style", "IDE1006:Naming Styles", Justification = "UnrealBuildTool naming style")]
 	class UnrealBuildAcceleratorConfig
 	{
+		/// <summary>
+		/// List of provider names for remote connections.
+		/// </summary>
+		[XmlConfigFile(Category = "UnrealBuildAccelerator")]
+		[CommandLine("-UBAProviders=", ListSeparator = '+')]
+		public string[] Providers { get; set; } = [];
+
+		/// <summary>
+		/// List of provider names for remote connections from ini
+		/// It is recommended to set Providers in ini so the setting can be shared with the uba controller in editor, such as
+		/// Windows:	%ProgramData%/Epic/Unreal Engine/Engine/Config/SystemEngine.ini
+		/// Mac:		~/Library/Application Support/Epic/Unreal Engine/Engine/Config/SystemEngine.ini
+		/// Linux:		~/.config/Epic/Unreal Engine/Engine/Config/SystemEngine.ini
+		/// </summary>
+		[ConfigFile(ConfigHierarchyType.Engine, "UbaController", "Providers")]
+		public List<string> IniProviders { get; set; } = [];
+
+		/// <summary>
+		/// List of provider names for remote connections
+		/// </summary>
+		[XmlConfigFile(Category = "UnrealBuildAccelerator")]
+		public string[] BuildMachineProviders { get; set; } = [];
+
+		/// <summary>
+		/// List of provider names for remote connections from ini
+		/// </summary>
+		[ConfigFile(ConfigHierarchyType.Engine, "UbaController", "BuildMachineProviders")]
+		public List<string> IniBuildMachineProviders { get; set; } = [];
+
+
 		/// <summary>
 		/// When set to true, UBA will not use any remote help
 		/// </summary>
diff --git a/Engine/Source/Programs/UnrealBuildTool/Executors/UnrealBuildAccelerator/UnrealBuildAcceleratorHordeConfig.cs b/Engine/Source/Programs/UnrealBuildTool/Executors/UnrealBuildAccelerator/UnrealBuildAcceleratorHordeConfig.cs
--- a/Engine/Source/Programs/UnrealBuildTool/Executors/UnrealBuildAccelerator/UnrealBuildAcceleratorHordeConfig.cs
+++ b/Engine/Source/Programs/UnrealBuildTool/Executors/UnrealBuildAccelerator/UnrealBuildAcceleratorHordeConfig.cs
@@ -1,8 +1,12 @@
 // Copyright Epic Games, Inc. All Rights Reserved.
 
-using System;
 using EpicGames.Core;
 using EpicGames.Horde.Compute;
+using System;
+using System.Collections.Generic;
+using System.Linq;
+using System.Reflection;
+using UnrealBuildBase;
 
 namespace UnrealBuildTool
 {
@@ -15,7 +19,7 @@ namespace UnrealBuildTool
 		/// <summary>
 		/// Uri of the Horde server
 		/// </summary>
-		[ConfigFile(ConfigHierarchyType.Engine, "Horde", "ServerUrl")]
+		[ConfigFile(ConfigHierarchyType.Engine, "Uba.Provider.Horde", "ServerUrl")]
 		[XmlConfigFile(Category = "Horde", Name = "Server")]
 		[CommandLine("-UBAHorde=")]
 		public string? HordeServer { get; set; }
@@ -23,7 +27,7 @@ namespace UnrealBuildTool
 		/// <summary>
 		/// Auth token for the Horde server
 		/// </summary>
-		[ConfigFile(ConfigHierarchyType.Engine, "Horde", "UbaToken")]
+		[ConfigFile(ConfigHierarchyType.Engine, "Uba.Provider.Horde", "UbaToken")]
 		[XmlConfigFile(Category = "Horde", Name = "Token")]
 		[CommandLine("-UBAHordeToken=")]
 		public string? HordeToken { get; set; }
@@ -61,7 +65,6 @@ namespace UnrealBuildTool
 					Pool = LinuxHordePool;
 				}
 
-				//Console.WriteLine("CHOSEN UBA POOL: {0}", Pool);
 				return Pool;
 			}
 		}
@@ -69,7 +72,7 @@ namespace UnrealBuildTool
 		/// <summary>
 		/// Pool for the Horde agent to assign if no override current platform doesn't have it set
 		/// </summary>
-		[ConfigFile(ConfigHierarchyType.Engine, "Horde", "UbaPool")]
+		[ConfigFile(ConfigHierarchyType.Engine, "Uba.Provider.Horde", "UbaPool")]
 		[XmlConfigFile(Category = "Horde", Name = "Pool")]
 		public string? DefaultHordePool { get; set; }
 
@@ -82,28 +85,28 @@ namespace UnrealBuildTool
 		/// <summary>
 		/// Pool for the Horde agent to assign when on Linux
 		/// </summary>
-		[ConfigFile(ConfigHierarchyType.Engine, "Horde", "UbaLinuxPool")]
+		[ConfigFile(ConfigHierarchyType.Engine, "Uba.Provider.Horde", "UbaLinuxPool")]
 		[XmlConfigFile(Category = "Horde", Name = "LinuxPool")]
 		public string? LinuxHordePool { get; set; }
 
 		/// <summary>
 		/// Pool for the Horde agent to assign when on Mac
 		/// </summary>
-		[ConfigFile(ConfigHierarchyType.Engine, "Horde", "UbaMacPool")]
+		[ConfigFile(ConfigHierarchyType.Engine, "Uba.Provider.Horde", "UbaMacPool")]
 		[XmlConfigFile(Category = "Horde", Name = "MacPool")]
 		public string? MacHordePool { get; set; }
 
 		/// <summary>
 		/// Pool for the Horde agent to assign when on Windows
 		/// </summary>
-		[ConfigFile(ConfigHierarchyType.Engine, "Horde", "UbaWindowsPool")]
+		[ConfigFile(ConfigHierarchyType.Engine, "Uba.Provider.Horde", "UbaWindowsPool")]
 		[XmlConfigFile(Category = "Horde", Name = "WindowsPool")]
 		public string? WindowsHordePool { get; set; }
 
 		/// <summary>
 		/// Requirements for the Horde agent to assign
 		/// </summary>
-		[ConfigFile(ConfigHierarchyType.Engine, "Horde", "UbaRequirements")]
+		[ConfigFile(ConfigHierarchyType.Engine, "Uba.Provider.Horde", "UbaRequirements")]
 		[XmlConfigFile(Category = "Horde", Name = "Requirements")]
 		[CommandLine("-UBAHordeRequirements=")]
 		public string? HordeCondition { get; set; }
@@ -122,7 +125,7 @@ namespace UnrealBuildTool
 		/// Compute cluster ID to use in Horde. Set to "_auto" to let Horde server resolve a suitable cluster.
 		/// In multi-region setups this is can simplify configuration of UBT/UBA a lot.
 		/// </summary>
-		[ConfigFile(ConfigHierarchyType.Engine, "Horde", "UbaCluster")]
+		[ConfigFile(ConfigHierarchyType.Engine, "Uba.Provider.Horde", "UbaCluster")]
 		[XmlConfigFile(Category = "Horde", Name = "Cluster")]
 		[CommandLine("-UBAHordeCluster=")]
 		public string? HordeCluster { get; set; }
@@ -130,7 +133,7 @@ namespace UnrealBuildTool
 		/// <summary>
 		/// Which ip UBA server should give to agents. This will invert so host listens and agents connect
 		/// </summary>
-		[ConfigFile(ConfigHierarchyType.Engine, "Horde", "UbaLocalHost")]
+		[ConfigFile(ConfigHierarchyType.Engine, "Uba.Provider.Horde", "UbaLocalHost")]
 		[XmlConfigFile(Category = "Horde", Name = "LocalHost")]
 		[CommandLine("-UBAHordeHost")]
 		public string HordeHost { get; set; } = String.Empty;
@@ -138,7 +141,7 @@ namespace UnrealBuildTool
 		/// <summary>
 		/// Max cores allowed to be used by build session
 		/// </summary>
-		[ConfigFile(ConfigHierarchyType.Engine, "Horde", "MaxCores")]
+		[ConfigFile(ConfigHierarchyType.Engine, "Uba.Provider.Horde", "MaxCores")]
 		[XmlConfigFile(Category = "Horde", Name = "MaxCores")]
 		[CommandLine("-UBAHordeMaxCores")]
 		public int HordeMaxCores { get; set; } = 576;
@@ -146,7 +149,7 @@ namespace UnrealBuildTool
 		/// <summary>
 		/// How long UBT should wait to ask for help. Useful in build configs where machine can delay remote work and still get same wall time results (pch dependencies etc)
 		/// </summary>
-		[ConfigFile(ConfigHierarchyType.Engine, "Horde", "UbaDelay")]
+		[ConfigFile(ConfigHierarchyType.Engine, "Uba.Provider.Horde", "UbaDelay")]
 		[XmlConfigFile(Category = "Horde", Name = "StartupDelay")]
 		[CommandLine("-UBAHordeDelay")]
 		public int HordeDelay { get; set; } = 0;
@@ -154,7 +157,7 @@ namespace UnrealBuildTool
 		/// <summary>
 		/// Allow use of Wine. Only applicable to Horde agents running Linux. Can still be ignored if Wine executable is not set on agent.
 		/// </summary>
-		[ConfigFile(ConfigHierarchyType.Engine, "Horde", "UbaAllowWine")]
+		[ConfigFile(ConfigHierarchyType.Engine, "Uba.Provider.Horde", "UbaAllowWine")]
 		[XmlConfigFile(Category = "Horde", Name = "AllowWine")]
 		[CommandLine("-UBAHordeAllowWine", Value = "true")]
 		public bool bHordeAllowWine { get; set; } = true;
@@ -163,7 +166,7 @@ namespace UnrealBuildTool
 		/// Connection mode for agent/compute communication
 		/// <see cref="ConnectionMode" /> for valid modes.
 		/// </summary>
-		[ConfigFile(ConfigHierarchyType.Engine, "Horde", "UbaConnectionMode")]
+		[ConfigFile(ConfigHierarchyType.Engine, "Uba.Provider.Horde", "UbaConnectionMode")]
 		[XmlConfigFile(Category = "Horde", Name = "ConnectionMode")]
 		[CommandLine("-UBAHordeConnectionMode=")]
 		public string? HordeConnectionMode { get; set; }
@@ -172,7 +175,7 @@ namespace UnrealBuildTool
 		/// Encryption to use for agent/compute communication. Note that UBA agent uses its own encryption.
 		/// <see cref="Encryption" /> for valid modes.
 		/// </summary>
-		[ConfigFile(ConfigHierarchyType.Engine, "Horde", "UbaEncryption")]
+		[ConfigFile(ConfigHierarchyType.Engine, "Uba.Provider.Horde", "UbaEncryption")]
 		[XmlConfigFile(Category = "Horde", Name = "Encryption")]
 		[CommandLine("-UBAHordeEncryption=")]
 		public string? HordeEncryption { get; set; }
@@ -180,7 +183,7 @@ namespace UnrealBuildTool
 		/// <summary>
 		/// Sentry URL to send box data to. Optional.
 		/// </summary>
-		[ConfigFile(ConfigHierarchyType.Engine, "Horde", "UbaSentryUrl")]
+		[ConfigFile(ConfigHierarchyType.Engine, "Uba.Provider.Horde", "UbaSentryUrl")]
 		[XmlConfigFile(Category = "Horde")]
 		[CommandLine("-UBASentryUrl=")]
 		public string? UBASentryUrl { get; set; }
@@ -195,25 +198,84 @@ namespace UnrealBuildTool
 		/// <summary>
 		/// Enabled parameter for use by INI config, expects one of [True, False, BuildMachineOnly]
 		/// </summary>
-		[ConfigFile(ConfigHierarchyType.Engine, "Horde", "UbaEnabled")]
+		[ConfigFile(ConfigHierarchyType.Engine, "Uba.Provider.Horde", "UbaEnabled")]
 		public string? HordeEnabled { get; set; } = null;
 	}
 
-	class UnrealBuildAcceleratorHordeConfig2
+	#region Extension methods
+	/// <summary>
+	/// Extension methods for FileReference functionality
+	/// </summary>
+	internal static class UnrealBuildAcceleratorHordeConfigExtensionMethods
 	{
 		/// <summary>
-		/// Uri of the Horde server
+		/// Load horde config from a different section name
 		/// </summary>
-		[ConfigFile(ConfigHierarchyType.Engine, "Horde2", "ServerUrl")]
-		[XmlConfigFile(Category = "Horde2", Name = "Server")]
-		[CommandLine("-UBAHorde2=")]
-		public string? HordeServer { get; set; }
+		/// <param name="config"></param>
+		/// <param name="ini"></param>
+		/// <param name="section"></param>
+		internal static void LoadConfigProvider(this UnrealBuildAcceleratorHordeConfig config, ConfigHierarchy ini, string section)
+		{
+			object? TryGetValue(Type propertyType, string keyName)
+			{
+				if (propertyType == typeof(string))
+				{
+					if (ini.TryGetValue(section, keyName, out string? value))
+					{
+						return value;
+					}
+				}
+				else if (propertyType == typeof(bool))
+				{
+					if (ini.TryGetValue(section, keyName, out bool value))
+					{
+						return value;
+					}
+				}
+				else if (propertyType == typeof(int))
+				{
+					if (ini.TryGetValue(section, keyName, out int value))
+					{
+						return value;
+					}
+				}
+				else if (propertyType == typeof(float))
+				{
+					if (ini.TryGetValue(section, keyName, out float value))
+					{
+						return value;
+					}
+				}
+				else if (propertyType == typeof(double))
+				{
+					if (ini.TryGetValue(section, keyName, out double value))
+					{
+						return value;
+					}
+				}
+				else
+				{
+					throw new BuildLogEventException("UnrealBuildAcceleratorHordeConfig.LoadConfigProvider: Unsupported type {Type} for Key {KeyName}", propertyType, keyName);
+				}
+				return null;
+			}
 
-		/// <summary>
-		/// Disable horde all together
-		/// </summary>
-		[XmlConfigFile(Category = "UnrealBuildAccelerator")]
-		[CommandLine("-UBADisableHorde")]
-		public bool bDisableHorde { get; set; } = false;
+			IEnumerable<PropertyInfo> properties = config.GetType().GetProperties(BindingFlags.Instance | BindingFlags.Public | BindingFlags.NonPublic);
+			foreach (PropertyInfo property in properties)
+			{
+				Type propertyType = property.PropertyType;
+				ConfigFileAttribute? attribute = property.GetCustomAttributes<ConfigFileAttribute>().FirstOrDefault();
+				if (attribute != null)
+				{
+					string keyName = attribute.KeyName ?? property.Name;
+					object? value = TryGetValue(propertyType, keyName) ?? (keyName.StartsWith("Uba") ? TryGetValue(propertyType, keyName[3..]) : null);
+					if (value != null)
+					{
+						property.SetValue(config, value);
+					}
+				}
+			}
+		}
 	}
+	#endregion
 }
\ No newline at end of file
diff --git a/Engine/Source/Programs/UnrealBuildTool/System/ActionGraph.cs b/Engine/Source/Programs/UnrealBuildTool/System/ActionGraph.cs
--- a/Engine/Source/Programs/UnrealBuildTool/System/ActionGraph.cs
+++ b/Engine/Source/Programs/UnrealBuildTool/System/ActionGraph.cs
@@ -436,11 +436,6 @@ namespace UnrealBuildTool
 				}
 			}
 
-			if (BuildConfiguration.bAllowUBALocalExecutor && UBALocalExecutor.IsAvailable())
-			{
-				return new UBALocalExecutor(BuildConfiguration.MaxParallelActions, BuildConfiguration.bAllCores, BuildConfiguration.bCompactOutput, Logger, TargetDescriptors);
-			}
-
 			return new ParallelExecutor(BuildConfiguration.MaxParallelActions, BuildConfiguration.bAllCores, BuildConfiguration.bCompactOutput, Logger);
 		}
 
