ICX 2025.1 and Microsoft Clang 19 HWPGO integrations

diff --git a/Engine/Config/Windows/Windows_SDK.json b/Engine/Config/Windows/Windows_SDK.json
--- a/Engine/Config/Windows/Windows_SDK.json
+++ b/Engine/Config/Windows/Windows_SDK.json
@@ -38,9 +38,9 @@
 	"MinimumVisualCppVersion": "14.38.33130",
 
 	"PreferredIntelOneApiVersions": [
-		"2025.0.0-2025.9999"
+		"2025.1.0-2025.9999"
 	],
-	"MinimumIntelOneApiVersion": "2025.0.0",
+	"MinimumIntelOneApiVersion": "2025.1.0",
 
 	"//04.1": "Microsoft.VisualStudio.Component.VC.14.38.17.8.x86.x64 -- Only LTSC versions of MSVC should be suggested",
 	"//04.2": "Microsoft.VisualStudio.Component.VC.14.38.17.8.ATL -- Match version of MSVC",
diff --git a/Engine/Plugins/Developer/Concert/ConcertApp/MultiUserClient/Source/MultiUserClient/Private/Widgets/ActiveSession/Replication/Client/Multi/Columns/AssignProperty/SAssignPropertyComboBox.h b/Engine/Plugins/Developer/Concert/ConcertApp/MultiUserClient/Source/MultiUserClient/Private/Widgets/ActiveSession/Replication/Client/Multi/Columns/AssignProperty/SAssignPropertyComboBox.h
--- a/Engine/Plugins/Developer/Concert/ConcertApp/MultiUserClient/Source/MultiUserClient/Private/Widgets/ActiveSession/Replication/Client/Multi/Columns/AssignProperty/SAssignPropertyComboBox.h
+++ b/Engine/Plugins/Developer/Concert/ConcertApp/MultiUserClient/Source/MultiUserClient/Private/Widgets/ActiveSession/Replication/Client/Multi/Columns/AssignProperty/SAssignPropertyComboBox.h
@@ -50,8 +50,8 @@ namespace UE::MultiUserClient::Replication::MultiStreamColumns
 		SLATE_END_ARGS()
 
 		void Construct(const FArguments& InArgs,
-           FAssignPropertyModel& InModel UE_LIFETIMEBOUND,
-           FUnifiedClientView& InClientView UE_LIFETIMEBOUND
+           FAssignPropertyModel& InModel,
+           FUnifiedClientView& InClientView
 		);
 
 	private:
diff --git a/Engine/Plugins/Developer/Concert/ConcertApp/MultiUserClient/Source/MultiUserClient/Private/Widgets/ActiveSession/Replication/Client/Multi/Columns/AssignedClients/SAssignedClientsWidget.h b/Engine/Plugins/Developer/Concert/ConcertApp/MultiUserClient/Source/MultiUserClient/Private/Widgets/ActiveSession/Replication/Client/Multi/Columns/AssignedClients/SAssignedClientsWidget.h
--- a/Engine/Plugins/Developer/Concert/ConcertApp/MultiUserClient/Source/MultiUserClient/Private/Widgets/ActiveSession/Replication/Client/Multi/Columns/AssignedClients/SAssignedClientsWidget.h
+++ b/Engine/Plugins/Developer/Concert/ConcertApp/MultiUserClient/Source/MultiUserClient/Private/Widgets/ActiveSession/Replication/Client/Multi/Columns/AssignedClients/SAssignedClientsWidget.h
@@ -36,8 +36,8 @@ namespace UE::MultiUserClient::Replication::MultiStreamColumns
 
 		void Construct(
 			const FArguments& InArgs,
-			FAssignedClientsModel& InModel UE_LIFETIMEBOUND,
-			const FUnifiedClientView& InClientView UE_LIFETIMEBOUND
+			FAssignedClientsModel& InModel,
+			const FUnifiedClientView& InClientView
 			);
 	
 		virtual ~SAssignedClientsWidget() override;
diff --git a/Engine/Plugins/Developer/Concert/ConcertApp/MultiUserClient/Source/MultiUserClient/Private/Widgets/ActiveSession/Replication/Client/Multi/SMultiClientView.h b/Engine/Plugins/Developer/Concert/ConcertApp/MultiUserClient/Source/MultiUserClient/Private/Widgets/ActiveSession/Replication/Client/Multi/SMultiClientView.h
--- a/Engine/Plugins/Developer/Concert/ConcertApp/MultiUserClient/Source/MultiUserClient/Private/Widgets/ActiveSession/Replication/Client/Multi/SMultiClientView.h
+++ b/Engine/Plugins/Developer/Concert/ConcertApp/MultiUserClient/Source/MultiUserClient/Private/Widgets/ActiveSession/Replication/Client/Multi/SMultiClientView.h
@@ -42,9 +42,9 @@ namespace UE::MultiUserClient::Replication
 
 		void Construct(
 			const FArguments& InArgs, const TSharedRef<IConcertClient>& InConcertClient,
-			FMultiUserReplicationManager& InMultiUserReplicationManager UE_LIFETIMEBOUND,
-			IOnlineClientSelectionModel& InOnlineClientSelectionModel UE_LIFETIMEBOUND,
-			IOfflineClientSelectionModel& InOfflineClientSelectionModel UE_LIFETIMEBOUND
+			FMultiUserReplicationManager& InMultiUserReplicationManager,
+			IOnlineClientSelectionModel& InOnlineClientSelectionModel,
+			IOfflineClientSelectionModel& InOfflineClientSelectionModel
 			);
 		virtual ~SMultiClientView() override;
 
diff --git a/Engine/Plugins/Experimental/LearningAgents/Source/Learning/Public/LearningArray.h b/Engine/Plugins/Experimental/LearningAgents/Source/Learning/Public/LearningArray.h
--- a/Engine/Plugins/Experimental/LearningAgents/Source/Learning/Public/LearningArray.h
+++ b/Engine/Plugins/Experimental/LearningAgents/Source/Learning/Public/LearningArray.h
@@ -326,45 +326,6 @@ namespace UE::Learning
 			}
 		}
 
-
-		/**
-		* Check that an array view does not contain any NaN, Inf, -Inf
-		*/
-		template<uint8 InDimNum, typename InElementType>
-		inline void Check(const TLearningArrayView<InDimNum, const InElementType> View)
-		{
-			for (int64 Idx = 0; Idx < View.Num(); Idx++)
-			{
-				checkf(
-					FMath::IsFinite(View.GetData()[Idx]),
-					TEXT("Invalid value %f found at flat array index %lld"), View.GetData()[Idx], Idx);
-			}
-		}
-
-		/**
-		* Check that an array view does not contain any NaN, Inf, -Inf
-		*/
-		template<uint8 InDimNum, typename InElementType>
-		inline void Check(const TLearningArrayView<InDimNum, const InElementType> View, const FIndexSet Indices)
-		{
-			if constexpr (InDimNum > 1)
-			{
-				for (const int32 Idx : Indices)
-				{
-					Check(View[Idx]);
-				}
-			}
-			else
-			{
-				for (const int32 Idx : Indices)
-				{
-					checkf(
-						FMath::IsFinite(View.GetData()[Idx]) && View.GetData()[Idx],
-						TEXT("Invalid value %f found at flat array index %lld"), View.GetData()[Idx], Idx);
-				}
-			}
-		}
-
 		/**
 		* Check that an array view does not contain any NaN, Inf, -Inf
 		*/
diff --git a/Engine/Source/Programs/UnrealBuildTool/Platform/Windows/VCToolChain.cs b/Engine/Source/Programs/UnrealBuildTool/Platform/Windows/VCToolChain.cs
--- a/Engine/Source/Programs/UnrealBuildTool/Platform/Windows/VCToolChain.cs
+++ b/Engine/Source/Programs/UnrealBuildTool/Platform/Windows/VCToolChain.cs
@@ -498,7 +498,8 @@ namespace UnrealBuildTool
 				}
 
 				// Resource dir is for built-in includes. We need it for vfs but also for uba include crawler to find includes to transfer
-				DirectoryReference ResourceDir = DirectoryReference.Combine(EnvVars.CompilerDir!, "lib", "clang", EnvVars.CompilerVersion.Components[0].ToString());
+				VersionNumber ClangVersion = Target.WindowsPlatform.Compiler == WindowsCompiler.Intel ? MicrosoftPlatformSDK.GetClangVersionForIntelCompiler(EnvVars.CompilerPath) : EnvVars.CompilerVersion;
+				DirectoryReference ResourceDir = DirectoryReference.Combine(EnvVars.CompilerDir!, "lib", "clang", ClangVersion.GetComponent(0).ToString());
 				Arguments.Add($"-resource-dir=\"{NormalizeCommandLinePath(ResourceDir, CompileEnvironment.RootPaths)}\"");
 
 			}
@@ -939,11 +940,22 @@ namespace UnrealBuildTool
 				if (CompileEnvironment.bPGOProfile)
 				{
 					// Generate instrumented code.
-					if (Target.WindowsPlatform.Compiler.IsIntel() && Target.WindowsPlatform.bSampleBasedPGO)
+					if (Target.WindowsPlatform.bSampleBasedPGO)
 					{
-						Arguments.Add("-fprofile-sample-generate");
-						Arguments.Add($"-fprofile-dwo-dir=\"{CompileEnvironment.PGODirectory!}\"");
-						Arguments.Add("-gsplit-dwarf");
+						if (Target.WindowsPlatform.Compiler.IsIntel())
+						{
+							Arguments.Add("-fprofile-sample-generate");
+							Arguments.Add($"-fprofile-dwo-dir=\"{CompileEnvironment.PGODirectory!}\"");
+							Arguments.Add("-gsplit-dwarf");
+						}
+						else
+						{
+							Arguments.Add("-gdwarf");
+							Arguments.Add("-gsplit-dwarf");
+							Arguments.Add("-gline-tables-only");
+							Arguments.Add("/clang:-fdebug-info-for-profiling");
+							Arguments.Add("/clang:-funique-internal-linkage-names");
+						}
 					}
 					else
 					{
@@ -955,9 +967,20 @@ namespace UnrealBuildTool
 					// Use a merged profdata file.
 					string ProfDataFilename = GetClangProfDataFilename(CompileEnvironment);
 					Log.TraceInformationOnce($"Using PGO profile data \"{ProfDataFilename}\"");
-					if (Target.WindowsPlatform.Compiler.IsIntel() && Target.WindowsPlatform.bSampleBasedPGO)
+					if (Target.WindowsPlatform.bSampleBasedPGO)
 					{
-						Arguments.Add($"-fprofile-sample-use=\"{ProfDataFilename}\"");
+						if (Target.WindowsPlatform.Compiler.IsIntel())
+						{
+							Arguments.Add($"-fprofile-sample-use=\"{ProfDataFilename}\"");
+						}
+						else
+						{
+							Arguments.Add("-gdwarf");
+							Arguments.Add("-gline-tables-only");
+							Arguments.Add("/clang:-fdebug-info-for-profiling");
+							Arguments.Add("/clang:-funique-internal-linkage-names");
+							Arguments.Add($"/clang:-fprofile-sample-use=\"{ProfDataFilename}\"");
+						}
 					}
 					else
 					{
@@ -1040,7 +1063,7 @@ namespace UnrealBuildTool
 			}
 
 			// If enabled, create debug information.
-			if (CompileEnvironment.bCreateDebugInfo && !(Target.WindowsPlatform.Compiler.IsIntel() && (CompileEnvironment.bPGOOptimize || CompileEnvironment.bPGOProfile || CompileEnvironment.bAllowLTCG)))
+			if (CompileEnvironment.bCreateDebugInfo && !(Target.WindowsPlatform.Compiler.IsClang() && (CompileEnvironment.bPGOOptimize || CompileEnvironment.bPGOProfile || CompileEnvironment.bAllowLTCG)))
 			{
 				/*
 				Disabled until we can find a better means to limit the cost of /JMC as it can add considerable runtime overhead
@@ -1164,6 +1187,14 @@ namespace UnrealBuildTool
 				Arguments.Add("/Qopt-streaming-stores:auto");
 				Arguments.Add("/Qopt-jump-tables");
 				Arguments.Add("/Qbranches-within-32B-boundaries");
+
+				VersionNumber ClangVersion = MicrosoftPlatformSDK.GetClangVersionForIntelCompiler(EnvVars.CompilerPath);
+
+				// Intel ICX 2025.1 reports ldexpf as part of math library, but comes up as unresolved
+				if (ClangVersion >= new VersionNumber(20))
+				{
+					Arguments.Add("-fno-builtin-ldexpf");
+				}
 			}
 
 			// Intel oneAPI compiler does not support /Zo
@@ -1492,25 +1523,31 @@ namespace UnrealBuildTool
 			// Prevents the linker from displaying its logo for each invocation.
 			Arguments.Add("/NOLOGO");
 
+			// On Clang, use DWARF debug info format for PGO configs.
+			if (Target.WindowsPlatform.Compiler == WindowsCompiler.Clang && Target.WindowsPlatform.bAllowClangLinker &&
+				(LinkEnvironment.bPGOOptimize || LinkEnvironment.bPGOProfile))
+			{
+				Arguments.Add("/DEBUG:DWARF");
+			}
 			// Address sanitizer requires debug info for symbolizing callstacks whether
 			// we're building debug or shipping.
-			if (LinkEnvironment.bCreateDebugInfo || Target.WindowsPlatform.bEnableAddressSanitizer)
+			else if (LinkEnvironment.bCreateDebugInfo || Target.WindowsPlatform.bEnableAddressSanitizer)
 			{
 				// Output debug info for the linked executable.
 				// Beginning in Visual Studio 2017 /DEBUG defaults to /DEBUG:FASTLINK for debug builds
 				Arguments.Add("/DEBUG:FULL");
-			}
 
-			if (LinkEnvironment.bCreateDebugInfo && LinkEnvironment.bUseFastPDBLinking)
-			{
-				// Allow partial PDBs for faster linking on Clang and Intel
-				if (Target.WindowsPlatform.Compiler.IsClang() && Target.WindowsPlatform.bAllowClangLinker)
-				{
-					Arguments[^1] = "/DEBUG:GHASH";
-				}
-				else
+				if (LinkEnvironment.bCreateDebugInfo && LinkEnvironment.bUseFastPDBLinking)
 				{
-					Arguments[^1] = "/DEBUG:FASTLINK";
+					// Allow partial PDBs for faster linking on Clang and Intel
+					if (Target.WindowsPlatform.Compiler.IsClang() && Target.WindowsPlatform.bAllowClangLinker)
+					{
+						Arguments[^1] = "/DEBUG:GHASH";
+					}
+					else
+					{
+						Arguments[^1] = "/DEBUG:FASTLINK";
+					}
 				}
 			}
 
@@ -1659,7 +1696,7 @@ namespace UnrealBuildTool
 			}
 
 			// Identical COMDAT folding. Prevent Identical Code Folding when instrumentation is enabled to avoid thunks to all be folded together and prevent hotpatching.
-			if (Target.WindowsPlatform.bMergeIdenticalCOMDATs && !(Target.WindowsPlatform.Compiler.IsIntel() && Target.WindowsPlatform.bSampleBasedPGO && LinkEnvironment.bPGOProfile) && !Target.WindowsPlatform.bEnableInstrumentation)
+			if (Target.WindowsPlatform.bMergeIdenticalCOMDATs && !(Target.WindowsPlatform.Compiler.IsClang() && Target.WindowsPlatform.bSampleBasedPGO && LinkEnvironment.bPGOProfile) && !Target.WindowsPlatform.bEnableInstrumentation)
 			{
 				Arguments.Add("/OPT:ICF");
 			}
@@ -3483,7 +3520,7 @@ namespace UnrealBuildTool
 				// Link arguments used for PGO on Clang/Intel compiler
 				if (bPGOOptimize)
 				{
-					if (Target.WindowsPlatform.Compiler.IsIntel() && Target.WindowsPlatform.bSampleBasedPGO)
+					if (Target.WindowsPlatform.bSampleBasedPGO)
 					{
 						Log.TraceInformationOnce("Enabling using Sample-Based Profile Guided Optimization (SPGO). Linking will take a while.");
 					}
@@ -3494,11 +3531,16 @@ namespace UnrealBuildTool
 				}
 				else if (bPGOProfile)
 				{
-					// Clang does not yet support sample-based PGO
-					if (Target.WindowsPlatform.Compiler.IsIntel() && Target.WindowsPlatform.bSampleBasedPGO)
+					if (Target.WindowsPlatform.bSampleBasedPGO)
 					{
 						Log.TraceInformationOnce("Enabling generating Sample-based Profile Guided Optimization (SPGO). Linking will take a while.");
-						Arguments.Add("-profile-sample-generate");
+
+						if (Target.WindowsPlatform.Compiler.IsIntel())
+						{
+							Arguments.Add("-profile-sample-generate");
+						}
+
+						Arguments.Add($"/dwodir:\"{LinkEnvironment.PGODirectory!}\"");
 					}
 					else
 					{
diff --git a/Engine/Source/Programs/UnrealBuildTool/ToolChain/ClangWarnings.cs b/Engine/Source/Programs/UnrealBuildTool/ToolChain/ClangWarnings.cs
--- a/Engine/Source/Programs/UnrealBuildTool/ToolChain/ClangWarnings.cs
+++ b/Engine/Source/Programs/UnrealBuildTool/ToolChain/ClangWarnings.cs
@@ -33,7 +33,6 @@ namespace UnrealBuildTool
 		{
 			Arguments.Add("-Wno-deprecated-copy");
 			Arguments.Add("-Wno-deprecated-copy-with-user-provided-copy");
-			Arguments.Add("-Wno-enum-constexpr-conversion");
 			Arguments.Add("-Wno-format");
 			Arguments.Add("-Wno-implicit-float-size-conversion");
 			Arguments.Add("-Wno-null-pointer-subtraction");
diff --git a/Engine/Source/Runtime/Core/Public/IntelICX/IntelICXPlatform.h b/Engine/Source/Runtime/Core/Public/IntelICX/IntelICXPlatform.h
--- a/Engine/Source/Runtime/Core/Public/IntelICX/IntelICXPlatform.h
+++ b/Engine/Source/Runtime/Core/Public/IntelICX/IntelICXPlatform.h
@@ -6,6 +6,8 @@
 
 #pragma once
 
+#if defined(__INTEL_LLVM_COMPILER)
+
 #if !defined(__cpp_if_constexpr)
 	#error "Compiler is expected to support if constexpr"
 #endif
@@ -14,11 +16,17 @@
 	#error "Compiler is expected to support fold expressions"
 #endif
 
+#if !__has_feature(cxx_decltype_auto)
+	#error "Compiler is expected to support decltype(auto)"
+#endif
+
 #define PLATFORM_RETURN_ADDRESS()			__builtin_return_address(0)
 #define PLATFORM_RETURN_ADDRESS_POINTER()	__builtin_frame_address(0)
 
 #define UE_LIFETIMEBOUND [[clang::lifetimebound]]
 
+#define UE_NODEBUG [[gnu::nodebug]]
+
 // Intel ICX supports the following GNU attributes to let us describe allocation functions.
 // `gnu::malloc` maps to the `noalias` function return on the caller in the LLVM IR.
 // `gnu::alloc_size` maps to the `allocsize` function attribute on the caller in the LLVM IR.
@@ -34,8 +42,18 @@
 #define UE_ALLOCATION_FUNCTION_X(x, SIZE, ALIGN, FUNC, ...) FUNC
 #define UE_ALLOCATION_FUNCTION(...) UE_ALLOCATION_FUNCTION_X(,##__VA_ARGS__, UE_ALLOCATION_FUNCTION_2(__VA_ARGS__), UE_ALLOCATION_FUNCTION_1(__VA_ARGS__), UE_ALLOCATION_FUNCTION_0(__VA_ARGS__)) 
 
-// Intel ICX does not support constexpr for __builtin_FILE and __builtin_LINE
-// Check with (__has_constexpr_builtin(__builtin_FILE) && __has_constexpr_builtin(__builtin_LINE)
-#define PLATFORM_COMPILER_SUPPORTS_CONSTEXPR_BUILTIN_FILE_AND_LINE 0
+#define PLATFORM_COMPILER_SUPPORTS_BUILTIN_BITCAST 1
+
+#define FUNCTION_NON_NULL_RETURN_START [[gnu::returns_nonnull]]
+
+#if defined(__has_attribute)
+#	if __has_attribute(no_profile_instrument_function)
+#		define UE_NO_PROFILE_ATTRIBUTE __attribute__((no_profile_instrument_function))
+#	endif
+#endif
+
+#ifndef UE_NO_PROFILE_ATTRIBUTE
+#	define UE_NO_PROFILE_ATTRIBUTE
+#endif
 
-#define PLATFORM_COMPILER_SUPPORTS_BUILTIN_BITCAST 0
+#endif // defined(__INTEL_LLVM_COMPILER)
diff --git a/Engine/Source/Runtime/Core/Public/Misc/Attribute.h b/Engine/Source/Runtime/Core/Public/Misc/Attribute.h
--- a/Engine/Source/Runtime/Core/Public/Misc/Attribute.h
+++ b/Engine/Source/Runtime/Core/Public/Misc/Attribute.h
@@ -115,20 +115,6 @@ public:
 		return TAttribute(MoveTemp(InGetter), bExplicitConstructor);
 	}
 
-	/**
-	 * Creates an attribute by binding an arbitrary function that will be called to generate this attribute's value on demand. 
-	 * After binding, the attribute will no longer have a value that can be accessed directly, and instead the bound
-	 * function will always be called to generate the value.
-	 *
-	 * @param  InFuncPtr Member function to bind.  The function's structure (return value, arguments, etc) must match IBoundAttributeDelegate's definition.
-	 */
-	template <typename... VarTypes >
-	[[nodiscard]] static TAttribute CreateStatic( TIdentity_T< typename FGetter::template TFuncPtr< VarTypes... > > InFuncPtr, VarTypes... Vars )
-	{
-		const bool bExplicitConstructor = true;
-		return TAttribute( FGetter::CreateStatic( InFuncPtr, Vars... ), bExplicitConstructor );
-	}
-
 	/**
 	 * Creates an attribute by binding an arbitrary function that will be called to generate this attribute's value on demand.
 	 * After binding, the attribute will no longer have a value that can be accessed directly, and instead the bound
diff --git a/Engine/Source/Runtime/Core/Public/Windows/WindowsPlatform.h b/Engine/Source/Runtime/Core/Public/Windows/WindowsPlatform.h
--- a/Engine/Source/Runtime/Core/Public/Windows/WindowsPlatform.h
+++ b/Engine/Source/Runtime/Core/Public/Windows/WindowsPlatform.h
@@ -6,10 +6,10 @@
 
 #include <sal.h>
 
-#if defined(__clang__)
-	#include "Clang/ClangPlatform.h"
-#elif defined(__INTEL_LLVM_COMPILER)
+#if defined(__INTEL_LLVM_COMPILER)
 	#include "IntelICX/IntelICXPlatform.h"
+#elif defined(__clang__)
+	#include "Clang/ClangPlatform.h"
 #else
 	#include "MSVC/MSVCPlatform.h"
 #endif
diff --git a/Engine/Source/Runtime/D3D12RHI/Private/D3D12Resources.h b/Engine/Source/Runtime/D3D12RHI/Private/D3D12Resources.h
--- a/Engine/Source/Runtime/D3D12RHI/Private/D3D12Resources.h
+++ b/Engine/Source/Runtime/D3D12RHI/Private/D3D12Resources.h
@@ -1074,7 +1074,7 @@ class FD3D12ResourceBarrierBatcher
 {
 	PRAGMA_DISABLE_ENUM_CONSTEXPR_CONVERSION_WARNINGS
 	// Use the top bit of the flags enum to mark transitions as "idle" time (used to remove the swapchain wait time for back buffers).
-	static constexpr D3D12_RESOURCE_BARRIER_FLAGS BarrierFlag_CountAsIdleTime = D3D12_RESOURCE_BARRIER_FLAGS(1ull << ((sizeof(D3D12_RESOURCE_BARRIER_FLAGS) * 8) - 1));
+	static const D3D12_RESOURCE_BARRIER_FLAGS BarrierFlag_CountAsIdleTime = D3D12_RESOURCE_BARRIER_FLAGS(1ull << ((sizeof(D3D12_RESOURCE_BARRIER_FLAGS) * 8) - 1));
 	PRAGMA_ENABLE_ENUM_CONSTEXPR_CONVERSION_WARNINGS
 
 	struct FD3D12ResourceBarrier : public D3D12_RESOURCE_BARRIER
diff --git a/Engine/Source/Runtime/Slate/Private/Widgets/Text/STextBlock.cpp b/Engine/Source/Runtime/Slate/Private/Widgets/Text/STextBlock.cpp
--- a/Engine/Source/Runtime/Slate/Private/Widgets/Text/STextBlock.cpp
+++ b/Engine/Source/Runtime/Slate/Private/Widgets/Text/STextBlock.cpp
@@ -48,35 +48,35 @@ void STextBlock::PrivateRegisterAttributes(FSlateAttributeInitializer& Attribute
 		.OnValueChanged(FSlateAttributeDescriptor::FAttributeValueChangedDelegate::CreateStatic(FInvalidation::UpdateDesiredSize));
 	SLATE_ADD_MEMBER_ATTRIBUTE_DEFINITION(AttributeInitializer, Font, EInvalidateWidgetReason::Layout)
 		.OnValueChanged(FSlateAttributeDescriptor::FAttributeValueChangedDelegate::CreateStatic(FInvalidation::UpdateTextStyle));
-	SLATE_ADD_MEMBER_ATTRIBUTE_DEFINITION(AttributeInitializer, StrikeBrush, FInvalidation::GetInvalidationNoneIfSimpleTextMode)
+	SLATE_ADD_MEMBER_ATTRIBUTE_DEFINITION(AttributeInitializer, StrikeBrush, EInvalidateWidgetReason::Layout)
 		.OnValueChanged(FSlateAttributeDescriptor::FAttributeValueChangedDelegate::CreateStatic(FInvalidation::UpdateTextStyle));
-	SLATE_ADD_MEMBER_ATTRIBUTE_DEFINITION(AttributeInitializer, ColorAndOpacity, FInvalidation::GetInvalidationPaintIfSimpleTextMode)
+	SLATE_ADD_MEMBER_ATTRIBUTE_DEFINITION(AttributeInitializer, ColorAndOpacity, EInvalidateWidgetReason::Layout)
 		.OnValueChanged(FSlateAttributeDescriptor::FAttributeValueChangedDelegate::CreateStatic(FInvalidation::UpdateTextStyle));
 	SLATE_ADD_MEMBER_ATTRIBUTE_DEFINITION(AttributeInitializer, ShadowOffset, EInvalidateWidgetReason::Layout)
 		.OnValueChanged(FSlateAttributeDescriptor::FAttributeValueChangedDelegate::CreateStatic(FInvalidation::UpdateTextStyle));
-	SLATE_ADD_MEMBER_ATTRIBUTE_DEFINITION(AttributeInitializer, ShadowColorAndOpacity, FInvalidation::GetInvalidationPaintIfSimpleTextMode)
+	SLATE_ADD_MEMBER_ATTRIBUTE_DEFINITION(AttributeInitializer, ShadowColorAndOpacity, EInvalidateWidgetReason::Layout)
 		.OnValueChanged(FSlateAttributeDescriptor::FAttributeValueChangedDelegate::CreateStatic(FInvalidation::UpdateTextStyle));
-	SLATE_ADD_MEMBER_ATTRIBUTE_DEFINITION(AttributeInitializer, HighlightColor, FInvalidation::GetInvalidationNoneIfSimpleTextMode)
+	SLATE_ADD_MEMBER_ATTRIBUTE_DEFINITION(AttributeInitializer, HighlightColor, EInvalidateWidgetReason::Layout)
 		.OnValueChanged(FSlateAttributeDescriptor::FAttributeValueChangedDelegate::CreateStatic(FInvalidation::UpdateTextStyle));
-	SLATE_ADD_MEMBER_ATTRIBUTE_DEFINITION(AttributeInitializer, HighlightShape, FInvalidation::GetInvalidationNoneIfSimpleTextMode)
+	SLATE_ADD_MEMBER_ATTRIBUTE_DEFINITION(AttributeInitializer, HighlightShape, EInvalidateWidgetReason::Layout)
 		.OnValueChanged(FSlateAttributeDescriptor::FAttributeValueChangedDelegate::CreateStatic(FInvalidation::UpdateTextStyle));
-	SLATE_ADD_MEMBER_ATTRIBUTE_DEFINITION(AttributeInitializer, HighlightText, FInvalidation::GetInvalidationNoneIfSimpleTextMode)
+	SLATE_ADD_MEMBER_ATTRIBUTE_DEFINITION(AttributeInitializer, HighlightText, EInvalidateWidgetReason::Layout)
 		.OnValueChanged(FSlateAttributeDescriptor::FAttributeValueChangedDelegate::CreateStatic(FInvalidation::UpdateDesiredSize));
-	SLATE_ADD_MEMBER_ATTRIBUTE_DEFINITION(AttributeInitializer, WrapTextAt, FInvalidation::GetInvalidationNoneIfSimpleTextMode)
+	SLATE_ADD_MEMBER_ATTRIBUTE_DEFINITION(AttributeInitializer, WrapTextAt, EInvalidateWidgetReason::Layout)
 		.OnValueChanged(FSlateAttributeDescriptor::FAttributeValueChangedDelegate::CreateStatic(FInvalidation::UpdateDesiredSize));
-	SLATE_ADD_MEMBER_ATTRIBUTE_DEFINITION(AttributeInitializer, AutoWrapText, FInvalidation::GetInvalidationNoneIfSimpleTextMode)
+	SLATE_ADD_MEMBER_ATTRIBUTE_DEFINITION(AttributeInitializer, AutoWrapText, EInvalidateWidgetReason::Layout)
 		.OnValueChanged(FSlateAttributeDescriptor::FAttributeValueChangedDelegate::CreateStatic(FInvalidation::UpdateDesiredSize));
-	SLATE_ADD_MEMBER_ATTRIBUTE_DEFINITION(AttributeInitializer, WrappingPolicy, FInvalidation::GetInvalidationNoneIfSimpleTextMode)
+	SLATE_ADD_MEMBER_ATTRIBUTE_DEFINITION(AttributeInitializer, WrappingPolicy, EInvalidateWidgetReason::Layout)
 		.OnValueChanged(FSlateAttributeDescriptor::FAttributeValueChangedDelegate::CreateStatic(FInvalidation::UpdateDesiredSize));
-	SLATE_ADD_MEMBER_ATTRIBUTE_DEFINITION(AttributeInitializer, TransformPolicy, FInvalidation::GetInvalidationNoneIfSimpleTextMode)
+	SLATE_ADD_MEMBER_ATTRIBUTE_DEFINITION(AttributeInitializer, TransformPolicy, EInvalidateWidgetReason::Layout)
 		.OnValueChanged(FSlateAttributeDescriptor::FAttributeValueChangedDelegate::CreateStatic(FInvalidation::UpdateDesiredSize));
-	SLATE_ADD_MEMBER_ATTRIBUTE_DEFINITION(AttributeInitializer, Margin, FInvalidation::GetInvalidationNoneIfSimpleTextMode)
+	SLATE_ADD_MEMBER_ATTRIBUTE_DEFINITION(AttributeInitializer, Margin, EInvalidateWidgetReason::Layout)
 		.OnValueChanged(FSlateAttributeDescriptor::FAttributeValueChangedDelegate::CreateStatic(FInvalidation::UpdateDesiredSize));
-	SLATE_ADD_MEMBER_ATTRIBUTE_DEFINITION(AttributeInitializer, Justification, FInvalidation::GetInvalidationNoneIfSimpleTextMode)
+	SLATE_ADD_MEMBER_ATTRIBUTE_DEFINITION(AttributeInitializer, Justification, EInvalidateWidgetReason::Layout)
 		.OnValueChanged(FSlateAttributeDescriptor::FAttributeValueChangedDelegate::CreateStatic(FInvalidation::UpdateDesiredSize));
-	SLATE_ADD_MEMBER_ATTRIBUTE_DEFINITION(AttributeInitializer, LineHeightPercentage, FInvalidation::GetInvalidationNoneIfSimpleTextMode)
+	SLATE_ADD_MEMBER_ATTRIBUTE_DEFINITION(AttributeInitializer, LineHeightPercentage, EInvalidateWidgetReason::Layout)
 		.OnValueChanged(FSlateAttributeDescriptor::FAttributeValueChangedDelegate::CreateStatic(FInvalidation::UpdateDesiredSize));
-	SLATE_ADD_MEMBER_ATTRIBUTE_DEFINITION(AttributeInitializer, ApplyLineHeightToBottomLine, FInvalidation::GetInvalidationNoneIfSimpleTextMode)
+	SLATE_ADD_MEMBER_ATTRIBUTE_DEFINITION(AttributeInitializer, ApplyLineHeightToBottomLine, EInvalidateWidgetReason::Layout)
 		.OnValueChanged(FSlateAttributeDescriptor::FAttributeValueChangedDelegate::CreateStatic(FInvalidation::UpdateDesiredSize));
 	SLATE_ADD_MEMBER_ATTRIBUTE_DEFINITION(AttributeInitializer, MinDesiredWidth, EInvalidateWidgetReason::Layout);
 }
