[UBT] Minor enhancements for CppCompileWarnings system.

diff --git a/Engine/Source/Programs/Shared/UnrealBuildTool.Tests/CppCompileWarningsTests.cs b/Engine/Source/Programs/Shared/UnrealBuildTool.Tests/CppCompileWarningsTests.cs
--- a/Engine/Source/Programs/Shared/UnrealBuildTool.Tests/CppCompileWarningsTests.cs
+++ b/Engine/Source/Programs/Shared/UnrealBuildTool.Tests/CppCompileWarningsTests.cs
@@ -318,6 +318,41 @@ namespace UnrealBuildTool.Tests
 					baseWarnings.UndefinedIdentifierWarningLevel = WarningLevel.Default;
 				}
 
+				{
+					// DeterministicWarningLevel
+					VersionNumber newVersionNumber = new VersionNumber(18, 9, 9);
+
+					CppCompileEnvironment deterministicCompileEnvironment = new CppCompileEnvironment(UnrealTargetPlatform.Win64, CppConfiguration.Debug, new UnrealArchitectures(UnrealArch.X64), null!);
+					deterministicCompileEnvironment.bDeterministic = false;
+
+
+					args = baseWarnings.GenerateWarningCommandLineArgs(deterministicCompileEnvironment, typeof(VCToolChain), newVersionNumber);
+					Assert.IsFalse(args.Any());
+
+					deterministicCompileEnvironment.bDeterministic = true;
+
+					args = baseWarnings.GenerateWarningCommandLineArgs(deterministicCompileEnvironment, typeof(VCToolChain), newVersionNumber);
+					Assert.IsFalse(args.Any());
+
+					baseWarnings.DeterministicWarningLevel = WarningLevel.Off;
+					args = baseWarnings.GenerateWarningCommandLineArgs(deterministicCompileEnvironment, typeof(VCToolChain), newVersionNumber);
+
+					Assert.AreEqual(1, args.Count());
+					Assert.IsTrue(args.Contains("/wd5048"));
+
+					baseWarnings.DeterministicWarningLevel = WarningLevel.Warning;
+					args = baseWarnings.GenerateWarningCommandLineArgs(deterministicCompileEnvironment, typeof(VCToolChain), newVersionNumber);
+
+					Assert.AreEqual(0, args.Count());
+					Assert.IsFalse(args.Any());
+
+					baseWarnings.DeterministicWarningLevel = WarningLevel.Error;
+					args = baseWarnings.GenerateWarningCommandLineArgs(deterministicCompileEnvironment, typeof(VCToolChain), newVersionNumber);
+
+					Assert.AreEqual(1, args.Count());
+					Assert.IsTrue(args.Contains("/we5048"));
+				}
+
 				// Specialized VC toolchain check;  UnsafeTypeCast
 				{
 					RulesAssembly generatedTestGameRulesAssembly = RulesCompiler.CreateTargetRulesAssembly(testUproject, "TestStub", false, false, false, null, false, new TestLogger());
diff --git a/Engine/Source/Programs/Shared/UnrealBuildTool.Tests/TargetRulesTests.cs b/Engine/Source/Programs/Shared/UnrealBuildTool.Tests/TargetRulesTests.cs
--- a/Engine/Source/Programs/Shared/UnrealBuildTool.Tests/TargetRulesTests.cs
+++ b/Engine/Source/Programs/Shared/UnrealBuildTool.Tests/TargetRulesTests.cs
@@ -50,11 +50,40 @@ namespace UnrealBuildTool.Tests
 					testTargetRules.CppCompileWarningSettings.ShadowVariableWarningLevel = preWarningLevel;
 				}
 
+
+				// Mutate the object against a RequiresUniqueEnvironment (Clang only)
+				{
+					WarningLevel preWarningLevel = testTargetRules.CppCompileWarningSettings.NonTrivialMemAccessWarningLevel;
+					testTargetRules.CppCompileWarningSettings.NonTrivialMemAccessWarningLevel = WarningLevel.Warning;
+					bool result = testTargetRules.RequiresUniqueEnvironment(generatedTestRulesAssembly, null, new Dictionary<string, (string?, string?)>(), out baseTarget);
+
+					Assert.IsTrue(result);
+
+					testTargetRules.CppCompileWarningSettings.NonTrivialMemAccessWarningLevel = preWarningLevel;
+				}
+
+				// Mutate the object against a RequiresUniqueEnvironment (Default warning set)
+				{
+					WarningLevel preWarningLevel = testTargetRules.CppCompileWarningSettings.DeprecationWarningLevel;
+					testTargetRules.CppCompileWarningSettings.DeprecationWarningLevel = WarningLevel.Error;
+					bool result = testTargetRules.RequiresUniqueEnvironment(generatedTestRulesAssembly, null, new Dictionary<string, (string?, string?)>(), out baseTarget);
+
+					Assert.IsTrue(result);
+
+					testTargetRules.CppCompileWarningSettings.DeprecationWarningLevel = WarningLevel.Warning;
+					result = testTargetRules.RequiresUniqueEnvironment(generatedTestRulesAssembly, null, new Dictionary<string, (string?, string?)>(), out baseTarget);
+
+					Assert.IsFalse(result);
+
+					testTargetRules.CppCompileWarningSettings.DeprecationWarningLevel = preWarningLevel;
+				}
+
 				// Mutate the object against a RequiresUniqueEnvironment, and set bOverrideBuildEnvironment=true
 				{
 					WarningLevel preWarningLevel = testTargetRules.CppCompileWarningSettings.ShadowVariableWarningLevel;
 					testTargetRules.CppCompileWarningSettings.ShadowVariableWarningLevel = WarningLevel.Warning;
 					testTargetRules.bOverrideBuildEnvironment = true;
+
 					bool result = testTargetRules.RequiresUniqueEnvironment(generatedTestRulesAssembly, null, new Dictionary<string, (string?, string?)>(), out baseTarget);
 
 					Assert.IsFalse(result);
diff --git a/Engine/Source/Programs/UnrealBuildTool/Configuration/CppCompileWarnings.cs b/Engine/Source/Programs/UnrealBuildTool/Configuration/CppCompileWarnings.cs
--- a/Engine/Source/Programs/UnrealBuildTool/Configuration/CppCompileWarnings.cs
+++ b/Engine/Source/Programs/UnrealBuildTool/Configuration/CppCompileWarnings.cs
@@ -84,7 +84,6 @@ namespace UnrealBuildTool
 		///		 https://clang.llvm.org/docs/DiagnosticsReference.html#wshadow
 		/// </summary>
 		[CommandLine("-ShadowVariableErrors", Value = nameof(WarningLevel.Error))]
-		[RequiresUniqueBuildEnvironment]
 		[WarningsMSCV(["/wd4456", "/wd4458", "/wd4459"], null, ["/we4456", "/we4458", "/we4459"])]
 		[ShadowVariableWarningsClangToolChain(["-Wno-shadow"], ["-Wshadow", "-Wno-error=shadow"], ["-Wshadow"])]
 		[WarningLevelDefault(WarningLevel.Warning, BuildSettingsVersion.V1, BuildSettingsVersion.Latest, InitializationContext.Constructor)]
@@ -119,7 +118,6 @@ namespace UnrealBuildTool
 		/// <remarks>This should be kept in sync with PRAGMA_DISABLE_UNSAFE_TYPECAST_WARNINGS in ClangPlatformCompilerPreSetup.h</remarks>
 		/// <remarks>Clang: "shorten-64-to-32"; too many hits right now, probably want it *soon*; "sign-conversion"; too many hits right now, probably want it eventually</remarks>
 		[XmlConfigFile(Category = "BuildConfiguration")]
-		[RequiresUniqueBuildEnvironment]
 		[UnsafeTypeCastWarningsVCToolChain(["/wd4244", "/wd4838"], ["/w44244", "/w44838"], ["/we4244", "/we4838"])]
 		[WarningsClangToolChain(
 			["-Wno-float-conversion", "-Wno-implicit-float-conversion", "-Wno-implicit-int-conversion", "-Wno-c++11-narrowing"],
@@ -144,7 +142,6 @@ namespace UnrealBuildTool
 		/// </summary>
 		[XmlConfigFile(Category = "BuildConfiguration")]
 		[CommandLine("-UndefinedIdentifierWarningLevel=")]
-		[RequiresUniqueBuildEnvironment]
 		[UndefinedIdentifierWarningsVCToolChain(null, ["/w44668"], ["/we4668"])]
 		[WarningsClangToolChain(null, ["-Wundef", "-Wno-error=undef"], ["-Wundef"])]
 		[WarningLevelDefault(WarningLevel.Off)]
@@ -165,7 +162,6 @@ namespace UnrealBuildTool
 		/// </summary>
 		[XmlConfigFile(Category = "BuildConfiguration")]
 		[CommandLine("-SwitchUnhandledEnumeratorWarningLevel=")]
-		[RequiresUniqueBuildEnvironment]
 		[WarningsVCToolChain(null, ["/w44061"], ["/we4061"])]
 		[WarningsClangToolChain(null, ["-Wswitch-enum", "-Wno-error=switch-enum"], ["-Wswitch-enum"])]
 		[WarningLevelDefault(WarningLevel.Off)]
@@ -1068,7 +1064,9 @@ namespace UnrealBuildTool
 		///		https://clang.llvm.org/docs/DiagnosticsReference.html#wdeprecated-declarations
 		/// </summary>
 		/// <remarks>Disable needed for wxWidgets.</remarks>
+		[XmlConfigFile(Category = "BuildConfiguration", Name = nameof(DeprecationWarningLevel))]
 		[WarningsClangToolChain(["-Wno-deprecated-declarations"], null, ["-Werror=deprecated-declarations"])]
+		[WarningsVCToolChain(["/wd4996"], null, ["/we4996"])]
 		[WarningLevelDefault(WarningLevel.Warning)]
 		public WarningLevel DeprecationWarningLevel
 		{
@@ -1083,6 +1081,7 @@ namespace UnrealBuildTool
 		///		https://clang.llvm.org/docs/DiagnosticsReference.html#wdate-time
 		/// </summary>
 		[WarningsClangToolChain(null, ["-Wdate-time", "-Wno-error=date-time"], ["-Wdate-time"], [nameof(StandardFilters.DeterministcFlagSetFilter)])]
+		[WarningsMSCV(["/wd5048"], null, ["/we5048"], [nameof(StandardFilters.DeterministcFlagSetFilter)])]
 		[WarningLevelDefault(WarningLevel.Off)]
 		public WarningLevel DeterministicWarningLevel
 		{
diff --git a/Engine/Source/Programs/UnrealBuildTool/Configuration/ModuleRules.cs b/Engine/Source/Programs/UnrealBuildTool/Configuration/ModuleRules.cs
--- a/Engine/Source/Programs/UnrealBuildTool/Configuration/ModuleRules.cs
+++ b/Engine/Source/Programs/UnrealBuildTool/Configuration/ModuleRules.cs
@@ -971,13 +971,10 @@ namespace UnrealBuildTool
 		/// </summary>
 		public WarningLevel DeterministicWarningLevel
 		{
-			get => (DeterministicWarningLevelPrivate == WarningLevel.Default) ? (Target.bDeterministic ? WarningLevel.Warning : WarningLevel.Off) : DeterministicWarningLevelPrivate;
-			set => DeterministicWarningLevelPrivate = value;
+			get => CppCompileWarningSettings.DeterministicWarningLevel;
+			set => CppCompileWarningSettings.DeterministicWarningLevel = value;
 		}
 
-		/// <inheritdoc cref="DeterministicWarningLevelPrivate"/>
-		private WarningLevel DeterministicWarningLevelPrivate;
-
 		/// <summary>
 		/// Settings object used by the module.
 		/// </summary>
diff --git a/Engine/Source/Programs/UnrealBuildTool/Configuration/TargetRules.cs b/Engine/Source/Programs/UnrealBuildTool/Configuration/TargetRules.cs
--- a/Engine/Source/Programs/UnrealBuildTool/Configuration/TargetRules.cs
+++ b/Engine/Source/Programs/UnrealBuildTool/Configuration/TargetRules.cs
@@ -1864,17 +1864,12 @@ namespace UnrealBuildTool
 		/// <summary>
 		/// Level to report deprecation warnings as errors
 		/// </summary>
-		[RequiresUniqueBuildEnvironment]
 		public WarningLevel DeprecationWarningLevel
 		{
-			get => (DeprecationWarningLevelPrivate == WarningLevel.Default) ? DefaultWarningLevel : DeprecationWarningLevelPrivate;
-			set => DeprecationWarningLevelPrivate = value;
+			get => CppCompileWarningSettings.DeprecationWarningLevel;
+			set => CppCompileWarningSettings.DeprecationWarningLevel = value;
 		}
 
-		/// <inheritdoc cref="DeprecationWarningLevel"/>
-		[XmlConfigFile(Category = "BuildConfiguration", Name = nameof(DeprecationWarningLevel))]
-		private WarningLevel DeprecationWarningLevelPrivate;
-
 		/// <summary>
 		/// C++ Warnings settings object used by the target
 		/// </summary>
@@ -3023,6 +3018,31 @@ namespace UnrealBuildTool
 				}
 			}
 
+			// Iterate through all fields with the [ApplyWarnings] attribute, which are consumed by the CppCompileWarnings system and inserted into the compile arguments
+			for (int idx = 0; idx < baseObjects.Length; idx++)
+			{
+				Type objectType = baseObjects[idx].GetType();
+				foreach (FieldInfo field in objectType.GetFields().Where(field => field.GetCustomAttributes<ApplyWarningsAttribute>().FirstOrDefault() != null))
+				{
+					object? thisValue = field.GetValue(thisObjects[idx]);
+					object? baseValue = field.GetValue(baseObjects[idx]);
+					if (!CheckValuesMatch(field.FieldType, thisValue, baseValue))
+					{
+						propNamesThatRequireUnique.Add(field.Name, (thisValue?.ToString(), baseValue?.ToString()));
+					}
+				}
+
+				foreach (PropertyInfo property in objectType.GetProperties().Where(property => property.GetCustomAttributes<ApplyWarningsAttribute>().FirstOrDefault() != null))
+				{
+					object? thisValue = property.GetValue(thisObjects[idx]);
+					object? baseValue = property.GetValue(baseObjects[idx]);
+					if (!CheckValuesMatch(property.PropertyType, thisValue, baseValue))
+					{
+						propNamesThatRequireUnique.Add(property.Name, (thisValue?.ToString(), baseValue?.ToString()));
+					}
+				}
+			}
+
 			// if any properties require a unique build environment, return true
 			return propNamesThatRequireUnique.Count > 0;
 		}
diff --git a/Engine/Source/Programs/UnrealBuildTool/Configuration/UEBuildModuleCPP.cs b/Engine/Source/Programs/UnrealBuildTool/Configuration/UEBuildModuleCPP.cs
--- a/Engine/Source/Programs/UnrealBuildTool/Configuration/UEBuildModuleCPP.cs
+++ b/Engine/Source/Programs/UnrealBuildTool/Configuration/UEBuildModuleCPP.cs
@@ -2320,7 +2320,6 @@ namespace UnrealBuildTool
 			Result.StaticAnalyzerPVSDisabledErrors = Rules.StaticAnalyzerPVSDisabledErrors;
 			Result.bTreatAsEngineModule = Rules.bTreatAsEngineModule;
 			Result.IncludeOrderVersion = Rules.IncludeOrderVersion;
-			Result.DeterministicWarningLevel = Rules.DeterministicWarningLevel;
 			Result.bValidateFormatStrings = Rules.bValidateFormatStrings;
 			Result.bValidateInternalApi = Rules.bValidateInternalApi;
 			Result.bValidateExperimentalApi = Rules.bValidateExperimentalApi;
diff --git a/Engine/Source/Programs/UnrealBuildTool/Configuration/UEBuildTarget.cs b/Engine/Source/Programs/UnrealBuildTool/Configuration/UEBuildTarget.cs
--- a/Engine/Source/Programs/UnrealBuildTool/Configuration/UEBuildTarget.cs
+++ b/Engine/Source/Programs/UnrealBuildTool/Configuration/UEBuildTarget.cs
@@ -5276,7 +5276,6 @@ namespace UnrealBuildTool
 			GlobalCompileEnvironment.bEnableExceptions = Rules.bForceEnableExceptions || (Rules.bCompileAgainstEditor && !Rules.bUseAutoRTFMCompiler);
 			GlobalCompileEnvironment.bEnableObjCExceptions = Rules.bForceEnableObjCExceptions || Rules.bCompileAgainstEditor;
 			GlobalCompileEnvironment.DefaultWarningLevel = Rules.DefaultWarningLevel;
-			GlobalCompileEnvironment.DeprecationWarningLevel = Rules.DeprecationWarningLevel;
 			GlobalCompileEnvironment.CppCompileWarnings = Rules.CppCompileWarningSettings.CloneAsWriteable();
 			GlobalCompileEnvironment.bRetainFramePointers = Rules.bRetainFramePointers;
 			GlobalCompileEnvironment.bWarningsAsErrors = Rules.bWarningsAsErrors;
diff --git a/Engine/Source/Programs/UnrealBuildTool/Platform/Windows/VCToolChain.cs b/Engine/Source/Programs/UnrealBuildTool/Platform/Windows/VCToolChain.cs
--- a/Engine/Source/Programs/UnrealBuildTool/Platform/Windows/VCToolChain.cs
+++ b/Engine/Source/Programs/UnrealBuildTool/Platform/Windows/VCToolChain.cs
@@ -720,15 +720,6 @@ namespace UnrealBuildTool
 
 					// warning C5049: Embedding a full path may result in machine-dependent output
 					Arguments.Add("/wd5049");
-
-					if (CompileEnvironment.DeterministicWarningLevel == WarningLevel.Error)
-					{
-						Arguments.Add("/we5048");
-					}
-					else if (CompileEnvironment.DeterministicWarningLevel == WarningLevel.Off)
-					{
-						Arguments.Add("/wd5048");
-					}
 				}
 			}
 
@@ -1215,15 +1206,6 @@ namespace UnrealBuildTool
 				Arguments.Add("/WX");
 			}
 
-			if (CompileEnvironment.DeprecationWarningLevel == WarningLevel.Off)
-			{
-				Arguments.Add("/wd4996");
-			} 
-			else if (CompileEnvironment.DeprecationWarningLevel == WarningLevel.Error)
-			{
-				Arguments.Add("/we4996");
-			}
-
 			if (CompileEnvironment.bUseHeaderUnitsForPch)
 			{
 				Arguments.Add("/wd4324"); // 'struct_name' : structure was padded due to __declspec(align())
diff --git a/Engine/Source/Programs/UnrealBuildTool/System/CppCompileEnvironment.cs b/Engine/Source/Programs/UnrealBuildTool/System/CppCompileEnvironment.cs
--- a/Engine/Source/Programs/UnrealBuildTool/System/CppCompileEnvironment.cs
+++ b/Engine/Source/Programs/UnrealBuildTool/System/CppCompileEnvironment.cs
@@ -390,15 +390,6 @@ namespace UnrealBuildTool
 		/// </summary>
 		public WarningLevel DefaultWarningLevel = WarningLevel.Warning;
 
-		/// <summary>
-		/// Whether to warn about deprecated variables
-		/// </summary>
-		public WarningLevel DeprecationWarningLevel
-		{
-			get => CppCompileWarnings.DeprecationWarningLevel;
-			set => CppCompileWarnings.DeprecationWarningLevel = value;
-		}
-
 		/// <summary>
 		/// The compiler warnings associacted with the compile environment.
 		/// </summary>
