UnrealBuildTool: UBA cache providers

#rnx

diff --git a/Engine/Config/BaseEngine.ini b/Engine/Config/BaseEngine.ini
--- a/Engine/Config/BaseEngine.ini
+++ b/Engine/Config/BaseEngine.ini
@@ -3841,7 +3841,6 @@ SendOSData=true
 
 ; Example Horde UBA Distribution Settings - See Engine/Source/Programs/UnrealBuildTool/Executors/UnrealBuildAccelerator/UnrealBuildAcceleratorHordeConfig.cs
 ;[UbaController]
-; DefaultProvider=Uba.Provider.Horde.Default ; The default provider if none are set
 ; +BuildMachineProviders=Uba.Provider.Horde.BuildMachine ; Providers for build machines
 ; +Providers=Uba.Provider.Horde.MyOwnHordeServer ; Example to add multiple horde providers
 ; +Providers=Uba.Provider.Horde.MyOwnHordeServer2
@@ -3860,6 +3859,20 @@ SendOSData=true
 ; Pool=MyUbaPool2
 ; Enabled=True
 
+; Example UBA Cache Settings
+; Engine\Binaries\Win64\UnrealBuildAccelerator\x64\UbaCacheService.exe
+;[UbaController]
+; +BuildMachineCacheProviders=Uba.CacheProvider.Uba.MyOwnCacheServer ; Cache providers for build machines
+; +CacheProviders=Uba.CacheProvider.Uba.LocalHost ; Default localhost cache if running
+;[Uba.CacheProvider.Uba.MyOwnCacheServer]
+; Url=myowncacheserver:1347 ; Required
+; Write=BuildMachineOnly ; Optional: Only allow build machines to write to this cache. Writing is disabled by default
+; RequireVFS=True ; Optional: Requires a UnrealBuildTool action have VFS enabled to maintain consistent pathing. Not required by default
+;[Uba.CacheProvider.Uba.LocalHost]
+; Url=127.0.0.1:1347
+; Write=True
+; RequireVFS=False
+
 [/Script/Engine.WorldPartitionSettings]
 EditorHashDefaultClass=/Script/CoreUObject.Class'/Script/Engine.WorldPartitionEditorSpatialHash'
 RuntimeHashDefaultClass=/Script/CoreUObject.Class'/Script/Engine.WorldPartitionRuntimeHashSet'
diff --git a/Engine/Source/Programs/UnrealBuildTool/Executors/UnrealBuildAccelerator/UBAAgentCoordinatorHorde.cs b/Engine/Source/Programs/UnrealBuildTool/Executors/UnrealBuildAccelerator/UBAAgentCoordinatorHorde.cs
--- a/Engine/Source/Programs/UnrealBuildTool/Executors/UnrealBuildAccelerator/UBAAgentCoordinatorHorde.cs
+++ b/Engine/Source/Programs/UnrealBuildTool/Executors/UnrealBuildAccelerator/UBAAgentCoordinatorHorde.cs
@@ -739,7 +739,6 @@ namespace UnrealBuildTool
 			_hordeConfig = new();
 
 			ConfigHierarchy engineIni = ConfigCache.ReadHierarchy(ConfigHierarchyType.Engine, projectDir, BuildHostPlatform.Current.Platform);
-			ConfigHierarchySection section = engineIni.FindSection(provider);
 			_hordeConfig.LoadConfigProvider(engineIni, provider);
 
 			if (provider == ProviderPrefix)
diff --git a/Engine/Source/Programs/UnrealBuildTool/Executors/UnrealBuildAccelerator/UBAExecutor.cs b/Engine/Source/Programs/UnrealBuildTool/Executors/UnrealBuildAccelerator/UBAExecutor.cs
--- a/Engine/Source/Programs/UnrealBuildTool/Executors/UnrealBuildAccelerator/UBAExecutor.cs
+++ b/Engine/Source/Programs/UnrealBuildTool/Executors/UnrealBuildAccelerator/UBAExecutor.cs
@@ -27,7 +27,8 @@ namespace UnrealBuildTool
 		public IServer? Server { get; private set; }
 		ISessionServer? _session;
 		object _sessionLock = new();
-		ICacheClient? _cacheClient;
+		IEnumerable<UnrealBuildAcceleratorCacheConfig> _cacheConfigs = [];
+		Dictionary<UnrealBuildAcceleratorCacheConfig, ICacheClient> _cacheClients = [];
 		EpicGames.UBA.ILogger? _ubaLogger;
 		readonly List<IUBAAgentCoordinator> _agentCoordinators = new();
 		DirectoryReference? _rootDirRef;
@@ -66,8 +67,11 @@ namespace UnrealBuildTool
 		{
 			if (disposing)
 			{
-				_cacheClient?.Dispose();
-				_cacheClient = null;
+				foreach (KeyValuePair<UnrealBuildAcceleratorCacheConfig, ICacheClient> item in _cacheClients)
+				{
+					item.Value.Dispose();
+				}
+				_cacheClients.Clear();
 				_session?.Dispose();
 				_session = null;
 				_threadedLogger.Dispose();
@@ -125,17 +129,21 @@ namespace UnrealBuildTool
 			XmlConfig.ApplyTo(UBAConfig);
 			CommandLine.ParseArguments(Environment.GetCommandLineArgs(), this, logger);
 
+			bool bAnyHotReload = false;
 			foreach (TargetDescriptor targetDescriptor in targetDescriptors)
 			{
 				ConfigCache.ReadSettings(targetDescriptor.ProjectFile?.Directory, BuildHostPlatform.Current.Platform, UBAConfig);
-				if (targetDescriptor.HotReloadMode != HotReloadMode.Disabled)
-				{
-					UBAConfig.CacheServer = null; // Disable cache server when using hot reload path (there would likely never be any cache hits anyway
-				}
+				bAnyHotReload = bAnyHotReload || targetDescriptor.HotReloadMode != HotReloadMode.Disabled;
 				targetDescriptor.AdditionalArguments.ApplyTo(this);
 				targetDescriptor.AdditionalArguments.ApplyTo(UBAConfig);
 			}
 
+			// Disable cache server when using hot reload path (there would likely never be any cache hits anyway
+			if (!bAnyHotReload)
+			{
+				LoadCacheConfigs(targetDescriptors);
+			}
+
 			if (UBAConfig.bUseCrypto)
 			{
 				Crypto = CreateCrypto();
@@ -145,6 +153,34 @@ namespace UnrealBuildTool
 			_threadedLogger = new ThreadedLogger(logger);
 		}
 
+		private void LoadCacheConfigs(IEnumerable<TargetDescriptor> targetDescriptors)
+		{
+			const string ProviderPrefix = "Uba.CacheProvider";
+
+			IEnumerable<string> providers = Unreal.IsBuildMachine()
+			? [.. UBAConfig.BuildMachineCacheProviders, .. UBAConfig.IniBuildMachineCacheProviders]
+			: [.. UBAConfig.CacheProviders, .. UBAConfig.IniCacheProviders];
+
+			if (!providers.Any())
+			{
+				providers = [ProviderPrefix];
+			}
+
+			List<UnrealBuildAcceleratorCacheConfig> configs = [];
+			foreach (string provider in providers.Distinct().Where(x => x.StartsWith(ProviderPrefix)))
+			{
+				foreach (DirectoryReference? projectDir in targetDescriptors.Select(x => x.ProjectFile?.Directory))
+				{
+					UnrealBuildAcceleratorCacheConfig config = new UnrealBuildAcceleratorCacheConfig();
+					ConfigHierarchy engineIni = ConfigCache.ReadHierarchy(ConfigHierarchyType.Engine, projectDir, BuildHostPlatform.Current.Platform);
+					config.LoadConfigProvider(engineIni, provider);
+					configs.Add(config);
+				}
+			}
+
+			_cacheConfigs = configs.DistinctBy(x => x.CacheServer).Where(x => !String.IsNullOrEmpty(x.CacheServer));
+		}
+
 		public void UpdateStatus(uint statusRow, uint statusColumn, string statusText, LogEntryType statusType, string? statusLink)
 		{
 			lock (_sessionLock)
@@ -271,20 +307,41 @@ namespace UnrealBuildTool
 					uint bucket = _executor.GetCacheBucket(action);
 					ulong rootsHandle = _executor.GetActionRootsHandle(action);
 
-					FetchFromCacheResult result = _executor._cacheClient!.FetchFromCache(rootsHandle, bucket, startInfo);
-					int exitCode = 0;
-					RemoveLogLineSpam(result.LogLines);
-					if (result.Success && action.bForceWarningsAsError && result.LogLines.Any(x => x.Contains("): warning: ")))
+					IEnumerable<KeyValuePair<UnrealBuildAcceleratorCacheConfig, ICacheClient>> cacheClients;
+					if (_executor.UBAConfig.bCacheShuffle && _executor._cacheClients.Count > 1)
 					{
-						result = new(true, result.LogLines.Select(x => x.Replace("): warning: ", "): error: ", StringComparison.OrdinalIgnoreCase)).ToList());
-						exitCode = 1;
-						++_executor._errorCount;
+						KeyValuePair<UnrealBuildAcceleratorCacheConfig, ICacheClient>[] array = [.. _executor._cacheClients];
+						Random r = new Random();
+						r.Shuffle(array);
+						cacheClients = array;
 					}
-					if (result.Success)
+					else
 					{
-						_executor.DecrementCacheableAction();
+						cacheClients = _executor._cacheClients;
 					}
-					return new ActionArtifactResult(result.Success, result.LogLines, exitCode);
+
+					foreach (KeyValuePair<UnrealBuildAcceleratorCacheConfig, ICacheClient> item in cacheClients)
+					{
+						if (item.Key.bRequireVfs && !action.RootPaths.bUseVfs)
+						{
+							continue;
+						}
+						FetchFromCacheResult result = item.Value.FetchFromCache(rootsHandle, bucket, startInfo);
+						int exitCode = 0;
+						RemoveLogLineSpam(result.LogLines);
+						if (result.Success && action.bForceWarningsAsError && result.LogLines.Any(x => x.Contains("): warning: ")))
+						{
+							result = new(true, result.LogLines.Select(x => x.Replace("): warning: ", "): error: ", StringComparison.OrdinalIgnoreCase)).ToList());
+							exitCode = 1;
+							++_executor._errorCount;
+						}
+						if (result.Success)
+						{
+							_executor.DecrementCacheableAction();
+							return new ActionArtifactResult(result.Success, result.LogLines, exitCode);
+						}
+					}
+					return new ActionArtifactResult(false, [], 0);
 				}, cancellationToken, TaskCreationOptions.LongRunning | TaskCreationOptions.PreferFairness, TaskScheduler.Default);
 			}
 
@@ -331,16 +388,7 @@ namespace UnrealBuildTool
 				UBAConfig.bAllowRetry = true;
 			}
 
-			if (!String.IsNullOrEmpty(UBAConfig.CacheServer))
-			{
-				foreach (LinkedAction action in inputActions)
-				{
-					if (action.ArtifactMode.HasFlag(ArtifactMode.Enabled))
-					{
-						++_cacheableActionsLeft;
-					}
-				}
-			}
+			_cacheableActionsLeft = inputActions.Count(x => x.ArtifactMode.HasFlag(ArtifactMode.Enabled));
 
 			PrintConfiguration();
 			await WriteActionOutputFileAsync(inputActions);
@@ -417,7 +465,7 @@ namespace UnrealBuildTool
 
 				string arch = RuntimeInformation.ProcessArchitecture.ToString().ToLowerInvariant();
 				FileReference configFile = FileReference.Combine(UbaBinariesDir, "UbaHost.toml");
-				var config = EpicGames.UBA.IConfig.LoadConfig(configFile.FullName);
+				IConfig config = EpicGames.UBA.IConfig.LoadConfig(configFile.FullName);
 				if (!String.IsNullOrEmpty(UBAConfig.CompressionLevel))
 				{
 					config.AddValue("Storage", "CompressionLevel", UBAConfig.CompressionLevel);
@@ -436,36 +484,37 @@ namespace UnrealBuildTool
 
 							OnSessionCreated?.Invoke(_session);
 
-							_cacheClient = ICacheClient.CreateCacheClient(_session, UBAConfig.bReportCacheMissReason, "", _cacheableActionsLeft != 0 ? GetCacheClientHint() : "");
-
-							if (_cacheableActionsLeft != 0)
+							if (_cacheableActionsLeft != 0 && _cacheConfigs.Any())
 							{
-								string[] nameAndPort = UBAConfig.CacheServer!.Split(':');
-								int port = 1347;
-								if (nameAndPort.Length > 1)
-								{
-									port = Int32.Parse(nameAndPort[1]);
-								}
-
 								_session.UpdateStatus(1, 1, "Cache", LogEntryType.Info, null);
 								_session.UpdateStatus(1, 6, "Connecting...", LogEntryType.Info);
 
-								System.Diagnostics.Stopwatch stopwatch = System.Diagnostics.Stopwatch.StartNew();
-								bool cacheSuccess = _cacheClient.Connect(nameAndPort[0], port);
-								long totalMs = stopwatch.ElapsedMilliseconds;
-								string successText = cacheSuccess ? "Connected to" : "Failed to connect to";
-								logger.LogInformation("UbaCache - {SuccessText} {Name}:{Port} ({Seconds}.{Milliseconds}s)", successText, nameAndPort[0], port, totalMs / 1000, totalMs % 1000);
-								if (cacheSuccess)
-								{
-									_session.UpdateStatus(1, 6, "Connected", LogEntryType.Info);
-									actionArtifactCache = new UBAActionArtifactCache(this);
-								}
-								else
+								foreach (UnrealBuildAcceleratorCacheConfig cacheConfig in _cacheConfigs)
 								{
-									_cacheClient.Dispose();
-									_cacheClient = null;
-									_session.UpdateStatus(1, 6, "Failed to connect", LogEntryType.Info);
+									ICacheClient cacheClient = ICacheClient.CreateCacheClient(_session, UBAConfig.bReportCacheMissReason, "", _cacheableActionsLeft != 0 ? GetCacheClientHint() : "");
+									string[] nameAndPort = cacheConfig.CacheServer.Split(':');
+									int port = 1347;
+									if (nameAndPort.Length > 1)
+									{
+										port = Int32.Parse(nameAndPort[1]);
+									}
+
+									System.Diagnostics.Stopwatch stopwatch = System.Diagnostics.Stopwatch.StartNew();
+									bool cacheSuccess = cacheClient.Connect(nameAndPort[0], port);
+									long totalMs = stopwatch.ElapsedMilliseconds;
+									string successText = cacheSuccess ? "Connected to" : "Failed to connect to";
+									logger.LogInformation("UbaCache - {SuccessText} {Name}:{Port} ({Seconds}.{Milliseconds}s)", successText, nameAndPort[0], port, totalMs / 1000, totalMs % 1000);
+									if (cacheSuccess)
+									{
+										_cacheClients.Add(cacheConfig, cacheClient);
+										actionArtifactCache ??= new UBAActionArtifactCache(this);
+									}
+									else
+									{
+										cacheClient.Dispose();
+									}
 								}
+								_session.UpdateStatus(1, 6, _cacheClients.Any() ? "Connected" : "Failed to connect", LogEntryType.Info);
 							}
 
 							ubaStorage = ubaStorageServer;
@@ -495,8 +544,11 @@ namespace UnrealBuildTool
 						}
 						finally
 						{
-							_cacheClient?.Dispose();
-							_cacheClient = null;
+							foreach (KeyValuePair<UnrealBuildAcceleratorCacheConfig, ICacheClient> item in _cacheClients)
+							{
+								item.Value.Dispose();
+							}
+							_cacheClients.Clear();
 							lock (_sessionLock)
 							{
 								_agentCoordinators.ForEach(ac => ac.Done());
@@ -726,8 +778,11 @@ namespace UnrealBuildTool
 		{
 			if (Interlocked.Decrement(ref _cacheableActionsLeft) == 0)
 			{
-				_cacheClient?.Dispose();
-				_cacheClient = null;
+				foreach (KeyValuePair<UnrealBuildAcceleratorCacheConfig, ICacheClient> item in _cacheClients)
+				{
+					item.Value.Dispose();
+				}
+				_cacheClients.Clear();
 			}
 		}
 
@@ -771,7 +826,7 @@ namespace UnrealBuildTool
 					// And since pch contains absolute paths we unfortunately can't share cache data between machines that have different paths
 					if (action.ArtifactMode.HasFlag(ArtifactMode.AbsolutePath))
 					{
-						foreach (var (id, vfs, local) in action.RootPaths)
+						foreach ((uint id, DirectoryReference vfs, DirectoryReference local) in action.RootPaths)
 						{
 							hasher.Update(System.Text.Encoding.UTF8.GetBytes(local.FullName));
 						}
@@ -798,7 +853,7 @@ namespace UnrealBuildTool
 			using (BinaryWriter writer = new(outputsMemory, System.Text.Encoding.UTF8, true))
 			{
 				bool useVfs = action.RootPaths.bUseVfs;
-				foreach (var (id, vfs, local) in action.RootPaths)
+				foreach ((uint id, DirectoryReference vfs, DirectoryReference local) in action.RootPaths)
 				{
 					writer.Write((byte)id);
 					if (useVfs)
@@ -829,7 +884,7 @@ namespace UnrealBuildTool
 
 		bool WriteToCache(LinkedAction action, IProcess process)
 		{
-			if (!UBAConfig.bWriteCache || process.ExitCode != 0 || _cacheClient == null || !action.ArtifactMode.HasFlag(ArtifactMode.Enabled))
+			if (process.ExitCode != 0 || !action.ArtifactMode.HasFlag(ArtifactMode.Enabled) || !_cacheClients.Any(x => x.Key.CanWrite))
 			{
 				return true;
 			}
@@ -889,7 +944,16 @@ namespace UnrealBuildTool
 
 			uint bucket = GetCacheBucket(action);
 
-			bool success = _cacheClient!.WriteToCache(bucket, process, inputsMemory.GetBuffer(), (uint)inputsMemory.Position, outputsMemory.GetBuffer(), (uint)outputsMemory.Position);
+			bool success = true;
+			foreach (KeyValuePair<UnrealBuildAcceleratorCacheConfig, ICacheClient> item in _cacheClients.Where(x => x.Key.CanWrite))
+			{
+				if (item.Key.bRequireVfs && !action.RootPaths.bUseVfs)
+				{
+					continue;
+				}
+
+				success = item.Value.WriteToCache(bucket, process, inputsMemory.GetBuffer(), (uint)inputsMemory.Position, outputsMemory.GetBuffer(), (uint)outputsMemory.Position) || success;
+			}
 
 			DecrementCacheableAction();
 
diff --git a/Engine/Source/Programs/UnrealBuildTool/Executors/UnrealBuildAccelerator/UnrealBuildAcceleratorCacheConfig.cs b/Engine/Source/Programs/UnrealBuildTool/Executors/UnrealBuildAccelerator/UnrealBuildAcceleratorCacheConfig.cs
new file mode 100644
--- /dev/null
+++ b/Engine/Source/Programs/UnrealBuildTool/Executors/UnrealBuildAccelerator/UnrealBuildAcceleratorCacheConfig.cs
@@ -0,0 +1,123 @@
+// Copyright Epic Games, Inc. All Rights Reserved.
+
+using EpicGames.Core;
+using System;
+using System.Collections.Generic;
+using System.Linq;
+using System.Reflection;
+using UnrealBuildBase;
+
+namespace UnrealBuildTool
+{
+	/// <summary>
+	/// Configuration for Unreal Build Accelerator
+	/// </summary>
+	[System.Diagnostics.CodeAnalysis.SuppressMessage("Style", "IDE1006:Naming Styles", Justification = "UnrealBuildTool naming style")]
+	class UnrealBuildAcceleratorCacheConfig
+	{
+		/// <summary>
+		/// Address of the uba cache service. Will automatically use cache if connected
+		/// </summary>
+		[XmlConfigFile(Category = "UnrealBuildAccelerator", Name = "Cache")]
+		[ConfigFile(ConfigHierarchyType.Engine, "Uba.CacheProvider", "Url")]
+		[CommandLine("-UBACache=")]
+		public string CacheServer { get; set; } = String.Empty;
+
+		/// <summary>
+		/// Set cache to write, expects one of [True, False, BuildMachineOnly]
+		/// </summary>
+		[XmlConfigFile(Category = "UnrealBuildAccelerator", Name = "WriteCache")]
+		[ConfigFile(ConfigHierarchyType.Engine, "Uba.CacheProvider", "Write")]
+		[CommandLine("-UBAWriteCache", Value = "true")]
+		public string WriteCache { get; set; } = "BuildMachineOnly";
+
+		/// <summary>
+		/// Set cache to require VFS to be enabled
+		/// </summary>
+		[XmlConfigFile(Category = "UnrealBuildAccelerator", Name = "RequireVFS")]
+		[ConfigFile(ConfigHierarchyType.Engine, "Uba.CacheProvider", "RequireVFS")]
+		public bool bRequireVfs { get; set; } = false;
+
+		/// <summary>
+		/// If the cache can be written to
+		/// </summary>
+		public bool CanWrite => WriteCache.Equals("True", StringComparison.OrdinalIgnoreCase) || (Unreal.IsBuildMachine() && WriteCache.Equals("BuildMachineOnly", StringComparison.OrdinalIgnoreCase));
+	}
+
+	#region Extension methods
+	/// <summary>
+	/// Extension methods for FileReference functionality
+	/// </summary>
+	internal static class UnrealBuildAcceleratorCacheConfigExtensionMethods
+	{
+		/// <summary>
+		/// Load horde config from a different section name
+		/// </summary>
+		/// <param name="config"></param>
+		/// <param name="ini"></param>
+		/// <param name="section"></param>
+		internal static void LoadConfigProvider(this UnrealBuildAcceleratorCacheConfig config, ConfigHierarchy ini, string section)
+		{
+			object? TryGetValue(Type propertyType, string keyName)
+			{
+				if (propertyType == typeof(string))
+				{
+					if (ini.TryGetValue(section, keyName, out string? value))
+					{
+						return value;
+					}
+				}
+				else if (propertyType == typeof(bool))
+				{
+					if (ini.TryGetValue(section, keyName, out bool value))
+					{
+						return value;
+					}
+				}
+				else if (propertyType == typeof(int))
+				{
+					if (ini.TryGetValue(section, keyName, out int value))
+					{
+						return value;
+					}
+				}
+				else if (propertyType == typeof(float))
+				{
+					if (ini.TryGetValue(section, keyName, out float value))
+					{
+						return value;
+					}
+				}
+				else if (propertyType == typeof(double))
+				{
+					if (ini.TryGetValue(section, keyName, out double value))
+					{
+						return value;
+					}
+				}
+				else
+				{
+					throw new BuildLogEventException("UnrealBuildAcceleratorCacheConfig.LoadConfigProvider: Unsupported type {Type} for Key {KeyName}", propertyType, keyName);
+				}
+				return null;
+			}
+
+			IEnumerable<PropertyInfo> properties = config.GetType().GetProperties(BindingFlags.Instance | BindingFlags.Public | BindingFlags.NonPublic);
+			foreach (PropertyInfo property in properties)
+			{
+				Type propertyType = property.PropertyType;
+				ConfigFileAttribute? attribute = property.GetCustomAttributes<ConfigFileAttribute>().FirstOrDefault();
+				if (attribute != null)
+				{
+					string keyName = attribute.KeyName ?? property.Name;
+					object? value = TryGetValue(propertyType, keyName) ?? null;
+					if (value != null)
+					{
+						property.SetValue(config, value);
+					}
+				}
+			}
+		}
+	}
+	#endregion
+}
\ No newline at end of file
diff --git a/Engine/Source/Programs/UnrealBuildTool/Executors/UnrealBuildAccelerator/UnrealBuildAcceleratorConfig.cs b/Engine/Source/Programs/UnrealBuildTool/Executors/UnrealBuildAccelerator/UnrealBuildAcceleratorConfig.cs
--- a/Engine/Source/Programs/UnrealBuildTool/Executors/UnrealBuildAccelerator/UnrealBuildAcceleratorConfig.cs
+++ b/Engine/Source/Programs/UnrealBuildTool/Executors/UnrealBuildAccelerator/UnrealBuildAcceleratorConfig.cs
@@ -3,7 +3,6 @@
 using EpicGames.Core;
 using System;
 using System.Collections.Generic;
-using UnrealBuildBase;
 
 namespace UnrealBuildTool
 {
@@ -42,6 +41,30 @@ namespace UnrealBuildTool
 		[ConfigFile(ConfigHierarchyType.Engine, "UbaController", "BuildMachineProviders")]
 		public List<string> IniBuildMachineProviders { get; set; } = [];
 
+		/// <summary>
+		/// List of provider names for remote cache.
+		/// </summary>
+		[XmlConfigFile(Category = "UnrealBuildAccelerator")]
+		[CommandLine("-UBACacheProviders=", ListSeparator = '+')]
+		public string[] CacheProviders { get; set; } = [];
+
+		/// <summary>
+		/// List of provider names for remote cache from ini
+		/// </summary>
+		[ConfigFile(ConfigHierarchyType.Engine, "UbaController", "CacheProviders")]
+		public List<string> IniCacheProviders { get; set; } = [];
+
+		/// <summary>
+		/// List of provider names for remote cache
+		/// </summary>
+		[XmlConfigFile(Category = "UnrealBuildAccelerator")]
+		public string[] BuildMachineCacheProviders { get; set; } = [];
+
+		/// <summary>
+		/// List of provider names for remote cache from ini
+		/// </summary>
+		[ConfigFile(ConfigHierarchyType.Engine, "UbaController", "BuildMachineCacheProviders")]
+		public List<string> IniBuildMachineCacheProviders { get; set; } = [];
 
 		/// <summary>
 		/// When set to true, UBA will not use any remote help
@@ -261,20 +284,6 @@ namespace UnrealBuildTool
 		[CommandLine("-UBADetailedLog", Value = "true")]
 		public bool bDetailedLog { get; set; } = false;
 
-		/// <summary>
-		/// Address of the uba cache service. Will automatically use cache if connected
-		/// </summary>
-		[XmlConfigFile(Category = "UnrealBuildAccelerator", Name = "Cache")]
-		[CommandLine("-UBACache=")]
-		public string? CacheServer { get; set; }
-
-		/// <summary>
-		/// Set cache to write instead of fetch
-		/// </summary>
-		[XmlConfigFile(Category = "UnrealBuildAccelerator", Name = "WriteCache")]
-		[CommandLine("-UBAWriteCache")]
-		public bool bWriteCache { get; set; }
-
 		/// <summary>
 		/// Disable all detouring making all actions run outside uba
 		/// </summary>
@@ -289,6 +298,13 @@ namespace UnrealBuildTool
 		[CommandLine("-UBACacheMaxWorkers")]
 		public int CacheMaxWorkers { get; set; } = 32;
 
+		/// <summary>
+		/// Shuffle the list of cache servers before each read attempt
+		/// </summary>
+		[XmlConfigFile(Category = "UnrealBuildAccelerator", Name = "CacheShuffle")]
+		[CommandLine("-UBAShuffleCache")]
+		public bool bCacheShuffle { get; set; } = false;
+
 		/// <summary>
 		/// Report reason a cache miss happened. Useful when searching for determinism/portability issues
 		/// </summary>
