UnrealBuildTool: EngineIncludeOrderVersion.Unreal5_7

#rnx

diff --git a/Engine/Source/Programs/UnrealBuildTool/Configuration/TargetRules.cs b/Engine/Source/Programs/UnrealBuildTool/Configuration/TargetRules.cs
--- a/Engine/Source/Programs/UnrealBuildTool/Configuration/TargetRules.cs
+++ b/Engine/Source/Programs/UnrealBuildTool/Configuration/TargetRules.cs
@@ -188,12 +188,13 @@ namespace UnrealBuildTool
 		/// <summary>
 		/// Include order used in Unreal 5.4
 		/// </summary>
-		[Obsolete("The Unreal 5.4 include order is deprecated and will be unsupported in 5.7.")]
+		[Obsolete("The Unreal 5.4 include order is unsupported.")]
 		Unreal5_4,
 
 		/// <summary>
 		/// Include order used in Unreal 5.5
 		/// </summary>
+		[Obsolete("The Unreal 5.5 include order is deprecated and will be unsupported in 5.8.")]
 		Unreal5_5,
 
 		/// <summary>
@@ -201,19 +202,24 @@ namespace UnrealBuildTool
 		/// </summary>
 		Unreal5_6,
 
-		// *** When adding new entries here, be sure to update UEBuildModuleCPP.CurrentIncludeOrderDefine to ensure that the correct guard is used. ***
+		/// <summary>
+		/// Include order used in Unreal 5.7
+		/// </summary>
+		Unreal5_7,
+
+		// *** When adding new entries here, be sure to update EngineIncludeOrderHelper.GetDeprecationDefine to ensure that the correct guard is used. ***
 
 		/// <summary>
 		/// Always use the latest version of include order. This value is updated every Unreal release, use with caution if you intend to integrate newer Unreal releases.
 		/// </summary>
-		Latest = Unreal5_6,
+		Latest = Unreal5_7,
 
 		/// <summary>
 		/// Contains the oldest version of include order that the engine supports.
 		/// Do not delete old enum entries to prevent breaking project generation
 		/// </summary>
 #pragma warning disable CS0618 // Type or member is obsolete
-		Oldest = Unreal5_3,
+		Oldest = Unreal5_4,
 #pragma warning restore CS0618 // Type or member is obsolete
 	}
 
@@ -434,76 +440,45 @@ namespace UnrealBuildTool
 	/// </summary>
 	public static class EngineIncludeOrderHelper
 	{
-		private static string GetDeprecationDefine(EngineIncludeOrderVersion inVersion)
+		private static string GetDeprecationDefine(EngineIncludeOrderVersion inVersion) => inVersion switch
 		{
-			switch (inVersion)
-			{
-				default:
-				case EngineIncludeOrderVersion.Unreal5_6:
-					return "UE_ENABLE_INCLUDE_ORDER_DEPRECATED_IN_5_6";
-				case EngineIncludeOrderVersion.Unreal5_5:
-					return "UE_ENABLE_INCLUDE_ORDER_DEPRECATED_IN_5_5";
+			EngineIncludeOrderVersion.Unreal5_7 => "UE_ENABLE_INCLUDE_ORDER_DEPRECATED_IN_5_7",
+			EngineIncludeOrderVersion.Unreal5_6 => "UE_ENABLE_INCLUDE_ORDER_DEPRECATED_IN_5_6",
 #pragma warning disable CS0618 // Type or member is obsolete
-				case EngineIncludeOrderVersion.Unreal5_4:
-					return "UE_ENABLE_INCLUDE_ORDER_DEPRECATED_IN_5_4";
-				case EngineIncludeOrderVersion.Unreal5_3:
-					return "UE_ENABLE_INCLUDE_ORDER_DEPRECATED_IN_5_3";
-				case EngineIncludeOrderVersion.Unreal5_2:
-					return "UE_ENABLE_INCLUDE_ORDER_DEPRECATED_IN_5_2";
+			EngineIncludeOrderVersion.Unreal5_5 => "UE_ENABLE_INCLUDE_ORDER_DEPRECATED_IN_5_5",
+			EngineIncludeOrderVersion.Unreal5_4 => "UE_ENABLE_INCLUDE_ORDER_DEPRECATED_IN_5_4",
 #pragma warning restore CS0618 // Type or member is obsolete
-			}
-		}
+			_ => throw new BuildLogEventException("Unsupported EngineIncludeOrderVersion {Value}", inVersion)
+		};
 
-		private static string GetDeprecationDefine(EngineIncludeOrderVersion inTestVersion, EngineIncludeOrderVersion inVersion)
-		{
-			return GetDeprecationDefine(inTestVersion) + "=" + (inVersion < inTestVersion ? "1" : "0");
-		}
+		private static string GetDeprecationDefine(EngineIncludeOrderVersion inTestVersion, EngineIncludeOrderVersion inVersion) => $"{GetDeprecationDefine(inTestVersion)}={(inVersion < inTestVersion ? "1" : "0")}";
 
 		/// <summary>
 		/// Returns a list of every deprecation define available.
 		/// </summary>
 		/// <returns></returns>
-		public static List<string> GetAllDeprecationDefines()
-		{
-			return new List<string>()
-			{
-#pragma warning disable CS0618 // Type or member is obsolete
-				GetDeprecationDefine(EngineIncludeOrderVersion.Unreal5_2),
-				GetDeprecationDefine(EngineIncludeOrderVersion.Unreal5_3),
-				GetDeprecationDefine(EngineIncludeOrderVersion.Unreal5_4),
-#pragma warning restore CS0618 // Type or member is obsolete
-				GetDeprecationDefine(EngineIncludeOrderVersion.Unreal5_5),
-				GetDeprecationDefine(EngineIncludeOrderVersion.Unreal5_6),
-			};
-		}
+		public static IEnumerable<string> GetAllDeprecationDefines() => Enum.GetValues<EngineIncludeOrderVersion>()
+			.Where(x => x >= EngineIncludeOrderVersion.Oldest && x <= EngineIncludeOrderVersion.Latest)
+			.Select(GetDeprecationDefine)
+			.Distinct()
+			.Order();
 
 		/// <summary>
 		/// Get a list of every deprecation define and their value for the specified engine include order.
 		/// </summary>
 		/// <param name="inVersion"></param>
 		/// <returns></returns>
-		public static List<string> GetDeprecationDefines(EngineIncludeOrderVersion inVersion)
-		{
-			return new List<string>()
-			{
-#pragma warning disable CS0618 // Type or member is obsolete
-				GetDeprecationDefine(EngineIncludeOrderVersion.Unreal5_2, inVersion),
-				GetDeprecationDefine(EngineIncludeOrderVersion.Unreal5_3, inVersion),
-				GetDeprecationDefine(EngineIncludeOrderVersion.Unreal5_4, inVersion),
-#pragma warning restore CS0618 // Type or member is obsolete
-				GetDeprecationDefine(EngineIncludeOrderVersion.Unreal5_5, inVersion),
-				GetDeprecationDefine(EngineIncludeOrderVersion.Unreal5_6, inVersion),
-			};
-		}
+		public static IEnumerable<string> GetDeprecationDefines(EngineIncludeOrderVersion inVersion) => Enum.GetValues<EngineIncludeOrderVersion>()
+			.Where(x => x >= EngineIncludeOrderVersion.Oldest && x <= EngineIncludeOrderVersion.Latest)
+			.Select(x => GetDeprecationDefine(x, inVersion))
+			.Distinct()
+			.Order();
 
 		/// <summary>
 		/// Returns the latest deprecation define.
 		/// </summary>
 		/// <returns></returns>
-		public static string GetLatestDeprecationDefine()
-		{
-			return GetDeprecationDefine(EngineIncludeOrderVersion.Latest);
-		}
+		public static string GetLatestDeprecationDefine() => GetDeprecationDefine(EngineIncludeOrderVersion.Latest);
 	}
 
 	/// <summary>
