UnrealBuildTool: Don't limit EngineIncludeOrderVersion by latest because of a downstream issue

#rnx

diff --git a/Engine/Source/Programs/UnrealBuildTool/Configuration/TargetRules.cs b/Engine/Source/Programs/UnrealBuildTool/Configuration/TargetRules.cs
--- a/Engine/Source/Programs/UnrealBuildTool/Configuration/TargetRules.cs
+++ b/Engine/Source/Programs/UnrealBuildTool/Configuration/TargetRules.cs
@@ -458,7 +458,7 @@ namespace UnrealBuildTool
 		/// </summary>
 		/// <returns></returns>
 		public static IEnumerable<string> GetAllDeprecationDefines() => Enum.GetValues<EngineIncludeOrderVersion>()
-			.Where(x => x >= EngineIncludeOrderVersion.Oldest && x <= EngineIncludeOrderVersion.Latest)
+			.Where(x => x >= EngineIncludeOrderVersion.Oldest)
 			.Select(GetDeprecationDefine)
 			.Distinct()
 			.Order();
@@ -469,7 +469,7 @@ namespace UnrealBuildTool
 		/// <param name="inVersion"></param>
 		/// <returns></returns>
 		public static IEnumerable<string> GetDeprecationDefines(EngineIncludeOrderVersion inVersion) => Enum.GetValues<EngineIncludeOrderVersion>()
-			.Where(x => x >= EngineIncludeOrderVersion.Oldest && x <= EngineIncludeOrderVersion.Latest)
+			.Where(x => x >= EngineIncludeOrderVersion.Oldest)
 			.Select(x => GetDeprecationDefine(x, inVersion))
 			.Distinct()
 			.Order();
