[RaceDetector]
- Fix intellisense on racedetector when instrumentation is active

diff --git a/Engine/Source/Programs/UnrealBuildTool/Platform/Windows/UEBuildWindows.cs b/Engine/Source/Programs/UnrealBuildTool/Platform/Windows/UEBuildWindows.cs
--- a/Engine/Source/Programs/UnrealBuildTool/Platform/Windows/UEBuildWindows.cs
+++ b/Engine/Source/Programs/UnrealBuildTool/Platform/Windows/UEBuildWindows.cs
@@ -1927,6 +1927,11 @@ namespace UnrealBuildTool
 		/// <param name="LinkEnvironment">The link environment for this target</param>
 		public override void SetUpEnvironment(ReadOnlyTargetRules Target, CppCompileEnvironment CompileEnvironment, LinkEnvironment LinkEnvironment)
 		{
+			if (Target.WindowsPlatform.bEnableInstrumentation) 
+			{
+				CompileEnvironment.Definitions.Add("USING_INSTRUMENTATION=1");
+			}
+
 			// @todo Remove this hack to work around broken includes
 			CompileEnvironment.Definitions.Add("NDIS_MINIPORT_MAJOR_VERSION=0");
 
