UnrealBuidlTool: Add KnownLogEvent for circular dependencies

#rnx

diff --git a/Engine/Source/Programs/Shared/EpicGames.Core/KnownLogEvents.cs b/Engine/Source/Programs/Shared/EpicGames.Core/KnownLogEvents.cs
--- a/Engine/Source/Programs/Shared/EpicGames.Core/KnownLogEvents.cs
+++ b/Engine/Source/Programs/Shared/EpicGames.Core/KnownLogEvents.cs
@@ -50,6 +50,11 @@ namespace EpicGames.Core
 		/// </summary>
 		public static EventId Validation { get; } = new EventId(103);
 
+		/// <summary>
+		/// Circular dependency messages
+		/// </summary>
+		public static EventId CircularDependency { get; } = new EventId(104);
+
 		/// <summary>
 		/// UHT compiler message
 		/// </summary>
diff --git a/Engine/Source/Programs/UnrealBuildTool/Configuration/UEBuildModule.cs b/Engine/Source/Programs/UnrealBuildTool/Configuration/UEBuildModule.cs
--- a/Engine/Source/Programs/UnrealBuildTool/Configuration/UEBuildModule.cs
+++ b/Engine/Source/Programs/UnrealBuildTool/Configuration/UEBuildModule.cs
@@ -1285,25 +1285,29 @@ namespace UnrealBuildTool
 		/// <returns>If any circular references are found</returns>
 		bool ValidateCircularReferences(ILogger logger)
 		{
-			List<CircularReference> circularReferences = new();
-			CheckForCircularReferences(this, new(), new(), circularReferences);
+			List<CircularReference> circularReferences = [];
+			CheckForCircularReferences(this, [], [], circularReferences);
 
 			IEnumerable<CircularReference> allowed = circularReferences.Where(x => !x.IsError).OrderBy(x => x.ReferenceStack.Count);
 			IEnumerable<CircularReference> errors = circularReferences.Where(x => x.IsError).OrderBy(x => x.ReferenceStack.Count);
 
 			if (errors.Any())
 			{
-				logger.LogError(KnownLogEvents.Validation, "Circular dependency for {RulesFile} detected:", RulesFile.MakeRelativeTo(Unreal.EngineSourceDirectory));
+				logger.LogInformation(KnownLogEvents.Validation, "Circular dependency for '{RulesFile}' detected:", RulesFile.MakeRelativeTo(Unreal.EngineDirectory));
 				foreach (CircularReference circularReference in errors)
 				{
-					logger.LogError(KnownLogEvents.Validation, "* {ReferenceStack}", String.Join(" -> ", circularReference.ReferenceStack));
+					logger.LogInformation(KnownLogEvents.Validation, "* {ReferenceStack}", String.Join(" -> ", circularReference.ReferenceStack));
+				}
+				foreach (FileReference rulesFile in errors.SelectMany(x => x.ReferenceStack).Select(x => x.RulesFile).Distinct())
+				{
+					logger.LogError(KnownLogEvents.CircularDependency, "Circular dependency in '{Module}' possibly due to '{RulesFile}'", Name, rulesFile);
 				}
 				logger.LogInformation(KnownLogEvents.Validation, "Break this loop by moving dependencies into a separate module or using Private/PublicIncludePathModuleNames to reference declarations");
 			}
 
 			if (allowed.Any())
 			{
-				logger.LogDebug(KnownLogEvents.Validation, "Allowed circular dependency for {RulesFile} detected:", RulesFile.MakeRelativeTo(Unreal.EngineSourceDirectory));
+				logger.LogDebug(KnownLogEvents.Validation, "Allowed circular dependency for '{RulesFile}' detected:", RulesFile.MakeRelativeTo(Unreal.EngineDirectory));
 				foreach (CircularReference circularReference in allowed)
 				{
 					logger.LogDebug(KnownLogEvents.Validation, "* {ReferenceStack} (allowed via {GuiltyModules})", String.Join(" -> ", circularReference.ReferenceStack), String.Join(", ", circularReference.GuiltyModules.Select(x => $"{x.Key} -> {x.Value}")));
