Fix missing rpaths in universal build, by adding missing LinkEnvironment fields.

#rnx

diff --git a/Engine/Source/Programs/UnrealBuildTool/System/LinkEnvironment.cs b/Engine/Source/Programs/UnrealBuildTool/System/LinkEnvironment.cs
--- a/Engine/Source/Programs/UnrealBuildTool/System/LinkEnvironment.cs
+++ b/Engine/Source/Programs/UnrealBuildTool/System/LinkEnvironment.cs
@@ -402,6 +402,7 @@ namespace UnrealBuildTool
 			SystemLibraryPaths.AddRange(Other.SystemLibraryPaths);
 			ExcludedLibraries.AddRange(Other.ExcludedLibraries);
 			RuntimeLibraryPaths.AddRange(Other.RuntimeLibraryPaths);
+			RuntimeDependencies.AddRange(Other.RuntimeDependencies);
 			Frameworks.AddRange(Other.Frameworks);
 			AdditionalFrameworks.AddRange(Other.AdditionalFrameworks);
 			WeakFrameworks.AddRange(Other.WeakFrameworks);
@@ -413,6 +414,7 @@ namespace UnrealBuildTool
 			bIsBuildingLibrary = Other.bIsBuildingLibrary;
 			bDisableSymbolCache = Other.bDisableSymbolCache;
 			bIsBuildingDLL = Other.bIsBuildingDLL;
+			LinkType = Other.LinkType;
 			bUseStaticCRT = Other.bUseStaticCRT;
 			bUseDebugCRT = Other.bUseDebugCRT;
 			bIsBuildingConsoleApplication = Other.bIsBuildingConsoleApplication;
