[UBT]
* Fixed so merged modules does not write __declspec(dllexport) for _API macros

diff --git a/Engine/Source/Programs/UnrealBuildTool/Configuration/UEBuildModule.cs b/Engine/Source/Programs/UnrealBuildTool/Configuration/UEBuildModule.cs
--- a/Engine/Source/Programs/UnrealBuildTool/Configuration/UEBuildModule.cs
+++ b/Engine/Source/Programs/UnrealBuildTool/Configuration/UEBuildModule.cs
@@ -689,7 +689,7 @@ namespace UnrealBuildTool
 				{
 					Definitions.Add(ModuleApiDefine + "=DLLIMPORT");
 				}
-				else if (!Binary.bAllowExports)
+				else if (!Binary.bAllowExports || Rules.Target.bMergeModules)
 				{
 					Definitions.Add(ModuleApiDefine + "=");
 				}
