UnrealBuildTool: Fix UBA cache settings from the command line

#rnx

diff --git a/Engine/Source/Programs/UnrealBuildTool/Executors/UnrealBuildAccelerator/UBAExecutor.cs b/Engine/Source/Programs/UnrealBuildTool/Executors/UnrealBuildAccelerator/UBAExecutor.cs
--- a/Engine/Source/Programs/UnrealBuildTool/Executors/UnrealBuildAccelerator/UBAExecutor.cs
+++ b/Engine/Source/Programs/UnrealBuildTool/Executors/UnrealBuildAccelerator/UBAExecutor.cs
@@ -141,7 +141,7 @@ namespace UnrealBuildTool
 			// Disable cache server when using hot reload path (there would likely never be any cache hits anyway
 			if (!bAnyHotReload)
 			{
-				LoadCacheConfigs(targetDescriptors);
+				LoadCacheConfigs(targetDescriptors, logger);
 			}
 
 			if (UBAConfig.bUseCrypto)
@@ -153,7 +153,7 @@ namespace UnrealBuildTool
 			_threadedLogger = new ThreadedLogger(logger);
 		}
 
-		private void LoadCacheConfigs(IEnumerable<TargetDescriptor> targetDescriptors)
+		private void LoadCacheConfigs(IEnumerable<TargetDescriptor> targetDescriptors, Microsoft.Extensions.Logging.ILogger logger)
 		{
 			const string ProviderPrefix = "Uba.CacheProvider";
 
@@ -169,11 +169,17 @@ namespace UnrealBuildTool
 			List<UnrealBuildAcceleratorCacheConfig> configs = [];
 			foreach (string provider in providers.Distinct().Where(x => x.StartsWith(ProviderPrefix)))
 			{
-				foreach (DirectoryReference? projectDir in targetDescriptors.Select(x => x.ProjectFile?.Directory))
+				foreach (TargetDescriptor targetDescriptor in targetDescriptors)
 				{
 					UnrealBuildAcceleratorCacheConfig config = new UnrealBuildAcceleratorCacheConfig();
-					ConfigHierarchy engineIni = ConfigCache.ReadHierarchy(ConfigHierarchyType.Engine, projectDir, BuildHostPlatform.Current.Platform);
+					if (provider == ProviderPrefix)
+					{
+						XmlConfig.ApplyTo(config);
+						CommandLine.ParseArguments(Environment.GetCommandLineArgs(), this, logger);
+					}
+					ConfigHierarchy engineIni = ConfigCache.ReadHierarchy(ConfigHierarchyType.Engine, targetDescriptor.ProjectFile?.Directory, BuildHostPlatform.Current.Platform);
 					config.LoadConfigProvider(engineIni, provider);
+					targetDescriptor.AdditionalArguments.ApplyTo(config);
 					configs.Add(config);
 				}
 			}
