Create Android-arm64/x64 Config in VS

diff --git a/Engine/Source/Programs/UnrealBuildTool/ProjectFiles/VisualStudio/VCProject.cs b/Engine/Source/Programs/UnrealBuildTool/ProjectFiles/VisualStudio/VCProject.cs
--- a/Engine/Source/Programs/UnrealBuildTool/ProjectFiles/VisualStudio/VCProject.cs
+++ b/Engine/Source/Programs/UnrealBuildTool/ProjectFiles/VisualStudio/VCProject.cs
@@ -738,6 +738,14 @@ namespace UnrealBuildTool
 									}
 									else
 									{
+										if (BuildPlatform.ArchitectureConfig.Mode == UnrealArchitectureMode.SingleTargetLinkSeparately)
+										{
+											if (BuildPlatform.Platform == UnrealTargetPlatform.Android)
+											{
+												Architectures = BuildPlatform.ArchitectureConfig.AllSupportedArchitectures;
+											}
+										}
+
 										foreach (UnrealArch Arch in Architectures.Architectures)
 										{
 											AddProjectAndTargetCombination(Arch);
diff --git a/Engine/Source/Programs/UnrealBuildTool/ProjectFiles/VisualStudio/VCProjectFileGenerator.cs b/Engine/Source/Programs/UnrealBuildTool/ProjectFiles/VisualStudio/VCProjectFileGenerator.cs
--- a/Engine/Source/Programs/UnrealBuildTool/ProjectFiles/VisualStudio/VCProjectFileGenerator.cs
+++ b/Engine/Source/Programs/UnrealBuildTool/ProjectFiles/VisualStudio/VCProjectFileGenerator.cs
@@ -1090,7 +1090,7 @@ namespace UnrealBuildTool
 					{
 						ProjectTarget ProjectTarget = SolutionConfigKeyValue.Value.Item2.Item1;
 
-						Action<UnrealArch?, List<VCSolutionConfigCombination>> AddSolutionConfig = (UnrealArch? Arch, List<VCSolutionConfigCombination> OutSolutionConfigs) =>
+						void AddSolutionConfig(UnrealArch? Arch, List<VCSolutionConfigCombination> OutSolutionConfigs, bool ForceArchSuffix = false)
 						{
 							// e.g.  "Development|Win64 = Development|Win64"
 							string SolutionConfigName = SolutionConfigKeyValue.Key;
@@ -1100,7 +1100,7 @@ namespace UnrealBuildTool
 							string SolutionPlatformName = CurPlatform.ToString();
 							// We use RequiresArchitectureFilenames to determine whether the architecture suffix should be added.
 							// This is used to tell us what the "default" architecture is.
-							if (Arch != null && BuildPlatform.ArchitectureConfig.RequiresArchitectureFilenames(new UnrealArchitectures(Arch.Value)))
+							if (Arch != null && (ForceArchSuffix || BuildPlatform.ArchitectureConfig.RequiresArchitectureFilenames(new UnrealArchitectures(Arch.Value))))
 							{
 								SolutionPlatformName += $"-{Arch}";
 							}
@@ -1128,6 +1128,17 @@ namespace UnrealBuildTool
 							{
 								AddSolutionConfig(Arch, OutSolutionConfigs);
 							}
+
+							if (BuildPlatform.ArchitectureConfig.Mode == UnrealArchitectureMode.SingleTargetLinkSeparately)
+							{
+								if (BuildPlatform.Platform == UnrealTargetPlatform.Android)
+								{
+									foreach (UnrealArch Arch in BuildPlatform.ArchitectureConfig.AllSupportedArchitectures.Architectures)
+									{
+										AddSolutionConfig(Arch, OutSolutionConfigs, true);
+									}
+								}
+							}
 						}
 					}
 				}
