Fixed so merge modules compiling/linking works on all console platforms.
* Changed to use version script instead of dynamic list on clang platforms
* Renamed strip to export since that is what the action does, create export list

diff --git a/Engine/Source/Programs/UnrealBuildTool/Configuration/UEBuildBinary.cs b/Engine/Source/Programs/UnrealBuildTool/Configuration/UEBuildBinary.cs
--- a/Engine/Source/Programs/UnrealBuildTool/Configuration/UEBuildBinary.cs
+++ b/Engine/Source/Programs/UnrealBuildTool/Configuration/UEBuildBinary.cs
@@ -170,6 +170,8 @@ namespace UnrealBuildTool
 			this.PrimaryModule = PrimaryModule;
 			this.bUsePrecompiled = bUsePrecompiled;
 
+			PrimaryModule.bCreatePerModuleFile = true;
+
 			Modules.Add(PrimaryModule);
 		}
 
@@ -861,7 +863,7 @@ namespace UnrealBuildTool
 
 				string Name = OutputFilePaths.First().GetFileNameWithoutExtension();
 				string Ext = ToolChain.GetExtraLinkFileExtension();
-				FileReference extraObj = FileReference.Combine(IntermediateDirectory!, $"{Name}.extra.{Ext}");
+				FileReference extraObj = FileReference.Combine(IntermediateDirectory!, $"{Name}.exports.{Ext}");
 				BinaryLinkEnvironment.InputFiles.Add(FileItem.GetItemByFileReference(extraObj));
 			}
 
diff --git a/Engine/Source/Programs/UnrealBuildTool/Configuration/UEBuildModuleCPP.cs b/Engine/Source/Programs/UnrealBuildTool/Configuration/UEBuildModuleCPP.cs
--- a/Engine/Source/Programs/UnrealBuildTool/Configuration/UEBuildModuleCPP.cs
+++ b/Engine/Source/Programs/UnrealBuildTool/Configuration/UEBuildModuleCPP.cs
@@ -101,6 +101,11 @@ namespace UnrealBuildTool
 		/// </summary>
 		DirectoryReference? AssociatedVerseDirectory;
 
+		/// <summary>
+		/// Create PerModuleInline source file. Defaults to PrimaryModule in a binary but could be another one if primary module does not depend on core
+		/// </summary>
+		public bool bCreatePerModuleFile;
+
 		/// <summary>
 		/// The Verse source code directory associated with this module if any
 		/// </summary>
@@ -722,7 +727,7 @@ namespace UnrealBuildTool
 
 			// The operator new/delete overrides need to exist in all modules.
 			// We put these overrides in a separate file in order to be able to ensure only one file is added per binary for the primary module
-			if (Target.LinkType == TargetLinkType.Modular && Binary.PrimaryModule == this)
+			if (Target.LinkType == TargetLinkType.Modular && bCreatePerModuleFile)
 			{
 				FileReference NewDeleteOverrides = FileReference.Combine(IntermediateDirectory, "PerModuleInline.gen.cpp");
 				List<string> Content =
diff --git a/Engine/Source/Programs/UnrealBuildTool/Configuration/UEBuildTarget.cs b/Engine/Source/Programs/UnrealBuildTool/Configuration/UEBuildTarget.cs
--- a/Engine/Source/Programs/UnrealBuildTool/Configuration/UEBuildTarget.cs
+++ b/Engine/Source/Programs/UnrealBuildTool/Configuration/UEBuildTarget.cs
@@ -3048,12 +3048,12 @@ namespace UnrealBuildTool
 
 			HashSet<FileItem> allObjects = new(Binaries.SelectMany(x => x.InputObjects));
 
-			void CreateStripAction(string Name, IEnumerable<FileItem> toStrip, DirectoryReference IntermediateDirectory)
+			void CreateStripAction(string Name, UEBuildBinaryType binaryType, IEnumerable< FileItem> toStrip, DirectoryReference IntermediateDirectory)
 			{
 				HashSet<FileItem> otherObjects = new(allObjects.Except(toStrip));
-				FileItem stripRsp = FileItem.GetItemByFileReference(FileReference.Combine(IntermediateDirectory, $"{Name}.strip.rsp"));
+				FileItem exportsRsp = FileItem.GetItemByFileReference(FileReference.Combine(IntermediateDirectory, $"{Name}.exports.rsp"));
 				string Ext = TargetToolChain.GetExtraLinkFileExtension();
-				FileReference extraObj = FileReference.Combine(IntermediateDirectory, $"{Name}.extra.{Ext}");
+				FileReference extraObj = FileReference.Combine(IntermediateDirectory, $"{Name}.exports.{Ext}");
 				IEnumerable<string> arguments =
 				[
 					.. toStrip.Select(path => $"/S:{path.Location.ChangeExtension(".exi")}").OrderBy(x => x),
@@ -3061,30 +3061,31 @@ namespace UnrealBuildTool
 					$"/O:{extraObj}",
 					$"/T:{Platform}",
 					$"/M:{Name}",
+					..TargetToolChain.GetExtraLinkFileAdditionalSymbols(binaryType).Select(symbol=> $"/E:{symbol}"),
 				];
 
-				MakefileBuilder.CreateIntermediateTextFile(stripRsp, arguments);
+				MakefileBuilder.CreateIntermediateTextFile(exportsRsp, arguments);
 
-				Action stripAction = MakefileBuilder.CreateAction(ActionType.Link);
+				Action exportsAction = MakefileBuilder.CreateAction(ActionType.Link);
 
-				stripAction.PrerequisiteItems.Add(stripRsp);
-				stripAction.PrerequisiteItems.UnionWith(allObjects);
-				stripAction.ProducedItems.Add(FileItem.GetItemByFileReference(extraObj));
-				stripAction.CommandPath = commandPath;
-				stripAction.CommandArguments = $"@\"{TargetToolChain.NormalizeCommandLinePath(stripRsp.Location, CompileEnvironment.RootPaths)}\"";
-				stripAction.CommandDescription = "GenerateExports";
-				stripAction.WorkingDirectory = Unreal.EngineSourceDirectory;
-				stripAction.StatusDescription = Name;
-				stripAction.bCanExecuteRemotely = false; // Can run remotely but is so fast so it would be a waste
-				stripAction.ArtifactMode = ArtifactMode.Enabled;
-				stripAction.RootPaths = CompileEnvironment.RootPaths;
+				exportsAction.PrerequisiteItems.Add(exportsRsp);
+				exportsAction.PrerequisiteItems.UnionWith(allObjects);
+				exportsAction.ProducedItems.Add(FileItem.GetItemByFileReference(extraObj));
+				exportsAction.CommandPath = commandPath;
+				exportsAction.CommandArguments = $"@\"{TargetToolChain.NormalizeCommandLinePath(exportsRsp.Location, CompileEnvironment.RootPaths)}\"";
+				exportsAction.CommandDescription = "GenerateExports";
+				exportsAction.WorkingDirectory = Unreal.EngineSourceDirectory;
+				exportsAction.StatusDescription = Name;
+				exportsAction.bCanExecuteRemotely = false; // Can run remotely but is so fast so it would be a waste
+				exportsAction.ArtifactMode = ArtifactMode.Enabled;
+				exportsAction.RootPaths = CompileEnvironment.RootPaths;
 
-				Makefile.OutputItems.AddRange(stripAction.ProducedItems);
+				Makefile.OutputItems.AddRange(exportsAction.ProducedItems);
 			}
 
 			foreach (UEBuildBinary Binary in Binaries.Where(x => x.bStripUnusedExports))
 			{
-				CreateStripAction(Binary.OutputFilePaths.First().GetFileNameWithoutExtension(), Binary.InputObjects, Binary.IntermediateDirectory);
+				CreateStripAction(Binary.OutputFilePaths.First().GetFileNameWithoutExtension(), Binary.Type, Binary.InputObjects, Binary.IntermediateDirectory);
 			}
 		}
 
@@ -3481,58 +3482,48 @@ namespace UnrealBuildTool
 		/// <param name="Logger"></param>
 		protected void SetupMergeModules(ILogger Logger)
 		{
-			void UpdatePrimaryModule(UEBuildBinary binary)
+			void SetPerModuleInline(UEBuildBinary binary)
 			{
-				// Check all modules in this binary to see if any references set PER_MODULE_INLINE_FILE
-				Dictionary<UEBuildModuleCPP, string> perModuleInlineFiles = new();
-				{
-					HashSet<UEBuildModule> allReferencedModules = [];
-					foreach (UEBuildModuleCPP module in binary.Modules)
-					{
-						if (allReferencedModules.Contains(module))
-						{
-							continue;
-						}
-						List<UEBuildModule> referencedModules = [];
-						allReferencedModules.Add(module);
-						module.GetAllDependencyModules(referencedModules, allReferencedModules, false, false, false);
-						allReferencedModules.UnionWith(referencedModules);
-						allReferencedModules.UnionWith(referencedModules.SelectMany(x => x.PublicIncludePathModules ?? new()));
-					}
+				binary.PrimaryModule.bCreatePerModuleFile = false;
+
+				var queue = new Queue<(UEBuildModule Module, UEBuildModule Source)>();
+				var visited = new HashSet<UEBuildModule>();
 
-					foreach (UEBuildModuleCPP reference in allReferencedModules.OfType<UEBuildModuleCPP>())
+				// Queue up primary module and all direct dependencies
+				queue.Enqueue((binary.PrimaryModule, binary.PrimaryModule));
+				foreach (var m in binary.Modules)
+				{
+					if (visited.Add(m))
 					{
-						string? perModuleInlineFile = reference.Rules.PublicDefinitions.LastOrDefault(x => x.StartsWith("PER_MODULE_INLINE_FILE"));
-						if (perModuleInlineFile != null)
-						{
-							perModuleInlineFiles.Add(reference, perModuleInlineFile);
-						}
+						queue.Enqueue((m, m));
 					}
 				}
 
-				// If PER_MODULE_INLINE_FILE is found, move that definition to the primary module
-				if (perModuleInlineFiles.Any())
+				// Search for non-engine PerModuleInlineFile.. 
+				while (queue.Count > 0)
 				{
-					// Warn if there are PER_MODULE_INLINE_FILE conflicts
-					UEBuildModuleCPP module = perModuleInlineFiles.First().Key;
-					string definition = perModuleInlineFiles.First().Value;
-					if (perModuleInlineFiles.Values.Distinct().Count() > 1)
+					var current = queue.Dequeue();
+					if (current.Module.RulesFile.IsUnderDirectory(Unreal.EngineDirectory))
 					{
-						IOrderedEnumerable<UEBuildModuleCPP> ordered = perModuleInlineFiles.Keys
-							.OrderBy(m => !m.RulesFile.IsUnderDirectory(Unreal.EngineDirectory))
-							.ThenBy(m => m.Name);
-						Logger.LogWarning("Conflicting definitions of PER_MODULE_INLINE_FILE define when merging '{Modules}', first instance '{First}' will be used", String.Join(", ", ordered), perModuleInlineFiles[ordered.First()]);
-						module = ordered.First();
-						definition = perModuleInlineFiles[ordered.First()];
+						continue;
+					}
+					if (current.Module.Rules.PublicDefinitions.Exists(x => x.StartsWith("PER_MODULE_INLINE_FILE")))
+					{
+						((UEBuildModuleCPP)current.Source).bCreatePerModuleFile = true;
+						return;
 					}
 
-					if (!binary.PrimaryModule.Rules.PublicDefinitions.Contains(definition))
+					foreach (var dep in current.Module.PublicDependencyModules!.Union(current.Module.PrivateDependencyModules!))
 					{
-						Logger.LogDebug("Updating {Binary} primary module {Module} to include define {Define} from {SourceModule}", binary.OutputFilePath.GetFileNameWithoutExtension(), binary.PrimaryModule.Name, definition, module.Name);
-						binary.PrimaryModule.PublicDefinitions.Add(definition);
-						binary.PrimaryModule.PublicIncludePaths.UnionWith(module.PublicIncludePaths);
+						if (visited.Add(dep))
+						{
+							queue.Enqueue((dep, current.Source));
+						}
 					}
 				}
+
+				// No non-eninge PerModuleInline files found, Fall back to Core.
+				((UEBuildModuleCPP)Modules["Core"]).bCreatePerModuleFile = true;
 			}
 
 			void MergeModulesToDLL(string name, IEnumerable<UEBuildModuleCPP> modules)
@@ -3547,9 +3538,8 @@ namespace UnrealBuildTool
 				foreach (UEBuildModuleCPP module in modules)
 				{
 					module.Binary = mergedBinary;
+					module.bCreatePerModuleFile = false;
 				}
-				UpdatePrimaryModule(mergedBinary);
-
 				Binaries.Add(mergedBinary);
 			}
 
@@ -3563,8 +3553,8 @@ namespace UnrealBuildTool
 				{
 					Binaries[0].AddModule(module);
 					module.Binary = Binaries[0];
+					module.bCreatePerModuleFile = false;
 				}
-				UpdatePrimaryModule(Binaries[0]);
 			}
 
 			void ConsumeModules(string name, IEnumerable<UEBuildModuleCPP> modules)
@@ -3719,6 +3709,9 @@ namespace UnrealBuildTool
 			// Remove any binaries that no longer have modules due to merging
 			Binaries.ForEach(x => x.Modules.RemoveAll(y => y.Binary != x));
 			Binaries.RemoveAll(x => !x.Modules.Any());
+
+			// Set module that should pull in PerModuleInline file
+			Binaries.ForEach(binary => SetPerModuleInline(binary));
 		}
 
 		/// <summary>
diff --git a/Engine/Source/Programs/UnrealBuildTool/Platform/Linux/LinuxToolChain.cs b/Engine/Source/Programs/UnrealBuildTool/Platform/Linux/LinuxToolChain.cs
--- a/Engine/Source/Programs/UnrealBuildTool/Platform/Linux/LinuxToolChain.cs
+++ b/Engine/Source/Programs/UnrealBuildTool/Platform/Linux/LinuxToolChain.cs
@@ -1111,9 +1111,9 @@ namespace UnrealBuildTool
 			List<string> ResponseLines = new List<string>();
 			foreach (FileItem InputFile in LinkEnvironment.InputFiles)
 			{
-				if (InputFile.HasExtension(".dynlist"))
+				if (InputFile.HasExtension(".ldscript"))
 				{
-					ResponseLines.Add($"--dynamic-list=\"{NormalizeCommandLinePath(InputFile)}\"");
+					ResponseLines.Add($"--version-script=\"{NormalizeCommandLinePath(InputFile)}\"");
 				}
 				else
 				{
@@ -1576,7 +1576,7 @@ namespace UnrealBuildTool
 
 		public override string GetExtraLinkFileExtension()
 		{
-			return "dynlist";
+			return "ldscript";
 		}
 	}
 }
diff --git a/Engine/Source/Programs/UnrealBuildTool/Platform/Windows/VCToolChain.cs b/Engine/Source/Programs/UnrealBuildTool/Platform/Windows/VCToolChain.cs
--- a/Engine/Source/Programs/UnrealBuildTool/Platform/Windows/VCToolChain.cs
+++ b/Engine/Source/Programs/UnrealBuildTool/Platform/Windows/VCToolChain.cs
@@ -3559,7 +3559,7 @@ namespace UnrealBuildTool
 
 		public override string GetExtraLinkFileExtension()
 		{
-			return Target.WindowsPlatform.Compiler.IsClang() ? "dynlist" : "obj";
+			return "obj";
 		}
 
 		public override void SetUpGlobalEnvironment(ReadOnlyTargetRules Target, CppCompileEnvironment GlobalCompileEnvironment, LinkEnvironment GlobalLinkEnvironment)
diff --git a/Engine/Source/Programs/UnrealBuildTool/ToolChain/UEToolChain.cs b/Engine/Source/Programs/UnrealBuildTool/ToolChain/UEToolChain.cs
--- a/Engine/Source/Programs/UnrealBuildTool/ToolChain/UEToolChain.cs
+++ b/Engine/Source/Programs/UnrealBuildTool/ToolChain/UEToolChain.cs
@@ -376,6 +376,11 @@ namespace UnrealBuildTool
 			throw new NotSupportedException("This platform does not support merged modules.");
 		}
 
+		public virtual IEnumerable<string> GetExtraLinkFileAdditionalSymbols(UEBuildBinaryType binaryType)
+		{
+			return Enumerable.Empty<string>();
+		}
+
 		/// <summary>
 		/// Runs the provided tool and argument. Returns the output, using a rexex capture if one is provided
 		/// </summary>
