Full adoption of Clang warning & error levels for all initial disables.

diff --git a/Engine/Source/Programs/Shared/UnrealBuildTool.Tests/CppCompileWarningsTests.cs b/Engine/Source/Programs/Shared/UnrealBuildTool.Tests/CppCompileWarningsTests.cs
--- a/Engine/Source/Programs/Shared/UnrealBuildTool.Tests/CppCompileWarningsTests.cs
+++ b/Engine/Source/Programs/Shared/UnrealBuildTool.Tests/CppCompileWarningsTests.cs
@@ -498,7 +498,8 @@ namespace UnrealBuildTool.Tests
 					deterministicCompileEnvironment.bDeterministic = true;
 
 					args = baseWarnings.GenerateWarningCommandLineArgs(deterministicCompileEnvironment, typeof(ClangToolChain), newVersionNumber);
-					Assert.IsFalse(args.Any());
+					Assert.IsTrue(args.Count() == 1);
+					Assert.IsTrue(args.Contains("-Wno-date-time"));
 
 					deterministicCompileEnvironment.bDeterministic = false;
 					baseWarnings.DeterministicWarningLevel = WarningLevel.Warning;
diff --git a/Engine/Source/Programs/UnrealBuildTool/Configuration/CppCompileWarnings.cs b/Engine/Source/Programs/UnrealBuildTool/Configuration/CppCompileWarnings.cs
--- a/Engine/Source/Programs/UnrealBuildTool/Configuration/CppCompileWarnings.cs
+++ b/Engine/Source/Programs/UnrealBuildTool/Configuration/CppCompileWarnings.cs
@@ -179,7 +179,7 @@ namespace UnrealBuildTool
 		/// Clang -
 		///		https://clang.llvm.org/docs/DiagnosticsReference.html#wenum-enum-conversion
 		/// </summary>
-		[WarningsClangToolChain(["-Wno-enum-enum-conversion"])]
+		[WarningsClangToolChain(["-Wno-enum-enum-conversion"], ["-Wenum-enum-conversion", "-Wno-error=enum-enum-conversion"], ["-Wenum-enum-conversion"])]
 		[WarningLevelDefault(WarningLevel.Off)]
 		public WarningLevel EnumEnumConversionWarningLevel
 		{
@@ -193,7 +193,7 @@ namespace UnrealBuildTool
 		/// Clang - 
 		///		https://clang.llvm.org/docs/DiagnosticsReference.html#wenum-float-conversion
 		/// </summary>
-		[WarningsClangToolChain(["-Wno-enum-float-conversion"])]
+		[WarningsClangToolChain(["-Wno-enum-float-conversion"], ["-Wenum-float-conversion", "-Wno-error=enum-float-conversion"], ["-Wenum-float-conversion"])]
 		[WarningLevelDefault(WarningLevel.Off)]
 		public WarningLevel EnumFloatConversionWarningLevel
 		{
@@ -207,7 +207,7 @@ namespace UnrealBuildTool
 		/// Clang -
 		///		https://clang.llvm.org/docs/DiagnosticsReference.html#wambiguous-reversed-operator
 		/// </summary>
-		[WarningsClangToolChain(["-Wno-ambiguous-reversed-operator"], null, null, [nameof(FilterID.Cpp20Min)])]
+		[WarningsClangToolChain(["-Wno-ambiguous-reversed-operator"], ["-Wambiguous-reversed-operator", "-Wno-error=ambiguous-reversed-operator"], ["-Wambiguous-reversed-operator"], [nameof(FilterID.Cpp20Min)])]
 		[WarningLevelDefault(WarningLevel.Off)]
 		public WarningLevel AmbiguousReversedOperatorWarningLevel
 		{
@@ -221,7 +221,7 @@ namespace UnrealBuildTool
 		/// Clang -
 		///		https://clang.llvm.org/docs/DiagnosticsReference.html#wdeprecated-anon-enum-enum-conversion
 		/// </summary>
-		[WarningsClangToolChain(["-Wno-deprecated-anon-enum-enum-conversion"], null, null, [nameof(FilterID.Cpp20Min)])]
+		[WarningsClangToolChain(["-Wno-deprecated-anon-enum-enum-conversion"], ["-Wdeprecated-anon-enum-enum-conversion", "-Wno-error=deprecated-anon-enum-enum-conversion"], ["-Wdeprecated-anon-enum-enum-conversion"], [nameof(FilterID.Cpp20Min)])]
 		[WarningLevelDefault(WarningLevel.Off)]
 		public WarningLevel DeprecatedAnonEnumEnumConversionWarningLevel
 		{
@@ -235,7 +235,7 @@ namespace UnrealBuildTool
 		/// Clang -
 		///		https://clang.llvm.org/docs/DiagnosticsReference.html#wdeprecated-volatile
 		/// </summary>
-		[WarningsClangToolChain(["-Wno-deprecated-volatile"], null, null, [nameof(FilterID.Cpp20Min)])]
+		[WarningsClangToolChain(["-Wno-deprecated-volatile"], ["-Wdeprecated-volatile", "-Wno-error=deprecated-volatile"], ["-Wdeprecated-volatile"], [nameof(FilterID.Cpp20Min)])]
 		[WarningLevelDefault(WarningLevel.Off)]
 		public WarningLevel DeprecatedVolatileWarningLevel
 		{
@@ -249,7 +249,7 @@ namespace UnrealBuildTool
 		/// Clang -
 		///		https://clang.llvm.org/docs/DiagnosticsReference.html#wordered-compare-function-pointers
 		/// </summary>
-		[WarningsClangToolChain(["-Wno-ordered-compare-function-pointers"])]
+		[WarningsClangToolChain(["-Wno-ordered-compare-function-pointers"], ["-Wordered-compare-function-pointers", "-Wno-error=ordered-compare-function-pointers"], ["-Wordered-compare-function-pointers"])]
 		[WarningLevelDefault(WarningLevel.Off)]
 		public WarningLevel OrderedCompareFunctionPointers
 		{
@@ -263,7 +263,7 @@ namespace UnrealBuildTool
 		/// Clang -
 		///		https://clang.llvm.org/docs/DiagnosticsReference.html#wbitwise-instead-of-logical
 		/// </summary>
-		[WarningsClangToolChain(["-Wno-bitwise-instead-of-logical"], null, null, [nameof(FilterID.Version14Min)])]
+		[WarningsClangToolChain(["-Wno-bitwise-instead-of-logical"], ["-Wbitwise-instead-of-logical", "-Wno-error=bitwise-instead-of-logical"], ["-Wbitwise-instead-of-logical"], [nameof(FilterID.Version14Min)])]
 		[WarningLevelDefault(WarningLevel.Off)]
 		public WarningLevel BitwiseInsteadOfLogicalWarningLevel
 		{
@@ -277,7 +277,7 @@ namespace UnrealBuildTool
 		/// Clang -
 		///		https://clang.llvm.org/docs/DiagnosticsReference.html#wdeprecated-copy
 		/// </summary>
-		[WarningsClangToolChain(["-Wno-deprecated-copy"], null, null, [nameof(FilterID.Version16Min)])]
+		[WarningsClangToolChain(["-Wno-deprecated-copy"], ["-Wdeprecated-copy", "-Wno-error=deprecated-copy"], ["-Wdeprecated-copy"], [nameof(FilterID.Version16Min)])]
 		[WarningLevelDefault(WarningLevel.Off)]
 		public WarningLevel DeprecatedCopyWarningLevel
 		{
@@ -291,7 +291,7 @@ namespace UnrealBuildTool
 		/// Clang -
 		///		https://clang.llvm.org/docs/DiagnosticsReference.html#wdeprecated-copy-with-user-provided-copy
 		/// </summary>
-		[WarningsClangToolChain(["-Wno-deprecated-copy-with-user-provided-copy"], null, null, [nameof(FilterID.Version16Min)])]
+		[WarningsClangToolChain(["-Wno-deprecated-copy-with-user-provided-copy"], ["-Wdeprecated-copy-with-user-provided-copy", "-Wno-error=deprecated-copy-with-user-provided-copy"], ["-Wdeprecated-copy-with-user-provided-copy"], [nameof(FilterID.Version16Min)])]
 		[WarningLevelDefault(WarningLevel.Off)]
 		public WarningLevel DeprecatedCopyWithUserProvidedCopyWarningLevel
 		{
@@ -305,7 +305,7 @@ namespace UnrealBuildTool
 		/// Clang -
 		///		https://clang.llvm.org/docs/DiagnosticsReference.html#winvalid-unevaluated-string
 		/// </summary>
-		[WarningsClangToolChain(["-Wno-invalid-unevaluated-string"], null, null, [nameof(SpecializedFilters.AndroidNDKR26ToolChainVersionExclusion)])]
+		[WarningsClangToolChain(["-Wno-invalid-unevaluated-string"], ["-Winvalid-unevaluated-string", "-Wno-error=invalid-unevaluated-string"], ["-Winvalid-unevaluated-string"], [nameof(SpecializedFilters.AndroidNDKR26ToolChainVersionExclusion)])]
 		[WarningLevelDefault(WarningLevel.Off)]
 		public WarningLevel InvalidUnevaluatedStringWarningLevel
 		{
@@ -324,7 +324,7 @@ namespace UnrealBuildTool
 		/// It could be easily fixed in the engine, but would create the risk of Win64 building fine and the code failing just on Clang-based platforms. 
 		/// We tend to use NAN just as a bit pattern, and not rely on it in calculations, so it should be reasonably safe to disable.
 		/// </remarks>
-		[WarningsClangToolChain(["-Wno-nan-infinity-disabled"], null, null, [nameof(FilterID.Version18Min)])]
+		[WarningsClangToolChain(["-Wno-nan-infinity-disabled"], ["-Wnan-infinity-disabled", "-Wno-error=nan-infinity-disabled"], ["-Wnan-infinity-disabled"], [nameof(FilterID.Version18Min)])]
 		[WarningLevelDefault(WarningLevel.Off)]
 		public WarningLevel NaNInfinityDisabledWarningLevel
 		{
@@ -338,7 +338,7 @@ namespace UnrealBuildTool
 		/// Clang -
 		///		https://clang.llvm.org/docs/DiagnosticsReference.html#wextra-qualification
 		/// </summary>
-		[WarningsClangToolChain(["-Wno-extra-qualification"], null, null, [nameof(FilterID.Version19Min)])]
+		[WarningsClangToolChain(["-Wno-extra-qualification"], ["-Wextra-qualification", "-Wno-error=extra-qualification"], ["-Wextra-qualification"], [nameof(FilterID.Version19Min)])]
 		[WarningLevelDefault(WarningLevel.Off)]
 		public WarningLevel LevelExtraQualificationWarningLevel
 		{
@@ -352,7 +352,7 @@ namespace UnrealBuildTool
 		/// Clang -
 		///		https://clang.llvm.org/docs/DiagnosticsReference.html#wcast-function-type-mismatch
 		/// </summary>
-		[WarningsClangToolChain(["-Wno-cast-function-type-mismatch"], null, null, [nameof(SpecializedFilters.AndroidNDKR28ToolChainVersionExclusion)])]
+		[WarningsClangToolChain(["-Wno-cast-function-type-mismatch"], ["-Wcast-function-type-mismatch", "-Wno-error=cast-function-type-mismatch"], ["-Wcast-function-type-mismatch"], [nameof(SpecializedFilters.AndroidNDKR28ToolChainVersionExclusion)])]
 		[WarningLevelDefault(WarningLevel.Off)]
 		public WarningLevel CastFunctionTypeMismatchWarningLevel
 		{
@@ -366,7 +366,7 @@ namespace UnrealBuildTool
 		/// Clang -
 		///		https://clang.llvm.org/docs/DiagnosticsReference.html#wmissing-template-arg-list-after-template-kw
 		/// </summary>
-		[WarningsClangToolChain(["-Wno-missing-template-arg-list-after-template-kw"], null, null, [nameof(SpecializedFilters.AndroidNDKR28ToolChainVersionExclusion)])]
+		[WarningsClangToolChain(["-Wno-missing-template-arg-list-after-template-kw"], ["-Wmissing-template-arg-list-after-template-kw", "-Wno-error=missing-template-arg-list-after-template-kw"], ["-Wmissing-template-arg-list-after-template-kw"], [nameof(SpecializedFilters.AndroidNDKR28ToolChainVersionExclusion)])]
 		[WarningLevelDefault(WarningLevel.Off)]
 		public WarningLevel MissingTemplateArgListAfterTemplateWarningLevel
 		{
@@ -381,7 +381,7 @@ namespace UnrealBuildTool
 		///		https://clang.llvm.org/docs/DiagnosticsReference.html#wgnu-string-literal-operator-template
 		/// </summary>
 		/// <remarks>We use this feature to allow static FNames.</remarks>
-		[WarningsClangToolChain(["-Wno-gnu-string-literal-operator-template"])]
+		[WarningsClangToolChain(["-Wno-gnu-string-literal-operator-template"], ["-Wgnu-string-literal-operator-template", "-Wno-error=gnu-string-literal-operator-template"], ["-Wgnu-string-literal-operator-template"])]
 		[WarningLevelDefault(WarningLevel.Off)]
 		public WarningLevel GNUStringLiteralOperatorTemplateWarningLevel
 		{
@@ -395,7 +395,7 @@ namespace UnrealBuildTool
 		/// Clang -
 		///		https://clang.llvm.org/docs/DiagnosticsReference.html#winconsistent-missing-override
 		/// </summary>
-		[WarningsClangToolChain(["-Wno-inconsistent-missing-override"])]
+		[WarningsClangToolChain(["-Wno-inconsistent-missing-override"], ["-Winconsistent-missing-override", "-Wno-error=inconsistent-missing-override"], ["-Winconsistent-missing-override"])]
 		[WarningLevelDefault(WarningLevel.Off)]
 		public WarningLevel InconsistentMissingOverrideWarningLevel
 		{
@@ -410,7 +410,7 @@ namespace UnrealBuildTool
 		///		https://clang.llvm.org/docs/DiagnosticsReference.html#winvalid-offsetof
 		/// </summary>
 		/// <remarks>Needed to suppress warnings about using offsetof on non-POD types.</remarks>
-		[WarningsClangToolChain(["-Wno-invalid-offsetof"])]
+		[WarningsClangToolChain(["-Wno-invalid-offsetof"], ["-Winvalid-offsetof", "-Wno-error=invalid-offsetof"], ["-Winvalid-offsetof"])]
 		[WarningLevelDefault(WarningLevel.Off)]
 		public WarningLevel InvalidOffsetWarningLevel
 		{
@@ -425,7 +425,7 @@ namespace UnrealBuildTool
 		///		https://clang.llvm.org/docs/DiagnosticsReference.html#wswitch
 		/// </summary>
 		/// <remarks>This hides the "enumeration value 'XXXXX' not handled in switch [-Wswitch]" warnings - we should maybe remove this at some point and add UE_LOG(, Fatal, ) to default cases.</remarks>
-		[WarningsClangToolChain(["-Wno-switch"])]
+		[WarningsClangToolChain(["-Wno-switch"], ["-Wswitch", "-Wno-error=switch"], ["-Wswitch"])]
 		[WarningLevelDefault(WarningLevel.Off)]
 		public WarningLevel SwitchWarningLevel
 		{
@@ -440,7 +440,7 @@ namespace UnrealBuildTool
 		///		https://clang.llvm.org/docs/DiagnosticsReference.html#wtautological-compare
 		/// </summary>
 		/// <remarks>This hides the "warning : comparison of unsigned expression &lt; 0 is always false" type warnings due to constant comparisons, which are possible with template arguments/</remarks>
-		[WarningsClangToolChain(["-Wno-tautological-compare"])]
+		[WarningsClangToolChain(["-Wno-tautological-compare"], ["-Wtautological-compare", "Wno-error=tautological-compare"], ["-Wtautological-compare"])]
 		[WarningLevelDefault(WarningLevel.Off)]
 		public WarningLevel TautologicalCompareWarningLevel
 		{
@@ -455,7 +455,7 @@ namespace UnrealBuildTool
 		///		https://clang.llvm.org/docs/DiagnosticsReference.html#wunknown-pragmas
 		/// </summary>
 		/// <remarks>Slate triggers this (with its optimize on/off pragmas).</remarks>
-		[WarningsClangToolChain(["-Wno-unknown-pragmas"])]
+		[WarningsClangToolChain(["-Wno-unknown-pragmas"], ["-Wunknown-pragmas", "-Wno-error=unknown-pragmas"], ["-Wunknown-pragmas"])]
 		[WarningLevelDefault(WarningLevel.Off)]
 		public WarningLevel UnknownPragmasWarningLevel
 		{
@@ -481,7 +481,10 @@ namespace UnrealBuildTool
 		/// unused-local-typedef: Clang is being overly strict here? PhysX headers trigger this.
 		/// unused-private-field: This will prevent the issue of warnings for unused private variables. MultichannelTcpSocket.h triggers this, possibly more.
 		/// </remarks>
-		[WarningsClangToolChain(["-Wno-unused-but-set-variable", "-Wno-unused-but-set-parameter", "-Wno-unused-function", "-Wno-unused-lambda-capture", "-Wno-unused-local-typedef", "-Wno-unused-private-field", "-Wno-unused-variable"], null, null)]
+		[WarningsClangToolChain(
+		["-Wno-unused-but-set-variable", "-Wno-unused-but-set-parameter", "-Wno-unused-function", "-Wno-unused-lambda-capture", "-Wno-unused-local-typedef", "-Wno-unused-private-field", "-Wno-unused-variable"], 
+		["-Wunused-but-set-variable","-Wno-error=unused-but-set-variable", "-Wunused-but-set-parameter", "-Wno-error=unused-but-set-parameter", "-Wunused-function", "-Wno-error=unused-function", "-Wunused-lambda-capture", "-Wno-error=unused-lambda-capture", "-Wunused-local-typedef", "-Wno-error=unused-local-typedef", "-Wunused-private-field", "-Wno-error=unused-private-field", "-Wunused-variable", "-Wno-error=unused-variable"], 
+		["-Wunused-but-set-variable", "-Wunused-but-set-parameter", "-Wunused-function", "-Wunused-lambda-capture", "-Wunused-local-typedef", "-Wunused-private-field", "-Wunused-variable"])]
 		[WarningLevelDefault(WarningLevel.Off)]
 		public WarningLevel UnusedWarningLevel
 		{
@@ -496,7 +499,7 @@ namespace UnrealBuildTool
 		///		https://clang.llvm.org/docs/DiagnosticsReference.html#wundefined-var-template
 		/// </summary>
 		/// <remarks>Not a good warning to be disabling.</remarks>
-		[WarningsClangToolChain(["-Wno-undefined-var-template"])]
+		[WarningsClangToolChain(["-Wno-undefined-var-template"], ["-Wundefined-var-template", "-Wno-error=undefined-var-template"], ["-Wundefined-var-template"])]
 		[WarningLevelDefault(WarningLevel.Off)]
 		public WarningLevel UndefinedVarTemplateWarningLevel
 		{
@@ -512,7 +515,7 @@ namespace UnrealBuildTool
 		///		https://clang.llvm.org/docs/DiagnosticsReference.html#wprofile-instr-unprofiled
 		/// </summary>
 		/// <remarks> Clang emits warnings for each compiled object file that doesn't have a matching entry in the profile data. This can happen when the profile data is older than the binaries we're compiling. Disable these warnings.</remarks>
-		[WarningsClangToolChain(["-Wno-profile-instr-out-of-date", "-Wno-profile-instr-unprofiled"], null, null, [nameof(StandardFilters.PGOOptimizedFilter)])]
+		[WarningsClangToolChain(["-Wno-profile-instr-out-of-date", "-Wno-profile-instr-unprofiled"], ["-Wprofile-instr-out-of-date", "-Wno-error=profile-instr-out-of-date", "-Wprofile-instr-unprofiled", "-Wno-error=profile-instr-unprofiled"], ["-Wprofile-instr-out-of-date", "-Wprofile-instr-unprofiled"], [nameof(StandardFilters.PGOOptimizedFilter)])]
 		[WarningLevelDefault(WarningLevel.Off)]
 		public WarningLevel ProfileInstructWarningLevel
 		{
@@ -526,7 +529,7 @@ namespace UnrealBuildTool
 		///		https://clang.llvm.org/docs/DiagnosticsReference.html#wbackend-plugin
 		/// </summary>
 		/// <remarks>Apparently there can be hashing conflicts with PGO which can result in: 'Function control flow change detected (hash mismatch)' warnings. </remarks>
-		[WarningsClangToolChain(["-Wno-backend-plugin"], null, null, [nameof(StandardFilters.PGOOptimizedFilter)])]
+		[WarningsClangToolChain(["-Wno-backend-plugin"], ["-Wbackend-plugin", "-Wno-error=backend-plugin"], ["-Wbackend-plugin"], [nameof(StandardFilters.PGOOptimizedFilter)])]
 		[WarningLevelDefault(WarningLevel.Off)]
 		public WarningLevel BackendPluginWarningLevel
 		{
@@ -544,7 +547,7 @@ namespace UnrealBuildTool
 		/// Shipping builds will cause -Wno-unused-value to warn with "ensure", so disable only in those case.
 		/// -Wunused-result - we will always re-enable this as it's controlled by the same set. This will be moved out to a separate property. It is on by default, so it only must be re-enabled in the disable context.
 		/// </remarks>
-		[UnusedValueClangToolChain(["-Wno-unused-value", "-Wunused-result"])]
+		[UnusedValueClangToolChain(["-Wno-unused-value", "-Wunused-result"], ["-Wunused-value", "-Wno-error=unused-value",  "-Wunused-result"], ["-Wunused-value"])]
 		[WarningLevelDefault(WarningLevel.Off)]
 		public WarningLevel UnusedValueWarningLevel
 		{
@@ -559,7 +562,7 @@ namespace UnrealBuildTool
 		///		https://clang.llvm.org/docs/DiagnosticsReference.html#wshorten-64-to-32
 		/// </summary>
 		/// <remarks>-Wimplicit-int-conversion also controls -Wshorten-64-to-32 as of clang 19.</remarks>
-		[WarningsClangToolChain(["-Wno-shorten-64-to-32"], null, null, [nameof(FilterID.Version19Min)])]
+		[WarningsClangToolChain(["-Wno-shorten-64-to-32"], ["-Wshorten-64-to-32", "-Wno-error=shorten-64-to-32"], ["-Wshorten-64-to-32"], [nameof(FilterID.Version19Min)])]
 		[WarningLevelDefault(WarningLevel.Off)]
 		public WarningLevel Shorten64To32WarningLevel
 		{
@@ -574,7 +577,7 @@ namespace UnrealBuildTool
 		///		https://clang.llvm.org/docs/DiagnosticsReference.html#wdllexport-explicit-instantiation-decl
 		/// </summary>
 		/// <remarks>The code base contains lots of places where we do "extern template X_API class ..." and we want to keep doing that to reduce compile times.</remarks>
-		[WarningsClangToolChain(["-Wno-dllexport-explicit-instantiation-decl"])]
+		[WarningsClangToolChain(["-Wno-dllexport-explicit-instantiation-decl"], ["-Wdllexport-explicit-instantiation-decl", "-Wno-error=dllexport-explicit-instantiation-decl"], ["-Wdllexport-explicit-instantiation-decl"])]
 		[WarningLevelDefault(WarningLevel.Off)]
 		public WarningLevel DllExportExplicitInstantiationDeclWarningLevel
 		{
@@ -588,7 +591,7 @@ namespace UnrealBuildTool
 		/// Clang -
 		///		https://clang.llvm.org/docs/DiagnosticsReference.html#wnontrivial-memaccess
 		/// </summary>
-		[WarningsClangToolChain(["-Wno-nontrivial-memaccess"], null, null, [nameof(FilterID.Version20Min)])]
+		[WarningsClangToolChain(["-Wno-nontrivial-memaccess"], ["-Wnontrivial-memaccess", "-Wno-error=nontrivial-memaccess"], ["-Wnontrivial-memaccess"], [nameof(FilterID.Version20Min)])]
 		[WarningLevelDefault(WarningLevel.Off)]
 		public WarningLevel NonTrivialMemAccessWarningLevel
 		{
@@ -607,7 +610,7 @@ namespace UnrealBuildTool
 		///		https://clang.llvm.org/docs/DiagnosticsReference.html#wmicrosoft
 		/// </summary>
 		/// <remarks>Allow Microsoft-specific syntax to slide, even though it may be non-standard.  Needed for Windows headers.</remarks>
-		[WarningsVCClang(["-Wno-microsoft"])]
+		[WarningsVCClang(["-Wno-microsoft"], ["-Wmicrosoft", "-Wno-error=microsoft"], ["-Wmicrosoft"])]
 		[WarningLevelDefault(WarningLevel.Off)]
 		public WarningLevel MicrosoftGroupWarningLevel
 		{
@@ -622,7 +625,7 @@ namespace UnrealBuildTool
 		///		https://clang.llvm.org/docs/DiagnosticsReference.html#wmsvc-include
 		/// </summary>
 		/// <remarks> Hack due to how we have our 'DummyPCH' wrappers setup when using unity builds.  This warning should not be disabled!</remarks>
-		[WarningsVCClang(["-Wno-msvc-include"])]
+		[WarningsVCClang(["-Wno-msvc-include"], ["-Wmsvc-include", "-Wno-error=msvc-include"], ["-Wmsvc-include"])]
 		[WarningLevelDefault(WarningLevel.Off)]
 		public WarningLevel MSVCIncludeWarningLevel
 		{
@@ -640,7 +643,7 @@ namespace UnrealBuildTool
 		/// This is disabled because clang explicitly warns about changing pack alignment in a header and not restoring it afterwards, which is something we do with the Pre/PostWindowsApi.h headers.
 		/// @TODO clang: push/pop this in  Pre/PostWindowsApi.h headers instead?
 		/// </remarks>
-		[WarningsVCClang(["-Wno-pragma-pack"])]
+		[WarningsVCClang(["-Wno-pragma-pack"], ["-Wpragma-pack", "-Wno-error=pragma-pack"], ["-Wpragma-pack"])]
 		[WarningLevelDefault(WarningLevel.Off)]
 		public WarningLevel PragmaPackWarningLevel
 		{
@@ -655,7 +658,7 @@ namespace UnrealBuildTool
 		///		https://clang.llvm.org/docs/DiagnosticsReference.html#winline-new-delete
 		/// </summary>
 		/// <remarks>@todo clang: We declare operator new as inline.  Clang doesn't seem to like that.</remarks>
-		[WarningsVCClang(["-Wno-inline-new-delete"])]
+		[WarningsVCClang(["-Wno-inline-new-delete"], ["-Winline-new-delete", "-Wno-error=inline-new-delete"], ["-Winline-new-delete"])]
 		[WarningLevelDefault(WarningLevel.Off)]
 		public WarningLevel InlineNewDeleteWarningLevel
 		{
@@ -670,7 +673,7 @@ namespace UnrealBuildTool
 		///		https://clang.llvm.org/docs/DiagnosticsReference.html#wimplicit-exception-spec-mismatch
 		/// </summary>
 
-		[WarningsVCClang(["-Wno-implicit-exception-spec-mismatch"])]
+		[WarningsVCClang(["-Wno-implicit-exception-spec-mismatch"], ["-Wimplicit-exception-spec-mismatch", "-Wno-error=implicit-exception-spec-mismatch"], ["-Wimplicit-exception-spec-mismatch"])]
 		[WarningLevelDefault(WarningLevel.Off)]
 		public WarningLevel ImplicitExceptionSpecMismatchWarningLevel
 		{
@@ -685,7 +688,7 @@ namespace UnrealBuildTool
 		///		https://clang.llvm.org/docs/DiagnosticsReference.html#wundefined-bool-conversion
 		/// </summary>
 		/// <remarks>Sometimes we compare 'this' pointers against nullptr, which Clang warns about by default.</remarks>
-		[WarningsVCClang(["-Wno-undefined-bool-conversion"])]
+		[WarningsVCClang(["-Wno-undefined-bool-conversion"], ["-Wundefined-bool-conversion", "-Wno-error=undefined-bool-conversion"], ["-Wundefined-bool-conversion"])]
 		[WarningLevelDefault(WarningLevel.Off)]
 		public WarningLevel UndefinedBoolConversionWarningLevel
 		{
@@ -700,7 +703,7 @@ namespace UnrealBuildTool
 		///		https://clang.llvm.org/docs/DiagnosticsReference.html#wdeprecated-writable-strings
 		/// </summary>
 
-		[WarningsVCClang(["-Wno-deprecated-writable-strings"])]
+		[WarningsVCClang(["-Wno-deprecated-writable-strings"], ["-Wdeprecated-writable-strings", "-Wno-error=deprecated-writable-strings"], ["-Wdeprecated-writable-strings"])]
 		[WarningLevelDefault(WarningLevel.Off)]
 		public WarningLevel DeprecatedWriteableStringsWarningLevel
 		{
@@ -714,7 +717,7 @@ namespace UnrealBuildTool
 		/// Clang -
 		///		https://clang.llvm.org/docs/DiagnosticsReference.html#wdeprecated-register
 		/// </summary>
-		[WarningsVCClang(["-Wno-deprecated-register"])]
+		[WarningsVCClang(["-Wno-deprecated-register"], ["-Wdeprecated-register", "-Wno-error=deprecated-register"], ["-Wdeprecated-register"])]
 		[WarningLevelDefault(WarningLevel.Off)]
 		public WarningLevel DeprecatedRegisterWarningLevel
 		{
@@ -730,7 +733,7 @@ namespace UnrealBuildTool
 		///		https://clang.llvm.org/docs/DiagnosticsReference.html#wlogical-op-parentheses
 		/// </summary>
 		/// <remarks>Disable needed for external headers we shan't change.</remarks>
-		[WarningsVCClang(["-Wno-logical-op-parentheses"])]
+		[WarningsVCClang(["-Wno-logical-op-parentheses"], ["-Wlogical-op-parentheses", "-Wno-error=logical-op-parentheses"], ["-Wlogical-op-parentheses"])]
 		[WarningLevelDefault(WarningLevel.Off)]
 		public WarningLevel LogicalOpParenthesesWarningLevel
 		{
@@ -745,7 +748,7 @@ namespace UnrealBuildTool
 		///		https://clang.llvm.org/docs/DiagnosticsReference.html#wnull-arithmetic
 		/// </summary>
 		/// <remarks>Disable needed for external headers we shan't change.</remarks>
-		[WarningsVCClang(["-Wno-null-arithmetic"])]
+		[WarningsVCClang(["-Wno-null-arithmetic"], ["-Wnull-arithmetic", "-Wno-error=null-arithmetic"], ["-Wnull-arithmetic"])]
 		[WarningLevelDefault(WarningLevel.Off)]
 		public WarningLevel NullArithmeticWarningLevel
 		{
@@ -760,7 +763,7 @@ namespace UnrealBuildTool
 		///		https://clang.llvm.org/docs/DiagnosticsReference.html#wreturn-type-c-linkage
 		/// </summary>
 		/// <remarks>Disable needed for PhysX.</remarks>
-		[WarningsVCClang(["-Wno-return-type-c-linkage"])]
+		[WarningsVCClang(["-Wno-return-type-c-linkage"], ["-Wreturn-type-c-linkage", "-Wno-error=return-type-c-linkage"], ["-Wreturn-type-c-linkage"])]
 		[WarningLevelDefault(WarningLevel.Off)]
 		public WarningLevel ReturnTypeCLinkageWarningLevel
 		{
@@ -775,7 +778,7 @@ namespace UnrealBuildTool
 		///		https://clang.llvm.org/docs/DiagnosticsReference.html#wignored-attributes
 		/// </summary>
 		/// <remarks>Disable needed for nvtesslib.</remarks>
-		[WarningsVCClang(["-Wno-ignored-attributes"])]
+		[WarningsVCClang(["-Wno-ignored-attributes"], ["-Wignored-attributes", "-Wno-error=ignored-attributes"], ["-Wignored-attributes"])]
 		[WarningLevelDefault(WarningLevel.Off)]
 		public WarningLevel IgnoredAttributesWarningLevel
 		{
@@ -789,7 +792,7 @@ namespace UnrealBuildTool
 		/// Clang -
 		///		https://clang.llvm.org/docs/DiagnosticsReference.html#wuninitialized
 		/// </summary>
-		[WarningsVCClang(["-Wno-uninitialized"])]
+		[WarningsVCClang(["-Wno-uninitialized"], ["-Wuninitialized", "-Wno-error=uninitialized"], ["-Wuninitialized"])]
 		[WarningLevelDefault(WarningLevel.Off)]
 		public WarningLevel UninitializedWarningLevel
 		{
@@ -804,7 +807,7 @@ namespace UnrealBuildTool
 		///		https://clang.llvm.org/docs/DiagnosticsReference.html#wreturn-type
 		/// </summary>
 		/// <remarks>Disable needed for external headers we shan't change.</remarks>
-		[WarningsVCClang(["-Wno-return-type"])]
+		[WarningsVCClang(["-Wno-return-type"], ["-Wreturn-type", "-Wno-error=return-type"], ["-Wreturn-type"])]
 		[WarningLevelDefault(WarningLevel.Off)]
 		public WarningLevel ReturnTypeWarningLevel
 		{
@@ -819,7 +822,7 @@ namespace UnrealBuildTool
 		///		https://clang.llvm.org/docs/DiagnosticsReference.html#wunused-parameter
 		/// </summary>
 		/// <remarks>Unused function parameter. A lot are named 'bUnused'...</remarks>
-		[WarningsVCClang(["-Wno-unused-parameter"])]
+		[WarningsVCClang(["-Wno-unused-parameter"], ["-Wunused-parameter", "-Wno-error=unused-parameter"], ["-Wunused-parameter"])]
 		[WarningLevelDefault(WarningLevel.Off)]
 		public WarningLevel UnusedParameterWarningLevel
 		{
@@ -834,7 +837,7 @@ namespace UnrealBuildTool
 		///		https://clang.llvm.org/docs/DiagnosticsReference.html#wignored-qualifiers
 		/// </summary>
 		/// <remarks>const ignored when returning by value e.g. 'const int foo() { return 4; }'.</remarks>
-		[WarningsVCClang(["-Wno-ignored-qualifiers"])]
+		[WarningsVCClang(["-Wno-ignored-qualifiers"], ["-Wignored-qualifiers", "-Wno-error=ignored-qualifiers"], ["-Wignored-qualifiers"])]
 		[WarningLevelDefault(WarningLevel.Off)]
 		public WarningLevel IgnoredQualifiersWarningLevel
 		{
@@ -849,7 +852,7 @@ namespace UnrealBuildTool
 		///		https://clang.llvm.org/docs/DiagnosticsReference.html#wexpansion-to-defined
 		/// </summary>
 		/// <remarks>Usage of 'defined(X)' in a macro definition. Gives different results under MSVC.</remarks>
-		[WarningsVCClang(["-Wno-expansion-to-defined"])]
+		[WarningsVCClang(["-Wno-expansion-to-defined"], ["-Wexpansion-to-defined", "-Wno-error=expansion-to-defined"], ["-Wexpansion-to-defined"])]
 		[WarningLevelDefault(WarningLevel.Off)]
 		public WarningLevel ExpansionToDefined
 		{
@@ -864,7 +867,7 @@ namespace UnrealBuildTool
 		///		https://clang.llvm.org/docs/DiagnosticsReference.html#wsign-compare
 		/// </summary>
 		/// <remarks>Signed/unsigned comparison - millions of these.</remarks>
-		[WarningsVCClang(["-Wno-sign-compare"])]
+		[WarningsVCClang(["-Wno-sign-compare"], ["-Wsign-compare", "-Wno-error=sign-compare"], ["-Wsign-compare"])]
 		[WarningLevelDefault(WarningLevel.Off)]
 		public WarningLevel SignCompareWarningLevel
 		{
@@ -879,7 +882,7 @@ namespace UnrealBuildTool
 		///		https://clang.llvm.org/docs/DiagnosticsReference.html#wmissing-field-initializers
 		/// </summary>
 		/// <remarks>Excessive warning, generated when you initialize with MyStruct A = {0};</remarks>
-		[WarningsVCClang(["-Wno-missing-field-initializers"])]
+		[WarningsVCClang(["-Wno-missing-field-initializers"], ["-Wmissing-field-initializers", "-Wno-error=missing-field-initializers"], ["-Wmissing-field-initializers"])]
 		[WarningLevelDefault(WarningLevel.Off)]
 		public WarningLevel MissingFieldInitializersWarningLevel
 		{
@@ -893,7 +896,7 @@ namespace UnrealBuildTool
 		/// Clang -
 		///		https://clang.llvm.org/docs/DiagnosticsReference.html#wnonportable-include-path
 		/// </summary>
-		[WarningsVCClang(["-Wno-nonportable-include-path"])]
+		[WarningsVCClang(["-Wno-nonportable-include-path"], ["-Wnonportable-include-path", "-Wno-error=nonportable-include-path"], ["-Wnonportable-include-path"])]
 		[WarningLevelDefault(WarningLevel.Off)]
 		public WarningLevel NonPortableIncludePathWarningLevel
 		{
@@ -909,7 +912,7 @@ namespace UnrealBuildTool
 		///	MASVC -
 		///		https://learn.microsoft.com/en-us/cpp/preprocessor/token-pasting-operator-hash-hash?view=msvc-170
 		/// </summary>
-		[WarningsVCClang(["-Wno-invalid-token-paste"])]
+		[WarningsVCClang(["-Wno-invalid-token-paste"], ["-Winvalid-token-paste", "-Wno-error=invalid-token-paste"], ["-Winvalid-token-paste"])]
 		[WarningLevelDefault(WarningLevel.Off)]
 		public WarningLevel InvalidTokenPasteWarningLevel
 		{
@@ -923,7 +926,7 @@ namespace UnrealBuildTool
 		/// Clang -
 		///		https://clang.llvm.org/docs/DiagnosticsReference.html#wnull-pointer-arithmetic
 		/// </summary>
-		[WarningsVCClang(["-Wno-null-pointer-arithmetic"])]
+		[WarningsVCClang(["-Wno-null-pointer-arithmetic"], ["-Wnull-pointer-arithmetic", "-Wno-error=null-pointer-arithmetic"], ["-Wnull-pointer-arithmetic"])]
 		[WarningLevelDefault(WarningLevel.Off)]
 		public WarningLevel NullPointerArithmeticWarningLevel
 		{
@@ -938,7 +941,7 @@ namespace UnrealBuildTool
 		///		https://clang.llvm.org/docs/DiagnosticsReference.html#wconstant-logical-operand
 		/// </summary>
 		/// <remarks>Triggered by || of two template-derived values inside a static_assert.</remarks>
-		[WarningsVCClang(["-Wno-constant-logical-operand"])]
+		[WarningsVCClang(["-Wno-constant-logical-operand"], ["-Wconstant-logical-operand", "-Wno-error=constant-logical-operand"], ["-Wconstant-logical-operand"])]
 		[WarningLevelDefault(WarningLevel.Off)]
 		public WarningLevel ConstantLogicalOperandWarningLevel
 		{
@@ -952,7 +955,7 @@ namespace UnrealBuildTool
 		/// Clang -
 		///		https://clang.llvm.org/docs/DiagnosticsReference.html#wbitfield-enum-conversion
 		/// </summary>
-		[WarningsVCClang(["-Wno-bitfield-enum-conversion"])]
+		[WarningsVCClang(["-Wno-bitfield-enum-conversion"], ["-Wbitfield-enum-conversion", "-Wno-error=bitfield-enum-conversion"], ["-Wbitfield-enum-conversion"])]
 		[WarningLevelDefault(WarningLevel.Off)]
 		public WarningLevel BitfieldEnumConversion
 		{
@@ -966,7 +969,7 @@ namespace UnrealBuildTool
 		/// Clang -
 		///		https://clang.llvm.org/docs/DiagnosticsReference.html#wnull-pointer-subtraction
 		/// </summary>
-		[WarningsVCClang(["-Wno-null-pointer-subtraction"])]
+		[WarningsVCClang(["-Wno-null-pointer-subtraction"], ["-Wnull-pointer-subtraction", "-Wno-error=null-pointer-subtraction"], ["-Wnull-pointer-subtraction"])]
 		[WarningLevelDefault(WarningLevel.Off)]
 		public WarningLevel NullPointerSubtractionWarningLevel
 		{
@@ -980,7 +983,7 @@ namespace UnrealBuildTool
 		/// Clang -
 		///		https://clang.llvm.org/docs/DiagnosticsReference.html#wdangling
 		/// </summary>	
-		[WarningsVCClang(["-Wno-dangling"])]
+		[WarningsVCClang(["-Wno-dangling"], ["-Wdangling", "-Wno-error=dangling"], ["-Wdangling"])]
 		[WarningLevelDefault(WarningLevel.Off)]
 		public WarningLevel DanglingWarningLevel
 		{
@@ -1024,7 +1027,7 @@ namespace UnrealBuildTool
 		/// In order to ensure no change of behavior for MSVC toolchain within the clang compilation context, we will keep this in, and deprecate with assignments to the clang setting.
 		/// -Wunused-result - we will always re-enable this as it's controlled by the same set. This will be moved out to a separate property. It is on by default, so it only must be re-enabled in the disable context.
 		/// </remarks>
-		[WarningsVCClang(["-Wno-unused-value", "-Wunused-result"])]
+		[WarningsVCClang(["-Wno-unused-value", "-Wunused-result"], ["-Wunused-value", "-Wno-error=unused-value", "-Wunused-result"], ["-Wunused-value"])]
 		[WarningLevelDefault(WarningLevel.Off)]
 		public WarningLevel MSVCUnusedValueWarningLevel
 		{
@@ -1043,7 +1046,7 @@ namespace UnrealBuildTool
 		/// The reason for duplication is that the clang field will be defaulted to WARNING in clang, but not for msvc which is disabled.
 		/// In order to ensure no change of behavior for MSVC toolchain within the clang compilation context, we will keep this in, and deprecate with assignments to the clang setting.
 		/// </remarks>
-		[WarningsVCClang(["-Wno-deprecated-declarations"])]
+		[WarningsVCClang(["-Wno-deprecated-declarations"], ["-Wdeprecated-declarations", "-Wno-error=deprecated-declarations"], ["-Wdeprecated-declarations"])]
 		[WarningLevelDefault(WarningLevel.Off)]
 		public WarningLevel MSVCDeprecationWarningLevel
 		{
@@ -1065,7 +1068,7 @@ namespace UnrealBuildTool
 		/// </summary>
 		/// <remarks>Disable needed for wxWidgets.</remarks>
 		[XmlConfigFile(Category = "BuildConfiguration", Name = nameof(DeprecationWarningLevel))]
-		[WarningsClangToolChain(["-Wno-deprecated-declarations"], null, ["-Werror=deprecated-declarations"])]
+		[WarningsClangToolChain(["-Wno-deprecated-declarations"], ["-Wdeprecated-declarations", "-Wno-error=deprecated-declarations"], ["-Wdeprecated-declarations"])]
 		[WarningsVCToolChain(["/wd4996"], null, ["/we4996"])]
 		[WarningLevelDefault(WarningLevel.Warning)]
 		public WarningLevel DeprecationWarningLevel
@@ -1080,7 +1083,7 @@ namespace UnrealBuildTool
 		/// Clang -
 		///		https://clang.llvm.org/docs/DiagnosticsReference.html#wdate-time
 		/// </summary>
-		[WarningsClangToolChain(null, ["-Wdate-time", "-Wno-error=date-time"], ["-Wdate-time"], [nameof(StandardFilters.DeterministcFlagSetFilter)])]
+		[WarningsClangToolChain(["-Wno-date-time"], ["-Wdate-time", "-Wno-error=date-time"], ["-Wdate-time"], [nameof(StandardFilters.DeterministcFlagSetFilter)])]
 		[WarningsMSCV(["/wd5048"], null, ["/we5048"], [nameof(StandardFilters.DeterministcFlagSetFilter)])]
 		[WarningLevelDefault(WarningLevel.Off)]
 		public WarningLevel DeterministicWarningLevel
