Compile warning argument ordering

diff --git a/Engine/Source/Programs/Shared/UnrealBuildTool.Tests/CppCompileWarningsTests.cs b/Engine/Source/Programs/Shared/UnrealBuildTool.Tests/CppCompileWarningsTests.cs
--- a/Engine/Source/Programs/Shared/UnrealBuildTool.Tests/CppCompileWarningsTests.cs
+++ b/Engine/Source/Programs/Shared/UnrealBuildTool.Tests/CppCompileWarningsTests.cs
@@ -238,6 +238,42 @@ namespace UnrealBuildTool.Tests
 			}
 		}
 
+		[TestMethod]
+		public void TestGenerateWarningCommandLineAgsOrder()
+		{
+			DirectoryReference programsDirectory = DirectoryReference.Combine(Unreal.EngineSourceDirectory, "Programs");
+			FileReference testStubsDirectory = FileReference.Combine(programsDirectory, "Shared", "UnrealBuildTool.Tests", "UBT");
+
+			// Generate ephemeral TargetRules & uproject files
+			using (SafeTestDirectory testRootDirectory = SafeTestDirectory.CreateTestDirectory(Path.Combine(StaticTargetRulesInitializer.UBTTestFolderRoot, $"{StaticTargetRulesInitializer.UBTTestFolderPrefix}{Guid.NewGuid().ToString()}")))
+			using (SafeTestDirectory testSourceDirectory = SafeTestDirectory.CreateTestDirectory(Path.Combine(testRootDirectory.TemporaryDirectory, "Source")))
+			using (SafeTestFile testEditorTargetTestFile = new SafeTestFile(File.ReadAllText(Path.Combine(testStubsDirectory.FullName, "TargetStubs", "TestStubEditor.Target.ubttest")), "TestStubEditor.Target.cs", testSourceDirectory.TemporaryDirectory))
+			using (SafeTestFile testTargetTestFile = new SafeTestFile(File.ReadAllText(Path.Combine(testStubsDirectory.FullName, "TargetStubs", "TestStub.Target.ubttest")), "TestStub.Target.cs", testSourceDirectory.TemporaryDirectory))
+			using (SafeTestFile testUprojectTestFile = new SafeTestFile(File.ReadAllText(Path.Combine(testStubsDirectory.FullName, "UprojectStubs", "TestStub.uproject.ubttest")), "TestStub.uproject", testRootDirectory.TemporaryDirectory))
+			{
+				CppCompileEnvironment newCompileEnvironment = new CppCompileEnvironment(UnrealTargetPlatform.Win64, CppConfiguration.Development, new UnrealArchitectures(UnrealArch.X64), null!);
+
+				FileReference testUproject = new FileReference(testUprojectTestFile.TemporaryFile);
+				IEnumerable<string> args = null!;
+
+				RulesAssembly generatedTestGameRulesAssembly = RulesCompiler.CreateTargetRulesAssembly(testUproject, "TestStub", false, false, false, null, false, new TestLogger());
+				TargetRules testTargetRulesNonWin64 = generatedTestGameRulesAssembly.CreateTargetRules("TestStub", UnrealTargetPlatform.Android, UnrealTargetConfiguration.Development, null, testUproject, null, new TestLogger());
+				CppCompileWarnings specializedWarnings = new CppCompileWarnings(testTargetRulesNonWin64, null);
+
+				// Verify the ordering
+				{
+					VersionNumber newVersionNumber = new VersionNumber(18, 9, 9);
+					specializedWarnings.DeprecatedCopyWarningLevel = WarningLevel.Off;
+					specializedWarnings.DeprecatedCopyWithUserProvidedCopyWarningLevel = WarningLevel.Error;
+					args = specializedWarnings.GenerateWarningCommandLineArgs(newCompileEnvironment, typeof(ClangToolChain), newVersionNumber);
+
+					Assert.IsTrue(args.Count() == 2);
+					Assert.IsTrue(args.First() == "-Wno-deprecated-copy");
+					Assert.IsTrue(args.Last() == "-Wdeprecated-copy-with-user-provided-copy");
+				}
+			}
+		}
+
 		[TestMethod]
 		public void TestGenerateWarningCommandLineAgs()
 		{
diff --git a/Engine/Source/Programs/UnrealBuildTool/Configuration/CompileWarnings/ApplyWarningsAttribute.cs b/Engine/Source/Programs/UnrealBuildTool/Configuration/CompileWarnings/ApplyWarningsAttribute.cs
--- a/Engine/Source/Programs/UnrealBuildTool/Configuration/CompileWarnings/ApplyWarningsAttribute.cs
+++ b/Engine/Source/Programs/UnrealBuildTool/Configuration/CompileWarnings/ApplyWarningsAttribute.cs
@@ -6,6 +6,25 @@ using EpicGames.Core;
 
 namespace UnrealBuildTool
 {
+	/// <summary>
+	/// 
+	/// </summary>
+	public enum ArgPosition
+	{
+		/// <summary>
+		/// 
+		/// </summary>
+		Beginning,
+		/// <summary>
+		/// 
+		/// </summary>
+		Middle,
+		/// <summary>
+		/// 
+		/// </summary>
+		End
+	}
+
 	/// <summary>
 	/// Attribute used to mark fields which describes a set of <see cref="WarningLevel"/> arguments.
 	/// </summary>
@@ -16,6 +35,7 @@ namespace UnrealBuildTool
 		private readonly string[]? _warningArgs;
 		private readonly string[]? _errorArgs;
 		protected readonly List<IApplyWarningsFilter> ActiveFilters = [];
+		public readonly ArgPosition Position;
 
 		/// <summary>
 		/// Defines the argument metadata around a <see cref="CppCompileWarnings"/> <see cref="WarningLevel"/> property.
@@ -24,12 +44,14 @@ namespace UnrealBuildTool
 		/// <param name="warningArgs">The arguments to use when the property is defined as <see cref="WarningLevel.Warning"/>.</param>
 		/// <param name="errorArgs">The arguments to use when the property is defined as <see cref="WarningLevel.Error"/>.</param>
 		/// <param name="applicableFilters">The applicable filter criteria.</param>
-		internal ApplyWarningsAttribute(string[]? offArgs = null, string[]? warningArgs = null, string[]? errorArgs = null, string[]? applicableFilters = null)
+		/// <param name="argPosition">The phase of when to apply the warning.</param>
+		internal ApplyWarningsAttribute(string[]? offArgs = null, string[]? warningArgs = null, string[]? errorArgs = null, string[]? applicableFilters = null, ArgPosition argPosition = ArgPosition.Middle)
 		{
 			_offArgs = offArgs;
 			_warningArgs = warningArgs;
 			_errorArgs = errorArgs;
 			AppendNewFilters(applicableFilters);
+			Position = argPosition;
 		}
 
 		/// <summary>
@@ -142,7 +164,7 @@ namespace UnrealBuildTool
 #pragma warning restore CA1813
 	{
 		/// <inheritdoc/>
-		internal WarningsVCToolChainAttribute(string[]? offArgs = null, string[]? warningArgs = null, string[]? errorArgs = null, string[]? filterNames = null) : base(offArgs, warningArgs, errorArgs, filterNames)
+		internal WarningsVCToolChainAttribute(string[]? offArgs = null, string[]? warningArgs = null, string[]? errorArgs = null, string[]? filterNames = null, ArgPosition argPosition = ArgPosition.Middle) : base(offArgs, warningArgs, errorArgs, filterNames, argPosition)
 		{
 			AppendNewFilters([nameof(FilterID.VCToolChain)]);
 		}
@@ -153,7 +175,7 @@ namespace UnrealBuildTool
 	{
 #pragma warning restore CA1813
 		/// <inheritdoc/>
-		internal WarningsClangToolChainAttribute(string[]? offArgs = null, string[]? warningArgs = null, string[]? errorArgs = null, string[]? filterNames = null) : base(offArgs, warningArgs, errorArgs, filterNames)
+		internal WarningsClangToolChainAttribute(string[]? offArgs = null, string[]? warningArgs = null, string[]? errorArgs = null, string[]? filterNames = null, ArgPosition argPosition = ArgPosition.Middle) : base(offArgs, warningArgs, errorArgs, filterNames, argPosition)
 		{
 			AppendNewFilters([nameof(FilterID.ClangToolChain)]);
 		}
@@ -169,7 +191,7 @@ namespace UnrealBuildTool
 	internal sealed class WarningsMSCVAttribute : WarningsVCToolChainAttribute
 	{
 		/// <inheritdoc/>
-		internal WarningsMSCVAttribute(string[]? offArgs = null, string[]? warningArgs = null, string[]? errorArgs = null, string[]? filterNames = null) : base(offArgs, warningArgs, errorArgs, filterNames)
+		internal WarningsMSCVAttribute(string[]? offArgs = null, string[]? warningArgs = null, string[]? errorArgs = null, string[]? filterNames = null, ArgPosition argPosition = ArgPosition.Middle) : base(offArgs, warningArgs, errorArgs, filterNames, argPosition)
 		{
 			AppendNewFilters([nameof(StandardFilters.MSVCCompilerFilter)]);
 		}
@@ -181,7 +203,7 @@ namespace UnrealBuildTool
 	internal sealed class WarningsVCClangAttribute : WarningsVCToolChainAttribute
 	{
 		/// <inheritdoc/>
-		internal WarningsVCClangAttribute(string[]? offArgs = null, string[]? warningArgs = null, string[]? errorArgs = null) : base(offArgs, warningArgs, errorArgs)
+		internal WarningsVCClangAttribute(string[]? offArgs = null, string[]? warningArgs = null, string[]? errorArgs = null, ArgPosition argPosition = ArgPosition.Middle) : base(offArgs, warningArgs, errorArgs, argPosition: argPosition)
 		{
 			AppendNewFilters([nameof(StandardFilters.VCClangCompilerFilter)]);
 		}
diff --git a/Engine/Source/Programs/UnrealBuildTool/Configuration/CppCompileWarnings.cs b/Engine/Source/Programs/UnrealBuildTool/Configuration/CppCompileWarnings.cs
--- a/Engine/Source/Programs/UnrealBuildTool/Configuration/CppCompileWarnings.cs
+++ b/Engine/Source/Programs/UnrealBuildTool/Configuration/CppCompileWarnings.cs
@@ -291,7 +291,8 @@ namespace UnrealBuildTool
 		/// Clang -
 		///		https://clang.llvm.org/docs/DiagnosticsReference.html#wdeprecated-copy-with-user-provided-copy
 		/// </summary>
-		[WarningsClangToolChain(["-Wno-deprecated-copy-with-user-provided-copy"], ["-Wdeprecated-copy-with-user-provided-copy", "-Wno-error=deprecated-copy-with-user-provided-copy"], ["-Wdeprecated-copy-with-user-provided-copy"], [nameof(FilterID.Version16Min)])]
+		/// <remarks><see cref="ArgPosition.End"/> due to it's relation to <see cref="DeprecatedCopyWarningLevel"/>.</remarks>
+		[WarningsClangToolChain(["-Wno-deprecated-copy-with-user-provided-copy"], ["-Wdeprecated-copy-with-user-provided-copy", "-Wno-error=deprecated-copy-with-user-provided-copy"], ["-Wdeprecated-copy-with-user-provided-copy"], [nameof(FilterID.Version16Min)], ArgPosition.End)]
 		[WarningLevelDefault(WarningLevel.Off)]
 		public WarningLevel DeprecatedCopyWithUserProvidedCopyWarningLevel
 		{
@@ -482,8 +483,8 @@ namespace UnrealBuildTool
 		/// unused-private-field: This will prevent the issue of warnings for unused private variables. MultichannelTcpSocket.h triggers this, possibly more.
 		/// </remarks>
 		[WarningsClangToolChain(
-		["-Wno-unused-but-set-variable", "-Wno-unused-but-set-parameter", "-Wno-unused-function", "-Wno-unused-lambda-capture", "-Wno-unused-local-typedef", "-Wno-unused-private-field", "-Wno-unused-variable"], 
-		["-Wunused-but-set-variable","-Wno-error=unused-but-set-variable", "-Wunused-but-set-parameter", "-Wno-error=unused-but-set-parameter", "-Wunused-function", "-Wno-error=unused-function", "-Wunused-lambda-capture", "-Wno-error=unused-lambda-capture", "-Wunused-local-typedef", "-Wno-error=unused-local-typedef", "-Wunused-private-field", "-Wno-error=unused-private-field", "-Wunused-variable", "-Wno-error=unused-variable"], 
+		["-Wno-unused-but-set-variable", "-Wno-unused-but-set-parameter", "-Wno-unused-function", "-Wno-unused-lambda-capture", "-Wno-unused-local-typedef", "-Wno-unused-private-field", "-Wno-unused-variable"],
+		["-Wunused-but-set-variable", "-Wno-error=unused-but-set-variable", "-Wunused-but-set-parameter", "-Wno-error=unused-but-set-parameter", "-Wunused-function", "-Wno-error=unused-function", "-Wunused-lambda-capture", "-Wno-error=unused-lambda-capture", "-Wunused-local-typedef", "-Wno-error=unused-local-typedef", "-Wunused-private-field", "-Wno-error=unused-private-field", "-Wunused-variable", "-Wno-error=unused-variable"],
 		["-Wunused-but-set-variable", "-Wunused-but-set-parameter", "-Wunused-function", "-Wunused-lambda-capture", "-Wunused-local-typedef", "-Wunused-private-field", "-Wunused-variable"])]
 		[WarningLevelDefault(WarningLevel.Off)]
 		public WarningLevel UnusedWarningLevel
@@ -547,7 +548,7 @@ namespace UnrealBuildTool
 		/// Shipping builds will cause -Wno-unused-value to warn with "ensure", so disable only in those case.
 		/// -Wunused-result - we will always re-enable this as it's controlled by the same set. This will be moved out to a separate property. It is on by default, so it only must be re-enabled in the disable context.
 		/// </remarks>
-		[UnusedValueClangToolChain(["-Wno-unused-value", "-Wunused-result"], ["-Wunused-value", "-Wno-error=unused-value",  "-Wunused-result"], ["-Wunused-value"])]
+		[UnusedValueClangToolChain(["-Wno-unused-value", "-Wunused-result"], ["-Wunused-value", "-Wno-error=unused-value", "-Wunused-result"], ["-Wunused-value"])]
 		[WarningLevelDefault(WarningLevel.Off)]
 		public WarningLevel UnusedValueWarningLevel
 		{
@@ -562,7 +563,8 @@ namespace UnrealBuildTool
 		///		https://clang.llvm.org/docs/DiagnosticsReference.html#wshorten-64-to-32
 		/// </summary>
 		/// <remarks>-Wimplicit-int-conversion also controls -Wshorten-64-to-32 as of clang 19.</remarks>
-		[WarningsClangToolChain(["-Wno-shorten-64-to-32"], ["-Wshorten-64-to-32", "-Wno-error=shorten-64-to-32"], ["-Wshorten-64-to-32"], [nameof(FilterID.Version19Min)])]
+		/// <remarks><see cref="ArgPosition.End"/> due to it's relation to <see cref="UnsafeTypeCastWarningLevel"/>.</remarks>
+		[WarningsClangToolChain(["-Wno-shorten-64-to-32"], ["-Wshorten-64-to-32", "-Wno-error=shorten-64-to-32"], ["-Wshorten-64-to-32"], [nameof(FilterID.Version19Min)], ArgPosition.End)]
 		[WarningLevelDefault(WarningLevel.Off)]
 		public WarningLevel Shorten64To32WarningLevel
 		{
@@ -600,10 +602,6 @@ namespace UnrealBuildTool
 		}
 		private WarningLevel _nonTrivialMemAccessWarningLevel;
 
-		#endregion
-
-		#region -- Clang VC Compiler Disables --
-
 		/// <summary>
 		/// Indicates what warning/error level to treat warnings associated with Microsoft-specific language extensions and compatibility issues.
 		/// Clang -
@@ -1305,6 +1303,8 @@ namespace UnrealBuildTool
 		#endregion -- Public API --
 
 		private static IEnumerable<MemberInfo>? s_membersWithApplyWarningsAttribute = null;
+		private readonly static List<string> s_warningArgumentsBuffer = [];
+		private readonly static List<ValueTuple<ApplyWarningsAttribute, WarningLevel>> s_warningAttributesBuffer = [];
 
 		/// <summary>
 		/// Generates the compiler warnings command line arguments for the requested <see cref="UEToolChain"/> type, and version.
@@ -1324,11 +1324,12 @@ namespace UnrealBuildTool
 				return [];
 			}
 
-			CompilerWarningsToolChainContext toolChainContext = new(compileEnvironment, BuildContextProvider, toolChainType, toolChainVersion, analyzer); //pull out into request?
-
-			List<string> warningArguments = [];
+			CompilerWarningsToolChainContext toolChainContext = new CompilerWarningsToolChainContext(compileEnvironment, BuildContextProvider, toolChainType, toolChainVersion, analyzer); //pull out into request?
+			s_warningArgumentsBuffer.Clear();
 
 			s_membersWithApplyWarningsAttribute ??= GetMembersWithAttributes(typeof(CppCompileWarnings), typeof(ApplyWarningsAttribute), typeof(WarningLevel));
+			s_warningAttributesBuffer.Clear();
+
 			foreach (MemberInfo memberInfo in s_membersWithApplyWarningsAttribute)
 			{
 				ApplyWarningsAttribute? applyWarningsAttribute = null;
@@ -1345,11 +1346,18 @@ namespace UnrealBuildTool
 				if (applyWarningsAttribute != null)
 				{
 					WarningLevel instanceWarningLevel = GetCurrentWarningLevel(memberInfo, this);
-					applyWarningsAttribute.ApplyWarningsToArguments(instanceWarningLevel, toolChainContext, warningArguments);
+					s_warningAttributesBuffer.Add(new ValueTuple<ApplyWarningsAttribute, WarningLevel>(applyWarningsAttribute, instanceWarningLevel));
 				}
 			}
 
-			return warningArguments;
+			s_warningAttributesBuffer.SortBy(x => x.Item1.Position);
+
+			foreach (ValueTuple<ApplyWarningsAttribute, WarningLevel> activeAttribute in s_warningAttributesBuffer)
+			{
+				activeAttribute.Item1.ApplyWarningsToArguments(activeAttribute.Item2, toolChainContext, s_warningArgumentsBuffer);
+			}
+
+			return s_warningArgumentsBuffer;
 		}
 
 		#region -- Internal API --
