Multithreading fix for project gen

diff --git a/Engine/Source/Programs/UnrealBuildTool/Configuration/CppCompileWarnings.cs b/Engine/Source/Programs/UnrealBuildTool/Configuration/CppCompileWarnings.cs
--- a/Engine/Source/Programs/UnrealBuildTool/Configuration/CppCompileWarnings.cs
+++ b/Engine/Source/Programs/UnrealBuildTool/Configuration/CppCompileWarnings.cs
@@ -1303,8 +1303,6 @@ namespace UnrealBuildTool
 		#endregion -- Public API --
 
 		private static IEnumerable<MemberInfo>? s_membersWithApplyWarningsAttribute = null;
-		private readonly static List<string> s_warningArgumentsBuffer = [];
-		private readonly static List<ValueTuple<ApplyWarningsAttribute, WarningLevel>> s_warningAttributesBuffer = [];
 
 		/// <summary>
 		/// Generates the compiler warnings command line arguments for the requested <see cref="UEToolChain"/> type, and version.
@@ -1325,10 +1323,11 @@ namespace UnrealBuildTool
 			}
 
 			CompilerWarningsToolChainContext toolChainContext = new CompilerWarningsToolChainContext(compileEnvironment, BuildContextProvider, toolChainType, toolChainVersion, analyzer); //pull out into request?
-			s_warningArgumentsBuffer.Clear();
+			List<string> warningArguments = [];
+			List<ValueTuple<ApplyWarningsAttribute, WarningLevel>> warningAttributesBuffer = [];
 
 			s_membersWithApplyWarningsAttribute ??= GetMembersWithAttributes(typeof(CppCompileWarnings), typeof(ApplyWarningsAttribute), typeof(WarningLevel));
-			s_warningAttributesBuffer.Clear();
+
 
 			foreach (MemberInfo memberInfo in s_membersWithApplyWarningsAttribute)
 			{
@@ -1346,18 +1345,18 @@ namespace UnrealBuildTool
 				if (applyWarningsAttribute != null)
 				{
 					WarningLevel instanceWarningLevel = GetCurrentWarningLevel(memberInfo, this);
-					s_warningAttributesBuffer.Add(new ValueTuple<ApplyWarningsAttribute, WarningLevel>(applyWarningsAttribute, instanceWarningLevel));
+					warningAttributesBuffer.Add(new ValueTuple<ApplyWarningsAttribute, WarningLevel>(applyWarningsAttribute, instanceWarningLevel));
 				}
 			}
 
-			s_warningAttributesBuffer.SortBy(x => x.Item1.Position);
+			warningAttributesBuffer.SortBy(x => x.Item1.Position);
 
-			foreach (ValueTuple<ApplyWarningsAttribute, WarningLevel> activeAttribute in s_warningAttributesBuffer)
+			foreach (ValueTuple<ApplyWarningsAttribute, WarningLevel> activeAttribute in warningAttributesBuffer)
 			{
-				activeAttribute.Item1.ApplyWarningsToArguments(activeAttribute.Item2, toolChainContext, s_warningArgumentsBuffer);
+				activeAttribute.Item1.ApplyWarningsToArguments(activeAttribute.Item2, toolChainContext, warningArguments);
 			}
 
-			return s_warningArgumentsBuffer;
+			return warningArguments;
 		}
 
 		#region -- Internal API --
