Do not relink all dlls after failed compilation of a dependent dll.

diff --git a/Engine/Source/Programs/UnrealBuildTool/System/ActionGraph.cs b/Engine/Source/Programs/UnrealBuildTool/System/ActionGraph.cs
--- a/Engine/Source/Programs/UnrealBuildTool/System/ActionGraph.cs
+++ b/Engine/Source/Programs/UnrealBuildTool/System/ActionGraph.cs
@@ -772,6 +772,12 @@ namespace UnrealBuildTool
 			{
 				foreach (FileItem PrerequisiteItem in RootAction.PrerequisiteItems)
 				{
+					// Need to check for import libraries here too
+					if (bIgnoreOutdatedImportLibraries && IsImportLibraryDependency(RootAction, PrerequisiteItem))
+					{
+						continue;
+					}
+
 					if (PrerequisiteItem.Exists)
 					{
 						// allow a 1 second slop for network copies
@@ -779,14 +785,10 @@ namespace UnrealBuildTool
 						bool bPrerequisiteItemIsNewerThanLastExecution = TimeDifference.TotalSeconds > 1;
 						if (bPrerequisiteItemIsNewerThanLastExecution)
 						{
-							// Need to check for import libraries here too
-							if (!bIgnoreOutdatedImportLibraries || !IsImportLibraryDependency(RootAction, PrerequisiteItem))
-							{
-								Logger.LogDebug("{StatusDescription}: Prerequisite {PrerequisiteItem} is newer than the last execution of the action: {Message}", RootAction.StatusDescription, PrerequisiteItem.Location,
+							Logger.LogDebug("{StatusDescription}: Prerequisite {PrerequisiteItem} is newer than the last execution of the action: {Message}", RootAction.StatusDescription, PrerequisiteItem.Location,
 									$"{PrerequisiteItem.LastWriteTimeUtc.ToLocalTime().ToString(CultureInfo.CurrentCulture)} vs {LastExecutionTimeUtc.LocalDateTime.ToString(CultureInfo.CurrentCulture)}");
-								bIsOutdated = true;
-								break;
-							}
+							bIsOutdated = true;
+							break;
 						}
 					}
 					else
