CppCompileWarnings Documentation

diff --git a/Engine/Source/Programs/UnrealBuildTool/Configuration/CompileWarnings/ApplyWarningsAttribute.cs b/Engine/Source/Programs/UnrealBuildTool/Configuration/CompileWarnings/ApplyWarningsAttribute.cs
--- a/Engine/Source/Programs/UnrealBuildTool/Configuration/CompileWarnings/ApplyWarningsAttribute.cs
+++ b/Engine/Source/Programs/UnrealBuildTool/Configuration/CompileWarnings/ApplyWarningsAttribute.cs
@@ -7,20 +7,20 @@ using EpicGames.Core;
 namespace UnrealBuildTool
 {
 	/// <summary>
-	/// 
+	/// Denotes the position of a resolved compile warnings argument relative to other warnings.
 	/// </summary>
 	public enum ArgPosition
 	{
 		/// <summary>
-		/// 
+		/// At the beginning of the argument list.
 		/// </summary>
 		Beginning,
 		/// <summary>
-		/// 
+		/// In the middle of the argument list.
 		/// </summary>
 		Middle,
 		/// <summary>
-		/// 
+		/// At the end of the argument list.
 		/// </summary>
 		End
 	}
diff --git a/Engine/Source/Programs/UnrealBuildTool/Configuration/CppCompileWarnings.cs b/Engine/Source/Programs/UnrealBuildTool/Configuration/CppCompileWarnings.cs
--- a/Engine/Source/Programs/UnrealBuildTool/Configuration/CppCompileWarnings.cs
+++ b/Engine/Source/Programs/UnrealBuildTool/Configuration/CppCompileWarnings.cs
@@ -65,6 +65,7 @@ namespace UnrealBuildTool
 		 *		For both contexts - specify it's initialization context to Any
 		 * 4. Access the new property through the (TargetRules|ModuleRules|CppCompileEnvironment).CppCompileWarnings as needed
 		 * 5. (Optional) For compiler line support, utilize the ApplyWarningsAttribute system to participate in invocations to GenerateWarningCommandLineAgs.
+		 * 6. (Optional) For the ApplyWarningsAttribute, pass in the ArgPosition if the argument must be placed after another set.
 		 * 
 		 * Note: [REDACTED_PROFILE_URL]/clang/tree/master/test/SemaCXX contains practical examples for various clang warnings.
 		 */
