Add handtracking permissions for Quest

#android

diff --git a/Engine/Source/Programs/UnrealBuildTool/Platform/Android/UEDeployAndroid.cs b/Engine/Source/Programs/UnrealBuildTool/Platform/Android/UEDeployAndroid.cs
--- a/Engine/Source/Programs/UnrealBuildTool/Platform/Android/UEDeployAndroid.cs
+++ b/Engine/Source/Programs/UnrealBuildTool/Platform/Android/UEDeployAndroid.cs
@@ -3115,6 +3115,10 @@ namespace UnrealBuildTool
 				{
 					// needed for virtual keyboard to work on MetaQuest
 					Text.AppendLine("\t<uses-feature android:name=\"oculus.software.overlay_keyboard\" android:required=\"false\"/>");
+
+					// need for hand tracking on MetaQuest
+					Text.AppendLine("\t<uses-feature android:name=\"oculus.software.handtracking\" android:required=\"false\"/>");
+					Text.AppendLine("\t<uses-permission android:name=\"com.oculus.permission.HAND_TRACKING\"/>");
 				}
 
 				Text.AppendLine("\t<uses-feature android:glEsVersion=\"" + AndroidToolChain.GetGLESVersion(bBuildForES31) + "\" android:required=\"true\" />");
