UnrealBuildTool: Don't lowercase horde pool

#rnx

diff --git a/Engine/Source/Programs/UnrealBuildTool/Executors/UnrealBuildAccelerator/UBAAgentCoordinatorHorde.cs b/Engine/Source/Programs/UnrealBuildTool/Executors/UnrealBuildAccelerator/UBAAgentCoordinatorHorde.cs
--- a/Engine/Source/Programs/UnrealBuildTool/Executors/UnrealBuildAccelerator/UBAAgentCoordinatorHorde.cs
+++ b/Engine/Source/Programs/UnrealBuildTool/Executors/UnrealBuildAccelerator/UBAAgentCoordinatorHorde.cs
@@ -768,13 +768,8 @@ namespace UnrealBuildTool
 				_hordeConfig.UBASentryUrl = null;
 			}
 
-			// Normalize URLs and pools
+			// Normalize URLs
 			_hordeConfig.HordeServer = _hordeConfig.HordeServer?.TrimEnd('/').ToLowerInvariant();
-			_hordeConfig.DefaultHordePool = _hordeConfig.DefaultHordePool?.ToLowerInvariant();
-			_hordeConfig.OverrideHordePool = _hordeConfig.OverrideHordePool?.ToLowerInvariant();
-			_hordeConfig.WindowsHordePool = _hordeConfig.WindowsHordePool?.ToLowerInvariant();
-			_hordeConfig.MacHordePool = _hordeConfig.MacHordePool?.ToLowerInvariant();
-			_hordeConfig.LinuxHordePool = _hordeConfig.LinuxHordePool?.ToLowerInvariant();
 			_hordeConfig.UBASentryUrl = _hordeConfig.UBASentryUrl?.TrimEnd('/').ToLowerInvariant();
 		}
 
