Increase MinimumWindowsX64TargetVersion from 0x601 (Windows 7) to 0xA00 (Windows 10)

Also modify code to default to minimum NTDDI_VERSION if TargetWindowsMinorVersion is unspecified, e.g. 0x0A000000 for Windows 10.

Looking to remove legacy XAudio2 libraries in the future and the Microsoft header is picky about Windows version macros.

diff --git a/Engine/Config/Windows/Windows_SDK.json b/Engine/Config/Windows/Windows_SDK.json
--- a/Engine/Config/Windows/Windows_SDK.json
+++ b/Engine/Config/Windows/Windows_SDK.json
@@ -68,8 +68,8 @@
 
 	"MinimumVisualStudio2022Version": "17.8",
 
-	"//05.1": "X64 -- 0x601 = Win7",
+	"//05.1": "X64 -- 0xA00 = Win10",
 	"//05.2": "Arm64 -- 0xA00 = Win10",
-	"MinimumWindowsX64TargetVersion": "0x601",
+	"MinimumWindowsX64TargetVersion": "0xA00",
 	"MinimumWindowsArm64TargetVersion": "0xA00"
 }
diff --git a/Engine/Source/Programs/UnrealBuildTool/Platform/Windows/UEBuildWindows.cs b/Engine/Source/Programs/UnrealBuildTool/Platform/Windows/UEBuildWindows.cs
--- a/Engine/Source/Programs/UnrealBuildTool/Platform/Windows/UEBuildWindows.cs
+++ b/Engine/Source/Programs/UnrealBuildTool/Platform/Windows/UEBuildWindows.cs
@@ -441,7 +441,12 @@ namespace UnrealBuildTool
 		/// </summary>
 		[RequiresUniqueBuildEnvironment]
 		[ConfigFile(ConfigHierarchyType.Engine, "/Script/WindowsTargetPlatform.WindowsTargetSettings", "TargetWindowsMinorVersion")]
-		public int? TargetWindowsMinorVersion = null;
+		public int TargetWindowsMinorVersion
+		{
+			get => TargetWindowsMinorVersionPrivate ?? (TargetWindowsVersion << 16);
+			set => TargetWindowsMinorVersionPrivate = value;
+		}
+		private int? TargetWindowsMinorVersionPrivate = null;
 
 		/// <summary>
 		/// Enable PIX debugging (automatically disabled in Shipping and Test configs)
@@ -969,7 +974,7 @@ namespace UnrealBuildTool
 
 		public int TargetWindowsVersion => Inner.TargetWindowsVersion;
 
-		public int? TargetWindowsMinorVersion => Inner.TargetWindowsMinorVersion;
+		public int TargetWindowsMinorVersion => Inner.TargetWindowsMinorVersion;
 
 		public bool bPixProfilingEnabled => Inner.bPixProfilingEnabled;
 
@@ -1939,11 +1944,7 @@ namespace UnrealBuildTool
 
 			CompileEnvironment.Definitions.Add(String.Format("_WIN32_WINNT=0x{0:X4}", Target.WindowsPlatform.TargetWindowsVersion));
 			CompileEnvironment.Definitions.Add(String.Format("WINVER=0x{0:X4}", Target.WindowsPlatform.TargetWindowsVersion));
-
-			if (Target.WindowsPlatform.TargetWindowsMinorVersion != null)
-			{
-				CompileEnvironment.Definitions.Add(String.Format("NTDDI_VERSION=0x{0:X8}", Target.WindowsPlatform.TargetWindowsMinorVersion));
-			}
+			CompileEnvironment.Definitions.Add(String.Format("NTDDI_VERSION=0x{0:X8}", Target.WindowsPlatform.TargetWindowsMinorVersion));
 
 			CompileEnvironment.Definitions.Add("PLATFORM_WINDOWS=1");
 			CompileEnvironment.Definitions.Add("PLATFORM_MICROSOFT=1");
diff --git a/Engine/Source/Runtime/AutoRTFM/Private/Context.cpp b/Engine/Source/Runtime/AutoRTFM/Private/Context.cpp
--- a/Engine/Source/Runtime/AutoRTFM/Private/Context.cpp
+++ b/Engine/Source/Runtime/AutoRTFM/Private/Context.cpp
@@ -17,7 +17,7 @@
 #include "Utils.h"
 
 #if AUTORTFM_PLATFORM_WINDOWS
-extern "C" __declspec(dllimport) void __stdcall GetCurrentThreadStackLimits(void**, void**);
+#include "WindowsHeader.h"
 #endif
 
 namespace
@@ -389,7 +389,7 @@ AutoRTFM::FStackRange FContext::GetThreadStackRange()
 		FStackRange Stack;
 
 #if AUTORTFM_PLATFORM_WINDOWS
-		GetCurrentThreadStackLimits(&Stack.Low, &Stack.High);
+		GetCurrentThreadStackLimits(reinterpret_cast<PULONG_PTR>(&Stack.Low), reinterpret_cast<PULONG_PTR>(&Stack.High));
 #elif defined(__APPLE__)         
 		Stack.High = pthread_get_stackaddr_np(pthread_self());
 		size_t StackSize = pthread_get_stacksize_np(pthread_self());
diff --git a/Engine/Source/Runtime/Core/Private/Sanitizer/RaceDetectorWindows.cpp b/Engine/Source/Runtime/Core/Private/Sanitizer/RaceDetectorWindows.cpp
--- a/Engine/Source/Runtime/Core/Private/Sanitizer/RaceDetectorWindows.cpp
+++ b/Engine/Source/Runtime/Core/Private/Sanitizer/RaceDetectorWindows.cpp
@@ -20,8 +20,6 @@ THIRD_PARTY_INCLUDES_START
 #include <process.h>  // for _beginthreadex
 THIRD_PARTY_INCLUDES_END
 
-extern "C" __declspec(dllimport) void __stdcall GetCurrentThreadStackLimits(void**, void**);
-
 using namespace UE::Sanitizer;
 using namespace UE::Sanitizer::RaceDetector;
 
@@ -123,9 +121,9 @@ INSTRUMENTATION_FUNCTION_ATTRIBUTES void InstrumentEndThread()
 	}
 	
 	// Clean the stack to avoid false positives when it gets reused for another thread.
-	void* LowLimit; void* HighLimit;
+	ULONG_PTR LowLimit, HighLimit;
 	GetCurrentThreadStackLimits(&LowLimit, &HighLimit);
-	UE::Sanitizer::RaceDetector::FreeMemoryRange(LowLimit, (UPTRINT)HighLimit - (UPTRINT)LowLimit);
+	UE::Sanitizer::RaceDetector::FreeMemoryRange(reinterpret_cast<void*>(LowLimit), HighLimit - LowLimit);
 
 	ReleaseCurrentContext();
 
@@ -884,7 +882,7 @@ namespace UE::Sanitizer::RaceDetector::Platform {
 	
 	INSTRUMENTATION_FUNCTION_ATTRIBUTES void GetCurrentThreadStackLimits(void** LowLimit, void** HighLimit)
 	{
-		::GetCurrentThreadStackLimits(LowLimit, HighLimit);
+		::GetCurrentThreadStackLimits(reinterpret_cast<PULONG_PTR>(LowLimit), reinterpret_cast<PULONG_PTR>(HighLimit));
 	}
 
 	INSTRUMENTATION_FUNCTION_ATTRIBUTES bool IsThreadAlive(uint32 ThreadId)
