UnrealBuildTool: Ensure lld-link error is matched

#rnx

diff --git a/Engine/Source/Programs/Shared/UnrealBuildTool.Tests/LogEventMatcherTests.cs b/Engine/Source/Programs/Shared/UnrealBuildTool.Tests/LogEventMatcherTests.cs
--- a/Engine/Source/Programs/Shared/UnrealBuildTool.Tests/LogEventMatcherTests.cs
+++ b/Engine/Source/Programs/Shared/UnrealBuildTool.Tests/LogEventMatcherTests.cs
@@ -714,6 +714,19 @@ namespace UnrealBuildTool.Tests
 			Assert.AreEqual("IMPLEMENT_MODULE_DataInterface", symbolProperty.Properties![LogEventPropertyName.Identifier]?.ToString());
 		}
 
+		[TestMethod]
+		public void LinkEventMatcher5()
+		{
+			string[] lines =
+			[
+				@"  lld-link: error: Malformed block",
+			];
+
+			List<LogEvent> logEvents = Parse(String.Join("\n", lines));
+			Assert.AreEqual(1, logEvents.Count);
+			CheckEventGroup(logEvents, 0, 1, LogLevel.Error, KnownLogEvents.Linker);
+		}
+
 		[TestMethod]
 		public void MacLinkErrorMatcher()
 		{
diff --git a/Engine/Source/Programs/UnrealBuildTool/Matchers/LinkEventMatcher.cs b/Engine/Source/Programs/UnrealBuildTool/Matchers/LinkEventMatcher.cs
--- a/Engine/Source/Programs/UnrealBuildTool/Matchers/LinkEventMatcher.cs
+++ b/Engine/Source/Programs/UnrealBuildTool/Matchers/LinkEventMatcher.cs
@@ -13,7 +13,7 @@ namespace UnrealBuildTool.Matchers
 	class LinkEventMatcher : ILogEventMatcher
 	{
 		static readonly Regex s_undefinedReferencePattern = new Regex(
-			@": undefined reference to |undefined symbol|^\s*(ld|ld.lld|(.*clang)):|^[^:]*[^a-zA-Z][a-z]?ld: |^\s*>>>");
+			@": undefined reference to |undefined symbol|^\s*(ld|ld.lld|lld-link|(.*clang)):|^[^:]*[^a-zA-Z][a-z]?ld: |^\s*>>>");
 
 		static readonly Regex s_undefinedSymbolsForArchPattern = new Regex(
 			@"^(\s*)Undefined symbols for architecture");
@@ -21,6 +21,9 @@ namespace UnrealBuildTool.Matchers
 		static readonly Regex s_ldClangPattern = new Regex(
 			@"^\s*(ld|ld.lld|lld-link|clang):");
 
+		static readonly Regex s_ldClangErrorPattern = new Regex(
+			@"^\s*(ld|ld.lld|lld-link):\s+error:\s+(?<error>.+)$");
+
 		static readonly Regex s_ldMultiplyDefinedPattern = new Regex(
 			@"(?<severity>error): L\d+: symbol `(?<symbol>[^`]+)' multiply defined$");
 
@@ -151,6 +154,13 @@ namespace UnrealBuildTool.Matchers
 
 				return builder.ToMatch(LogEventPriority.Highest, LogLevel.Error, KnownLogEvents.Linker_DuplicateSymbol);
 			}
+			if (cursor.TryMatch(s_ldClangErrorPattern, out match))
+			{
+				LogEventBuilder builder = new LogEventBuilder(cursor);
+				builder.Annotate(match.Groups["error"], LogEventMarkup.Message);
+
+				return builder.ToMatch(LogEventPriority.Highest, LogLevel.Error, KnownLogEvents.Linker);
+			}
 
 			return null;
 		}
