Start loading libUnreal dependencies early for faster startup time

#android

diff --git a/Engine/Build/Android/Java/src/com/epicgames/unreal/GameActivity.java.template b/Engine/Build/Android/Java/src/com/epicgames/unreal/GameActivity.java.template
--- a/Engine/Build/Android/Java/src/com/epicgames/unreal/GameActivity.java.template
+++ b/Engine/Build/Android/Java/src/com/epicgames/unreal/GameActivity.java.template
@@ -260,7 +260,8 @@ public class GameActivity extends $${gameActivitySuperClass}$$ implements Surfac
 															ComponentCallbacks2
 {
 	public static final boolean IS_SHIPPING_CONFIG = $${gameActivityIsShippingConfig}$$;
-	
+	public static final boolean HAS_SPLASH_ACTIVITY = $${gameActivityHasSplashActivity}$$;
+
 	public int VersionCode = 0;
 	public String VersionName = "Unknown";
 	public static int targetSdkVersion = 0;
@@ -3740,7 +3741,11 @@ Trace.endSection();
 		MainThreadPriority = android.os.Process.getThreadPriority(MainThreadId);
 
 		// block getGraphicsInformation on ProcessSystemInfoThread until GLES loads
-		ProcessSystemInfoLock.lock();
+		// no need if using SplashActivity as it uses the GameApplication lock instead
+		if (!HAS_SPLASH_ACTIVITY)
+		{
+			ProcessSystemInfoLock.lock();
+		}
 
 //$${gameActivityOnCreateBeforeProcessSystemInfoAdditions}$$
 		ProcessSystemInfoThread = new Thread(new Runnable() {
@@ -3750,7 +3755,14 @@ Trace.endSection();
 				ParseCommandline(ProjectName, UseExternalFilesDir);
 Trace.beginSection("processSystemInfo");
 				// collect system info
-				processSystemInfo(ProcessSystemInfoLock);
+				if (HAS_SPLASH_ACTIVITY)
+				{
+					processSystemInfo(com.epicgames.unreal.GameApplication.ProcessSystemInfoLock);
+				}
+				else
+				{
+					processSystemInfo(ProcessSystemInfoLock);
+				}
 Trace.endSection();
 
 				// run config rules
@@ -3768,18 +3780,22 @@ Trace.endSection();
 
 		// load GLES/EGL now so available for getGraphicsInformation in ProcessSystemInfo thread without RenderDoc hang
 		// cannot be in another thread as it will just block (fairly quick to load so not a big deal to do it here)
-Trace.beginSection("loadGLES");
-		try
-		{
-			System.loadLibrary("GLESv3");
-			System.loadLibrary("EGL");
-		}
-		catch (Exception e)
+		// SplashActivity would have already tried to load early so can skip if we have it
+		if (!HAS_SPLASH_ACTIVITY)
 		{
-		}
-		// unlock so getGraphicsInformation is allowed to run
-		ProcessSystemInfoLock.unlock();
+Trace.beginSection("loadGLES");
+			try
+			{
+				System.loadLibrary("GLESv3");
+				System.loadLibrary("EGL");
+			}
+			catch (Exception e)
+			{
+			}
+			// unlock so getGraphicsInformation is allowed to run
+			ProcessSystemInfoLock.unlock();
 Trace.endSection();
+		}
 
 //$${gameActivityOnCreateBeforeLoadLibrariesAdditions}$$
 Trace.beginSection("loadLibraries");
@@ -4453,6 +4469,12 @@ Trace.endSection();
 			// configuration failed (show error)
 			return;
 		}
+		
+		if (HAS_SPLASH_ACTIVITY)
+		{
+			com.epicgames.unreal.GameApplication.waitLoadLibraries();
+		}
+
 		Log.debug("$$$ configrules: apply config rules");
 		applyConfigRules(ConfigRulesVars, ProjectName, appPackageName);
 		Log.debug("$$$ configrules: finish config rules");
diff --git a/Engine/Build/Android/Java/src/com/epicgames/unreal/GameApplication.java.template b/Engine/Build/Android/Java/src/com/epicgames/unreal/GameApplication.java.template
--- a/Engine/Build/Android/Java/src/com/epicgames/unreal/GameApplication.java.template
+++ b/Engine/Build/Android/Java/src/com/epicgames/unreal/GameApplication.java.template
@@ -8,6 +8,11 @@ import androidx.lifecycle.LifecycleObserver;
 import androidx.lifecycle.OnLifecycleEvent;
 import androidx.lifecycle.ProcessLifecycleOwner;
 
+import java.lang.Thread;
+import java.lang.Runnable;
+import java.util.concurrent.locks.Lock;
+import java.util.concurrent.locks.ReentrantLock;
+
 import com.epicgames.unreal.network.NetworkChangedManager;
 //$${gameApplicationImportAdditions}$$
 
@@ -21,6 +26,13 @@ public class GameApplication extends Application implements LifecycleObserver
 
 	private static Context context;
 
+	protected static int MainThreadId;
+	protected static int MainThreadPriority;
+
+	private static Thread LoadLibrariesThread = null;
+	private static boolean bLoadedLibraries = false;
+	public static Lock ProcessSystemInfoLock = new ReentrantLock();
+
 //$${gameApplicationClassAdditions}$$
 
 	@Override
@@ -104,4 +116,68 @@ public class GameApplication extends Application implements LifecycleObserver
 	public static boolean isAppInBackground() {
 		return !isForeground;
 	}
+
+	protected static void LoadLibraries(String CommandLine)
+	{
+		// We need to decide on enabling memory tracing before loading native code.
+		// ignore for now
+		try
+		{
+//$${gameApplicationEarlyLoadAdditions}$$
+		}
+		catch (Exception e)
+		{
+			Log.debug(e.toString());
+		}
+	}
+
+	public static void startLoadLibraries()
+	{
+		if (bLoadedLibraries)
+		{
+			return;
+		}
+		if (LoadLibrariesThread != null)
+		{
+			return;
+		}
+
+		MainThreadId = android.os.Process.myTid();
+		MainThreadPriority = android.os.Process.getThreadPriority(MainThreadId);
+
+		LoadLibrariesThread = new Thread(new Runnable() {
+			public void run() {
+				ProcessSystemInfoLock.lock();
+				android.os.Process.setThreadPriority(android.os.Process.myTid(), MainThreadPriority);
+Trace.beginSection("EarlyLoadLibraries");
+				LoadLibraries("");
+Trace.endSection();
+			}
+		});
+		LoadLibrariesThread.setName("LoadLibraries");
+		LoadLibrariesThread.start();
+	}
+
+	public static boolean waitLoadLibraries()
+	{
+		if (bLoadedLibraries)
+		{
+			return true;
+		}
+		if (LoadLibrariesThread == null)
+		{
+			return true;
+		}
+
+		try
+		{
+			LoadLibrariesThread.join();
+			bLoadedLibraries = true;
+			return true;
+		}
+		catch (Exception e)
+		{
+			return false;
+		}
+	}
 }
diff --git a/Engine/Build/Android/Java/src/com/epicgames/unreal/SplashActivity.java b/Engine/Build/Android/Java/src/com/epicgames/unreal/SplashActivity.java
--- a/Engine/Build/Android/Java/src/com/epicgames/unreal/SplashActivity.java
+++ b/Engine/Build/Android/Java/src/com/epicgames/unreal/SplashActivity.java
@@ -43,8 +43,6 @@ public class SplashActivity extends Activity
 	@Override
 	protected void onCreate(Bundle savedInstanceState)
 	{
-		super.onCreate(savedInstanceState);
-
 		// check if already running
 		if (!isTaskRoot())
 		{
@@ -52,12 +50,19 @@ public class SplashActivity extends Activity
 			// see if foregrounded by launcher main action (like from home screen icon)
 			if (intent.hasCategory(Intent.CATEGORY_LAUNCHER) && intent.getAction() != null && intent.getAction().equals(Intent.ACTION_MAIN))
 			{
+				super.onCreate(savedInstanceState);
+
 				// leave previous top activity at top instead of forcing GameActivity
 				finish();
 				return;
 			}
 		}
 
+		// trigger library loading
+		com.epicgames.unreal.GameApplication.startLoadLibraries();
+
+		super.onCreate(savedInstanceState);
+
 		packageName = getPackageName();
 		pm = getPackageManager();
 
diff --git a/Engine/Source/Programs/UnrealBuildTool/Platform/Android/AndroidToolChain.cs b/Engine/Source/Programs/UnrealBuildTool/Platform/Android/AndroidToolChain.cs
--- a/Engine/Source/Programs/UnrealBuildTool/Platform/Android/AndroidToolChain.cs
+++ b/Engine/Source/Programs/UnrealBuildTool/Platform/Android/AndroidToolChain.cs
@@ -243,6 +243,7 @@ namespace UnrealBuildTool
 			}
 
 			ReadElfPath = Path.Combine(NDKPath, @"toolchains/llvm", ArchitecturePath, @"bin/llvm-readelf" + ExeExtension);
+			ObjDumpPath = Path.Combine(NDKPath, @"toolchains/llvm", ArchitecturePath, @"bin/llvm-objdump" + ExeExtension);
 		}
 
 		IEnumerable<string> AdditionalPaths(CppRootPaths Roots)
@@ -562,6 +563,8 @@ namespace UnrealBuildTool
 			{
 				Arguments.Add("-fno-inline-functions");
 			}
+
+			Arguments.Add("-fdebug-types-section");
 		}
 
 		/// <inheritdoc/>
@@ -1032,6 +1035,7 @@ namespace UnrealBuildTool
 		protected static string? ArPathArm64;
 		protected static string? ArPathx64;
 		protected static string? ReadElfPath;
+		protected static string? ObjDumpPath;
 
 		public static string GetStripExecutablePath(UnrealArch UnrealArch)
 		{
@@ -1043,6 +1047,16 @@ namespace UnrealBuildTool
 			return StripPath.Replace("-ar", "-strip");
 		}
 
+		public static string GetReadElfExecutablePath(UnrealArch UnrealArch)
+		{
+			return ReadElfPath!;
+		}
+
+		public static string GetObjDumpExecutablePath(UnrealArch UnrealArch)
+		{
+			return ObjDumpPath!;
+		}
+
 		private HashSet<UnrealArch> HasHandledLaunchModule = new();
 		private HashSet<UnrealArch> HasHandledCoreModule = new();
 
diff --git a/Engine/Source/Programs/UnrealBuildTool/Platform/Android/UEDeployAndroid.cs b/Engine/Source/Programs/UnrealBuildTool/Platform/Android/UEDeployAndroid.cs
--- a/Engine/Source/Programs/UnrealBuildTool/Platform/Android/UEDeployAndroid.cs
+++ b/Engine/Source/Programs/UnrealBuildTool/Platform/Android/UEDeployAndroid.cs
@@ -1296,6 +1296,49 @@ namespace UnrealBuildTool
 			Utils.RunLocalProcessAndLogOutput(StartInfo, Logger);
 		}
 
+		private List<string> GetSODependencies(string SOFilename, UnrealArch UnrealArch)
+		{
+			List<string> DependencyList = new List<string>();
+			string ReadElfExecutable = AndroidToolChain.GetReadElfExecutablePath(UnrealArch);
+			string? Output = Utils.RunLocalProcessAndReturnStdOut(ReadElfExecutable, "--dynamic \"" + SOFilename + "\"");
+			if (Output != null)
+			{
+				// ex:   0x0000000000000001 (NEEDED)         Shared library: [libGLESv3.so]
+				int ScanStart = 0;
+				int SharedIndex = Output.IndexOf("Shared library: [", ScanStart);
+				while (SharedIndex >= 0)
+				{
+					ScanStart = SharedIndex + 17;
+					int EndIndex = Output.IndexOf(']', ScanStart);
+					if (EndIndex > 0)
+					{
+						string SharedLib = Output.Substring(ScanStart, EndIndex - ScanStart);
+
+						// strip off lib (not needed for System.loadLibrary()
+						if (SharedLib.StartsWith("lib"))
+						{
+							SharedLib = SharedLib.Substring(3);
+						}
+						// strip off .so (not needed for System.loadLibrary()
+						if (SharedLib.EndsWith(".so"))
+						{
+							SharedLib = SharedLib.Substring(0, SharedLib.Length - 3);
+						}
+
+						DependencyList.Add(SharedLib);
+					}
+					else
+					{
+						// Unpaired [, break due to unexpected result
+						break;
+					}
+					ScanStart = EndIndex + 1;
+					SharedIndex = Output.IndexOf("Shared library: [", ScanStart);
+				}
+			}
+			return DependencyList;
+		}
+
 		private static void CopySO(string FinalName, string SourceName)
 		{
 			// check to see if file is newer than last time we copied
@@ -5070,6 +5113,9 @@ dependencies.implementation('androidx.tracing:tracing:1.2.0')
 				File.WriteAllText(Path.Combine(IntermediateAndroidPath, Arch + "_AndroidManifest.xml"), Manifest);
 				File.WriteAllText(Path.Combine(UnrealBuildPath, "AndroidManifest.xml"), Manifest);
 
+				// Check if included SplashActivity in AndroidManifest
+				bool bHasSplashActivity = Manifest.Contains("com.epicgames.unreal.SplashActivity");
+
 				// copy prebuild plugin files
 				UPL.ProcessPluginNode(NDKArch, "prebuildCopies", "");
 
@@ -5080,10 +5126,6 @@ dependencies.implementation('androidx.tracing:tracing:1.2.0')
 				UnrealOBBDataFileName = GetUnrealJavaOBBDataFileName(Path.Combine(UnrealBuildPath, "src", PackageName.Replace('.', Path.DirectorySeparatorChar)));
 				WriteJavaOBBDataFile(UnrealOBBDataFileName, PackageName, RequiredOBBFiles, CookFlavor, bPackageDataInsideApk, Arch);
 
-				// update GameActivity.java and GameApplication.java if out of date
-				UpdateGameActivity(Arch, NDKArch, EngineDirectory, UnrealBuildPath, bDontBundleLibrariesInAPK, bJavaPerfettoMarkers, Configuration);
-				UpdateGameApplication(Arch, NDKArch, EngineDirectory, UnrealBuildPath, bJavaPerfettoMarkers, Configuration);
-
 				// we don't actually need the SO for the bSkipGradleBuild case
 				string? SourceSOName = null;
 				string DestApkDirectory = Path.Combine(ProjectDirectory, "Binaries", "Android");
@@ -5217,6 +5259,10 @@ popd
 					Utils.LinkOrCopyIfDifferent(Path.Combine(JniLibsDir, "libUnreal.so"), SourceSOName);
 				}
 
+				// update GameActivity.java and GameApplication.java if out of date
+				UpdateGameActivity(Arch, NDKArch, EngineDirectory, UnrealBuildPath, bDontBundleLibrariesInAPK, bJavaPerfettoMarkers, Configuration, bHasSplashActivity);
+				UpdateGameApplication(Arch, NDKArch, EngineDirectory, UnrealBuildPath, bJavaPerfettoMarkers, Configuration, SourceSOName);
+
 				bool bSkipLibCpp = false;
 				Ini.GetBool("/Script/AndroidRuntimeSettings.AndroidRuntimeSettings", "bSkipLibCpp", out bSkipLibCpp);
 				if (!bSkipLibCpp)
@@ -6095,7 +6141,8 @@ popd
 			return "empty";
 		}
 
-		private void UpdateGameActivity(UnrealArch UnrealArch, string NDKArch, string EngineDir, string UnrealBuildPath, bool bDontBundleLibrariesInAPK, bool bJavaPerfettoMarkers, UnrealTargetConfiguration Configuration)
+		private void UpdateGameActivity(UnrealArch UnrealArch, string NDKArch, string EngineDir, string UnrealBuildPath, bool bDontBundleLibrariesInAPK,
+				bool bJavaPerfettoMarkers, UnrealTargetConfiguration Configuration, bool bHasSplashActivity)
 		{
 			string SourceFilename = Path.Combine(EngineDir, "Build", "Android", "Java", "src", "com", "epicgames", "unreal", "GameActivity.java.template");
 			string DestFilename = Path.Combine(UnrealBuildPath, "src", "com", "epicgames", "unreal", "GameActivity.java");
@@ -6230,6 +6277,7 @@ import java.util.Collection;
 				{ "//$${soLoadLibrary}$$", UPL.ProcessPluginNode(NDKArch, "soLoadLibrary", LoadLibraryDefaults)},
 				{ "$${gameActivitySuperClass}$$", SuperClassDefault},
 				{ "$${gameActivityIsShippingConfig}$$", Configuration == UnrealTargetConfiguration.Shipping ? "true" : "false"},
+				{ "$${gameActivityHasSplashActivity}$$", bHasSplashActivity ? "true" : "false"},
 			};
 
 			string[] TemplateSrc = File.ReadAllLines(SourceFilename);
@@ -6312,17 +6360,36 @@ import java.util.Collection;
 			}
 		}
 
-		private void UpdateGameApplication(UnrealArch UnrealArch, string NDKArch, string EngineDir, string UnrealBuildPath, bool bJavaPerfettoMarkers, UnrealTargetConfiguration Configuration)
+		private void UpdateGameApplication(UnrealArch UnrealArch, string NDKArch, string EngineDir, string UnrealBuildPath,
+				bool bJavaPerfettoMarkers, UnrealTargetConfiguration Configuration, string SOName)
 		{
 			string SourceFilename = Path.Combine(EngineDir, "Build", "Android", "Java", "src", "com", "epicgames", "unreal", "GameApplication.java.template");
 			string DestFilename = Path.Combine(UnrealBuildPath, "src", "com", "epicgames", "unreal", "GameApplication.java");
 
-			String GameApplicationImportAdditionsDefault = "";
+			string GameApplicationImportAdditionsDefault = "";
 			if (bJavaPerfettoMarkers)
 			{
 				GameApplicationImportAdditionsDefault += "import androidx.tracing.Trace;\n";
 			}
 
+			// add dependencies to early load during SplashActivity (optional)
+			List<string> UnrealDependencies = GetSODependencies(SOName, UnrealArch);
+			string GameApplicationEarlyLoad = "";
+			if (UnrealDependencies.Count > 0)
+			{
+				// Make sure these are first then release the lock
+				UnrealDependencies.Remove("GLESv3");
+				UnrealDependencies.Remove("EGL");
+				GameApplicationEarlyLoad += $"Trace.beginSection(\"GLESv3\");\n\t\t\tSystem.loadLibrary(\"GLESv3\");\nTrace.endSection();\n";
+				GameApplicationEarlyLoad += $"Trace.beginSection(\"EGL\");\n\t\t\tSystem.loadLibrary(\"EGL\");\nTrace.endSection();\n";
+				GameApplicationEarlyLoad += $"\t\t\tProcessSystemInfoLock.unlock();\n";
+
+				foreach (string SharedLibrary in UnrealDependencies)
+				{
+					GameApplicationEarlyLoad += $"Trace.beginSection(\"{SharedLibrary}\");\n\t\t\tSystem.loadLibrary(\"{SharedLibrary}\");\nTrace.endSection();\n";
+				}
+			}
+
 			Dictionary<string, string> Replacements = new Dictionary<string, string>{
 				{ "//$${gameApplicationImportAdditions}$$", UPL!.ProcessPluginNode(NDKArch, "gameApplicationImportAdditions", GameApplicationImportAdditionsDefault)},
 				{ "//$${gameApplicationOnCreateAdditions}$$", UPL.ProcessPluginNode(NDKArch, "gameApplicationOnCreateAdditions", "")},
@@ -6332,6 +6399,7 @@ import java.util.Collection;
 				{ "//$${gameApplicationOnConfigurationChangedAdditions}$$", UPL.ProcessPluginNode(NDKArch, "gameApplicationOnConfigurationChangedAdditions", "")},
 				{ "//$${gameApplicationClassAdditions}$$", UPL.ProcessPluginNode(NDKArch, "gameApplicationClassAdditions", "")},
 				{ "$${gameApplicationIsShippingConfig}$$", Configuration == UnrealTargetConfiguration.Shipping ? "true" : "false"},
+				{ "//$${gameApplicationEarlyLoadAdditions}$$", GameApplicationEarlyLoad},
 			};
 
 			string[] TemplateSrc = File.ReadAllLines(SourceFilename);
