Removing deprecated plugins from some uprojects and added more details in the logged message

diff --git a/Engine/Source/Programs/LiveLinkHub/Source/LiveLinkHub.Target.cs b/Engine/Source/Programs/LiveLinkHub/Source/LiveLinkHub.Target.cs
--- a/Engine/Source/Programs/LiveLinkHub/Source/LiveLinkHub.Target.cs
+++ b/Engine/Source/Programs/LiveLinkHub/Source/LiveLinkHub.Target.cs
@@ -47,7 +47,6 @@ public class LiveLinkHubTarget : TargetRules
 			"ProceduralMeshComponent",
 			"PropertyAccessEditor",
 			"PythonScriptPlugin",
-			"StructUtils",
 			"UdpMessaging",
 			"CameraCalibrationCore",
 			"AppleARKitFaceSupport",
diff --git a/Engine/Source/Programs/UnrealBuildTool/Configuration/UEBuildPlugin.cs b/Engine/Source/Programs/UnrealBuildTool/Configuration/UEBuildPlugin.cs
--- a/Engine/Source/Programs/UnrealBuildTool/Configuration/UEBuildPlugin.cs
+++ b/Engine/Source/Programs/UnrealBuildTool/Configuration/UEBuildPlugin.cs
@@ -97,8 +97,9 @@ namespace UnrealBuildTool
 		/// Public entry point to validate a module
 		/// </summary>
 		/// <param name="logger"></param>
+		/// <param name="rules"></param>
 		/// <returns>true if there are any fatal errors</returns>
-		public bool ValidatePlugin(ILogger logger)
+		public bool ValidatePlugin(ILogger logger, ReadOnlyTargetRules rules)
 		{
 			bool anyErrors = false;
 			if (Dependencies != null)
@@ -125,10 +126,8 @@ namespace UnrealBuildTool
 
 			if (!String.IsNullOrEmpty(Descriptor.DeprecatedEngineVersion))
 			{
-				// Eventually we need to promote this to an actual warning but there are a number of uses of deprecated plugins that would cause a lot
-				// of noise and we want to gently nudge people to migrate plugins.
-				logger.LogInformation("Project depends on plugin '{PluginName}' which was deprecated in {EngineVersion} and will soon be removed. Please update your dependencies.",
-									  Name, Descriptor.DeprecatedEngineVersion);
+				logger.LogWarning("Project '{ProjectName}' depends on plugin '{PluginName}' which was deprecated in {EngineVersion} and will soon be removed. Please update your dependencies.",
+					rules.Name, Name, Descriptor.DeprecatedEngineVersion);
 			}
 			
 			// Check that any plugins with the NoCode specifier do not contain modules.
diff --git a/Engine/Source/Programs/UnrealBuildTool/Configuration/UEBuildTarget.cs b/Engine/Source/Programs/UnrealBuildTool/Configuration/UEBuildTarget.cs
--- a/Engine/Source/Programs/UnrealBuildTool/Configuration/UEBuildTarget.cs
+++ b/Engine/Source/Programs/UnrealBuildTool/Configuration/UEBuildTarget.cs
@@ -1450,7 +1450,7 @@ namespace UnrealBuildTool
 		/// <returns>true if there are any fatal errors</returns>
 		bool ValidatePlugin(UEBuildPlugin plugin, ILogger logger)
 		{
-			bool anyErrors = plugin.ValidatePlugin(logger);
+			bool anyErrors = plugin.ValidatePlugin(logger, Rules);
 
 			// Check that each plugin declares its dependencies explicitly
 			foreach (UEBuildModule module in plugin.Modules)
