UnrealBuildTool: Don't set headers as ClCompile for generated visual studio projects by default, as this causes intellisense issues
Removes the ability to compile a single header in the IDE unless enabled

#rnx

diff --git a/Engine/Source/Programs/UnrealBuildTool/ProjectFiles/VisualStudio/VCProject.cs b/Engine/Source/Programs/UnrealBuildTool/ProjectFiles/VisualStudio/VCProject.cs
--- a/Engine/Source/Programs/UnrealBuildTool/ProjectFiles/VisualStudio/VCProject.cs
+++ b/Engine/Source/Programs/UnrealBuildTool/ProjectFiles/VisualStudio/VCProject.cs
@@ -2104,7 +2104,7 @@ namespace UnrealBuildTool
 				case VCFileType.CCode:
 					return "ClCompile";
 				case VCFileType.Header:
-					return "ClCompile";
+					return Settings.bHeadersAsClCompile ? "ClCompile" : "ClInclude";
 				case VCFileType.Inline:
 					return "ClInclude";
 				case VCFileType.Resource:
@@ -2121,7 +2121,7 @@ namespace UnrealBuildTool
 
 		private bool CanVCFileTypeBeCompiled(VCFileType FileType)
 		{
-			return FileType == VCFileType.CCode || FileType == VCFileType.Header;
+			return FileType == VCFileType.CCode || (FileType == VCFileType.Header && Settings.bHeadersAsClCompile);
 		}
 
 		/// <summary>
diff --git a/Engine/Source/Programs/UnrealBuildTool/ProjectFiles/VisualStudio/VCProjectFileGenerator.cs b/Engine/Source/Programs/UnrealBuildTool/ProjectFiles/VisualStudio/VCProjectFileGenerator.cs
--- a/Engine/Source/Programs/UnrealBuildTool/ProjectFiles/VisualStudio/VCProjectFileGenerator.cs
+++ b/Engine/Source/Programs/UnrealBuildTool/ProjectFiles/VisualStudio/VCProjectFileGenerator.cs
@@ -80,6 +80,12 @@ namespace UnrealBuildTool
 		[XmlConfigFile(Category = "BuildConfiguration")]
 		public bool bUsePerFileIntellisense = true;
 
+		/// <summary>
+		/// Whether to treat headers as ClCompile. May cause intellisense issues but allows the VS Compile command to build headers.
+		/// </summary>
+		[XmlConfigFile(Category = "VCProjectFileGenerator")]
+		public bool bHeadersAsClCompile = false;
+
 		/// <summary>
 		/// Whether to include a dependency on ShaderCompileWorker when generating project files for the editor.
 		/// </summary>
