Bug fix: Arm64 SxS manifest is now generated as an intermediate file after runtime dependencies have been collected. This allows it to get updated with all potential runtime dependencies in a multi-target build.

- removed the option for overriding for SxS manifest template because multiple targets can stage to the same folder.
- clean out stale references in the assembly maniest when staging.
- add separate manifest for arm64ec.

#rnx

diff --git a/Engine/Build/Windows/Resources/arm64/DefaultSxS-Win64-arm64.manifest b/Engine/Build/Windows/Resources/arm64/DefaultSxS-Win64-arm64.manifest
deleted file mode 100644
diff --git a/Engine/Build/Windows/Resources/arm64/arm64.manifest b/Engine/Build/Windows/Resources/arm64/arm64.manifest
new file mode 100644
--- /dev/null
+++ b/Engine/Build/Windows/Resources/arm64/arm64.manifest
@@ -0,0 +1,4 @@
+<?xml version="1.0" encoding="UTF-8"?>
+<assembly xmlns="urn:schemas-microsoft-com:asm.v1" manifestVersion="1.0">
+    <assemblyIdentity type="Win32" name="arm64" version="1.0.0.0" processorArchitecture="arm64" />
+</assembly>
diff --git a/Engine/Build/Windows/Resources/arm64ec/arm64ec.manifest b/Engine/Build/Windows/Resources/arm64ec/arm64ec.manifest
new file mode 100644
--- /dev/null
+++ b/Engine/Build/Windows/Resources/arm64ec/arm64ec.manifest
@@ -0,0 +1,4 @@
+<?xml version="1.0" encoding="UTF-8"?>
+<assembly xmlns="urn:schemas-microsoft-com:asm.v1" manifestVersion="1.0">
+    <assemblyIdentity type="Win32" name="arm64ec" version="1.0.0.0" processorArchitecture="arm64" />
+</assembly>
diff --git a/Engine/Source/Programs/AutomationTool/Win/WinPlatform.Automation.cs b/Engine/Source/Programs/AutomationTool/Win/WinPlatform.Automation.cs
--- a/Engine/Source/Programs/AutomationTool/Win/WinPlatform.Automation.cs
+++ b/Engine/Source/Programs/AutomationTool/Win/WinPlatform.Automation.cs
@@ -203,17 +203,9 @@ public class Win64Platform : Platform
 
 		if (!string.IsNullOrWhiteSpace(Params.AppLocalDirectory))
 		{
-			bool bIsStagingX64   = (SC.StageTargets.Any(Target => Target.Receipt.Architectures.SingleArchitecture.bIsX64));
-			bool bIsStagingARM64 = (SC.StageTargets.Any(Target => !Target.Receipt.Architectures.SingleArchitecture.bIsX64));
-
-			if (bIsStagingX64)
-			{
-				StageAppLocalDependencies(Params, SC, "Win64", UnrealArch.X64);
-			}
-
-			if (bIsStagingARM64)
+			foreach (UnrealArch Architecture in SC.StageTargets.Select( T => T.Receipt.Architectures.SingleArchitecture ).Distinct() )
 			{
-				StageAppLocalDependencies(Params, SC, "Win64", UnrealArch.Arm64); // also used for Arm64ec
+				StageAppLocalDependencies(Params, SC, "Win64", Architecture);
 			}
 		}
 	}
@@ -468,7 +460,8 @@ public class Win64Platform : Platform
 		// x64 DLL dependencies are placed directly in StagedBinariesDir
 		bool bUseArchitectureSubdirectory = !Architecture.bIsX64;
 
-		StagedDirectoryReference StagedArchitectureBinariesDir = bUseArchitectureSubdirectory ? StagedDirectoryReference.Combine(StagedBinariesDir, Architecture.WindowsToolChain) : StagedBinariesDir;
+		string AssemblyName = Architecture.ToString().ToLower();
+		StagedDirectoryReference StagedArchitectureBinariesDir = bUseArchitectureSubdirectory ? StagedDirectoryReference.Combine(StagedBinariesDir, AssemblyName) : StagedBinariesDir;
 
 		// Check if there are any executables being staged in this directory. Usually we only need to stage runtime dependencies next to the executable, but we may be staging
 		// other engine executables too (eg. CEF)
@@ -495,11 +488,33 @@ public class Win64Platform : Platform
 		}
 
 		// try to update the staged SxS manifest
-		DirectoryReference ArchitectureDir = DirectoryReference.Combine( SC.StageDirectory, StagedBinariesDir.Name, Architecture.WindowsToolChain );
-		FileReference ManifestFile = FileReference.Combine(ArchitectureDir, Architecture + ".manifest");
+		string AssemblyName = Architecture.ToString().ToLower();
+		DirectoryReference ArchitectureDir = DirectoryReference.Combine( SC.StageDirectory, StagedBinariesDir.Name, AssemblyName );
+		FileReference ManifestFile = FileReference.Combine(ArchitectureDir, AssemblyName + ".manifest");
 		if (FileReference.Exists(ManifestFile))
 		{
-			WindowsExports.UpdateSxSManifest(ManifestFile);
+			try
+			{
+				XDocument XmlManifest = XDocument.Load(ManifestFile.FullName);
+				if (XmlManifest.Root != null && XmlManifest.Root.Name.LocalName == "assembly")
+				{
+					XNamespace XmlNs = XmlManifest.Root.Name.Namespace;
+
+					// remove all existing file elements & add the new ones (this ensures that stale items from the Binaries folder are not included)
+					XmlManifest.Root.Elements(XmlNs+"file").ToList().ForEach( X => X.Remove() );
+					foreach (FileReference DLLFile in DirectoryReference.EnumerateFiles(ArchitectureDir, "*.dll"))
+					{
+						string DLLName = DLLFile.MakeRelativeTo(ArchitectureDir);
+						XmlManifest.Root.Add( new XElement(XmlNs+"file", new XAttribute("name", DLLName ) ) );
+					}
+
+					XmlManifest.Save(ManifestFile.FullName);
+				}
+			}
+			catch(Exception e)
+			{
+				Logger.LogWarning("Could not update SxS manifest {name} : {msg}", ManifestFile.GetFileName(), e.Message);
+			}
 		}
 	}
 
diff --git a/Engine/Source/Programs/UnrealBuildTool/Platform/Windows/UEBuildWindows.cs b/Engine/Source/Programs/UnrealBuildTool/Platform/Windows/UEBuildWindows.cs
--- a/Engine/Source/Programs/UnrealBuildTool/Platform/Windows/UEBuildWindows.cs
+++ b/Engine/Source/Programs/UnrealBuildTool/Platform/Windows/UEBuildWindows.cs
@@ -552,12 +552,6 @@ namespace UnrealBuildTool
 		/// </summary>
 		public string? ManifestFile;
 
-		/// <summary>
-		/// Specifies the path to a manifest file that defines our SxS assembly on non-x64 architecture, allowing both x64 and arm64 to be staged together with separate runtime dependencies. 
-		/// Defaults to the manifest in Engine/Build/Windows/Resources/{architecture}/. Can be assigned to null if the SxS manifest is not required
-		/// </summary>
-		public string? SxSAssemblyManifestFile;
-
 		/// <summary>
 		/// Specifies the path to an .ico file to use as the appliction icon. Can be assigned to null and will default to Engine/Build/Windows/Resources/Default.ico for engine targets or Build/Windows/Application.ico for projects.
 		/// </summary>
@@ -904,15 +898,6 @@ namespace UnrealBuildTool
 				Platform += "-arm64";
 			}
 			ManifestFile = FileReference.Combine(Unreal.EngineDirectory, "Build", "Windows", "Resources", String.Format("Default-{0}.manifest", Platform)).FullName;
-
-			if (Target.Platform == UnrealTargetPlatform.Win64 && !Target.Architecture.bIsX64)
-			{
-				SxSAssemblyManifestFile = FileReference.Combine(Unreal.EngineDirectory, "Build", "Windows", "Resources", "arm64", String.Format("DefaultSxS-{0}.manifest", Platform)).FullName;
-			}
-			else
-			{
-				SxSAssemblyManifestFile = null;
-			}
 		}
 	}
 
@@ -1016,8 +1001,6 @@ namespace UnrealBuildTool
 
 		public string? ManifestFile => Inner.ManifestFile;
 		
-		public string? SxSAssemblyManifestFile => Inner.SxSAssemblyManifestFile;
-		
 		public string? ApplicationIcon => Inner.ApplicationIcon;
 
 		public bool bNeedsLegacyStdioDefinitionsLib => Inner.bNeedsLegacyStdioDefinitionsLib;
diff --git a/Engine/Source/Programs/UnrealBuildTool/Platform/Windows/VCToolChain.cs b/Engine/Source/Programs/UnrealBuildTool/Platform/Windows/VCToolChain.cs
--- a/Engine/Source/Programs/UnrealBuildTool/Platform/Windows/VCToolChain.cs
+++ b/Engine/Source/Programs/UnrealBuildTool/Platform/Windows/VCToolChain.cs
@@ -1492,9 +1492,10 @@ namespace UnrealBuildTool
 				Arguments.Add($"/MANIFESTINPUT:\"{NormalizeCommandLinePath(ManifestFile, RootPaths)}\"");
 
 				// Embed a dependency on the SxS assembly manifest. this is for dependent runtime DLLs
-				if (Target.WindowsPlatform.SxSAssemblyManifestFile != null && !Target.Architecture.bIsX64)
+				if (!Target.Architecture.bIsX64)
 				{
-					Arguments.Add("/MANIFESTDEPENDENCY:\"name='arm64' processorArchitecture='arm64' version='1.0.0.0' type='win32'\"");
+					string AssemblyName = Target.Architecture.ToString().ToLower();
+					Arguments.Add("/MANIFESTDEPENDENCY:\"name='" + AssemblyName + "' processorArchitecture='arm64' version='1.0.0.0' type='win32'\"");
 				}
 			}
 
@@ -2403,7 +2404,67 @@ namespace UnrealBuildTool
 					}
 				}
 			}
+
+			OptionalUpdateSxSManifest(Target, MakefileBuilder);
+		}
+		
+		protected void OptionalUpdateSxSManifest(ReadOnlyTargetRules Target, TargetMakefileBuilder MakefileBuilder)
+		{
+			if (Target.WindowsPlatform.ManifestFile != null && !Target.bShouldCompileAsDLL && !Target.Architecture.bIsX64)
+			{
+				string AssemblyName = Target.Architecture.ToString().ToLower();
+
+				// collect the DLLs that should be included in the SxS manifest
+				DirectoryReference ManifestDir = DirectoryReference.Combine(MakefileBuilder.Makefile.ExecutableFile.Directory, AssemblyName);
+				IEnumerable<FileItem> DLLs = MakefileBuilder.Makefile.OutputItems.Where( Item =>
+					Item.HasExtension(".dll") &&
+					Item.Location.IsUnderDirectory(ManifestDir)
+				)
+				.Distinct();
+
+				// load or create the intermediate SxS manifest
+				bool bDirty = false;
+				DirectoryReference BaseDir = Target.ProjectFile?.Directory ?? Unreal.EngineDirectory;
+				FileReference IntermediateManifest = FileReference.Combine(  BaseDir, "Intermediate", "SxSManifest", AssemblyName, AssemblyName + ".manifest");
+				FileReference TemplateManifest = FileReference.Combine(Unreal.EngineDirectory, "Build", "Windows", "Resources", AssemblyName, AssemblyName + ".manifest");			
+				XDocument XmlManifest;
+				if (FileReference.Exists(IntermediateManifest))
+				{
+					XmlManifest = XDocument.Load(IntermediateManifest.FullName);
+				}
+				else
+				{
+					XmlManifest = XDocument.Load(TemplateManifest.FullName);
+					bDirty = true;
+				}
+
+				// update & save the intermediate SxS manifest
+				if (XmlManifest.Root != null && XmlManifest.Root.Name.LocalName == "assembly")
+				{
+					XNamespace XmlNs = XmlManifest.Root.Name.Namespace;
+					HashSet<string> ExistingDLLNames = [.. XmlManifest.Root.Elements(XmlNs+"file").Select( X => X.Attribute("name")!.Value )];
+
+					foreach (FileItem DLL in DLLs)
+					{
+						string DLLName = DLL.Location.MakeRelativeTo(ManifestDir);
+						if (!ExistingDLLNames.Contains(DLLName))
+						{
+							XmlManifest.Root.Add( new XElement(XmlNs+"file", new XAttribute("name", DLLName ) ) );
+							bDirty = true;
+						}
+					}
+				}
+				if (bDirty)
+				{
+					DirectoryReference.CreateDirectory(IntermediateManifest.Directory);
+					XmlManifest.Save(IntermediateManifest.FullName);
+				}
+
+				// include the intermediate SxS manifest as a dependency
+				MakefileBuilder.Makefile.InternalDependencies.Add( FileItem.GetItemByFileReference(IntermediateManifest) );
+			}
 		}
+		
 
 		public override void PrepareRuntimeDependencies(List<RuntimeDependency> RuntimeDependencies, Dictionary<FileReference, FileReference> TargetFileToSourceFile, DirectoryReference ExeDir)
 		{
@@ -2455,85 +2516,22 @@ namespace UnrealBuildTool
 				}
 			}
 
-			// mark the generated SxS manifest as a runtime dependency
-			if (Target.WindowsPlatform.SxSAssemblyManifestFile != null && !Target.Architecture.bIsX64)
+			// add a dependency on the SxS manifest
+			if (Target.WindowsPlatform.ManifestFile != null && !Target.bShouldCompileAsDLL && !Target.Architecture.bIsX64)
 			{
-				string Architecture = Target.Architecture.WindowsToolChain;
-				FileReference OutputManifestFile = FileReference.Combine(ExeDir, Architecture, Architecture + ".manifest");
-				RuntimeDependencies.Add( new RuntimeDependency(OutputManifestFile));
-			}
-		}
+				string AssemblyName = Target.Architecture.ToString().ToLower();
+				DirectoryReference BaseDir = Target.ProjectFile?.Directory ?? Unreal.EngineDirectory;
 
+				FileReference SrcFile = FileReference.Combine(  BaseDir, "Intermediate", "SxSManifest", AssemblyName, AssemblyName + ".manifest"); // created in OptionalUpdateSxSManifest during a post-build step
+				FileReference DstFile = FileReference.Combine(ExeDir, AssemblyName, AssemblyName + ".manifest");
 
-
-		[ToolMode("UpdateSxSManifest", ToolModeOptions.None)]
-		private class UpdateSxSManifest : ToolMode
-		{
-			[CommandLine("-Template", Required = true)]
-			FileReference? Template = null;
-
-			[CommandLine("-Manifest", Required = true)]
-			FileReference? Manifest = null;
-
-			public override System.Threading.Tasks.Task<int> ExecuteAsync(EpicGames.Core.CommandLineArguments Arguments, ILogger Logger)
-			{
-				Arguments.ApplyTo(this);
-				Arguments.CheckAllArgumentsUsed();
-
-				bool bSuccess = false;
-				try
-				{
-					if (FileReference.Exists(Manifest!))
-					{
-						FileReference.SetAttributes(Manifest!, FileAttributes.Normal);
-						FileReference.Delete(Manifest!);
-					}
-
-					FileReference.Copy(Template!, Manifest!, true);
-					FileReference.SetAttributes(Manifest!, FileAttributes.Normal);
-
-					bSuccess = WindowsExports.UpdateSxSManifest(Manifest!);
-				}
-				catch(Exception e)
-				{
-					Logger.LogError(e.Message);
-				}
-
-				return System.Threading.Tasks.Task.FromResult(bSuccess ? 0 : 1);
+				TargetFileToSourceFile[DstFile] = SrcFile;
+				RuntimeDependencies.Add(new RuntimeDependency(DstFile, StagedFileType.NonUFS));
 			}
-		};
-
-		public override ICollection<FileItem> PostBuild(ReadOnlyTargetRules Target, FileItem Executable, LinkEnvironment BinaryLinkEnvironment, IActionGraphBuilder Graph)
-		{
-			ICollection<FileItem> OutputFiles = base.PostBuild(Target, Executable, BinaryLinkEnvironment, Graph);
-
-			if (Target.WindowsPlatform.ManifestFile != null && Target.WindowsPlatform.SxSAssemblyManifestFile != null && !Target.Architecture.bIsX64 && !BinaryLinkEnvironment.bIsBuildingDLL)
-			{
-				string Architecture = Target.Architecture.WindowsToolChain;
-				DirectoryReference ManifestDir = DirectoryReference.Combine(Executable.Directory.Location, Architecture);
-				FileItem ManifestFileItem = FileItem.GetItemByFileReference(FileReference.Combine( ManifestDir, Architecture + ".manifest" ));
-
-				string ActionArgs = $"-Template={Target.WindowsPlatform.SxSAssemblyManifestFile} -Manifest={ManifestFileItem.Location}";
-				Action PostBuildAction = Graph.CreateRecursiveAction<UpdateSxSManifest>(ActionType.CreateAppBundle, ActionArgs);
-				PostBuildAction.ProducedItems.Add(ManifestFileItem);
-				PostBuildAction.DeleteItems.Add(ManifestFileItem);
-				PostBuildAction.StatusDescription = $"{Architecture}.manifest";
-				PostBuildAction.bCanExecuteRemotely = false;
+		}
 
-				foreach ( ModuleRules.RuntimeDependency Dependency in BinaryLinkEnvironment.RuntimeDependencies.Where( D => D?.Type == StagedFileType.NonUFS))
-				{
-					FileReference DependencyFile = FileReference.FromString( Dependency.Path );
-					if (DependencyFile.HasExtension(".dll") && DependencyFile.IsUnderDirectory(ManifestDir))
-					{
-						PostBuildAction.PrerequisiteItems.Add( FileItem.GetItemByFileReference(DependencyFile) );
-					}
-				}
 
-				OutputFiles.Add(ManifestFileItem);
-			}
 
-			return OutputFiles;
-		}
 
 		public virtual FileReference GetApplicationIcon(FileReference? ProjectFile)
 		{
@@ -3069,11 +3067,6 @@ namespace UnrealBuildTool
 				{
 					PrerequisiteItems.Add(FileItem.GetItemByPath(Target.WindowsPlatform.ManifestFile));
 				}
-
-				if (Target.WindowsPlatform.SxSAssemblyManifestFile != null && !LinkEnvironment.bIsBuildingDLL)
-				{
-					PrerequisiteItems.Add(FileItem.GetItemByPath(Target.WindowsPlatform.SxSAssemblyManifestFile));
-				}
 			}
 
 			Arguments.AddRange(InputFileNames);
diff --git a/Engine/Source/Programs/UnrealBuildTool/Platform/Windows/WindowsExports.cs b/Engine/Source/Programs/UnrealBuildTool/Platform/Windows/WindowsExports.cs
--- a/Engine/Source/Programs/UnrealBuildTool/Platform/Windows/WindowsExports.cs
+++ b/Engine/Source/Programs/UnrealBuildTool/Platform/Windows/WindowsExports.cs
@@ -126,46 +126,5 @@ namespace UnrealBuildTool
 			OutRedistDir = null;
 			return MicrosoftPlatformSDK.TryGetToolChainDir(Compiler, CompilerVersion, Architecture, Logger, out OutToolChainVersion, out OutToolChainDir, out OutRedistDir);
 		}
-
-
-
-		/// <summary>
-		/// Updates the given SxS manifest to include all of the DLL files in its directory
-		/// </summary>
-		/// <param name="SxSManifestFile"></param>
-		/// <returns>True if it succeeded</returns>
-		public static bool UpdateSxSManifest(FileReference SxSManifestFile)
-		{
-			bool bSuccess = false;
-			try
-			{
-				XDocument XmlManifest = XDocument.Load(SxSManifestFile.FullName);
-				if (XmlManifest.Root != null && XmlManifest.Root.Name.LocalName == "assembly")
-				{
-					XNamespace XmlNs = XmlManifest.Root.Name.Namespace;
-
-					// update the manifest with any DLLs that are in this folder
-					HashSet<string> ExistingDLLNames = [.. XmlManifest.Root.Elements(XmlNs+"file").Select( X => X.Attribute("name")!.Value )];
-					IEnumerable<FileReference> DLLs = DirectoryReference.EnumerateFiles(SxSManifestFile.Directory, "*.dll");
-					foreach (FileReference DLL in DLLs)
-					{
-						string DLLName = DLL.MakeRelativeTo(SxSManifestFile.Directory);
-						if (!ExistingDLLNames.Contains(DLLName))
-						{
-							XmlManifest.Root.Add( new XElement(XmlNs+"file", new XAttribute("name", DLLName ) ) );
-						}
-					}
-
-					// save the manifest again
-					XmlManifest.Save(SxSManifestFile.FullName);
-					bSuccess = true;
-				}
-			}
-			catch(Exception)
-			{
-			}
-
-			return bSuccess;
-		}
 	}
 }
