Fix packaged game crash on mac with mismatching project folder name

#rnx

diff --git a/Engine/Source/Programs/UnrealBuildTool/Platform/Mac/MacToolChain.cs b/Engine/Source/Programs/UnrealBuildTool/Platform/Mac/MacToolChain.cs
--- a/Engine/Source/Programs/UnrealBuildTool/Platform/Mac/MacToolChain.cs
+++ b/Engine/Source/Programs/UnrealBuildTool/Platform/Mac/MacToolChain.cs
@@ -777,7 +777,7 @@ namespace UnrealBuildTool
 							// handle staged rpaths
 							if (ProjectFile != null && DestDir.IsUnderDirectory(ProjectFile.Directory))
 							{
-								RPathEntry = StagedBundlePrefix + DestDir.MakeRelativeTo(ProjectFile.Directory.ParentDirectory!);
+								RPathEntry = StagedBundlePrefix + ProjectFile.GetFileNameWithoutAnyExtensions() + "/" + DestDir.MakeRelativeTo(ProjectFile.Directory);
 							}
 							else
 							{
