UnrealBuildTool: Add pgo training data files as inputs when pgo optimizing

#rnx

diff --git a/Engine/Source/Programs/UnrealBuildTool/Platform/Linux/LinuxToolChain.cs b/Engine/Source/Programs/UnrealBuildTool/Platform/Linux/LinuxToolChain.cs
--- a/Engine/Source/Programs/UnrealBuildTool/Platform/Linux/LinuxToolChain.cs
+++ b/Engine/Source/Programs/UnrealBuildTool/Platform/Linux/LinuxToolChain.cs
@@ -1057,6 +1057,13 @@ namespace UnrealBuildTool
 			// ignored as a prerequisite for other actions
 			LinkAction.bProducesImportLibrary = LinkEnvironment.bIsBuildingDLL;
 
+			// Add profdata as a prerequisite for pgo optimize
+			if (LinkEnvironment.bPGOOptimize && LinkEnvironment.PGODirectory != null && LinkEnvironment.PGOFilenamePrefix != null)
+			{
+				DirectoryReference PGODir = new DirectoryReference(LinkEnvironment.PGODirectory);
+				LinkAction.PrerequisiteItems.Add(FileItem.GetItemByFileReference(FileReference.Combine(PGODir, LinkEnvironment.PGOFilenamePrefix)));
+			}
+
 			bool needsRelink = LinkEnvironment.bIsBuildingDLL && LinkEnvironment.bIsCrossReferenced;
 
 			// If there will be a relink we let the link step write to a file under intermediate folder
diff --git a/Engine/Source/Programs/UnrealBuildTool/Platform/Windows/VCToolChain.cs b/Engine/Source/Programs/UnrealBuildTool/Platform/Windows/VCToolChain.cs
--- a/Engine/Source/Programs/UnrealBuildTool/Platform/Windows/VCToolChain.cs
+++ b/Engine/Source/Programs/UnrealBuildTool/Platform/Windows/VCToolChain.cs
@@ -365,12 +365,12 @@ namespace UnrealBuildTool
 			Arguments.Add($"/clang:-MD /clang:-MF\"{SourceDependsFileString}\"");
 		}
 
-		protected static string GetClangProfDataFilename(CppCompileEnvironment CompileEnvironment)
+		protected static FileReference GetClangProfDataFilename(CppCompileEnvironment CompileEnvironment)
 		{
 			string ProfDataFilename = Path.Combine(CompileEnvironment.PGODirectory!, CompileEnvironment.PGOFilenamePrefix!);
 			if (File.Exists(ProfDataFilename)) 
 			{
-				return ProfDataFilename;
+				return new FileReference(ProfDataFilename);
 			}
 			// If an exact match doesn't exist, fall back to an alternative. This is supported for building Test with Shipping PGO data for example
 			string[] ProfDataFiles = Directory.GetFiles(CompileEnvironment.PGODirectory!, "*.profdata");
@@ -382,7 +382,7 @@ namespace UnrealBuildTool
 			{
 				throw new BuildException($"No .profdata files found in \"{CompileEnvironment.PGODirectory}\".");
 			}
-			return ProfDataFiles.First();
+			return new FileReference(ProfDataFiles.First());
 		}	
 
 		protected virtual void AppendCLArguments_Global(CppCompileEnvironment CompileEnvironment, List<string> Arguments)
@@ -956,7 +956,7 @@ namespace UnrealBuildTool
 				else if (CompileEnvironment.bPGOOptimize)
 				{
 					// Use a merged profdata file.
-					string ProfDataFilename = GetClangProfDataFilename(CompileEnvironment);
+					FileReference ProfDataFilename = GetClangProfDataFilename(CompileEnvironment);
 					Log.TraceInformationOnce($"Using PGO profile data \"{ProfDataFilename}\"");
 					if (Target.WindowsPlatform.bSampleBasedPGO)
 					{
@@ -1877,6 +1877,11 @@ namespace UnrealBuildTool
 				}
 			}
 
+			if (EnvVars.Compiler.IsClang() && CompileEnvironment.bPGOOptimize)
+			{
+				BaseCompileAction.AdditionalPrerequisiteItems.Add(FileItem.GetItemByFileReference(GetClangProfDataFilename(CompileEnvironment)));
+			}
+
 			// Generate the timing info
 			if (CompileEnvironment.bPrintTimingInfo || Target.WindowsPlatform.bCompilerTrace)
 			{
@@ -3150,6 +3155,14 @@ namespace UnrealBuildTool
 			// Allow the toolchain to adjust/process the link arguments
 			ModifyFinalLinkArguments(LinkEnvironment, Arguments, bBuildImportLibraryOnly);
 
+			// Add training data as a prerequisite for pgo optimize
+			if (EnvVars.Compiler.IsMSVC() && LinkEnvironment.bPGOOptimize && LinkEnvironment.PGOFilenamePrefix != null && LinkEnvironment.OutputFilePath.FullName.EndsWith(".exe"))
+			{
+				// training data files are copied to the output directory in ModifyFinalLinkArguments() -> PreparePGOFilesMsvc()
+				PrerequisiteItems.AddRange(DirectoryReference.EnumerateFiles(LinkEnvironment.OutputDirectory!, $"{LinkEnvironment.PGOFilenamePrefix}*.pgd").Select(FileItem.GetItemByFileReference));
+				PrerequisiteItems.AddRange(DirectoryReference.EnumerateFiles(LinkEnvironment.OutputDirectory!, $"{LinkEnvironment.PGOFilenamePrefix}*.pgc").Select(FileItem.GetItemByFileReference));
+			}
+
 			if (IsDynamicDebuggingEnabled)
 			{
 				HashSet<string> exts = [".dll", ".exe", ".exp", ".lib", ".map", ".pdb"];
diff --git a/Engine/Source/Programs/UnrealBuildTool/ToolChain/ClangToolChain.cs b/Engine/Source/Programs/UnrealBuildTool/ToolChain/ClangToolChain.cs
--- a/Engine/Source/Programs/UnrealBuildTool/ToolChain/ClangToolChain.cs
+++ b/Engine/Source/Programs/UnrealBuildTool/ToolChain/ClangToolChain.cs
@@ -1111,7 +1111,14 @@ namespace UnrealBuildTool
 				OutputFile = FileItem.GetItemByFileReference(FileReference.Combine(OutputDir, GetFileNameFromExtension(FileName, ObjectFileExtension)));
 				CompileResult.ObjectFiles.Add(OutputFile);
 			}
-	
+
+			// Add profdata as a prerequisite for pgo optimize
+			if (CompileEnvironment.bPGOOptimize && CompileEnvironment.PGODirectory != null && CompileEnvironment.PGOFilenamePrefix != null)
+			{
+				DirectoryReference PGODir = new DirectoryReference(CompileEnvironment.PGODirectory);
+				CompileAction.PrerequisiteItems.Add(FileItem.GetItemByFileReference(FileReference.Combine(PGODir, CompileEnvironment.PGOFilenamePrefix)));
+			}
+
 			// FPassPlugin
 			foreach (string FPassPlugin in CompileEnvironment.FPassPlugins)
 			{
@@ -1120,7 +1127,7 @@ namespace UnrealBuildTool
     			{
     				Arguments.Add($"-fpass-plugin=\"{NormalizeCommandLinePath(FPassPluginFile, CompileEnvironment.RootPaths)}\"");
     				CompileAction.PrerequisiteItems.Add(FileItem.GetItemByFileReference(FPassPluginFile));
-    			}			
+    			}
 			}
 
 			// Add the output file to the produced item list.
