Submitting on behalf of Joe.Kirchoff:
UnrealBuildTool: Fix /Z7 not being passed for clang pgo

#rnx
#[robomerge] 5.6

diff --git a/Engine/Source/Programs/UnrealBuildTool/Platform/Windows/VCToolChain.cs b/Engine/Source/Programs/UnrealBuildTool/Platform/Windows/VCToolChain.cs
--- a/Engine/Source/Programs/UnrealBuildTool/Platform/Windows/VCToolChain.cs
+++ b/Engine/Source/Programs/UnrealBuildTool/Platform/Windows/VCToolChain.cs
@@ -1050,7 +1050,7 @@ namespace UnrealBuildTool
 			}
 
 			// If enabled, create debug information.
-			if (CompileEnvironment.bCreateDebugInfo && !(Target.WindowsPlatform.Compiler.IsClang() && (CompileEnvironment.bPGOOptimize || CompileEnvironment.bPGOProfile || CompileEnvironment.bAllowLTCG)))
+			if (CompileEnvironment.bCreateDebugInfo)
 			{
 				/*
 				Disabled until we can find a better means to limit the cost of /JMC as it can add considerable runtime overhead
