UnrealBuildTool: Update cache bucket to include if a target is pgo

#rnx

diff --git a/Engine/Source/Programs/UnrealBuildTool/ToolChain/ISPCToolChain.cs b/Engine/Source/Programs/UnrealBuildTool/ToolChain/ISPCToolChain.cs
--- a/Engine/Source/Programs/UnrealBuildTool/ToolChain/ISPCToolChain.cs
+++ b/Engine/Source/Programs/UnrealBuildTool/ToolChain/ISPCToolChain.cs
@@ -395,6 +395,15 @@ namespace UnrealBuildTool
 				cacheBucket = GetStringHash(target.Architecture.ToString(), cacheBucket); // arm64 or x64
 			}
 
+			if (target.bPGOProfile)
+			{
+				cacheBucket = GetStringHash("PGOProfile", cacheBucket); // pgo profile
+			}
+			if (target.bPGOOptimize)
+			{
+				cacheBucket = GetStringHash("PGOOptimize", cacheBucket); // pgo profile
+			}
+
 			if (target.StaticAnalyzer != StaticAnalyzer.None)
 			{
 				cacheBucket = GetStringHash(target.StaticAnalyzer.ToString(), cacheBucket); // If static analyzer
