[UBA]
* Added -killtcphogs to UbaAgent command line on mac

diff --git a/Engine/Source/Programs/UnrealBuildTool/Executors/UnrealBuildAccelerator/UBAAgentCoordinatorHorde.cs b/Engine/Source/Programs/UnrealBuildTool/Executors/UnrealBuildAccelerator/UBAAgentCoordinatorHorde.cs
--- a/Engine/Source/Programs/UnrealBuildTool/Executors/UnrealBuildAccelerator/UBAAgentCoordinatorHorde.cs
+++ b/Engine/Source/Programs/UnrealBuildTool/Executors/UnrealBuildAccelerator/UBAAgentCoordinatorHorde.cs
@@ -574,6 +574,7 @@ namespace UnrealBuildTool
 							// we need to populate the cas with all known xcodes so we can serve all the ones we have installed
 							string xcodeVersion = Utils.RunLocalProcessAndReturnStdOut("/bin/sh", "-c '/usr/bin/defaults read $(xcode-select -p)/../version.plist ProductBuildVersion");
 							arguments.Add($"-populateCasFromXcodeVersion={xcodeVersion}");
+							arguments.Add("-killtcphogs");
 						}
 
 						arguments.Add("-Dir=%UE_HORDE_SHARED_DIR%\\Uba");
