Remove support for C++17, Unreal Engine minimum is now C++20
C++17 is now a validation error for modules and will prevent building
Add highly experimental support for C++23

diff --git a/Engine/Source/Programs/Enterprise/Datasmith/DatasmithMaxExporter/DatasmithMax2017.Build.cs b/Engine/Source/Programs/Enterprise/Datasmith/DatasmithMaxExporter/DatasmithMax2017.Build.cs
--- a/Engine/Source/Programs/Enterprise/Datasmith/DatasmithMaxExporter/DatasmithMax2017.Build.cs
+++ b/Engine/Source/Programs/Enterprise/Datasmith/DatasmithMaxExporter/DatasmithMax2017.Build.cs
@@ -1,5 +1,7 @@
 // Copyright Epic Games, Inc. All Rights Reserved.
+using System;
 using System.IO;
+using System.Linq;
 
 namespace UnrealBuildTool.Rules
 {
@@ -86,7 +88,8 @@ namespace UnrealBuildTool.Rules
 
 			if (int.Parse(GetMaxVersion()) < 2022)
 			{
-				CppStandard = CppStandardVersion.Cpp17;
+				// Being cheeky here as suppressing the warning doesn't work in ModuleRules for some reason
+				CppStandard = Enum.GetValues<CppStandardVersion>().Where(x => x.ToString().Contains("Cpp17")).First();
 			}
 		}
 
diff --git a/Engine/Source/Programs/Shared/UnrealBuildTool.Tests/CppCompileWarningsTests.cs b/Engine/Source/Programs/Shared/UnrealBuildTool.Tests/CppCompileWarningsTests.cs
--- a/Engine/Source/Programs/Shared/UnrealBuildTool.Tests/CppCompileWarningsTests.cs
+++ b/Engine/Source/Programs/Shared/UnrealBuildTool.Tests/CppCompileWarningsTests.cs
@@ -502,7 +502,9 @@ namespace UnrealBuildTool.Tests
 					VersionNumber newVersionNumber = new VersionNumber(18, 9, 9);
 
 					CppCompileEnvironment cppStandardCompileEnvironment = new CppCompileEnvironment(UnrealTargetPlatform.Win64, CppConfiguration.Debug, new UnrealArchitectures(UnrealArch.X64), null!);
+#pragma warning disable CS0618 // Type or member is obsolete
 					cppStandardCompileEnvironment.CppStandard = CppStandardVersion.Cpp17;
+#pragma warning restore CS0618 // Type or member is obsolete
 
 					args = baseWarnings.GenerateWarningCommandLineArgs(cppStandardCompileEnvironment, typeof(ClangToolChain), newVersionNumber);
 					Assert.IsFalse(args.Any());
diff --git a/Engine/Source/Programs/UnrealBuildTool/Configuration/TargetRules.cs b/Engine/Source/Programs/UnrealBuildTool/Configuration/TargetRules.cs
--- a/Engine/Source/Programs/UnrealBuildTool/Configuration/TargetRules.cs
+++ b/Engine/Source/Programs/UnrealBuildTool/Configuration/TargetRules.cs
@@ -140,7 +140,7 @@ namespace UnrealBuildTool
 
 		/// <summary>
 		/// New defaults for 5.3:
-		/// * TargetRules.CppStandard = CppStandard.Default has changed from Cpp17 to Cpp20
+		/// * TargetRules.CppStandard = CppStandardVersion.Default has changed from Cpp17 to Cpp20
 		/// * TargetRules.WindowsPlatform.bStrictConformanceMode = true
 		/// </summary>
 		V4,
@@ -2551,7 +2551,9 @@ namespace UnrealBuildTool
 		[XmlConfigFile(Category = "BuildConfiguration")]
 		public CppStandardVersion CppStandard
 		{
+#pragma warning disable CS0618 // Type or member is obsolete
 			get => CppStandardPrivate ?? (DefaultBuildSettings < BuildSettingsVersion.V4 ? CppStandardVersion.Cpp17 : CppStandardVersion.Default);
+#pragma warning restore CS0618 // Type or member is obsolete
 			set => CppStandardPrivate = value;
 		}
 		private CppStandardVersion? CppStandardPrivate;
@@ -3562,7 +3564,7 @@ namespace UnrealBuildTool
 			if (CppStandardEngine < CppStandardVersion.EngineDefault)
 			{
 				diagnostics.Add("[Upgrade]");
-				diagnostics.Add($"[Upgrade] The latest version of UE no longer supports CppStandardVersion.{CppStandardEngine} which may require code changes.");
+				diagnostics.Add($"[Upgrade] The latest version of UE no longer supports CppStandardVersion.{CppStandardEngine} for engine code which may require code changes.");
 				diagnostics.Add($"[Upgrade] Suppress this message by removing 'CppStandardEngine = CppStandardVersion.{CppStandardEngine};' in {File!.GetFileName()}.");
 				diagnostics.Add("[Upgrade]");
 			}
@@ -3581,7 +3583,7 @@ namespace UnrealBuildTool
 				}
 
 				diagnostics.Add("[Upgrade]");
-				diagnostics.Add($"[Upgrade] The latest version of UE no longer supports CppStandardVersion.{CppStandard} which may require code changes.");
+				diagnostics.Add($"[Upgrade] The latest version of UE no longer {(CppStandard < CppStandardVersion.Minimum ? "supports" : "recommends")} CppStandardVersion.{CppStandard} which may require code changes.");
 				diagnostics.Add($"[Upgrade] Suppress this message by setting 'CppStandard = CppStandardVersion.{defaultCppStandard};' in {File!.GetFileName()}.");
 				diagnostics.Add("[Upgrade]");
 			}
diff --git a/Engine/Source/Programs/UnrealBuildTool/Configuration/UEBuildModuleCPP.cs b/Engine/Source/Programs/UnrealBuildTool/Configuration/UEBuildModuleCPP.cs
--- a/Engine/Source/Programs/UnrealBuildTool/Configuration/UEBuildModuleCPP.cs
+++ b/Engine/Source/Programs/UnrealBuildTool/Configuration/UEBuildModuleCPP.cs
@@ -1,5 +1,8 @@
 // Copyright Epic Games, Inc. All Rights Reserved.
 
+using EpicGames.Core;
+using Microsoft.CodeAnalysis;
+using Microsoft.Extensions.Logging;
 using System;
 using System.Collections.Concurrent;
 using System.Collections.Generic;
@@ -7,8 +10,6 @@ using System.IO;
 using System.Linq;
 using System.Text;
 using System.Text.Json;
-using EpicGames.Core;
-using Microsoft.Extensions.Logging;
 using UnrealBuildBase;
 
 namespace UnrealBuildTool
@@ -290,6 +291,13 @@ namespace UnrealBuildTool
 				}
 			}
 
+			CppStandardVersion CppStandard = Rules.CppStandard ?? (Rules.bTreatAsEngineModule ? Rules.Target.CppStandardEngine : Rules.Target.CppStandard);
+			if (CppStandard < CppStandardVersion.Minimum)
+			{
+				logger.LogError("{Module} CppStandard CppStandardVersion.{CppStandard} is no longer supported.", Name, CppStandard);
+				anyErrors = true;
+			}
+
 			// Make sure that engine modules use shared PCHs or have an explicit private PCH
 			if (Rules.PCHUsage == ModuleRules.PCHUsageMode.NoSharedPCHs && Rules.PrivatePCHHeaderFile == null)
 			{
@@ -1734,8 +1742,8 @@ namespace UnrealBuildTool
 			// Always add the cpp standard into the filename to prevent variants when the engine standard is mismatched with a project
 			switch (CompileEnvironment.CppStandard)
 			{
-				case CppStandardVersion.Cpp17: Variant += ".Cpp17"; break;
 				case CppStandardVersion.Cpp20: Variant += ".Cpp20"; break;
+				case CppStandardVersion.Cpp23: Variant += ".Cpp23"; break;
 				case CppStandardVersion.Latest: Variant += ".CppLatest"; break;
 			}
 
@@ -2389,13 +2397,8 @@ namespace UnrealBuildTool
 			{
 				Result.CppStandard = Result.CppStandardEngine;
 			}
-			// CppModules require at least Cpp20.
-			if (Target.bEnableCppModules && Result.CppStandard < CppStandardVersion.Cpp20)
-			{
-				Result.CppStandard = CppStandardVersion.Cpp20;
-			}
 
-			if (Result.CppStandard < Result.CppStandardEngine)
+			if (Result.CppStandard != Result.CppStandardEngine)
 			{
 				// SharedPCH is disallowed for modules that compile against an older CppStandard than the engine
 				if (Rules.PCHUsage == ModuleRules.PCHUsageMode.UseExplicitOrSharedPCHs)
diff --git a/Engine/Source/Programs/UnrealBuildTool/Platform/Windows/PVSToolChain.cs b/Engine/Source/Programs/UnrealBuildTool/Platform/Windows/PVSToolChain.cs
--- a/Engine/Source/Programs/UnrealBuildTool/Platform/Windows/PVSToolChain.cs
+++ b/Engine/Source/Programs/UnrealBuildTool/Platform/Windows/PVSToolChain.cs
@@ -694,20 +694,17 @@ namespace UnrealBuildTool
 			}
 		}
 
-		public static readonly VersionNumber CLVerWithCPP20Support = new(14, 23);
-
 		const string CPP_20 = "c++20";
-		const string CPP_17 = "c++17";
-		const string CPP_14 = "c++14";
+		const string CPP_23 = "c++23";
 
 		public static string GetLangStandForCfgFile(CppStandardVersion cppStandard, VersionNumber compilerVersion)
 		{
 			return cppStandard switch
 			{
-				CppStandardVersion.Cpp17 => CPP_17,
 				CppStandardVersion.Cpp20 => CPP_20,
-				CppStandardVersion.Latest => VersionNumber.Compare(compilerVersion, CLVerWithCPP20Support) >= 0 ? CPP_20 : CPP_17,
-				_ => CPP_14,
+				CppStandardVersion.Cpp23 => CPP_23,
+				CppStandardVersion.Latest => CPP_23,
+				_ => CPP_20,
 			};
 		}
 
diff --git a/Engine/Source/Programs/UnrealBuildTool/Platform/Windows/VCToolChain.cs b/Engine/Source/Programs/UnrealBuildTool/Platform/Windows/VCToolChain.cs
--- a/Engine/Source/Programs/UnrealBuildTool/Platform/Windows/VCToolChain.cs
+++ b/Engine/Source/Programs/UnrealBuildTool/Platform/Windows/VCToolChain.cs
@@ -1288,12 +1288,19 @@ namespace UnrealBuildTool
 
 			switch (CompileEnvironment.CppStandard)
 			{
-				case CppStandardVersion.Cpp17:
-					Arguments.Add("/std:c++17");
-					break;
 				case CppStandardVersion.Cpp20:
 					Arguments.Add("/std:c++20");
 					break;
+				case CppStandardVersion.Cpp23:
+					if (EnvVars.Compiler.IsMSVC() && EnvVars.CompilerVersion >= new VersionNumber(14, 43, 34808))
+					{
+						Arguments.Add("/std:c++23preview");
+					}
+					else
+					{
+						Arguments.Add("/std:c++latest");
+					}
+					break;
 				case CppStandardVersion.Latest:
 					Arguments.Add("/std:c++latest");
 					break;
@@ -1323,13 +1330,6 @@ namespace UnrealBuildTool
 				{
 					Arguments.Add("/await:strict");
 				}
-				else if (Target.WindowsPlatform.Compiler.IsClang() && !Target.WindowsPlatform.Compiler.IsIntel())
-				{
-					if (CompileEnvironment.CppStandard < CppStandardVersion.Cpp20)
-					{
-						Arguments.Add("-fcoroutines-ts");
-					}
-				}
 			}
 
 			if (Target.WindowsPlatform.Compiler.IsClang())
diff --git a/Engine/Source/Programs/UnrealBuildTool/ProjectFiles/Rider/RiderProjectFile.cs b/Engine/Source/Programs/UnrealBuildTool/ProjectFiles/Rider/RiderProjectFile.cs
--- a/Engine/Source/Programs/UnrealBuildTool/ProjectFiles/Rider/RiderProjectFile.cs
+++ b/Engine/Source/Programs/UnrealBuildTool/ProjectFiles/Rider/RiderProjectFile.cs
@@ -691,12 +691,12 @@ namespace UnrealBuildTool
 				// Do not use version.ToString(). See: https://youtrack.jetbrains.com/issue/RIDER-68030
 				switch (version)
 				{
-					case CppStandardVersion.Cpp17:
-						Writer.WriteValue(Name, "Cpp17");
-						break;
 					case CppStandardVersion.Cpp20:
 						Writer.WriteValue(Name, "Cpp20");
 						break;
+					case CppStandardVersion.Cpp23:
+						Writer.WriteValue(Name, "Cpp23");
+						break;
 					case CppStandardVersion.Latest:
 						Writer.WriteValue(Name, "Latest");
 						break;
diff --git a/Engine/Source/Programs/UnrealBuildTool/ProjectFiles/VisualStudio/VCProject.cs b/Engine/Source/Programs/UnrealBuildTool/ProjectFiles/VisualStudio/VCProject.cs
--- a/Engine/Source/Programs/UnrealBuildTool/ProjectFiles/VisualStudio/VCProject.cs
+++ b/Engine/Source/Programs/UnrealBuildTool/ProjectFiles/VisualStudio/VCProject.cs
@@ -1,5 +1,8 @@
 // Copyright Epic Games, Inc. All Rights Reserved.
 
+using EpicGames.Core;
+using Microsoft.CodeAnalysis;
+using Microsoft.Extensions.Logging;
 using System;
 using System.Collections.Concurrent;
 using System.Collections.Generic;
@@ -12,9 +15,6 @@ using System.Text.RegularExpressions;
 using System.Threading.Tasks;
 using System.Xml;
 using System.Xml.Linq;
-using EpicGames.Core;
-using Microsoft.CodeAnalysis;
-using Microsoft.Extensions.Logging;
 using UnrealBuildBase;
 using static UnrealBuildTool.PlatformProjectGenerator;
 
@@ -557,10 +557,9 @@ namespace UnrealBuildTool
 		{
 			switch (Version)
 			{
-				case CppStandardVersion.Cpp17:
-					return "/std:c++17";
 				case CppStandardVersion.Cpp20:
 					return "/std:c++20";
+				case CppStandardVersion.Cpp23:
 				case CppStandardVersion.Latest:
 					return "/std:c++latest";
 				default:
diff --git a/Engine/Source/Programs/UnrealBuildTool/ProjectFiles/VisualStudioCode/VSCodeProjectFileGenerator.cs b/Engine/Source/Programs/UnrealBuildTool/ProjectFiles/VisualStudioCode/VSCodeProjectFileGenerator.cs
--- a/Engine/Source/Programs/UnrealBuildTool/ProjectFiles/VisualStudioCode/VSCodeProjectFileGenerator.cs
+++ b/Engine/Source/Programs/UnrealBuildTool/ProjectFiles/VisualStudioCode/VSCodeProjectFileGenerator.cs
@@ -706,17 +706,17 @@ namespace UnrealBuildTool
 				OutFile.EndArray();
 			}
 
+			OutFile.AddField("cStandard", "c17");
+
 			switch (CppStandard)
 			{
-				case CppStandardVersion.Cpp17:
-					OutFile.AddField("cStandard", "c17");
-					OutFile.AddField("cppStandard", "c++17");
-					break;
 				case CppStandardVersion.Cpp20:
-				case CppStandardVersion.Latest:
-					OutFile.AddField("cStandard", "c17");
 					OutFile.AddField("cppStandard", "c++20");
 					break;
+				case CppStandardVersion.Cpp23:
+					OutFile.AddField("cppStandard", "c++23");
+					break;
+				case CppStandardVersion.Latest:
 				default:
 					throw new BuildException($"Unsupported C++ standard type set: {CppStandard}");
 			}
diff --git a/Engine/Source/Programs/UnrealBuildTool/System/CppCompileEnvironment.cs b/Engine/Source/Programs/UnrealBuildTool/System/CppCompileEnvironment.cs
--- a/Engine/Source/Programs/UnrealBuildTool/System/CppCompileEnvironment.cs
+++ b/Engine/Source/Programs/UnrealBuildTool/System/CppCompileEnvironment.cs
@@ -34,8 +34,9 @@ namespace UnrealBuildTool
 		Cpp14,
 
 		/// <summary>
-		/// C++17 standard. No longer supported, code may not compile with this version.
+		/// Obsolete C++17 standard. No longer allowed, code will not compile with this version.
 		/// </summary>
+		[Obsolete("C++17 is no longer allowed, code will not compile with this version")]
 		Cpp17,
 
 		/// <summary>
@@ -43,6 +44,11 @@ namespace UnrealBuildTool
 		/// </summary>
 		Cpp20,
 
+		/// <summary>
+		/// C++23 standard. Experimental, compiler support is currently limited.
+		/// </summary>
+		Cpp23,
+
 		/// <summary>
 		/// Latest standard supported by the compiler
 		/// </summary>
@@ -61,7 +67,7 @@ namespace UnrealBuildTool
 		/// <summary>
 		/// The minimum allowed C++ standard version. Lower values than this will cause a build error.
 		/// </summary>
-		Minimum = Cpp17,
+		Minimum = Cpp20,
 	}
 
 	/// <summary>
diff --git a/Engine/Source/Programs/UnrealBuildTool/ToolChain/ClangToolChain.cs b/Engine/Source/Programs/UnrealBuildTool/ToolChain/ClangToolChain.cs
--- a/Engine/Source/Programs/UnrealBuildTool/ToolChain/ClangToolChain.cs
+++ b/Engine/Source/Programs/UnrealBuildTool/ToolChain/ClangToolChain.cs
@@ -381,14 +381,14 @@ namespace UnrealBuildTool
 			// https://clang.llvm.org/cxx_status.html
 			switch (CompileEnvironment.CppStandard)
 			{
-				case CppStandardVersion.Cpp17:
-					Arguments.Add("-std=c++17");
-					break;
 				case CppStandardVersion.Cpp20:
 					Arguments.Add("-std=c++20");
 					break;
+				case CppStandardVersion.Cpp23:
+					Arguments.Add("-std=c++23");
+					break;
 				case CppStandardVersion.Latest:
-					Arguments.Add(CompilerVersionLessThan(13, 0, 0) ? "-std=c++20" : "-std=c++2b");
+					Arguments.Add("-std=c++2c");
 					break;
 				default:
 					throw new BuildException($"Unsupported C++ standard type set: {CompileEnvironment.CppStandard}");
@@ -396,11 +396,6 @@ namespace UnrealBuildTool
 
 			if (CompileEnvironment.bEnableCoroutines)
 			{
-				if (CompileEnvironment.CppStandard < CppStandardVersion.Cpp20)
-				{
-					Arguments.Add("-fcoroutines-ts");
-				}
-
 				if (!CompileEnvironment.bEnableExceptions)
 				{
 					Arguments.Add("-Wno-coroutine-missing-unhandled-exception");
