UnrealBuildTool: Fix legacy c++17 projects breaking project gen

#rnx

diff --git a/Engine/Source/Programs/UnrealBuildTool/Configuration/UEBuildModuleCPP.cs b/Engine/Source/Programs/UnrealBuildTool/Configuration/UEBuildModuleCPP.cs
--- a/Engine/Source/Programs/UnrealBuildTool/Configuration/UEBuildModuleCPP.cs
+++ b/Engine/Source/Programs/UnrealBuildTool/Configuration/UEBuildModuleCPP.cs
@@ -292,7 +292,7 @@ namespace UnrealBuildTool
 			}
 
 			CppStandardVersion CppStandard = Rules.CppStandard ?? (Rules.bTreatAsEngineModule ? Rules.Target.CppStandardEngine : Rules.Target.CppStandard);
-			if (CppStandard < CppStandardVersion.Minimum)
+			if (CppStandard < CppStandardVersion.Minimum && !Rules.Target.bGenerateProjectFiles)
 			{
 				logger.LogError("{Module} CppStandard CppStandardVersion.{CppStandard} is no longer supported.", Name, CppStandard);
 				anyErrors = true;
