[UBA]
* Fixed bug where different root dir could be used if there are fewer actions than max number of parallel processes and dir is provided by coordinators

diff --git a/Engine/Source/Programs/UnrealBuildTool/Executors/UnrealBuildAccelerator/UBAExecutor.cs b/Engine/Source/Programs/UnrealBuildTool/Executors/UnrealBuildAccelerator/UBAExecutor.cs
--- a/Engine/Source/Programs/UnrealBuildTool/Executors/UnrealBuildAccelerator/UBAExecutor.cs
+++ b/Engine/Source/Programs/UnrealBuildTool/Executors/UnrealBuildAccelerator/UBAExecutor.cs
@@ -410,18 +410,11 @@ namespace UnrealBuildTool
 				_rootDirRef = DirectoryReference.FromString(Environment.GetEnvironmentVariable("UBA_ROOT") ?? Environment.GetEnvironmentVariable("BOX_ROOT"));
 			}
 
-			List<Task> coordinatorInitTasks = new();
-			if (!UBAConfig.bDisableRemote)
+			if (_rootDirRef == null)
 			{
 				foreach (IUBAAgentCoordinator coordinator in _agentCoordinators)
 				{
-					// Do not await here, it will block local tasks from running
-					coordinatorInitTasks.Add(coordinator.InitAsync(this));
-
-					if (_rootDirRef == null)
-					{
-						_rootDirRef = coordinator.GetUBARootDir();
-					}
+					_rootDirRef = _rootDirRef == null ? coordinator.GetUBARootDir() : _rootDirRef;
 				}
 			}
 
@@ -437,6 +430,16 @@ namespace UnrealBuildTool
 				}
 			}
 
+			List<Task> coordinatorInitTasks = new();
+			if (!UBAConfig.bDisableRemote)
+			{
+				foreach (IUBAAgentCoordinator coordinator in _agentCoordinators)
+				{
+					// Do not await here, it will block local tasks from running
+					coordinatorInitTasks.Add(coordinator.InitAsync(this));
+				}
+			}
+
 			DirectoryReference.CreateDirectory(_rootDirRef);
 
 			FileReference ubaTraceFile;
