[UBA]
* Added rootdir to ubt debug log

diff --git a/Engine/Source/Programs/UnrealBuildTool/Executors/UnrealBuildAccelerator/UBAExecutor.cs b/Engine/Source/Programs/UnrealBuildTool/Executors/UnrealBuildAccelerator/UBAExecutor.cs
--- a/Engine/Source/Programs/UnrealBuildTool/Executors/UnrealBuildAccelerator/UBAExecutor.cs
+++ b/Engine/Source/Programs/UnrealBuildTool/Executors/UnrealBuildAccelerator/UBAExecutor.cs
@@ -202,6 +202,7 @@ namespace UnrealBuildTool
 			{
 				_threadedLogger.LogInformation("  Object Compression Allowed");
 			}
+			_threadedLogger.LogDebug($"  RootDir {_rootDirRef}");
 		}
 
 		private async Task WriteActionOutputFileAsync(IEnumerable<LinkedAction> inputActions)
@@ -396,7 +397,6 @@ namespace UnrealBuildTool
 
 			_cacheableActionsLeft = inputActions.Count(x => x.ArtifactMode.HasFlag(ArtifactMode.Enabled));
 
-			PrintConfiguration();
 			await WriteActionOutputFileAsync(inputActions);
 
 			logger = _threadedLogger;
@@ -440,6 +440,8 @@ namespace UnrealBuildTool
 				}
 			}
 
+			PrintConfiguration();
+
 			DirectoryReference.CreateDirectory(_rootDirRef);
 
 			FileReference ubaTraceFile;
