Intel CppCompileWarnings moves

diff --git a/Engine/Source/Programs/UnrealBuildTool/Configuration/CompileWarnings/ApplyWarningsAttribute.cs b/Engine/Source/Programs/UnrealBuildTool/Configuration/CompileWarnings/ApplyWarningsAttribute.cs
--- a/Engine/Source/Programs/UnrealBuildTool/Configuration/CompileWarnings/ApplyWarningsAttribute.cs
+++ b/Engine/Source/Programs/UnrealBuildTool/Configuration/CompileWarnings/ApplyWarningsAttribute.cs
@@ -181,6 +181,17 @@ namespace UnrealBuildTool
 		}
 	}
 
+#pragma warning disable CA1813
+	internal class WarningsIntelCompilerAttribute : ApplyWarningsAttribute
+#pragma warning restore CA1813
+	{
+		/// <inheritdoc/>
+		internal WarningsIntelCompilerAttribute(string[]? offArgs = null, string[]? warningArgs = null, string[]? errorArgs = null, string[]? filterNames = null, ArgPosition argPosition = ArgPosition.Middle) : base(offArgs, warningArgs, errorArgs, filterNames, argPosition)
+		{
+			AppendNewFilters([nameof(StandardFilters.IntelCompilerFilter)]);
+		}
+	}
+
 	#endregion -- ToolChains --
 
 	#region -- Compiler Specializations -- 
diff --git a/Engine/Source/Programs/UnrealBuildTool/Configuration/CompileWarnings/StandardApplyWarningsFilterDelegates.cs b/Engine/Source/Programs/UnrealBuildTool/Configuration/CompileWarnings/StandardApplyWarningsFilterDelegates.cs
--- a/Engine/Source/Programs/UnrealBuildTool/Configuration/CompileWarnings/StandardApplyWarningsFilterDelegates.cs
+++ b/Engine/Source/Programs/UnrealBuildTool/Configuration/CompileWarnings/StandardApplyWarningsFilterDelegates.cs
@@ -42,6 +42,19 @@ namespace UnrealBuildTool
 			return readOnlyTargetRules != null && readOnlyTargetRules.WindowsPlatform != null && readOnlyTargetRules.WindowsPlatform.Compiler.IsMSVC();
 		}
 
+		/// <summary>
+		/// Filter for Intel compiler.
+		/// </summary>
+		/// <param name="context">The context in which to evaluate the applicability of the filter.</param>
+		/// <returns>True if the ApplyWarnings will apply, false otherwise.</returns>
+		[ApplyWarningsFilterDelegate(nameof(IntelCompilerFilter))]
+		internal static bool IntelCompilerFilter(CompilerWarningsToolChainContext context)
+		{
+			ReadOnlyTargetRules? readOnlyTargetRules = context._buildSystemContext.GetReadOnlyTargetRules();
+
+			return readOnlyTargetRules != null && readOnlyTargetRules.WindowsPlatform != null && readOnlyTargetRules.WindowsPlatform.Compiler.IsIntel();
+		}
+
 		/// <summary>
 		/// Filter for Clang compiler within the <see cref="VCToolChain"/>.
 		/// </summary>
diff --git a/Engine/Source/Programs/UnrealBuildTool/Configuration/CppCompileWarnings.cs b/Engine/Source/Programs/UnrealBuildTool/Configuration/CppCompileWarnings.cs
--- a/Engine/Source/Programs/UnrealBuildTool/Configuration/CppCompileWarnings.cs
+++ b/Engine/Source/Programs/UnrealBuildTool/Configuration/CppCompileWarnings.cs
@@ -1058,6 +1058,125 @@ namespace UnrealBuildTool
 
 		#endregion -- Clang VC Compiler Disables --
 
+		#region -- Intel Compiler Disables --
+
+		/// <summary>
+		/// Indicates what warning/error level to treat instances of mismatched printf style format specifiers.
+		/// </summary>
+		[WarningsIntelCompiler(["-Wno-format"])]
+		[WarningLevelDefault(WarningLevel.Off)]
+		public WarningLevel FormatWarningLevel
+		{
+			get => ResolveWarning(ParentCppCompileWarnings?.FormatWarningLevel, _formatWarningLevel);
+			set => _formatWarningLevel = value;
+		}
+		private WarningLevel _formatWarningLevel;
+
+		/// <summary>
+		/// Indicates what warning/error level to treat instances of implicit conversion between floating-point types of different precision.
+		/// </summary>
+		[WarningsIntelCompiler(["-Wno-implicit-float-size-conversion"])]
+		[WarningLevelDefault(WarningLevel.Off)]
+		public WarningLevel ImplicitFloatSizeConversionWarningLevel
+		{
+			get => ResolveWarning(ParentCppCompileWarnings?.ImplicitFloatSizeConversionWarningLevel, _implicitFloatSizeConversionWarningLevel);
+			set => _implicitFloatSizeConversionWarningLevel = value;
+		}
+		private WarningLevel _implicitFloatSizeConversionWarningLevel;
+
+		/// <summary>
+		/// Indicates what warning/error level to treat instances where a command line argument is provided, but not used for anything.
+		/// </summary>
+		[WarningsIntelCompiler(["-Wno-single-bit-bitfield-constant-conversion"])]
+		[WarningLevelDefault(WarningLevel.Off)]
+		public WarningLevel SingleBitfieldConstantConversionWarningLevel
+		{
+			get => ResolveWarning(ParentCppCompileWarnings?.SingleBitfieldConstantConversionWarningLevel, _singleBitfieldConstantConversionWarningLevel);
+			set => _singleBitfieldConstantConversionWarningLevel = value;
+		}
+		private WarningLevel _singleBitfieldConstantConversionWarningLevel;
+
+		/// <summary>
+		/// 
+		/// Clang - https://clang.llvm.org/docs/DiagnosticsReference.html#wunused-command-line-argument
+		/// </summary>
+		[WarningsIntelCompiler(["-Wno-unused-command-line-argument"])]
+		[WarningLevelDefault(WarningLevel.Off)]
+		public WarningLevel UnusedCommandLineArgumentWarningLevel
+		{
+			get => ResolveWarning(ParentCppCompileWarnings?.UnusedCommandLineArgumentWarningLevel, _unusedCommandLineArgumentWarningLevel);
+			set => _unusedCommandLineArgumentWarningLevel = value;
+		}
+		private WarningLevel _unusedCommandLineArgumentWarningLevel;
+
+		/// <summary>
+		/// Indicates what warning/error level to treat instances of malformed comments.
+		///  Clang - https://clang.llvm.org/docs/DiagnosticsReference.html#wcomment
+		/// </summary>
+		[WarningsIntelCompiler(["-Wno-comment"])]
+		[WarningLevelDefault(WarningLevel.Off)]
+		public WarningLevel CommentWarningLevel
+		{
+			get => ResolveWarning(ParentCppCompileWarnings?.CommentWarningLevel, _commentWarningLevel);
+			set => _commentWarningLevel = value;
+		}
+		private WarningLevel _commentWarningLevel;
+
+
+		/// <summary>
+		/// Indicates what warning/error level to treat instances of copy construct in range loop iteration.
+		///  Clang - https://clang.llvm.org/docs/DiagnosticsReference.html#wrange-loop-construct
+		/// </summary>
+		[WarningsIntelCompiler(["-Wno-range-loop-construct"])]
+		[WarningLevelDefault(WarningLevel.Off)]
+		public WarningLevel RangeLoopConstructWarningLevel
+		{
+			get => ResolveWarning(ParentCppCompileWarnings?.RangeLoopConstructWarningLevel, _rangeLoopConstructWarningLevel);
+			set => _rangeLoopConstructWarningLevel = value;
+		}
+		private WarningLevel _rangeLoopConstructWarningLevel;
+
+		/// <summary>
+		/// Indicates what warning/error level to treat instances where a #pragma once is used in a source file instead of a header.
+		/// Clang - https://clang.llvm.org/docs/DiagnosticsReference.html#wpragma-once-outside-header
+		/// </summary>
+		[WarningsIntelCompiler(["-Wno-pragma-once-outside-header"])]
+		[WarningLevelDefault(WarningLevel.Off)]
+		public WarningLevel PragmaOnceOutsideHeaderWarningLevel
+		{
+			get => ResolveWarning(ParentCppCompileWarnings?.PragmaOnceOutsideHeaderWarningLevel, _pragmaOnceOutsideHeaderWarningLevel);
+			set => _pragmaOnceOutsideHeaderWarningLevel = value;
+		}
+		private WarningLevel _pragmaOnceOutsideHeaderWarningLevel;
+
+		/// <summary>
+		/// Indicates what warning/error level to treat instances where the logical NOT operator (!) is used in a way that's ambiguous with respect to parentheses.
+		/// Clang - https://clang.llvm.org/docs/DiagnosticsReference.html#wlogical-not-parentheses
+		/// </summary>
+		[WarningsIntelCompiler(["-Wno-logical-not-parentheses"])]
+		[WarningLevelDefault(WarningLevel.Off)]
+		public WarningLevel LogicalNotParenthesesWarningLevel
+		{
+			get => ResolveWarning(ParentCppCompileWarnings?.LogicalNotParenthesesWarningLevel, _logicalNotParenthesesWarningLevel);
+			set => _logicalNotParenthesesWarningLevel = value;
+		}
+		private WarningLevel _logicalNotParenthesesWarningLevel;
+
+		/// <summary>
+		///  Indicates what warning/error level to treat instances where a C++20 feature is being used, but compiling with an older standard.
+		/// Clang - https://clang.llvm.org/docs/DiagnosticsReference.html#wc-20-extensions
+		/// </summary>
+		[WarningsIntelCompiler(["-Wno-c++20-extensions"])]
+		[WarningLevelDefault(WarningLevel.Off)]
+		public WarningLevel Cpp20ExtensionsWarningLevel
+		{
+			get => ResolveWarning(ParentCppCompileWarnings?.Cpp20ExtensionsWarningLevel, _cpp20ExtensionsWarningLevel);
+			set => _cpp20ExtensionsWarningLevel = value;
+		}
+		private WarningLevel _cpp20ExtensionsWarningLevel;
+
+		#endregion -- Intel Compiler Disables --
+
 		#region -- CppCompileEnvironment Settings
 
 		/// <summary>
@@ -1329,7 +1448,6 @@ namespace UnrealBuildTool
 
 			s_membersWithApplyWarningsAttribute ??= GetMembersWithAttributes(typeof(CppCompileWarnings), typeof(ApplyWarningsAttribute), typeof(WarningLevel));
 
-
 			foreach (MemberInfo memberInfo in s_membersWithApplyWarningsAttribute)
 			{
 				ApplyWarningsAttribute? applyWarningsAttribute = null;
diff --git a/Engine/Source/Programs/UnrealBuildTool/Configuration/ReadOnlyCppCompileWarnings.cs b/Engine/Source/Programs/UnrealBuildTool/Configuration/ReadOnlyCppCompileWarnings.cs
--- a/Engine/Source/Programs/UnrealBuildTool/Configuration/ReadOnlyCppCompileWarnings.cs
+++ b/Engine/Source/Programs/UnrealBuildTool/Configuration/ReadOnlyCppCompileWarnings.cs
@@ -80,7 +80,6 @@ namespace UnrealBuildTool
 		#endregion -- Clang Compiler Disables --
 
 		#region -- Clang VC Compiler Disables --
-
 		public WarningLevel MicrosoftGroupWarningLevel => _inner.MicrosoftGroupWarningLevel;
 		public WarningLevel MSVCIncludeWarningLevel => _inner.MSVCIncludeWarningLevel;
 		public WarningLevel PragmaPackWarningLevel => _inner.PragmaPackWarningLevel;
@@ -113,11 +112,25 @@ namespace UnrealBuildTool
 		public WarningLevel MSVCSwitchEnumWarningLevel => _inner.MSVCSwitchEnumWarningLevel;
 		public WarningLevel MSVCUnusedValueWarningLevel => _inner.MSVCUnusedValueWarningLevel;
 		public WarningLevel MSVCDeprecationWarningLevel => _inner.MSVCDeprecationWarningLevel;
-		
+
 		#endregion -- Duplicated Clang VC Compiler Disables --
 
 		#endregion -- Clang VC Compiler Disables --
 
+		#region -- Intel Compiler Disables --
+
+		public WarningLevel FormatWarningLevel => _inner.FormatWarningLevel;
+		public WarningLevel ImplicitFloatSizeConversionWarningLevel => _inner.ImplicitFloatSizeConversionWarningLevel;
+		public WarningLevel SingleBitfieldConstantConversionWarningLevel => _inner.SingleBitfieldConstantConversionWarningLevel;
+		public WarningLevel UnusedCommandLineArgumentWarningLevel => _inner.UnusedCommandLineArgumentWarningLevel;
+		public WarningLevel CommentWarningLevel => _inner.CommentWarningLevel;
+		public WarningLevel RangeLoopConstructWarningLevel => _inner.RangeLoopConstructWarningLevel;
+		public WarningLevel PragmaOnceOutsideHeaderWarningLevel => _inner.PragmaOnceOutsideHeaderWarningLevel;
+		public WarningLevel LogicalNotParenthesesWarningLevel => _inner.LogicalNotParenthesesWarningLevel;
+		public WarningLevel Cpp20ExtensionsWarningLevel => _inner.Cpp20ExtensionsWarningLevel;
+
+		#endregion
+
 		#endregion -- Compiler Warning Settings --
 
 		#region -- Performance Warning Settings --
diff --git a/Engine/Source/Programs/UnrealBuildTool/Platform/Windows/VCToolChain.cs b/Engine/Source/Programs/UnrealBuildTool/Platform/Windows/VCToolChain.cs
--- a/Engine/Source/Programs/UnrealBuildTool/Platform/Windows/VCToolChain.cs
+++ b/Engine/Source/Programs/UnrealBuildTool/Platform/Windows/VCToolChain.cs
@@ -1360,11 +1360,6 @@ namespace UnrealBuildTool
 					Arguments.Add("-Xclang -mllvm -Xclang -autortfm-closed-static-linkage");
 				}
 			}
-
-			if (Target.WindowsPlatform.Compiler == WindowsCompiler.Intel)
-			{
-				ClangWarnings.GetIntelDisabledWarnings(Arguments);
-			}
 		}
 
 		void AppendCLArguments_C(CppCompileEnvironment CompileEnvironment, List<string> Arguments)
diff --git a/Engine/Source/Programs/UnrealBuildTool/ToolChain/ClangWarnings.cs b/Engine/Source/Programs/UnrealBuildTool/ToolChain/ClangWarnings.cs
--- a/Engine/Source/Programs/UnrealBuildTool/ToolChain/ClangWarnings.cs
+++ b/Engine/Source/Programs/UnrealBuildTool/ToolChain/ClangWarnings.cs
@@ -28,28 +28,6 @@ namespace UnrealBuildTool
 			Arguments.Add("-Wno-#pragma-messages");
 		}
 
-		// Additional disabled warnings for Intel. Everything below should be checked if it is necessary
-		internal static void GetIntelDisabledWarnings(List<string> Arguments)
-		{
-			Arguments.Add("-Wno-deprecated-copy");
-			Arguments.Add("-Wno-deprecated-copy-with-user-provided-copy");
-			Arguments.Add("-Wno-format");
-			Arguments.Add("-Wno-implicit-float-size-conversion");
-			Arguments.Add("-Wno-null-pointer-subtraction");
-			Arguments.Add("-Wno-single-bit-bitfield-constant-conversion");
-			Arguments.Add("-Wno-invalid-unevaluated-string");
-			Arguments.Add("-Wno-unused-command-line-argument");
-			Arguments.Add("-Wno-dangling");
-			Arguments.Add("-Wno-comment");
-			Arguments.Add("-Wno-range-loop-construct");
-			Arguments.Add("-Wno-pragma-once-outside-header");
-			Arguments.Add("-Wno-extra-qualification");
-			Arguments.Add("-Wno-logical-not-parentheses");
-			Arguments.Add("-Wno-c++20-extensions");
-			Arguments.Add("-Wno-deprecated-declarations");
-			Arguments.Add("-Wno-shorten-64-to-32");
-		}
-
 		// TODO: Get valid checkers by calling clang -cc1 -analyzer-checker-help
 		static Lazy<Dictionary<string, int>> CheckerAddedVersion = new Lazy<Dictionary<string, int>>(() => new()
 			{
