Improve module merging logic

- Account for game-specific changes to the working plugin set
- Add game-level support for shared merged libraries

#rnx

diff --git a/Engine/Source/Programs/UnrealBuildTool/Configuration/ReadOnlyTargetRules.cs b/Engine/Source/Programs/UnrealBuildTool/Configuration/ReadOnlyTargetRules.cs
--- a/Engine/Source/Programs/UnrealBuildTool/Configuration/ReadOnlyTargetRules.cs
+++ b/Engine/Source/Programs/UnrealBuildTool/Configuration/ReadOnlyTargetRules.cs
@@ -573,6 +573,8 @@ namespace UnrealBuildTool
 
 		public IReadOnlyList<string> MergePlugins => Inner.MergePlugins;
 
+		public IReadOnlyList<string> MergePluginsIgnored => Inner.MergePluginsIgnored;
+
 		public IReadOnlyDictionary<string, IEnumerable<string>> MergePluginsShared => Inner.MergePluginsShared.ToDictionary(k => k.Key, v => v.Value.AsEnumerable());
 
 		public IReadOnlyList<string> MergePluginsLaunch => Inner.MergePluginsLaunch;
diff --git a/Engine/Source/Programs/UnrealBuildTool/Configuration/TargetRules.cs b/Engine/Source/Programs/UnrealBuildTool/Configuration/TargetRules.cs
--- a/Engine/Source/Programs/UnrealBuildTool/Configuration/TargetRules.cs
+++ b/Engine/Source/Programs/UnrealBuildTool/Configuration/TargetRules.cs
@@ -2628,6 +2628,15 @@ namespace UnrealBuildTool
 		[CommandLine("-MergePlugins=", ListSeparator = '+')]
 		public List<string> MergePlugins { get; } = new();
 
+		/// <summary>
+		/// Experimental: List of plugins (and their dependencies) to ignore when processing merged libraries. Requires bMergeModules to be enabled
+		/// Ignored plugins will still be compiled, but if they're kept in the base executable, their dependencies will not be forced out of merged library with them.
+		/// This should be used carefully, for organisational plugins known to have dependencies they don't actually need at runtime.
+		/// </summary>
+		[RequiresUniqueBuildEnvironment]
+		[CommandLine("-MergePluginsIgnored=", ListSeparator = '+')]
+		public List<string> MergePluginsIgnored { get; } = new();
+
 		/// <summary>
 		/// Experimental: List of merged plugins to move common dependencies to a shared library, requires MergePlugins. Can be chained
 		/// </summary>
diff --git a/Engine/Source/Programs/UnrealBuildTool/Configuration/UEBuildTarget.cs b/Engine/Source/Programs/UnrealBuildTool/Configuration/UEBuildTarget.cs
--- a/Engine/Source/Programs/UnrealBuildTool/Configuration/UEBuildTarget.cs
+++ b/Engine/Source/Programs/UnrealBuildTool/Configuration/UEBuildTarget.cs
@@ -3568,8 +3568,31 @@ namespace UnrealBuildTool
 					MergeModulesToDLL(name, modules);
 				}
 			}
+
+			void GetPluginDependencies(UEBuildPlugin plugin, ref HashSet<UEBuildPlugin> dependencies, bool recursively = false)
+			{
+				foreach (PluginReferenceDescriptor dependency in plugin.Info.Descriptor.Plugins ?? new())
+				{
+					if (!Rules.ShouldIgnorePluginDependency(plugin.Info, dependency))
+					{
+						IEnumerable<UEBuildPlugin> dependentPluginCandidates = (plugin.Dependencies ?? new()).Where(x => x.Name == dependency.Name);
+						if (dependentPluginCandidates.Any())
+						{
+							UEBuildPlugin dependentPlugin = dependentPluginCandidates.First();
+							if (!dependencies.Contains(dependentPlugin))
+							{
+								dependencies.Add(dependentPlugin);
+								if (recursively)
+								{
+									GetPluginDependencies(dependentPlugin, ref dependencies, recursively);
+								}
+							}
+						}
+					}
+				}
+			}
 		
-			void GetDependenciesRecursively(IEnumerable<UEBuildModuleCPP> Modules, ref HashSet<UEBuildModuleCPP> Needed)
+			void GetModulesDependenciesRecursively(IEnumerable<UEBuildModuleCPP> Modules, ref HashSet<UEBuildModuleCPP> Needed)
 			{
 				IEnumerable<UEBuildModuleCPP> CPPModules = Modules.OfType<UEBuildModuleCPP>();
 
@@ -3586,14 +3609,77 @@ namespace UnrealBuildTool
 					{
 						if (item.PrivateDependencyModules != null && item.PrivateDependencyModules.Count > 0)
 						{
-							GetDependenciesRecursively(item.PrivateDependencyModules.OfType<UEBuildModuleCPP>(), ref Needed);
+							GetModulesDependenciesRecursively(item.PrivateDependencyModules.OfType<UEBuildModuleCPP>(), ref Needed);
 						}
 
 						if (item.PublicDependencyModules != null && item.PublicDependencyModules.Count > 0)
 						{
-							GetDependenciesRecursively(item.PublicDependencyModules.OfType<UEBuildModuleCPP>(), ref Needed);
+							GetModulesDependenciesRecursively(item.PublicDependencyModules.OfType<UEBuildModuleCPP>(), ref Needed);
+						}
+					}
+				}
+			}
+			
+			void HandleSharedLibraries(ref Dictionary<string, HashSet<UEBuildPlugin>> pluginGroups, IReadOnlySet<string> nonRootSharedPlugins)
+			{
+				foreach (KeyValuePair<string, IEnumerable<string>> group in Rules.MergePluginsShared)
+				{
+					// Get the existing groups that will be shared
+					Dictionary<string, HashSet<UEBuildPlugin>> shared = new();
+					foreach (string dependency in group.Value)
+					{
+						if (pluginGroups.TryGetValue(dependency, out HashSet<UEBuildPlugin>? value))
+						{
+							shared.Add(dependency, value);
+						}
+					}
+
+					// Find all plugins that are referenced by more than one group or are forced into the shared set
+					HashSet<UEBuildPlugin> common = new(shared.SelectMany(x => x.Value).Where(x => shared.Count(y => y.Value.Contains(x)) > 1));
+					HashSet<UEBuildPlugin> forced = new();
+					foreach (KeyValuePair<string, HashSet<UEBuildPlugin>> item in shared)
+					{
+						if (nonRootSharedPlugins.Contains(item.Key))
+						{
+							forced.UnionWith(item.Value);
+							pluginGroups.Remove(item.Key);
+						}
+						else
+						{
+							item.Value.ExceptWith(common);
+							pluginGroups.Remove(item.Key);
+							pluginGroups.Add(item.Key, item.Value);
 						}
 					}
+
+					common.UnionWith(forced);
+					pluginGroups.Add(group.Key, common);
+				}
+			}
+
+			void RemoveNonExclusivePlugins(ref Dictionary<string, HashSet<UEBuildPlugin>> pluginGroups, ref HashSet<UEBuildPlugin> conflictingPlugins)
+			{
+				if (pluginGroups.Any())
+				{
+					HashSet<UEBuildPlugin> commonPluginDependencies = new();
+					foreach (KeyValuePair<string, HashSet<UEBuildPlugin>> item in pluginGroups)
+					{
+						IEnumerable<UEBuildPlugin> otherPluginDependencies = pluginGroups
+							.Where(x => x.Key != item.Key)       // get other groups
+							.SelectMany(x => x.Value)            // get plugins
+							.Distinct()                          // unique
+							.Where(x => item.Value.Contains(x)); // that are shared here
+
+						commonPluginDependencies.UnionWith(otherPluginDependencies);
+						conflictingPlugins.UnionWith(otherPluginDependencies);
+					}
+
+					Dictionary<string, HashSet<UEBuildPlugin>> filtered = new();
+					foreach (KeyValuePair<string, HashSet<UEBuildPlugin>> item in pluginGroups)
+					{
+						filtered.Add(item.Key, new(item.Value.Except(commonPluginDependencies)));
+					}
+					pluginGroups = filtered;
 				}
 			}
 
@@ -3605,67 +3691,154 @@ namespace UnrealBuildTool
 
 			if (Rules.MergePlugins.Any() && BuildPlugins != null)
 			{
-				// Find the plugin to merge and gather all dependencies
+				// Dependency rules:
+				// - Regular (non-shared) merged libraries cannot contain plugins that are dependencies of outside plugins.
+				// - Regular (non-shared) merged libraries can depend on any number of shared merged libraries.
+				// - Shared merged libraries cannot have dependencies at all (other than the base executable) to avoid circularity.
+				// - Plugins that are part of no merged library (shared or not) will be pulled into the executable *with their dependencies*, so they should be minimized.
+				// - Some plugins are organizational in nature and should be ignored for dependencies.
+
+				// Syntax
+				// - Rules.MergePlugin represents merged libraries: everything in that plugin, direct plugin dependencies and dependent modules.
+				// - Rules.MergePluginsShared represents shared merged libraries: same as above, except it's a list of plugins, and plugins in a merged library are excluded.
+				// - Rules.MergePluginsIgnored: explicitly ignored plugins.
+
+				// Shared merged libraries hold modules that are common to at least two in a list of root GFPs.
+				// If an entry for a shared merged library isn't also defined as a root, it'll be fully merged into the shared.
+				HashSet<string> nonRootSharedPlugins = Rules.MergePluginsShared.SelectMany(x => x.Value).Except(Rules.MergePlugins).ToHashSet();
+				IEnumerable<string> mergedPlugins = Rules.MergePlugins.Union(nonRootSharedPlugins);
+				IEnumerable<string> sharedPlugins = nonRootSharedPlugins.Where(plugin => Rules.MergePluginsShared.SelectMany(x => x.Value).Where(x => x == plugin).Count() > 1);
+				nonRootSharedPlugins.ExceptWith(sharedPlugins);
+
+				// Find the plugins to merge and gather their direct dependencies
 				Dictionary<string, HashSet<UEBuildPlugin>> pluginGroups = new();
-				foreach (string name in Rules.MergePlugins)
+				foreach (string name in mergedPlugins)
 				{
 					UEBuildPlugin? plugin = BuildPlugins.FirstOrDefault(x => x.Name == name);
 					if (plugin == null)
 					{
-						Logger.LogWarning("Plugin {Name} not found in BuildPlugin and will not be merged", name);
+						Logger.LogWarning("Plugin {Name} not found in BuildPlugins and will not be merged", name);
 						continue;
 					}
 					HashSet<UEBuildPlugin> needed = new();
 					needed.Add(plugin);
-					needed.UnionWith(plugin.Dependencies ?? new());
+					GetPluginDependencies(plugin, ref needed);
 					pluginGroups.Add(plugin.Name, needed);
 				}
 
-				// Move common plugins to a new group
-				foreach (KeyValuePair<string, IEnumerable<string>> group in Rules.MergePluginsShared)
+				// For each group, find out plugin dependencies that would make sense to also merge
+				Dictionary<string, HashSet<UEBuildPlugin>> desiredPluginGroups = new();
+				foreach (KeyValuePair<string, HashSet<UEBuildPlugin>> item in pluginGroups)
 				{
-					// Get the existing groups that will be shared
-					Dictionary<string, HashSet<UEBuildPlugin>> shared = new();
-					foreach (string dependecy in group.Value)
+					HashSet<UEBuildPlugin> desired = new();
+					foreach (UEBuildPlugin plugin in item.Value)
 					{
-						if (pluginGroups.TryGetValue(dependecy, out HashSet<UEBuildPlugin>? value))
+						if (!desiredPluginGroups.ContainsKey(plugin.Name))
 						{
-							shared.Add(dependecy, value);
+							GetPluginDependencies(plugin, ref desired, true);
 						}
 					}
+					desiredPluginGroups.Add(item.Key, desired);
+				}
 
-					// Find all plugins that are referenced by more than one group
-					HashSet<UEBuildPlugin> common = new(shared.SelectMany(x => x.Value).Where(x => shared.Count(y => y.Value.Contains(x)) > 1));
-					foreach (KeyValuePair<string, HashSet<UEBuildPlugin>> item in shared)
+				// Move common plugins to a new group and then remove duplicates
+				HashSet<UEBuildPlugin> conflictingPlugins = new();
+				HandleSharedLibraries(ref pluginGroups, nonRootSharedPlugins);
+				RemoveNonExclusivePlugins(ref pluginGroups, ref conflictingPlugins);
+
+				// Merge secondary dependencies into their owning set
+				HandleSharedLibraries(ref desiredPluginGroups, nonRootSharedPlugins);
+				RemoveNonExclusivePlugins(ref desiredPluginGroups, ref conflictingPlugins);
+
+				// Merge remaining secondary dependencies into primaries
+				foreach (string group in pluginGroups.Keys.ToArray())
+				{
+					if (desiredPluginGroups.TryGetValue(group, out HashSet<UEBuildPlugin>? value))
 					{
-						item.Value.ExceptWith(common);
-						pluginGroups.Remove(item.Key);
-						pluginGroups.Add(item.Key, item.Value);
+						pluginGroups[group].UnionWith(value);
 					}
-					pluginGroups.Add(group.Key, common);
 				}
 
-				// Remove any plugins that are shared between multiple merge groups
-				if (pluginGroups.Any())
+				// Find remaining unmerged plugins
+				IEnumerable<UEBuildPlugin> ignoredPlugins = BuildPlugins.Where(x => Rules.MergePluginsIgnored.Contains(x.Name));
+				HashSet<UEBuildPlugin> unmergablePlugins = BuildPlugins.Except(ignoredPlugins).Except(pluginGroups.SelectMany(x => x.Value)).ToHashSet();
+
+				// For all unreferenced plugins, attempt to merge them through their dependencies
+				for (int unreferencedPhaseIndex = 0; unreferencedPhaseIndex < 3; unreferencedPhaseIndex++)
 				{
-					Dictionary<string, HashSet<UEBuildPlugin>> filtered = new();
-					foreach (KeyValuePair<string, HashSet<UEBuildPlugin>> item in pluginGroups)
+					Logger.LogInformation("Unreferenced phase {unreferencedPhaseIndex}", unreferencedPhaseIndex);
+
+					// For all unreferenced plugins, attempt to merge them through their dependencies
+					foreach (UEBuildPlugin plugin in unmergablePlugins)
 					{
-						IEnumerable<UEBuildPlugin> otherPluginDependencies = pluginGroups.Where(x => x.Key != item.Key).SelectMany(x => x.Value).Distinct();
-						filtered.Add(item.Key, new(item.Value.Except(otherPluginDependencies)));
+						if (!pluginGroups.SelectMany(x => x.Value).Contains(plugin))
+						{
+							HashSet<UEBuildPlugin> dependencies = new();
+							GetPluginDependencies(plugin, ref dependencies);
+
+							if (dependencies.Count() > 0)
+							{
+								string? foundPluginGroup = pluginGroups.FirstOrDefault(x => x.Value.Intersect(dependencies).Count() == dependencies.Count()).Key;
+
+								if (foundPluginGroup != null)
+								{
+									Logger.LogInformation("{0}: was unreferenced but will be merged with {1}", plugin.Name, foundPluginGroup);
+									pluginGroups[foundPluginGroup].Add(plugin);
+									pluginGroups[foundPluginGroup].UnionWith(dependencies);
+								}
+							}
+						}
 					}
-					pluginGroups = filtered;
 				}
-
-				// Find remaining unmerged plugins and gather their dependencies
-				HashSet<UEBuildPlugin> unmergablePlugins = new(BuildPlugins);
 				unmergablePlugins.ExceptWith(pluginGroups.SelectMany(x => x.Value));
-				unmergablePlugins = new(unmergablePlugins.Union(unmergablePlugins.SelectMany(x => x.Dependencies ?? new())));
 
-				HashSet<UEBuildModuleCPP> allPluginModules = new HashSet<UEBuildModuleCPP>();
+				// At this point, we have to exclude any code that falls under remaining unmergeable plugins, dependencies included.
+				// We need that because there is no way the executable can depend on a merged library, and this remaining chunk would always be loaded as its own library.
+
+				// Update unmergeable plugins by accounting for their dependencies
+				Dictionary<UEBuildPlugin, HashSet<string>> unmergeablePluginDependencies = new();
+				foreach (UEBuildPlugin plugin in unmergablePlugins)
+				{
+					if (conflictingPlugins.Contains(plugin))
+					{
+						Logger.LogInformation("{0}: unmergeable as a conflicted plugin", plugin.Name);
+					}
+					else
+					{
+						Logger.LogInformation("{0}: unmergeable as an unreferenced plugin", plugin.Name);
+					}
+
+					HashSet<UEBuildPlugin> dependencies = new();
+					GetPluginDependencies(plugin, ref dependencies, true);
 
-				// Remove unmergable plugin modules from plugin modules
-				// and create a set of all plugin modules
+					foreach (UEBuildPlugin dependency in dependencies)
+					{
+						if (!unmergeablePluginDependencies.ContainsKey(dependency))
+						{
+							HashSet<string> pluginNames = new();
+							pluginNames.Add(plugin.Name);
+							unmergeablePluginDependencies.Add(dependency, pluginNames);
+						}
+						else
+						{
+							unmergeablePluginDependencies[dependency].Add(plugin.Name);
+						}
+					}
+				}
+
+				// Report unmergeable dependencies
+				unmergablePlugins = new(unmergablePlugins.Union(unmergeablePluginDependencies.Keys));
+				foreach (KeyValuePair<UEBuildPlugin, HashSet<string>> dependency in unmergeablePluginDependencies)
+				{
+					if (dependency.Key.Modules.Any())
+					{
+						Logger.LogInformation("{0}: unmergeable as a dependency of: {1}", dependency.Key, string.Join(",", dependency.Value));
+					}
+				}
+				Logger.LogInformation("Found a total of {0} plugins to merge", pluginGroups.SelectMany(x => x.Value).Count());
+
+				// Remove unmergable plugin modules from plugin modules and create a set of all plugin modules
+				HashSet<UEBuildModuleCPP> allPluginModules = new HashSet<UEBuildModuleCPP>();
 				foreach (KeyValuePair<string, HashSet<UEBuildPlugin>> item in pluginGroups)
 				{
 					// Find all modules for the unique plugins to be merged, except those that were unmergable due to being a dependency of an unmerged plugin
@@ -3680,8 +3853,8 @@ namespace UnrealBuildTool
 				allLeftoverModules.ExceptWith(allPluginModules);
 
 				// create a set of all dependencies of leftover modules
-				HashSet<UEBuildModuleCPP> allLeftoverDependenies = new HashSet<UEBuildModuleCPP>();
-				GetDependenciesRecursively(allLeftoverModules.OfType<UEBuildModuleCPP>(), ref allLeftoverDependenies);
+				HashSet<UEBuildModuleCPP> allLeftoverDependencies = new HashSet<UEBuildModuleCPP>();
+				GetModulesDependenciesRecursively(allLeftoverModules.OfType<UEBuildModuleCPP>(), ref allLeftoverDependencies);
 
 				// Remove all leftover dependencies from plugin modules
 				// then remove remaining plugin modules from all modules set
@@ -3691,7 +3864,7 @@ namespace UnrealBuildTool
 					// Find all modules for the unique plugins to be merged, except those that were unmergable due to being a dependency of an unmerged plugin
 					HashSet<UEBuildModuleCPP> pluginModules = item.Value.SelectMany(x => x.Modules).ToHashSet();
 
-					pluginModules.ExceptWith(allLeftoverDependenies);
+					pluginModules.ExceptWith(allLeftoverDependencies);
 					allModules.ExceptWith(pluginModules);
 
 					ConsumeModules(item.Key, pluginModules);
@@ -4916,8 +5089,36 @@ namespace UnrealBuildTool
 					{
 						foreach (PluginReferenceDescriptor NextReference in Instance.Descriptor.Plugins)
 						{
-							string NextReferenceChain = String.Format("{0} -> {1}", ReferenceChain, Instance.File.GetFileName());
-							AddPlugin(NextReference, NextReferenceChain, ExcludeFolders, NameToInstance, NameToInfos, Logger);
+							
+							PluginInfo? Info = null;
+
+							PluginSet? PluginVersions;
+							if (NameToInfos.TryGetValue(Reference.Name, out PluginVersions))
+							{
+								if (Reference.RequestedVersion.HasValue)
+								{
+									Info = PluginVersions.KnownVersions.Find(x => x.Descriptor.Version == Reference.RequestedVersion.Value);
+									if (Info != null)
+									{
+										PluginVersions.ChoiceVersion = Info;
+									}
+									else
+									{
+										Logger.LogWarning("Failed to find specific version (v{RequestedVersion}) of plugin '{ReferenceName}' (referenced via {ReferenceChain}). Other versions exist, but the explicit version requested could not be found.",
+											Reference.RequestedVersion.Value, Reference.Name, ReferenceChain);
+									}
+								}
+								else
+								{
+									Info = PluginVersions.ChoiceVersion;
+								}
+							}
+
+							if (Info == null || !Rules.ShouldIgnorePluginDependency(Info, NextReference))
+							{
+								string NextReferenceChain = String.Format("{0} -> {1}", ReferenceChain, Instance.File.GetFileName());
+								AddPlugin(NextReference, NextReferenceChain, ExcludeFolders, NameToInstance, NameToInfos, Logger);
+							}
 						}
 					}
 				}
