[UBA]
* Changed so proxy port is not local only and not opening a relay port

diff --git a/Engine/Source/Programs/UnrealBuildTool/Executors/UnrealBuildAccelerator/UBAAgentCoordinatorHorde.cs b/Engine/Source/Programs/UnrealBuildTool/Executors/UnrealBuildAccelerator/UBAAgentCoordinatorHorde.cs
--- a/Engine/Source/Programs/UnrealBuildTool/Executors/UnrealBuildAccelerator/UBAAgentCoordinatorHorde.cs
+++ b/Engine/Source/Programs/UnrealBuildTool/Executors/UnrealBuildAccelerator/UBAAgentCoordinatorHorde.cs
@@ -275,7 +275,7 @@ namespace UnrealBuildTool
 				{
 					ModePreference = _connectionMode,
 					Encryption = _encryption,
-					Ports = { { UbaPortName, UbaPort }, { UbaProxyPortName, UbaProxyPort } },
+					Ports = { { UbaPortName, UbaPort } },
 					InactivityTimeoutMs = 90000
 				};
 				
