Add IOS_18 as a minimum deployment target.  Add ability to specify specific deployment version. ie: 17.4
Also, suppress warning for deprecated apis when deployment target is higher than deprecated api version.

diff --git a/Engine/Source/Programs/UnrealBuildTool/Platform/IOS/UEBuildIOS.cs b/Engine/Source/Programs/UnrealBuildTool/Platform/IOS/UEBuildIOS.cs
--- a/Engine/Source/Programs/UnrealBuildTool/Platform/IOS/UEBuildIOS.cs
+++ b/Engine/Source/Programs/UnrealBuildTool/Platform/IOS/UEBuildIOS.cs
@@ -398,10 +398,19 @@ namespace UnrealBuildTool
 						return "16.0";
 					case "IOS_17":
 						return "17.0";
+					case "IOS_18":
+						return "18.0";
 					case "IOS_Minimum":
 					default:
-						// default to the UE min in IOS_SDK.json
-						return UEBuildPlatformSDK.GetSDKForPlatform("IOS")!.GetSoftwareInfo()!.Min!;
+						if (String.IsNullOrEmpty(MinimumIOSVersion) || MinimumIOSVersion.StartsWith("IOS_"))
+						{
+							return UEBuildPlatformSDK.GetSDKForPlatform("IOS")!.GetSoftwareInfo()!.Min!;
+						}
+						else
+						{
+							// Assume the user is overriding the defaults, ie: "17.4.1"
+							return MinimumIOSVersion;
+						}
 				}
 			}
 		}
diff --git a/Engine/Source/Runtime/Core/Private/IOS/IOSPlatformMisc.cpp b/Engine/Source/Runtime/Core/Private/IOS/IOSPlatformMisc.cpp
--- a/Engine/Source/Runtime/Core/Private/IOS/IOSPlatformMisc.cpp
+++ b/Engine/Source/Runtime/Core/Private/IOS/IOSPlatformMisc.cpp
@@ -491,7 +491,11 @@ void FIOSPlatformMisc::SetAllowedDeviceOrientation(EDeviceScreenOrientation NewA
 #if !PLATFORM_VISIONOS
 			else
 			{
+// Deprecated iOS16 - suppress warning if deployment target	>= iOS16
+#pragma clang diagnostic push
+#pragma clang diagnostic ignored "-Wdeprecated-declarations"
 				[UIViewController attemptRotationToDeviceOrientation];
+#pragma clang diagnostic pop
 			}
 #endif
 		}
diff --git a/Engine/Source/Runtime/Launch/Private/IOS/LaunchIOS.cpp b/Engine/Source/Runtime/Launch/Private/IOS/LaunchIOS.cpp
--- a/Engine/Source/Runtime/Launch/Private/IOS/LaunchIOS.cpp
+++ b/Engine/Source/Runtime/Launch/Private/IOS/LaunchIOS.cpp
@@ -318,7 +318,11 @@ void FAppEntry::PreInit(IOSAppDelegate* AppDelegate, UIApplication* Application)
 
 #if !PLATFORM_TVOS
 	// reset badge count on launch
+	// Deprecated iOS17 - suppress warning if deployment target	> iOS17
+#pragma clang diagnostic push
+#pragma clang diagnostic ignored "-Wdeprecated-declarations"
 	Application.applicationIconBadgeNumber = 0;
+#pragma clang diagnostic pop
 #endif
 	
 #endif
