Fix legacy hot reload

#rnx

diff --git a/Engine/Source/Programs/UnrealBuildTool/Configuration/UEBuildBinary.cs b/Engine/Source/Programs/UnrealBuildTool/Configuration/UEBuildBinary.cs
--- a/Engine/Source/Programs/UnrealBuildTool/Configuration/UEBuildBinary.cs
+++ b/Engine/Source/Programs/UnrealBuildTool/Configuration/UEBuildBinary.cs
@@ -398,20 +398,12 @@ namespace UnrealBuildTool
 		}
 
 		/// <summary>
-		/// Called to allow the binary to find game modules.
+		/// Called to allow the binary to find hot reloadable modules.
 		/// </summary>
-		/// <returns>The OnlyModule if found, null if not</returns>
-		public List<UEBuildModule> FindHotReloadModules()
+		/// <returns>Enumerable of modules</returns>
+		public IEnumerable<UEBuildModule> FindHotReloadModules()
 		{
-			List<UEBuildModule> GameModules = new List<UEBuildModule>();
-			foreach (UEBuildModule Module in Modules)
-			{
-				if (Module.Rules.Context.bCanHotReload)
-				{
-					GameModules.Add(Module);
-				}
-			}
-			return GameModules;
+			return Modules.Where(x => x.Rules.Context.bCanHotReload);
 		}
 
 		/// <summary>
diff --git a/Engine/Source/Programs/UnrealBuildTool/Configuration/UEBuildTarget.cs b/Engine/Source/Programs/UnrealBuildTool/Configuration/UEBuildTarget.cs
--- a/Engine/Source/Programs/UnrealBuildTool/Configuration/UEBuildTarget.cs
+++ b/Engine/Source/Programs/UnrealBuildTool/Configuration/UEBuildTarget.cs
@@ -2536,7 +2536,7 @@ namespace UnrealBuildTool
 			Makefile.PreBuildTargets = Rules.PreBuildTargets.ToArray();
 
 			// Setup the hot reload module list
-			Makefile.HotReloadModuleNames = GetHotReloadModuleNames();
+			Makefile.HotReloadModuleNames = GetHotReloadModuleNames(TargetDescriptor);
 
 			// If we're compiling monolithic, make sure the executable knows about all referenced modules
 			if (ShouldCompileMonolithic())
@@ -4156,15 +4156,19 @@ namespace UnrealBuildTool
 		/// Finds a list of module names which can be hot-reloaded
 		/// </summary>
 		/// <returns>Set of module names</returns>
-		private HashSet<string> GetHotReloadModuleNames()
+		private HashSet<string> GetHotReloadModuleNames(TargetDescriptor TargetDescriptor)
 		{
 			HashSet<string> HotReloadModuleNames = new HashSet<string>(StringComparer.OrdinalIgnoreCase);
 			if (Rules.bAllowHotReload && !ShouldCompileMonolithic())
 			{
-				foreach (UEBuildBinary Binary in Binaries)
+				// Legacy hot reload cannot reload modules from the executable binary
+				IEnumerable<UEBuildBinary> ValidBinaries = TargetDescriptor.HotReloadMode == HotReloadMode.FromEditor || TargetDescriptor.HotReloadMode == HotReloadMode.FromIDE
+					? Binaries.Skip(1)
+					: Binaries;
+				foreach (UEBuildBinary Binary in ValidBinaries)
 				{
-					List<UEBuildModule> GameModules = Binary.FindHotReloadModules();
-					if (GameModules != null && GameModules.Count > 0)
+					IEnumerable<UEBuildModule> GameModules = Binary.FindHotReloadModules();
+					if (GameModules.Any())
 					{
 						if (!Unreal.IsProjectInstalled() || EnabledPlugins!.Where(x => x.Type == PluginType.Mod).Any(x => Binary.OutputFilePaths[0].IsUnderDirectory(x.Directory)))
 						{
diff --git a/Engine/Source/Programs/UnrealBuildTool/System/HotReload.cs b/Engine/Source/Programs/UnrealBuildTool/System/HotReload.cs
--- a/Engine/Source/Programs/UnrealBuildTool/System/HotReload.cs
+++ b/Engine/Source/Programs/UnrealBuildTool/System/HotReload.cs
@@ -1,5 +1,8 @@
 // Copyright Epic Games, Inc. All Rights Reserved.
 
+using EpicGames.Core;
+using Microsoft.Extensions.Logging;
+using OpenTracing.Util;
 using System;
 using System.Collections.Generic;
 using System.Diagnostics;
@@ -8,9 +11,6 @@ using System.Linq;
 using System.Text;
 using System.Text.Json;
 using System.Threading;
-using EpicGames.Core;
-using Microsoft.Extensions.Logging;
-using OpenTracing.Util;
 using UnrealBuildBase;
 
 namespace UnrealBuildTool
@@ -743,18 +743,15 @@ namespace UnrealBuildTool
 				for (int LastNumFilesWithNewSuffix = 0; FilesRequiringSuffix.Count > LastNumFilesWithNewSuffix;)
 				{
 					LastNumFilesWithNewSuffix = FilesRequiringSuffix.Count;
-					foreach (LinkedAction PrerequisiteAction in PrerequisiteActions)
+					foreach (LinkedAction PrerequisiteAction in PrerequisiteActions.Where(x => !TargetActionsToExecute.Contains(x) && x.PrerequisiteItems.Intersect(FilesRequiringSuffix).Any()))
 					{
-						if (!TargetActionsToExecute.Contains(PrerequisiteAction))
+						foreach (FileItem ProducedItem in PrerequisiteAction.ProducedItems)
 						{
-							foreach (FileItem ProducedItem in PrerequisiteAction.ProducedItems)
+							FileItem[]? DependentItems;
+							if (HotReloadItemToDependentItems.TryGetValue(ProducedItem, out DependentItems))
 							{
-								FileItem[]? DependentItems;
-								if (HotReloadItemToDependentItems.TryGetValue(ProducedItem, out DependentItems))
-								{
-									TargetActionsToExecute.Add(PrerequisiteAction);
-									FilesRequiringSuffix.UnionWith(DependentItems);
-								}
+								TargetActionsToExecute.Add(PrerequisiteAction);
+								FilesRequiringSuffix.UnionWith(DependentItems);
 							}
 						}
 					}
@@ -1056,13 +1053,9 @@ namespace UnrealBuildTool
 		/// </summary>
 		public static Dictionary<FileReference, FileReference> PatchActionGraph(IEnumerable<LinkedAction> Actions, Dictionary<FileReference, FileReference> OriginalFileToHotReloadFile)
 		{
-			// Gather all of the response files for link actions.  We're going to need to patch 'em up after we figure out new
+			// Gather all of the files for link action to be patched.  We're going to need to patch 'em up after we figure out new
 			// names for all of the output files and import libraries
-			List<string> ResponseFilePaths = new List<string>();
-
-			// Same as Response files but for all of the link.sh files for link actions.
-			// Only used on BuildHostPlatform Linux
-			List<string> LinkScriptFilePaths = new List<string>();
+			Dictionary<FileReference, FileReference> OriginalToNewFilePaths = [];
 
 			// Keep a map of the original file names and their new file names, so we can fix up response files after
 			Dictionary<string, string> OriginalFileNameAndNewFileNameList_NoExtensions = new Dictionary<string, string>();
@@ -1070,111 +1063,71 @@ namespace UnrealBuildTool
 			// Finally, we'll keep track of any file items that we had to create counterparts for change file names, so we can fix those up too
 			Dictionary<FileItem, FileItem> AffectedOriginalFileItemAndNewFileItemMap = new Dictionary<FileItem, FileItem>();
 
+			string ResponseFileExtension = UEToolChain.ResponseExt;
+			string ScriptExtension = (BuildHostPlatform.Current.Platform == UnrealTargetPlatform.Win64) ? ".bat" : ".sh";
+
 			foreach (LinkedAction Action in Actions.Where((Action) => Action.ActionType == ActionType.Link))
 			{
-				FileItem FirstProducedItem = Action.ProducedItems.First();
+				// Find the first produced item that needs to be renamed
+				FileItem? ProducedItem = Action.ProducedItems.FirstOrDefault(x => OriginalFileToHotReloadFile.ContainsKey(x.Location));
+				if (ProducedItem == null)
+				{
+					continue;
+				}
 
-				// Assume that the first produced item (with no extension) is our output file name
 				FileReference? HotReloadFile;
-				if (!OriginalFileToHotReloadFile.TryGetValue(FirstProducedItem.Location, out HotReloadFile))
+				if (!OriginalFileToHotReloadFile.TryGetValue(ProducedItem.Location, out HotReloadFile))
 				{
 					continue;
 				}
 
-				string OriginalFileNameWithoutExtension = Utils.GetFilenameWithoutAnyExtensions(FirstProducedItem.AbsolutePath);
+				string OriginalFileNameWithoutExtension = Utils.GetFilenameWithoutAnyExtensions(ProducedItem.AbsolutePath);
 				string NewFileNameWithoutExtension = Utils.GetFilenameWithoutAnyExtensions(HotReloadFile.FullName);
 
-				// Find the response file in the command line.  We'll need to make a copy of it with our new file name.
-				string ResponseFileExtension = UEToolChain.ResponseExt;
-				int ResponseExtensionIndex = Action.CommandArguments.IndexOf(ResponseFileExtension, StringComparison.InvariantCultureIgnoreCase);
-				if (ResponseExtensionIndex != -1)
-				{
-					int ResponseFilePathIndex = Action.CommandArguments.LastIndexOf("@\"", ResponseExtensionIndex);
-					if (ResponseFilePathIndex == -1)
-					{
-						throw new BuildException("Couldn't find response file path in action's command arguments when hot reloading");
-					}
-
-					string OriginalResponseFilePathWithoutExtension = Action.CommandArguments.Substring(ResponseFilePathIndex + 2, (ResponseExtensionIndex - ResponseFilePathIndex) - 2);
-					string OriginalResponseFilePath = OriginalResponseFilePathWithoutExtension + ResponseFileExtension;
-
-					string NewResponseFilePath = ReplaceBaseFileName(OriginalResponseFilePath, OriginalFileNameWithoutExtension, NewFileNameWithoutExtension);
-
-					// Copy the old response file to the new path
-					if (!String.Equals(OriginalResponseFilePath, NewResponseFilePath, StringComparison.OrdinalIgnoreCase))
-					{
-						File.Copy(OriginalResponseFilePath, NewResponseFilePath, overwrite: true);
-					}
-
-					// Keep track of the new response file name.  We'll have to do some edits afterwards.
-					ResponseFilePaths.Add(NewResponseFilePath);
-				}
-
 				// Duplicate the action
 				Action NewAction = new Action(Action);
 				Action.Inner = NewAction;
 
-				// Find the *.link.sh file in the command line.  We'll need to make a copy of it with our new file name.
-				// Only currently used on Linux
-				if (UEBuildPlatform.IsPlatformInGroup(BuildHostPlatform.Current.Platform, UnrealPlatformGroup.Unix))
+				// Update this action's list of prerequisite items too
+				List<FileItem> UpdatePrerequisiteItems = new List<FileItem>(NewAction.PrerequisiteItems);
+				for (int ItemIndex = 0; ItemIndex < UpdatePrerequisiteItems.Count; ++ItemIndex)
 				{
-					string LinkScriptFileExtension = ".link.sh";
-					int LinkScriptExtensionIndex = Action.CommandArguments.IndexOf(LinkScriptFileExtension, StringComparison.InvariantCultureIgnoreCase);
-					if (LinkScriptExtensionIndex != -1)
-					{
-						// We expect the script invocation to be quoted
-						int LinkScriptFilePathIndex = Action.CommandArguments.LastIndexOf("\"", LinkScriptExtensionIndex);
-						if (LinkScriptFilePathIndex == -1)
-						{
-							throw new BuildException("Couldn't find link script file path in action's command arguments when hot reloading. Is the path quoted?");
-						}
-
-						string OriginalLinkScriptFilePathWithoutExtension = Action.CommandArguments.Substring(LinkScriptFilePathIndex + 1, (LinkScriptExtensionIndex - LinkScriptFilePathIndex) - 1);
-						string OriginalLinkScriptFilePath = OriginalLinkScriptFilePathWithoutExtension + LinkScriptFileExtension;
+					FileItem OriginalPrerequisiteItem = UpdatePrerequisiteItems[ItemIndex];
+					string NewPrerequisiteItemFilePath = ReplaceBaseFileName(OriginalPrerequisiteItem.AbsolutePath, OriginalFileNameWithoutExtension, NewFileNameWithoutExtension);
 
-						string NewLinkScriptFilePath = ReplaceBaseFileName(OriginalLinkScriptFilePath, OriginalFileNameWithoutExtension, NewFileNameWithoutExtension);
-
-						// Copy the old response file to the new path
-						File.Copy(OriginalLinkScriptFilePath, NewLinkScriptFilePath, overwrite: true);
-
-						// Keep track of the new response file name.  We'll have to do some edits afterwards.
-						LinkScriptFilePaths.Add(NewLinkScriptFilePath);
-					}
-
-					// Update this action's list of prerequisite items too
-					List<FileItem> UpdatePrerequisiteItems = new List<FileItem>(NewAction.PrerequisiteItems);
-					for (int ItemIndex = 0; ItemIndex < UpdatePrerequisiteItems.Count; ++ItemIndex)
+					if (OriginalPrerequisiteItem.AbsolutePath != NewPrerequisiteItemFilePath)
 					{
-						FileItem OriginalPrerequisiteItem = UpdatePrerequisiteItems[ItemIndex];
-						string NewPrerequisiteItemFilePath = ReplaceBaseFileName(OriginalPrerequisiteItem.AbsolutePath, OriginalFileNameWithoutExtension, NewFileNameWithoutExtension);
+						// OK, the prerequisite item's file name changed so we'll update it to point to our new file
+						FileItem NewPrerequisiteItem = FileItem.GetItemByPath(NewPrerequisiteItemFilePath);
+						UpdatePrerequisiteItems[ItemIndex] = NewPrerequisiteItem;
+
+						// Keep track of it so we can fix up dependencies in a second pass afterwards
+						AffectedOriginalFileItemAndNewFileItemMap.Add(OriginalPrerequisiteItem, NewPrerequisiteItem);
 
-						if (OriginalPrerequisiteItem.AbsolutePath != NewPrerequisiteItemFilePath)
+						int ResponseExtensionIndex = OriginalPrerequisiteItem.AbsolutePath.IndexOf(ResponseFileExtension, StringComparison.InvariantCultureIgnoreCase);
+						if (ResponseExtensionIndex != -1)
 						{
-							// OK, the prerequisite item's file name changed so we'll update it to point to our new file
-							FileItem NewPrerequisiteItem = FileItem.GetItemByPath(NewPrerequisiteItemFilePath);
-							UpdatePrerequisiteItems[ItemIndex] = NewPrerequisiteItem;
+							string OriginalResponseFilePathWithoutExtension = OriginalPrerequisiteItem.AbsolutePath.Substring(0, ResponseExtensionIndex);
+							string OriginalResponseFilePath = OriginalResponseFilePathWithoutExtension + ResponseFileExtension;
 
-							// Keep track of it so we can fix up dependencies in a second pass afterwards
-							AffectedOriginalFileItemAndNewFileItemMap.Add(OriginalPrerequisiteItem, NewPrerequisiteItem);
+							string NewResponseFilePath = ReplaceBaseFileName(OriginalResponseFilePath, OriginalFileNameWithoutExtension, NewFileNameWithoutExtension);
 
-							ResponseExtensionIndex = OriginalPrerequisiteItem.AbsolutePath.IndexOf(ResponseFileExtension, StringComparison.InvariantCultureIgnoreCase);
-							if (ResponseExtensionIndex != -1)
-							{
-								string OriginalResponseFilePathWithoutExtension = OriginalPrerequisiteItem.AbsolutePath.Substring(0, ResponseExtensionIndex);
-								string OriginalResponseFilePath = OriginalResponseFilePathWithoutExtension + ResponseFileExtension;
+							OriginalToNewFilePaths.Add(new FileReference(OriginalResponseFilePath), new FileReference(NewResponseFilePath));
+						}
 
-								string NewResponseFilePath = ReplaceBaseFileName(OriginalResponseFilePath, OriginalFileNameWithoutExtension, NewFileNameWithoutExtension);
+						int ScriptExtensionIndex = OriginalPrerequisiteItem.AbsolutePath.IndexOf(ScriptExtension, StringComparison.InvariantCultureIgnoreCase);
+						if (ScriptExtensionIndex != -1)
+						{
+							string OriginalScriptFilePathWithoutExtension = OriginalPrerequisiteItem.AbsolutePath.Substring(0, ScriptExtensionIndex);
+							string OriginalScriptFilePath = OriginalScriptFilePathWithoutExtension + ScriptExtension;
 
-								// Copy the old response file to the new path
-								File.Copy(OriginalResponseFilePath, NewResponseFilePath, overwrite: true);
+							string NewScriptFilePath = ReplaceBaseFileName(OriginalScriptFilePath, OriginalFileNameWithoutExtension, NewFileNameWithoutExtension);
 
-								// Keep track of the new response file name.  We'll have to do some edits afterwards.
-								ResponseFilePaths.Add(NewResponseFilePath);
-							}
+							OriginalToNewFilePaths.Add(new FileReference(OriginalScriptFilePath), new FileReference(NewScriptFilePath));
 						}
 					}
-					NewAction.PrerequisiteItems = new SortedSet<FileItem>(UpdatePrerequisiteItems);
 				}
+				NewAction.PrerequisiteItems = new SortedSet<FileItem>(UpdatePrerequisiteItems);
 
 				// Update this action's list of produced items too
 				List<FileItem> UpdateProducedItems = new List<FileItem>(NewAction.ProducedItems);
@@ -1257,10 +1210,10 @@ namespace UnrealBuildTool
 					}
 				}
 
-				foreach (string ResponseFilePath in ResponseFilePaths)
+				foreach (KeyValuePair<FileReference, FileReference> Item in OriginalToNewFilePaths)
 				{
 					// Load the file up
-					string FileContents = Utils.ReadAllText(ResponseFilePath);
+					string FileContents = FileReference.ReadAllText(Item.Key);
 
 					// Replace all of the old file names with new ones
 					foreach (KeyValuePair<string, string> FileNameTuple in OriginalFileNameAndNewFileNameList_NoExtensions)
@@ -1271,29 +1224,8 @@ namespace UnrealBuildTool
 						FileContents = ReplaceBaseFileName(FileContents, OriginalFileNameWithoutExtension, NewFileNameWithoutExtension);
 					}
 
-					// Overwrite the original file
-					File.WriteAllText(ResponseFilePath, FileContents, new System.Text.UTF8Encoding(false));
-				}
-
-				if (UEBuildPlatform.IsPlatformInGroup(BuildHostPlatform.Current.Platform, UnrealPlatformGroup.Unix))
-				{
-					foreach (string LinkScriptFilePath in LinkScriptFilePaths)
-					{
-						// Load the file up
-						string FileContents = Utils.ReadAllText(LinkScriptFilePath);
-
-						// Replace all of the old file names with new ones
-						foreach (KeyValuePair<string, string> FileNameTuple in OriginalFileNameAndNewFileNameList_NoExtensions)
-						{
-							string OriginalFileNameWithoutExtension = FileNameTuple.Key;
-							string NewFileNameWithoutExtension = FileNameTuple.Value;
-
-							FileContents = ReplaceBaseFileName(FileContents, OriginalFileNameWithoutExtension, NewFileNameWithoutExtension);
-						}
-
-						// Overwrite the original file
-						File.WriteAllText(LinkScriptFilePath, FileContents, new System.Text.UTF8Encoding(false));
-					}
+					// Write out the updated file location
+					FileReference.WriteAllText(Item.Value, FileContents, new System.Text.UTF8Encoding(false));
 				}
 			}
 
diff --git a/Engine/Source/Programs/UnrealBuildTool/System/RulesCompiler.cs b/Engine/Source/Programs/UnrealBuildTool/System/RulesCompiler.cs
--- a/Engine/Source/Programs/UnrealBuildTool/System/RulesCompiler.cs
+++ b/Engine/Source/Programs/UnrealBuildTool/System/RulesCompiler.cs
@@ -156,6 +156,7 @@ namespace UnrealBuildTool
 			using (GlobalTracer.Instance.BuildSpan("Finding plugin modules").StartActive())
 			{
 				ModuleRulesContext PluginsModuleContext = new ModuleRulesContext(PluginsScope, RootDirectories[0]);
+				PluginsModuleContext.bCanHotReload = true;
 				FindModuleRulesForPlugins(Plugins, PluginsModuleContext, ModuleFileToContext);
 				// Plugin test target rules only
 				FindTestRulesForPlugins(Plugins, PluginsModuleContext, ModuleFileToContext, ProgramTargetFiles);
