Additional logging for XMLConfig file cache search paths.

diff --git a/Engine/Source/Programs/UnrealBuildTool/System/XmlConfig.cs b/Engine/Source/Programs/UnrealBuildTool/System/XmlConfig.cs
--- a/Engine/Source/Programs/UnrealBuildTool/System/XmlConfig.cs
+++ b/Engine/Source/Programs/UnrealBuildTool/System/XmlConfig.cs
@@ -184,6 +184,13 @@ namespace UnrealBuildTool
 							x => x.Value.Select(x => x.Value).ToArray()));
 					s_values.Write(CacheFile);
 				}
+				else
+				{
+					foreach (FileReference inputFile in s_values.InputFiles)
+					{
+						logger.LogDebug("Reading configuration file from cache ({CacheLocaltion}): {Location}", CacheFile.FullName, inputFile.FullName);
+					}
+				}
 			}
 
 			// Apply all the static field values
