Roll back Windows & Linux clang & verse-clang toolchains to LLVM 18.1.8

Updating to LLVM 19 has uncovered a UB issue in UE which will take some time to resolve. We're rolling back the C++ toolchains to 18 as a temporary workaround.

diff --git a/Engine/Build/BatchFiles/Linux/Toolchain/RunMe.bat b/Engine/Build/BatchFiles/Linux/Toolchain/RunMe.bat
--- a/Engine/Build/BatchFiles/Linux/Toolchain/RunMe.bat
+++ b/Engine/Build/BatchFiles/Linux/Toolchain/RunMe.bat
@@ -1,9 +1,9 @@
 @REM @echo off
 set BUILD_ONLY=0
-set TOOLCHAIN_VERSION=v24
-set LLVM_VERSION=19.1.7
-set LLVM_BRANCH=release/19.x
-set LLVM_TAG=llvmorg-19.1.7
+set TOOLCHAIN_VERSION=v25
+set LLVM_VERSION=18.1.0
+set LLVM_BRANCH=release/18.x
+set LLVM_TAG=llvmorg-18.1.0
 
 set CMAKE_BINARY=%CD%\..\..\..\..\Extras\ThirdPartyNotUE\CMake\bin\cmake.exe
 set PYTHON_BINARY=%CD%\..\..\..\..\Binaries\ThirdParty\Python\Win64\python.exe
@@ -78,6 +78,7 @@ for %%G in (aarch64-unknown-linux-gnueabi x86_64-unknown-linux-gnu) do (
     copy "install\bin\ld.lld.exe" OUTPUT\%%G\bin
     copy "install\bin\lld.exe" OUTPUT\%%G\bin
     copy "install\bin\llvm-ar.exe" OUTPUT\%%G\bin
+    copy "install\bin\llvm-ranlib.exe" OUTPUT\%%G\bin
     copy "install\bin\llvm-profdata.exe" OUTPUT\%%G\bin
     copy "install\bin\llvm-symbolizer.exe" OUTPUT\%%G\bin
     copy "install\bin\llvm-objcopy.exe" OUTPUT\%%G\bin
@@ -111,4 +112,4 @@ popd
 
 :endscript
 echo.
-echo Done.
+echo Done.
\ No newline at end of file
diff --git a/Engine/Build/BatchFiles/Linux/Toolchain/RunMe.sh b/Engine/Build/BatchFiles/Linux/Toolchain/RunMe.sh
--- a/Engine/Build/BatchFiles/Linux/Toolchain/RunMe.sh
+++ b/Engine/Build/BatchFiles/Linux/Toolchain/RunMe.sh
@@ -34,9 +34,6 @@ if ! [ -f "/.dockerenv" ]; then
   #echo Removing ${ImageName}...
   docker rm ${ImageName}
 
-  # Copy LibC++ files into place
-  update_libc++.sh
-
 else
 
 
diff --git a/Engine/Build/BatchFiles/Linux/Toolchain/build_linux_toolchain.sh b/Engine/Build/BatchFiles/Linux/Toolchain/build_linux_toolchain.sh
--- a/Engine/Build/BatchFiles/Linux/Toolchain/build_linux_toolchain.sh
+++ b/Engine/Build/BatchFiles/Linux/Toolchain/build_linux_toolchain.sh
@@ -3,9 +3,9 @@
 set -x
 set -eu
 
-ToolChainVersion=v24
-LLVM_VERSION_MAJOR=19
-LLVM_VERSION=${LLVM_VERSION_MAJOR}.1.7
+ToolChainVersion=v25
+LLVM_VERSION_MAJOR=18
+LLVM_VERSION=${LLVM_VERSION_MAJOR}.1.0
 LLVM_BRANCH=release/${LLVM_VERSION_MAJOR}.x
 LLVM_TAG=llvmorg-${LLVM_VERSION}
 
@@ -171,6 +171,7 @@ if [ ! -f copy-linux-toolchain.done ]; then
 		cp -L ${InstallClangDir}/bin/lld             ${OutputDirLinux}/$arch/bin/
 		cp -L ${InstallClangDir}/bin/ld.lld          ${OutputDirLinux}/$arch/bin/
 		cp -L ${InstallClangDir}/bin/llvm-ar         ${OutputDirLinux}/$arch/bin/
+		cp -L ${InstallClangDir}/bin/llvm-ranlib     ${OutputDirLinux}/$arch/bin/
 		cp -L ${InstallClangDir}/bin/llvm-profdata   ${OutputDirLinux}/$arch/bin/
 		cp -L ${InstallClangDir}/bin/llvm-objcopy    ${OutputDirLinux}/$arch/bin/
 		cp -L ${InstallClangDir}/bin/llvm-symbolizer ${OutputDirLinux}/$arch/bin/
@@ -180,9 +181,10 @@ if [ ! -f copy-linux-toolchain.done ]; then
 			cp -r -L ${InstallClangDir}/lib/clang ${OutputDirLinux}/$arch/lib/
 
 			# copy libc++.a and friends to lib64.  This will need to move out to where the libc++ stuff currently lives in the tree
-			mkdir -p install-libc++-$arch/{lib,include/c++}
-			cp -L ${InstallClangDir}/lib/$arch/* install-libc++-$arch/lib
-			cp -r -L ${InstallClangDir}/include/$arch/c++/v1 install-libc++-$arch/include/c++
+			cp -L ${InstallClangDir}/lib/$arch/* ${OutputDirLinux}/$arch/lib64
+			# __config_site ends up in a different dir for x86_64, make sure to grab it and the regular includes
+			cp -r -L ${InstallClangDir}/include/$arch/c++/v1 ${OutputDirLinux}/$arch/include/c++
+			cp -r -L ${InstallClangDir}/include/c++/v1 ${OutputDirLinux}/$arch/include/c++
 		fi
 	done
 	touch copy-linux-toolchain.done
@@ -266,6 +268,10 @@ for arch in $TARGETS; do
 
 		make -j$CORES && make install
 
+	echo "Copying libc++ for ${arch}..."
+	cp -L ../install-libc++-${arch}/lib/* ${OutputDirLinux}/$arch/lib64
+	cp -r -L ../install-libc++-${arch}/include/c++/v1 ${OutputDirLinux}/$arch/include/c++
+
 	popd
 
 done
@@ -303,6 +309,10 @@ if [ ! -f copy-windows-toolchain.done ]; then
 			# Copy compiler-rt
 			cp -r -L ${OutputDirLinux}/$arch/lib/clang lib/
 
+			# copy libc++ from the linux builds so that we don't need to build it again in windows
+			cp -r -L ${OutputDirLinux}/$arch/include/c++/v1 include/c++
+			cp -r -L ${OutputDirLinux}/$arch/lib64/libc++* lib64
+
 			# Copy linux llvm-symbolizer so that we can resolve ASan hits
 			cp -r -L ${OutputDirLinux}/$arch/bin/llvm-symbolizer bin
 
@@ -366,7 +376,5 @@ pushd ${OutputDirWindows}
 	zip -r /src/build/${ToolChainVersionName}-windows.zip *
 popd
 
-
-echo "Don't forget to run update_libc++.sh and perform a p4 reconcile in Engine/Source/ThirdParty/Unix/LibCxx..."
-
 echo done.
+
diff --git a/Engine/Config/Linux/Linux_SDK.json b/Engine/Config/Linux/Linux_SDK.json
--- a/Engine/Config/Linux/Linux_SDK.json
+++ b/Engine/Config/Linux/Linux_SDK.json
@@ -1,7 +1,7 @@
 {
-	"MainVersion" : "v24_clang-19.1.7-rockylinux8",
-	"MinVersion" : "v24_clang-19.1.7-rockylinux8",
-	"MaxVersion" : "v24_clang-19.1.7-rockylinux8",
+	"MainVersion" : "v25_clang-18.1.0-rockylinux8",
+	"MinVersion" : "v25_clang-18.1.0-rockylinux8",
+	"MaxVersion" : "v25_clang-18.1.0-rockylinux8",
 
 	"AutoSDKPlatform" : "Linux_x64"
 }
diff --git a/Engine/Config/Windows/Windows_SDK.json b/Engine/Config/Windows/Windows_SDK.json
--- a/Engine/Config/Windows/Windows_SDK.json
+++ b/Engine/Config/Windows/Windows_SDK.json
@@ -6,8 +6,8 @@
 	"MinSoftwareVersion": "10.0.19041.0",
 
 	"PreferredClangVersions": [
-		"19.1.0-19.999",
-		"18.1.3-18.999"
+		"18.1.3-18.999",
+		"19.1.0-19.999"
 	],
 	"MinimumClangVersion": "18.1.3",
 
diff --git a/Engine/Source/Programs/UnrealBuildTool/Platform/Linux/LinuxPlatformSDK.cs b/Engine/Source/Programs/UnrealBuildTool/Platform/Linux/LinuxPlatformSDK.cs
--- a/Engine/Source/Programs/UnrealBuildTool/Platform/Linux/LinuxPlatformSDK.cs
+++ b/Engine/Source/Programs/UnrealBuildTool/Platform/Linux/LinuxPlatformSDK.cs
@@ -92,6 +92,11 @@ namespace UnrealBuildTool
 		/// </summary>
 		private int bForceUseSystemCompiler = -1;
 
+		/// <summary>
+		/// Force using legacy LibCxx in Engine/Source/ThirdParty/Unix
+		/// </summary>
+		private int bForceUseLegacyLibCxx = -1;
+
 		/// <summary>
 		/// Whether to compile with the verbose flag
 		/// </summary>
@@ -181,6 +186,30 @@ namespace UnrealBuildTool
 			return bForceUseSystemCompiler == 1;
 		}
 
+		/// <summary>
+		/// Force use of Legacy LibCxx in Engine/Source/ThirdParty/Unix, which will be deprecated in the future
+		/// </summary>
+		public virtual bool ForceUseLegacyLibCxx()
+		{
+			// by default tools chains don't parse arguments, but we want to be able to check the -bForceUseSystemCompiler flag.
+			if (bForceUseLegacyLibCxx == -1)
+			{
+				bForceUseLegacyLibCxx = 0;
+				string[] CmdLine = Environment.GetCommandLineArgs();
+
+				foreach (string CmdLineArg in CmdLine)
+				{
+					if (CmdLineArg.Equals("-ForceUseLegacyLibCxx", StringComparison.OrdinalIgnoreCase))
+					{
+						bForceUseLegacyLibCxx = 1;
+						break;
+					}
+				}
+			}
+
+			return bForceUseLegacyLibCxx == 1;
+		}
+
 		/// <summary>
 		/// Returns the root SDK path for all architectures
 		/// WARNING: Do not cache this value - it may be changed after sourcing OutputEnvVars.txt
diff --git a/Engine/Source/Programs/UnrealBuildTool/Platform/Linux/LinuxToolChain.cs b/Engine/Source/Programs/UnrealBuildTool/Platform/Linux/LinuxToolChain.cs
--- a/Engine/Source/Programs/UnrealBuildTool/Platform/Linux/LinuxToolChain.cs
+++ b/Engine/Source/Programs/UnrealBuildTool/Platform/Linux/LinuxToolChain.cs
@@ -55,7 +55,7 @@ namespace UnrealBuildTool
 			: this(UnrealTargetPlatform.Linux, InArchitecture, InSDK, InOptions, InLogger)
 		{
 			// prevent unknown clangs since the build is likely to fail on too old or too new compilers
-			if ((CompilerVersionLessThan(19, 0, 0) || CompilerVersionGreaterOrEqual(20, 0, 0)) && !Options.HasFlag(ClangToolChainOptions.UseAutoRTFMCompiler))
+			if ((CompilerVersionLessThan(18, 0, 0) || CompilerVersionGreaterOrEqual(20, 0, 0)) && !Options.HasFlag(ClangToolChainOptions.UseAutoRTFMCompiler))
 			{
 				throw new BuildException(
 					String.Format("This version of the Unreal Engine can only be compiled with clang 18.x. clang {0} may not build it - please use a different version.",
@@ -509,8 +509,24 @@ namespace UnrealBuildTool
 			if (ShouldUseLibcxx())
 			{
 				Arguments.Add("-nostdinc++");
-				Arguments.Add(GetSystemIncludePathArgument(DirectoryReference.Combine(Unreal.EngineSourceDirectory, "ThirdParty", "Unix", "LibCxx", "include"), CompileEnvironment));
-				Arguments.Add(GetSystemIncludePathArgument(DirectoryReference.Combine(Unreal.EngineSourceDirectory, "ThirdParty", "Unix", "LibCxx", "include", "c++", "v1"), CompileEnvironment));
+				if (PlatformSDK.ForceUseLegacyLibCxx())
+				{
+					Arguments.Add(GetSystemIncludePathArgument(DirectoryReference.Combine(Unreal.EngineSourceDirectory, "ThirdParty", "Unix", "LibCxx", "include"), CompileEnvironment));
+					Arguments.Add(GetSystemIncludePathArgument(DirectoryReference.Combine(Unreal.EngineSourceDirectory, "ThirdParty", "Unix", "LibCxx", "include", "c++", "v1"), CompileEnvironment));
+				}
+				else
+				{
+					if (!PlatformSDK.ForceUseSystemCompiler())
+					{
+						// LinuxInfo.BaseLinuxPath should only be null in the case that we've been forced to use the system compiler.
+						// In that case, project build.cs files should be adjusted to include the relevant directories to the system include paths
+						if(LinuxInfo.BaseLinuxPath != null)
+						{
+							Arguments.Add(GetSystemIncludePathArgument(DirectoryReference.Combine(LinuxInfo.BaseLinuxPath, "include"), CompileEnvironment));
+							Arguments.Add(GetSystemIncludePathArgument(DirectoryReference.Combine(LinuxInfo.BaseLinuxPath, "include", "c++", "v1"), CompileEnvironment));
+						}
+					}
+				}
 			}
 
 			if (CompilerVersionGreaterOrEqual(12, 0, 0))
@@ -1351,9 +1367,13 @@ namespace UnrealBuildTool
 			{
 				// libc++ and its abi lib
 				LinkCommandString += " -nodefaultlibs";
-				LinkCommandString += " -L" + "ThirdParty/Unix/LibCxx/lib/Unix/" + LinkEnvironment.Architecture.LinuxName + "/";
-				LinkCommandString += " " + "ThirdParty/Unix/LibCxx/lib/Unix/" + LinkEnvironment.Architecture.LinuxName + "/libc++.a";
-				LinkCommandString += " " + "ThirdParty/Unix/LibCxx/lib/Unix/" + LinkEnvironment.Architecture.LinuxName + "/libc++abi.a";
+				string LibCxxLibDir = LinuxInfo.BaseLinuxPath + "/lib64/";
+				if(PlatformSDK.ForceUseLegacyLibCxx())
+				{
+					LibCxxLibDir = "ThirdParty/Unix/LibCxx/lib/Unix/" + LinkEnvironment.Architecture.LinuxName + "/";
+				}
+				LinkCommandString += " " +  LibCxxLibDir + "libc++.a";
+				LinkCommandString += " " +  LibCxxLibDir + "libc++abi.a";
 				LinkCommandString += " -lm";
 				LinkCommandString += " -lc";
 				LinkCommandString += " -lpthread"; // pthread_mutex_trylock is missing from libc stubs
