UnrealBuildTool: Respect -Log= when generating project files

#rnx

diff --git a/Engine/Source/Programs/UnrealBuildTool/Modes/GenerateProjectFilesMode.cs b/Engine/Source/Programs/UnrealBuildTool/Modes/GenerateProjectFilesMode.cs
--- a/Engine/Source/Programs/UnrealBuildTool/Modes/GenerateProjectFilesMode.cs
+++ b/Engine/Source/Programs/UnrealBuildTool/Modes/GenerateProjectFilesMode.cs
@@ -74,8 +74,15 @@ namespace UnrealBuildTool
 			}
 
 			// set up logging (taken from BuildMode)
-			FileReference LogFile = FileReference.Combine(Unreal.EngineProgramSavedDirectory, "UnrealBuildTool", "Log_GPF.txt");
-			Log.AddFileWriter("DefaultLogTraceListener", LogFile);
+			if (!Log.HasFileWriter())
+			{
+				FileReference LogFile = FileReference.Combine(Unreal.EngineProgramSavedDirectory, "UnrealBuildTool", "Log_GPF.txt");
+				Log.AddFileWriter("DefaultLogTraceListener", LogFile);
+			}
+			else
+			{
+				Log.RemoveStartupTraceListener();
+			}
 
 			// Parse rocket-specific arguments.
 			FileReference? ProjectFile;
