UnrealBuildTool: If an AutoSDK path contains Version.txt, add it to the UBA path hash

#rnx

diff --git a/Engine/Source/Programs/UnrealBuildTool/ToolChain/ClangToolChain.cs b/Engine/Source/Programs/UnrealBuildTool/ToolChain/ClangToolChain.cs
--- a/Engine/Source/Programs/UnrealBuildTool/ToolChain/ClangToolChain.cs
+++ b/Engine/Source/Programs/UnrealBuildTool/ToolChain/ClangToolChain.cs
@@ -254,7 +254,15 @@ namespace UnrealBuildTool
 				GlobalCompileEnvironment.RootPaths[CppRootPathFolder.PlatformSDK] = LazyInfo.Value.BasePath;
 				GlobalLinkEnvironment.RootPaths[CppRootPathFolder.PlatformSDK] = LazyInfo.Value.BasePath;
 
-				EpicGames.UBA.Utils.RegisterPathHash(LazyInfo.Value.BasePath.FullName, LazyInfo.Value.ClangVersion.ToString());
+				string sdkVersion = UEBuildPlatformSDK.GetSDKForPlatform(Target.Platform.ToString())?.GetInstalledVersion() ?? "0";
+				string pathHash = $"{Target.Platform}-{sdkVersion}-{LazyInfo.Value.ClangVersion}";
+				FileReference versionFile = FileReference.Combine(LazyInfo.Value.BasePath, "..", "Version.txt"); // Optional AutoSDK version, may not exist
+				if (FileReference.Exists(versionFile))
+				{
+					string contents = FileReference.ReadAllText(versionFile).ReplaceLineEndings("-").Trim();
+					pathHash = $"{pathHash}-{contents}";
+				}
+				EpicGames.UBA.Utils.RegisterPathHash(LazyInfo.Value.BasePath.FullName, pathHash);
 			}
 
 			PreprocessDepends = Target.bPreprocessDepends;
