UnrealBuildTool: Automatically disable pgo optimize if .profdata was generated from a newer major version of llvm

#rnx

diff --git a/Engine/Source/Programs/UnrealBuildTool/Platform/Windows/UEBuildWindows.cs b/Engine/Source/Programs/UnrealBuildTool/Platform/Windows/UEBuildWindows.cs
--- a/Engine/Source/Programs/UnrealBuildTool/Platform/Windows/UEBuildWindows.cs
+++ b/Engine/Source/Programs/UnrealBuildTool/Platform/Windows/UEBuildWindows.cs
@@ -2048,6 +2048,29 @@ namespace UnrealBuildTool
 
 			CompileEnvironment.Definitions.Add("WINDOWS_MAX_NUM_TLS_SLOTS=" + Target.WindowsPlatform.MaxNumTlsSlots.ToString());
 			CompileEnvironment.Definitions.Add("WINDOWS_MAX_NUM_THREADS_WITH_TLS_SLOTS=" + Target.WindowsPlatform.MaxNumThreadsWithTlsSlots.ToString());
+
+			// Validate PGO data
+			if (CompileEnvironment.bPGOOptimize && CompileEnvironment.PGODirectory != null && Target.WindowsPlatform.Compiler.IsClang() && Target.WindowsPlatform.bIgnoreStalePGOData)
+			{
+				string CompilerVersionFilename = Path.Combine(CompileEnvironment.PGODirectory!, "PGOProfileCompilerInfo.txt");
+				if (File.Exists(CompilerVersionFilename) && Target.WindowsPlatform.Environment != null)
+				{
+					string[] ProfileCompilerVersionLines = File.ReadAllLines(CompilerVersionFilename);
+					if (ProfileCompilerVersionLines.Length >= 2 & Enum.TryParse(ProfileCompilerVersionLines[0].Trim(), false, out WindowsCompiler ProfileCompiler) && VersionNumber.TryParse(ProfileCompilerVersionLines[1].Trim(), out VersionNumber? ProfileVersion))
+					{
+						VersionNumber CompilerVersion = Target.WindowsPlatform.Compiler.IsIntel() ? MicrosoftPlatformSDK.GetClangVersionForIntelCompiler(Target.WindowsPlatform.Environment.CompilerPath) : Target.WindowsPlatform.Environment.CompilerVersion;
+						if (ProfileCompiler.IsClang() && ProfileVersion.Components[0] > CompilerVersion.Components[0])
+						{
+							Logger.LogWarning("PGO Profile major LLVM version is newer than PGO Optimize LLVM major version. Profile Guided Optimization will be disabled for {Config} until new PGO Profile data is generated.", Target.Configuration.ToString());
+							Logger.LogInformation("Profile: {ProfileCompiler} ({ProfileVersion}) Optimize: {OptimizeCompiler} ({OptimizeVersion})", ProfileCompiler, ProfileVersion, Target.WindowsPlatform.Compiler, CompilerVersion);
+							CompileEnvironment.bPGOOptimize = false;
+							LinkEnvironment.bPGOOptimize = false;
+							CompileEnvironment.bAllowLTCG = true;
+							LinkEnvironment.bAllowLTCG = true;
+						}
+					}
+				}
+			}
 		}
 
 		/// <summary>
diff --git a/Engine/Source/Programs/UnrealBuildTool/Platform/Windows/VCToolChain.cs b/Engine/Source/Programs/UnrealBuildTool/Platform/Windows/VCToolChain.cs
--- a/Engine/Source/Programs/UnrealBuildTool/Platform/Windows/VCToolChain.cs
+++ b/Engine/Source/Programs/UnrealBuildTool/Platform/Windows/VCToolChain.cs
@@ -2294,6 +2294,14 @@ namespace UnrealBuildTool
 
 		public override void FinalizeOutput(ReadOnlyTargetRules Target, TargetMakefileBuilder MakefileBuilder)
 		{
+			if (Target.bPGOOptimize)
+			{
+				// Win64 PGO folder is Windows, the rest match the platform name
+				string PGOPlatform = Target.Platform == UnrealTargetPlatform.Win64 ? "Windows" : Target.Platform.ToString();
+				DirectoryReference PGODirectory = DirectoryReference.Combine(Target.ProjectFile?.Directory ?? Unreal.WritableEngineDirectory, "Platforms", PGOPlatform, "Build", "PGO");
+				MakefileBuilder.Makefile.InternalDependencies.Add(FileItem.GetItemByFileReference(FileReference.Combine(PGODirectory, "PGOProfileCompilerInfo.txt")));
+			}
+
 			if (Target.bPrintToolChainTimingInfo || Target.WindowsPlatform.bCompilerTrace)
 			{
 				TargetMakefile Makefile = MakefileBuilder.Makefile;
@@ -3377,7 +3385,8 @@ namespace UnrealBuildTool
 			// We may want to use this to ensure compatibility
 			if (bPGOProfile)
 			{
-				string[] CompilerVersionLines = { EnvVars.Compiler.ToString(), EnvVars.CompilerVersion.ToString() };
+				VersionNumber CompilerVersion = EnvVars.Compiler.IsIntel() ? MicrosoftPlatformSDK.GetClangVersionForIntelCompiler(EnvVars.CompilerPath) : EnvVars.CompilerVersion;
+				string[] CompilerVersionLines = { EnvVars.Compiler.ToString(), CompilerVersion.ToString() };
 				string CompilerVersionFilename = Path.Combine(LinkEnvironment.PGODirectory!, "PGOProfileCompilerInfo.txt");
 				Utils.WriteFileIfChanged(new FileReference(CompilerVersionFilename), CompilerVersionLines, Logger);
 			}
@@ -3628,6 +3637,28 @@ namespace UnrealBuildTool
 			{
 				EpicGames.UBA.Utils.RegisterPathHash(EnvVars.ToolChainDir.FullName, EnvVars.ToolChainVersion.ToString());
 			}
+
+			// Validate PGO data
+			if (EnvVars.Compiler.IsClang() && GlobalLinkEnvironment.bPGOOptimize && Target.WindowsPlatform.bIgnoreStalePGOData)
+			{
+				string CompilerVersionFilename = Path.Combine(GlobalLinkEnvironment.PGODirectory!, "PGOProfileCompilerInfo.txt");
+				if (File.Exists(CompilerVersionFilename))
+				{
+					string[] ProfileCompilerVersionLines = File.ReadAllLines(CompilerVersionFilename);
+					if (ProfileCompilerVersionLines.Length >= 2 & Enum.TryParse(ProfileCompilerVersionLines[0].Trim(), false, out WindowsCompiler ProfileCompiler) && VersionNumber.TryParse(ProfileCompilerVersionLines[1].Trim(), out VersionNumber? ProfileVersion))
+					{
+						VersionNumber CompilerVersion = EnvVars.Compiler.IsIntel() ? MicrosoftPlatformSDK.GetClangVersionForIntelCompiler(EnvVars.CompilerPath) : EnvVars.CompilerVersion;
+						if (ProfileCompiler.IsClang() && ProfileVersion.Components[0] > CompilerVersion.Components[0])
+						{
+							Logger.LogWarning("PGO Profile major LLVM version is newer than PGO Optimize LLVM major version. Profile Guided Optimization will be disabled for {Config} until new PGO Profile data is generated.", Target.Configuration.ToString());
+							Logger.LogInformation("Profile: {ProfileCompiler} ({ProfileVersion}) Optimize: {OptimizeCompiler} ({OptimizeVersion})", ProfileCompiler, ProfileVersion, EnvVars.Compiler, CompilerVersion);
+							GlobalCompileEnvironment.bPGOOptimize = false;
+							GlobalLinkEnvironment.bPGOOptimize = false;
+							throw new BuildException("test");
+						}
+					}
+				}
+			}
 		}
 
 		public override void ModifyBuildProducts(ReadOnlyTargetRules Target, UEBuildBinary Binary, IEnumerable<string> Libraries, IEnumerable<UEBuildBundleResource> BundleResources, Dictionary<FileReference, BuildProductType> BuildProducts)
