Made Android NDK toolchains detection for manual SDK installs more robust. Made JAVA_HOME environment setting for manual SDK installs more robust. Added missing setting of ndkPath to gradle property value

#android
#tests Tested changing ndks and building

diff --git a/Engine/Source/Programs/UnrealBuildTool/Platform/Android/AndroidPlatformSDK.cs b/Engine/Source/Programs/UnrealBuildTool/Platform/Android/AndroidPlatformSDK.cs
--- a/Engine/Source/Programs/UnrealBuildTool/Platform/Android/AndroidPlatformSDK.cs
+++ b/Engine/Source/Programs/UnrealBuildTool/Platform/Android/AndroidPlatformSDK.cs
@@ -137,9 +137,14 @@ namespace UnrealBuildTool
 				return null;
 			}
 
+			return GetNdkNameFromDirectory(NDKPath);
+		}
+
+		static string? GetNdkNameFromDirectory(string Directory)
+		{
 			// figure out the NDK version
 			string? NDKToolchainVersion = null;
-			string SourcePropFilename = Path.Combine(NDKPath, "source.properties");
+			string SourcePropFilename = Path.Combine(Directory, "source.properties");
 			if (File.Exists(SourcePropFilename))
 			{
 				string RevisionString = "";
@@ -164,7 +169,7 @@ namespace UnrealBuildTool
 			}
 			else
 			{
-				string ReleaseFilename = Path.Combine(NDKPath, "RELEASE.TXT");
+				string ReleaseFilename = Path.Combine(Directory, "RELEASE.TXT");
 				if (File.Exists(ReleaseFilename))
 				{
 					string[] PropertyContents = File.ReadAllLines(SourcePropFilename);
@@ -215,18 +220,13 @@ namespace UnrealBuildTool
 			}
 		}
 
-		string NdkVersionToName(string[] Parts)
-		{
-			return Parts[1] == "0" ? $"r{Parts[0]}" : $"r{Parts[0]}{Convert.ToChar('a' + int.Parse(Parts[1]))}";
-		}
-		
 		public override string[] GetAllInstalledSDKVersions()
 		{
 			string StudioSdkDir = GetStudioSdkDir();
 			
 			if (Directory.Exists(Path.Combine(StudioSdkDir, "platforms", GetPlatformSpecificVersion("platforms"))))
 			{
-				return Directory.EnumerateDirectories(Path.Combine(StudioSdkDir, "ndk")).Select(Directory => NdkVersionToName(Path.GetFileName(Directory).Split("."))).Distinct().ToArray();
+				return Directory.EnumerateDirectories(Path.Combine(StudioSdkDir, "ndk")).Select(GetNdkNameFromDirectory).Where(Directory => Directory != null).Distinct().ToArray()!;
 			}
 			
 			return [];
@@ -241,21 +241,42 @@ namespace UnrealBuildTool
 			
 			if (string.IsNullOrEmpty(CurrentSdkDir) || !Directory.Exists(Path.Combine(CurrentSdkDir, "platforms", GetPlatformSpecificVersion("platforms"))) || GetInstalledSDKVersion() != SelectedSDKVersion)
 			{
-				string StudioDir = GetStudioDir();
-				string StudioSdkDir = GetStudioSdkDir();
-				string NdkBaseDir = Path.Combine(StudioSdkDir, "ndk");
-				string NdkDir = Path.Combine(NdkBaseDir, string.Join(".", Directory.EnumerateDirectories(NdkBaseDir).Select(Directory => Path.GetFileName(Directory).Split(".")).Where(Parts => NdkVersionToName(Parts) == SelectedSDKVersion).OrderByDescending(Parts => int.Parse(Parts[2])).First()));
+				try
+				{
+					string StudioSdkDir = GetStudioSdkDir();
+					string NdkDir = Directory.EnumerateDirectories(Path.Combine(StudioSdkDir, "ndk")).First(Directory => GetNdkNameFromDirectory(Directory) == SelectedSDKVersion);
+
+					EnvVarValues = new()
+					{
+						["ANDROID_HOME"] = StudioSdkDir,
+						["ANDROID_SDK_HOME"] = "",
+						["ANDROID_NDK_ROOT"] = NdkDir,
+						["NDK_ROOT"] = NdkDir,
+						["NDKROOT"] = NdkDir
+					};
 
-				EnvVarValues = new()
+					string StudioDir = GetStudioDir();
+					string JavaDir = Path.Combine(StudioDir, "jbr");
+
+					if (!Directory.Exists(JavaDir))
+					{
+						JavaDir = Path.Combine(StudioDir, "jre");
+
+						if (!Directory.Exists(JavaDir))
+						{
+							return true;
+						}
+					}
+
+					EnvVarValues["JAVA_HOME"] = JavaDir;
+					return true;
+				}
+				catch (DirectoryNotFoundException)
 				{
-					["ANDROID_HOME"] = StudioSdkDir,
-					["ANDROID_SDK_HOME"] = "",
-					["ANDROID_NDK_ROOT"] = NdkDir,
-					["NDK_ROOT"] = NdkDir,
-					["NDKROOT"] = NdkDir,
-					["JAVA_HOME"] = Path.Combine(StudioDir, int.Parse(GetVersionFromConfig("AutoSDKDirectory")![1..]) <= 25 ? "jre" : "jbr"),
-				};
-				return true;
+				}
+				catch (InvalidOperationException)
+				{
+				}
 			}
 
 			EnvVarValues = null;
