- Added the missing fields to copy over in LinkEnvironment when making multiple architecture LinkEnvs.

diff --git a/Engine/Source/Programs/UnrealBuildTool/System/LinkEnvironment.cs b/Engine/Source/Programs/UnrealBuildTool/System/LinkEnvironment.cs
--- a/Engine/Source/Programs/UnrealBuildTool/System/LinkEnvironment.cs
+++ b/Engine/Source/Programs/UnrealBuildTool/System/LinkEnvironment.cs
@@ -391,12 +391,14 @@ namespace UnrealBuildTool
 			Platform = Other.Platform;
 			Configuration = Other.Configuration;
 			Architectures = Other.Architectures;
+			LinkType = Other.LinkType;
 			RootPaths = new(Other.RootPaths);
 			BundleDirectory = Other.BundleDirectory;
 			OutputDirectory = Other.OutputDirectory;
 			IntermediateDirectory = Other.IntermediateDirectory;
 			LocalShadowDirectory = Other.LocalShadowDirectory;
 			OutputFilePaths = Other.OutputFilePaths.ToList();
+			RuntimeDependencies.AddRange(Other.RuntimeDependencies);
 			Libraries.AddRange(Other.Libraries);
 			SystemLibraries.AddRange(Other.SystemLibraries);
 			SystemLibraryPaths.AddRange(Other.SystemLibraryPaths);
