LowLevelTests: Disabled AutoSDK for test targets

diff --git a/Engine/Source/Programs/UnrealBuildTool/Configuration/TestTargetRules.cs b/Engine/Source/Programs/UnrealBuildTool/Configuration/TestTargetRules.cs
--- a/Engine/Source/Programs/UnrealBuildTool/Configuration/TestTargetRules.cs
+++ b/Engine/Source/Programs/UnrealBuildTool/Configuration/TestTargetRules.cs
@@ -329,6 +329,7 @@ namespace UnrealBuildTool
 				GlobalDefinitions.AddRange(TestedTarget.GlobalDefinitions);
 			}
 
+			GlobalDefinitions.Add("AUTOSDKS_ENABLED=0");
 			GlobalDefinitions.Add("STATS=0");
 			GlobalDefinitions.Add("TEST_FOR_VALID_FILE_SYSTEM_MEMORY=0");
 
