Fix boot-up crashes in merged modules by excluding invalid dependencies

#rnx

diff --git a/Engine/Source/Programs/UnrealBuildTool/Configuration/UEBuildTarget.cs b/Engine/Source/Programs/UnrealBuildTool/Configuration/UEBuildTarget.cs
--- a/Engine/Source/Programs/UnrealBuildTool/Configuration/UEBuildTarget.cs
+++ b/Engine/Source/Programs/UnrealBuildTool/Configuration/UEBuildTarget.cs
@@ -3569,7 +3569,13 @@ namespace UnrealBuildTool
 				}
 			}
 
-			void GetPluginDependencies(UEBuildPlugin plugin, ref HashSet<UEBuildPlugin> dependencies, bool recursively = false)
+			bool IsGameFeaturePlugin(UEBuildPlugin plugin)
+			{
+				DirectoryReference gfpDir = DirectoryReference.Combine(ProjectDirectory, "Plugins", "GameFeatures");
+				return plugin.File.IsUnderDirectory(gfpDir);
+			}
+
+			void GetPluginDependencies(UEBuildPlugin plugin, ref HashSet<UEBuildPlugin> dependencies, bool recursively = false, bool checkAllPlugins = false)
 			{
 				foreach (PluginReferenceDescriptor dependency in plugin.Info.Descriptor.Plugins ?? new())
 				{
@@ -3579,12 +3585,12 @@ namespace UnrealBuildTool
 						if (dependentPluginCandidates.Any())
 						{
 							UEBuildPlugin dependentPlugin = dependentPluginCandidates.First();
-							if (!dependencies.Contains(dependentPlugin))
+							if (!dependencies.Contains(dependentPlugin) && (IsGameFeaturePlugin(dependentPlugin) || checkAllPlugins))
 							{
 								dependencies.Add(dependentPlugin);
 								if (recursively)
 								{
-									GetPluginDependencies(dependentPlugin, ref dependencies, recursively);
+									GetPluginDependencies(dependentPlugin, ref dependencies, recursively, checkAllPlugins);
 								}
 							}
 						}
@@ -3657,22 +3663,32 @@ namespace UnrealBuildTool
 				}
 			}
 
-			void RemoveNonExclusivePlugins(ref Dictionary<string, HashSet<UEBuildPlugin>> pluginGroups, ref HashSet<UEBuildPlugin> conflictingPlugins)
+			void FindCommonPluginDependencies(ref Dictionary<string, HashSet<UEBuildPlugin>> pluginGroups, ref Dictionary<string, HashSet<UEBuildPlugin>> testPluginGroups,
+				ref HashSet<UEBuildPlugin> outCommonPluginDependencies, ref HashSet<UEBuildPlugin> outConflictingPlugins)
 			{
-				if (pluginGroups.Any())
+				if (pluginGroups.Any() && testPluginGroups.Any())
 				{
-					HashSet<UEBuildPlugin> commonPluginDependencies = new();
 					foreach (KeyValuePair<string, HashSet<UEBuildPlugin>> item in pluginGroups)
 					{
-						IEnumerable<UEBuildPlugin> otherPluginDependencies = pluginGroups
+						IEnumerable<UEBuildPlugin> otherPluginDependencies = testPluginGroups
 							.Where(x => x.Key != item.Key)       // get other groups
 							.SelectMany(x => x.Value)            // get plugins
 							.Distinct()                          // unique
 							.Where(x => item.Value.Contains(x)); // that are shared here
 
-						commonPluginDependencies.UnionWith(otherPluginDependencies);
-						conflictingPlugins.UnionWith(otherPluginDependencies);
+						outCommonPluginDependencies.UnionWith(otherPluginDependencies);
+						outConflictingPlugins.UnionWith(otherPluginDependencies);
 					}
+				}
+			}
+
+			void RemoveNonExclusivePlugins(ref Dictionary<string, HashSet<UEBuildPlugin>> pluginGroups, ref Dictionary<string, HashSet<UEBuildPlugin>> additionalPluginGroups, ref HashSet<UEBuildPlugin> outConflictingPlugins)
+			{
+				if (pluginGroups.Any())
+				{
+					HashSet<UEBuildPlugin> commonPluginDependencies = new();
+					FindCommonPluginDependencies(ref pluginGroups, ref pluginGroups, ref commonPluginDependencies, ref outConflictingPlugins);
+					FindCommonPluginDependencies(ref pluginGroups, ref additionalPluginGroups, ref commonPluginDependencies, ref outConflictingPlugins);
 
 					Dictionary<string, HashSet<UEBuildPlugin>> filtered = new();
 					foreach (KeyValuePair<string, HashSet<UEBuildPlugin>> item in pluginGroups)
@@ -3717,7 +3733,7 @@ namespace UnrealBuildTool
 					UEBuildPlugin? plugin = BuildPlugins.FirstOrDefault(x => x.Name == name);
 					if (plugin == null)
 					{
-						Logger.LogWarning("Plugin {Name} not found in BuildPlugins and will not be merged", name);
+						Logger.LogInformation("Plugin {Name} not found in BuildPlugins and will not be merged", name);
 						continue;
 					}
 					HashSet<UEBuildPlugin> needed = new();
@@ -3743,14 +3759,21 @@ namespace UnrealBuildTool
 
 				// Move common plugins to a new group and then remove duplicates
 				HashSet<UEBuildPlugin> conflictingPlugins = new();
+				Dictionary<string, HashSet<UEBuildPlugin>> empty = new();
 				HandleSharedLibraries(ref pluginGroups, nonRootSharedPlugins);
-				RemoveNonExclusivePlugins(ref pluginGroups, ref conflictingPlugins);
+				RemoveNonExclusivePlugins(ref pluginGroups, ref empty, ref conflictingPlugins);
 
-				// Merge secondary dependencies into their owning set
-				HandleSharedLibraries(ref desiredPluginGroups, nonRootSharedPlugins);
-				RemoveNonExclusivePlugins(ref desiredPluginGroups, ref conflictingPlugins);
+				foreach (KeyValuePair<string, HashSet<UEBuildPlugin>> item in pluginGroups)
+				{
+					foreach (UEBuildPlugin plugin in item.Value)
+					{
+						Logger.LogInformation("{0}: merged into {1} (requested)", plugin.Name, item.Key);
+					}
+				}
 
-				// Merge remaining secondary dependencies into primaries
+				// Apply the same logic to desired dependencies and merge them
+				HandleSharedLibraries(ref desiredPluginGroups, nonRootSharedPlugins);
+				RemoveNonExclusivePlugins(ref desiredPluginGroups, ref pluginGroups, ref conflictingPlugins);
 				foreach (string group in pluginGroups.Keys.ToArray())
 				{
 					if (desiredPluginGroups.TryGetValue(group, out HashSet<UEBuildPlugin>? value))
@@ -3790,26 +3813,58 @@ namespace UnrealBuildTool
 						}
 					}
 				}
-				unmergablePlugins.ExceptWith(pluginGroups.SelectMany(x => x.Value));
+
+				// Avoid merging anything that depends on a plugin already merged elsewhere, as we cannot have lateral dependencies between merged libraries
+				foreach (KeyValuePair<string, HashSet<UEBuildPlugin>> item in pluginGroups)
+				{
+					IEnumerable<UEBuildPlugin> otherMergedPlugins = pluginGroups.Where(x => x.Key != item.Key).SelectMany(x => x.Value);
+					foreach (UEBuildPlugin plugin in item.Value)
+					{
+						bool shouldRemove = false;
+						List<string> conflicts = new();
+
+						HashSet<UEBuildPlugin> dependencies = new();
+						GetPluginDependencies(plugin, ref dependencies);
+
+						foreach (UEBuildPlugin dependency in dependencies)
+						{
+							if (otherMergedPlugins.Contains(dependency))
+							{
+								shouldRemove = true;
+								conflicts.Add(dependency.Name);
+							}
+						}
+
+						if (shouldRemove)
+						{
+							item.Value.Remove(plugin);
+							Logger.LogInformation("{0}: filtered out of {group} because of dependencies: {1}", plugin.Name, item.Key, String.Join(",", conflicts));
+						}
+					}
+				}
 
 				// At this point, we have to exclude any code that falls under remaining unmergeable plugins, dependencies included.
 				// We need that because there is no way the executable can depend on a merged library, and this remaining chunk would always be loaded as its own library.
+				unmergablePlugins.ExceptWith(pluginGroups.SelectMany(x => x.Value));
 
 				// Update unmergeable plugins by accounting for their dependencies
 				Dictionary<UEBuildPlugin, HashSet<string>> unmergeablePluginDependencies = new();
 				foreach (UEBuildPlugin plugin in unmergablePlugins)
 				{
-					if (conflictingPlugins.Contains(plugin))
-					{
-						Logger.LogInformation("{0}: unmergeable as a conflicted plugin", plugin.Name);
-					}
-					else
+					if (IsGameFeaturePlugin(plugin))
 					{
-						Logger.LogInformation("{0}: unmergeable as an unreferenced plugin", plugin.Name);
+						if (conflictingPlugins.Contains(plugin))
+						{
+							Logger.LogInformation("{0}: unmergeable (conflicted plugin)", plugin.Name);
+						}
+						else
+						{
+							Logger.LogInformation("{0}: unmergeable (unreferenced or filtered out)", plugin.Name);
+						}
 					}
 
 					HashSet<UEBuildPlugin> dependencies = new();
-					GetPluginDependencies(plugin, ref dependencies, true);
+					GetPluginDependencies(plugin, ref dependencies, true, true);
 
 					foreach (UEBuildPlugin dependency in dependencies)
 					{
@@ -3830,11 +3885,15 @@ namespace UnrealBuildTool
 				unmergablePlugins = new(unmergablePlugins.Union(unmergeablePluginDependencies.Keys));
 				foreach (KeyValuePair<UEBuildPlugin, HashSet<string>> dependency in unmergeablePluginDependencies)
 				{
-					if (dependency.Key.Modules.Any())
+					if (IsGameFeaturePlugin(dependency.Key))
 					{
-						Logger.LogInformation("{0}: unmergeable as a dependency of: {1}", dependency.Key, string.Join(",", dependency.Value));
+						if (dependency.Key.Modules.Any())
+						{
+							Logger.LogInformation("{0}: unmergeable (dependency of: {1})", dependency.Key, string.Join(",", dependency.Value));
+						}
 					}
 				}
+
 				Logger.LogInformation("Found a total of {0} plugins to merge", pluginGroups.SelectMany(x => x.Value).Count());
 
 				// Remove unmergable plugin modules from plugin modules and create a set of all plugin modules
