QAGame fix for forked clang warnings, authored by julian.gamble
#rnx

diff --git a/Engine/Source/Programs/UnrealBuildTool/Platform/Android/AndroidToolChain.cs b/Engine/Source/Programs/UnrealBuildTool/Platform/Android/AndroidToolChain.cs
--- a/Engine/Source/Programs/UnrealBuildTool/Platform/Android/AndroidToolChain.cs
+++ b/Engine/Source/Programs/UnrealBuildTool/Platform/Android/AndroidToolChain.cs
@@ -484,6 +484,7 @@ namespace UnrealBuildTool
 			base.GetCompileArguments_WarningsAndErrors(CompileEnvironment, Arguments);
 
 			// @todo unlikely all needed
+			Arguments.Add("-Wno-unknown-warning-option");   // Clang has been forked for Android - some warnings may not apply.
 			Arguments.Add("-Wno-local-type-template-args"); // engine triggers this
 			Arguments.Add("-Wno-return-type-c-linkage");    // needed for PhysX
 			Arguments.Add("-Wno-reorder");                  // member initialization order
