UnrealBuildTool: Ensure directories for produced items exist for non-standard modes like GenerateClangDatabase

#rnx

diff --git a/Engine/Source/Programs/UnrealBuildTool/Modes/ExecuteMode.cs b/Engine/Source/Programs/UnrealBuildTool/Modes/ExecuteMode.cs
--- a/Engine/Source/Programs/UnrealBuildTool/Modes/ExecuteMode.cs
+++ b/Engine/Source/Programs/UnrealBuildTool/Modes/ExecuteMode.cs
@@ -55,6 +55,7 @@ namespace UnrealBuildTool
 			using (GlobalTracer.Instance.BuildSpan("ActionGraph.ExecuteActions()").StartActive())
 			{
 				List<TargetDescriptor> TargetDescriptors = TargetDescriptor.ParseCommandLine(Arguments, BuildConfiguration, Logger);
+				ActionGraph.CreateDirectoriesForProducedItems(Actions);
 				await ActionGraph.ExecuteActionsAsync(BuildConfiguration, Actions, TargetDescriptors, Logger);
 			}
 
diff --git a/Engine/Source/Programs/UnrealBuildTool/Modes/GenerateClangDatabase.cs b/Engine/Source/Programs/UnrealBuildTool/Modes/GenerateClangDatabase.cs
--- a/Engine/Source/Programs/UnrealBuildTool/Modes/GenerateClangDatabase.cs
+++ b/Engine/Source/Programs/UnrealBuildTool/Modes/GenerateClangDatabase.cs
@@ -123,6 +123,7 @@ namespace UnrealBuildTool
 						if (PrerequisiteActions.Any())
 						{
 							Logger.LogInformation("Executing actions that produce source files...");
+							ActionGraph.CreateDirectoriesForProducedItems(PrerequisiteActions);
 							await ActionGraph.ExecuteActionsAsync(BuildConfiguration, PrerequisiteActions, new List<TargetDescriptor> { TargetDescriptor }, Logger);
 						}
 					}
diff --git a/Engine/Source/Programs/UnrealBuildTool/Modes/JsonExportMode.cs b/Engine/Source/Programs/UnrealBuildTool/Modes/JsonExportMode.cs
--- a/Engine/Source/Programs/UnrealBuildTool/Modes/JsonExportMode.cs
+++ b/Engine/Source/Programs/UnrealBuildTool/Modes/JsonExportMode.cs
@@ -67,6 +67,7 @@ namespace UnrealBuildTool
 						if (PrerequisiteActions.Count > 0)
 						{
 							Logger.LogInformation("Executing actions that produce source files...");
+							ActionGraph.CreateDirectoriesForProducedItems(PrerequisiteActions);
 							await ActionGraph.ExecuteActionsAsync(BuildConfiguration, PrerequisiteActions, new List<TargetDescriptor> { TargetDescriptor }, Logger);
 						}
 					}
