UnrealBuildTool: Remove no longer neededversion check for dsymutil for Apple platforms
Prevents running "/usr/bin/dsymuti -version" once per binary being linked if the makefile is invalidated which both spams the console window and adds unnecessary runtime  to UBT as it has to spawn hundreds of processes for editor targets

#rnx

diff --git a/Engine/Source/Programs/UnrealBuildTool/Platform/IOS/IOSToolChain.cs b/Engine/Source/Programs/UnrealBuildTool/Platform/IOS/IOSToolChain.cs
--- a/Engine/Source/Programs/UnrealBuildTool/Platform/IOS/IOSToolChain.cs
+++ b/Engine/Source/Programs/UnrealBuildTool/Platform/IOS/IOSToolChain.cs
@@ -652,26 +652,23 @@ namespace UnrealBuildTool
 
 			GenDebugAction.WorkingDirectory = GetMacDevSrcRoot();
 			GenDebugAction.CommandPath = BuildHostPlatform.Current.Shell;
-			string ExtraOptions;
-			string DsymutilPath = GetDsymutilPath(Logger, out ExtraOptions, bIsForLTOBuild: bIsForLTOBuild);
+			string DsymutilPath = "/usr/bin/dsymutil";
 			if (ProjectSettings.bGeneratedSYMBundle)
 			{
-				GenDebugAction.CommandArguments = String.Format("-c \"rm -rf \\\"{2}\\\"; \\\"{0}\\\" \\\"{1}\\\" {4} -o \\\"{2}\\\"; cd \\\"{2}/..\\\"; zip -r -y -1 {3}.zip {3}\"",
+				GenDebugAction.CommandArguments = String.Format("-c \"rm -rf \\\"{2}\\\"; \\\"{0}\\\" \\\"{1}\\\" -o \\\"{2}\\\"; cd \\\"{2}/..\\\"; zip -r -y -1 {3}.zip {3}\"",
 					DsymutilPath,
 					Executable.AbsolutePath,
 					OutputFile.AbsolutePath,
-					Path.GetFileName(FullDestPathRoot),
-					ExtraOptions);
+					Path.GetFileName(FullDestPathRoot));
 				GenDebugAction.ProducedItems.Add(ZipOutputFile);
 				Logger.LogInformation("Zip file: {File}", ZipOutputFile.AbsolutePath);
 			}
 			else
 			{
-				GenDebugAction.CommandArguments = String.Format("-c \"rm -rf \\\"{2}\\\"; \\\"{0}\\\" \\\"{1}\\\" {3} -f -o \\\"{2}\\\"\"",
+				GenDebugAction.CommandArguments = String.Format("-c \"rm -rf \\\"{2}\\\"; \\\"{0}\\\" \\\"{1}\\\" -f -o \\\"{2}\\\"\"",
 						DsymutilPath,
 						Executable.AbsolutePath,
-						OutputFile.AbsolutePath,
-						ExtraOptions);
+						OutputFile.AbsolutePath);
 			}
 			GenDebugAction.PrerequisiteItems.Add(Executable);
 			GenDebugAction.ProducedItems.Add(OutputFile);
diff --git a/Engine/Source/Programs/UnrealBuildTool/Platform/Mac/MacToolChain.cs b/Engine/Source/Programs/UnrealBuildTool/Platform/Mac/MacToolChain.cs
--- a/Engine/Source/Programs/UnrealBuildTool/Platform/Mac/MacToolChain.cs
+++ b/Engine/Source/Programs/UnrealBuildTool/Platform/Mac/MacToolChain.cs
@@ -1010,8 +1010,7 @@ namespace UnrealBuildTool
 			// Deletes ay existing file on the building machine. Also, waits 30 seconds, if needed, for the input file to be created in an attempt to work around
 			// a problem where dsymutil would exit with an error saying the input file did not exist.
 			// Note that the source and dest are switched from a copy command
-			string ExtraOptions = "";
-			string DsymutilPath = GetDsymutilPath(Logger, out ExtraOptions, bIsForLTOBuild: false);
+			string DsymutilPath = "/usr/bin/dsymutil";
 			string UniversalDsymutilScriptPath = FileReference.Combine(Unreal.EngineDirectory, "Build/BatchFiles/Mac/GenerateUniversalDSYM.sh").FullName;
 
 			string ArgumentString = "-c \"";
@@ -1037,11 +1036,10 @@ namespace UnrealBuildTool
 			}
 			else
 			{
-				ArgumentString += String.Format(" \\\"{0}\\\" {3} -f \\\"{1}\\\" -o \\\"{2}\\\"; ",
+				ArgumentString += String.Format(" \\\"{0}\\\" -f \\\"{1}\\\" -o \\\"{2}\\\"; ",
 					DsymutilPath,
 					MachOBinary.AbsolutePath,
-					OutputFile.AbsolutePath,
-					ExtraOptions);
+					OutputFile.AbsolutePath);
 			}
 			ArgumentString += String.Format("fi; ");
 			ArgumentString += "\"";
diff --git a/Engine/Source/Programs/UnrealBuildTool/ToolChain/AppleToolChain.cs b/Engine/Source/Programs/UnrealBuildTool/ToolChain/AppleToolChain.cs
--- a/Engine/Source/Programs/UnrealBuildTool/ToolChain/AppleToolChain.cs
+++ b/Engine/Source/Programs/UnrealBuildTool/ToolChain/AppleToolChain.cs
@@ -756,67 +756,6 @@ namespace UnrealBuildTool
 			return CompileAction;
 		}
 
-		protected string GetDsymutilPath(ILogger Logger, out string ExtraOptions, bool bIsForLTOBuild = false)
-		{
-			FileReference DsymutilLocation = new FileReference("/usr/bin/dsymutil");
-
-			// dsymutil before 10.0.1 has a bug that causes issues, it's fixed in autosdks but not everyone has those set up so for the timebeing we have
-			// a version in P4 - first determine if we need it
-			string DsymutilVersionString = Utils.RunLocalProcessAndReturnStdOut(DsymutilLocation.FullName, "-version", Logger);
-
-			bool bUseInstalledDsymutil = true;
-			int Major = 0, Minor = 0, Patch = 0;
-
-			// tease out the version number
-			string[] Tokens = DsymutilVersionString.Split(" ".ToCharArray());
-
-			// sanity check
-			if (Tokens.Length < 4 || Tokens[3].Contains('.') == false)
-			{
-				Log.TraceInformationOnce("Unable to parse dsymutil version out of: {0}", DsymutilVersionString);
-			}
-			else
-			{
-				string[] Versions = Tokens[3].Split(".".ToCharArray());
-				if (Versions.Length < 3)
-				{
-					Log.TraceInformationOnce("Unable to parse version token: {0}", Tokens[3]);
-				}
-				else
-				{
-					if (!Int32.TryParse(Versions[0], out Major) || !Int32.TryParse(Versions[1], out Minor) || !Int32.TryParse(Versions[2], out Patch))
-					{
-						Log.TraceInformationOnce("Unable to parse version tokens: {0}", Tokens[3]);
-					}
-				}
-			}
-
-			// if the installed one is too old, use a fixed up one if it can
-			if (bUseInstalledDsymutil == false)
-			{
-				FileReference PatchedDsymutilLocation = FileReference.Combine(Unreal.EngineDirectory, "Restricted/NotForLicensees/Binaries/Mac/LLVM/bin/dsymutil");
-
-				if (File.Exists(PatchedDsymutilLocation.FullName))
-				{
-					DsymutilLocation = PatchedDsymutilLocation;
-				}
-
-				DirectoryReference? AutoSdkDir;
-				if (UEBuildPlatformSDK.TryGetHostPlatformAutoSDKDir(out AutoSdkDir))
-				{
-					FileReference AutoSdkDsymutilLocation = FileReference.Combine(AutoSdkDir, "Mac", "LLVM", "bin", "dsymutil");
-					if (FileReference.Exists(AutoSdkDsymutilLocation))
-					{
-						DsymutilLocation = AutoSdkDsymutilLocation;
-					}
-				}
-			}
-
-			// 10.0.1 has an issue with LTO builds where we need to limit the number of threads
-			ExtraOptions = (bIsForLTOBuild && Major == 10 && Minor == 0 && Patch == 1) ? "-j 1" : "";
-			return DsymutilLocation.FullName;
-		}
-
 		public override ICollection<FileItem> PostBuild(ReadOnlyTargetRules Target, FileItem Executable, LinkEnvironment BinaryLinkEnvironment, IActionGraphBuilder Graph)
 		{
 			List<FileItem> OutputFiles = new List<FileItem>(base.PostBuild(Target, Executable, BinaryLinkEnvironment, Graph));
