UnrealBuildTool: Pass additional arguments when -ffast-math is enabled (non-default)

#rnx

diff --git a/Engine/Source/Programs/UnrealBuildTool/ToolChain/ClangToolChain.cs b/Engine/Source/Programs/UnrealBuildTool/ToolChain/ClangToolChain.cs
--- a/Engine/Source/Programs/UnrealBuildTool/ToolChain/ClangToolChain.cs
+++ b/Engine/Source/Programs/UnrealBuildTool/ToolChain/ClangToolChain.cs
@@ -694,6 +694,12 @@ namespace UnrealBuildTool
 					break;
 				case FPSemanticsMode.Imprecise:
 					Arguments.Add("-ffast-math");
+					Arguments.Add("-fhonor-infinities");
+					if (Info.ClangVersion.Major >= 19)
+					{
+						Arguments.Add("-fhonor-nans");
+					}
+					Arguments.Add("-fno-reciprocal-math");
 					break;
 				default:
 					throw new BuildException($"Unsupported FP semantics: {CompileEnvironment.FPSemantics}");
