Add lto-sample-profile and dwodir flags to sample-based PGO optimize builds

#rnx

diff --git a/Engine/Source/Programs/UnrealBuildTool/Platform/Windows/VCToolChain.cs b/Engine/Source/Programs/UnrealBuildTool/Platform/Windows/VCToolChain.cs
--- a/Engine/Source/Programs/UnrealBuildTool/Platform/Windows/VCToolChain.cs
+++ b/Engine/Source/Programs/UnrealBuildTool/Platform/Windows/VCToolChain.cs
@@ -365,7 +365,7 @@ namespace UnrealBuildTool
 			Arguments.Add($"/clang:-MD /clang:-MF\"{SourceDependsFileString}\"");
 		}
 
-		protected static FileReference GetClangProfDataFilename(CppCompileEnvironment CompileEnvironment)
+		protected static FileReference GetClangCompileProfDataFilename(CppCompileEnvironment CompileEnvironment)
 		{
 			string ProfDataFilename = Path.Combine(CompileEnvironment.PGODirectory!, CompileEnvironment.PGOFilenamePrefix!);
 			if (File.Exists(ProfDataFilename)) 
@@ -383,7 +383,27 @@ namespace UnrealBuildTool
 				throw new BuildException($"No .profdata files found in \"{CompileEnvironment.PGODirectory}\".");
 			}
 			return new FileReference(ProfDataFiles.First());
-		}	
+		}
+
+		protected static FileReference GetClangLinkProfDataFilename(LinkEnvironment LinkEnvironment)
+		{
+			string ProfDataFilename = Path.Combine(LinkEnvironment.PGODirectory!, LinkEnvironment.PGOFilenamePrefix!);
+			if (File.Exists(ProfDataFilename))
+			{
+				return new FileReference(ProfDataFilename);
+			}
+			// If an exact match doesn't exist, fall back to an alternative. This is supported for building Test with Shipping PGO data for example
+			string[] ProfDataFiles = Directory.GetFiles(LinkEnvironment.PGODirectory!, "*.profdata");
+			if (ProfDataFiles.Length > 1)
+			{
+				throw new BuildException($"More than one .profdata file found in \"{LinkEnvironment.PGODirectory}\" and \"{ProfDataFilename}\" not found ");
+			}
+			if (ProfDataFiles.Length == 0)
+			{
+				throw new BuildException($"No .profdata files found in \"{LinkEnvironment.PGODirectory}\".");
+			}
+			return new FileReference(ProfDataFiles.First());
+		}
 
 		protected virtual void AppendCLArguments_Global(CppCompileEnvironment CompileEnvironment, List<string> Arguments)
 		{
@@ -956,7 +976,7 @@ namespace UnrealBuildTool
 				else if (CompileEnvironment.bPGOOptimize)
 				{
 					// Use a merged profdata file.
-					FileReference ProfDataFilename = GetClangProfDataFilename(CompileEnvironment);
+					FileReference ProfDataFilename = GetClangCompileProfDataFilename(CompileEnvironment);
 					Log.TraceInformationOnce($"Using PGO profile data \"{ProfDataFilename}\"");
 					if (Target.WindowsPlatform.bSampleBasedPGO)
 					{
@@ -1874,7 +1894,7 @@ namespace UnrealBuildTool
 
 			if (EnvVars.Compiler.IsClang() && CompileEnvironment.bPGOOptimize)
 			{
-				BaseCompileAction.AdditionalPrerequisiteItems.Add(FileItem.GetItemByFileReference(GetClangProfDataFilename(CompileEnvironment)));
+				BaseCompileAction.AdditionalPrerequisiteItems.Add(FileItem.GetItemByFileReference(GetClangCompileProfDataFilename(CompileEnvironment)));
 			}
 
 			// Generate the timing info
@@ -3241,6 +3261,11 @@ namespace UnrealBuildTool
 				LinkAction.bCanExecuteInUBA = false; // Disabled for now. Should revisit to see why it is not working
 			}
 
+			if (EnvVars.Compiler.IsClang() && Target.WindowsPlatform.bSampleBasedPGO && LinkEnvironment.bPGOOptimize)
+			{
+				LinkAction.PrerequisiteItems.Add(FileItem.GetItemByFileReference(GetClangLinkProfDataFilename(LinkEnvironment)));
+			}
+
 			// Create link repro if requested, this argument is intentionally not added to the response file
 			if (Target.WindowsPlatform.LinkReproDir != null)
 			{
@@ -3511,6 +3536,21 @@ namespace UnrealBuildTool
 					if (Target.WindowsPlatform.bSampleBasedPGO)
 					{
 						Log.TraceInformationOnce("Enabling using Sample-Based Profile Guided Optimization (SPGO). Linking will take a while.");
+
+						Arguments.Add($"/dwodir:\"{LinkEnvironment.PGODirectory!}\"");
+
+						FileReference ProfDataFilename = GetClangLinkProfDataFilename(LinkEnvironment);
+
+						if (LinkEnvironment.bAllowLTCG && Target.WindowsPlatform.bAllowClangLinker)
+						{
+							// Use a merged profdata file.
+							Arguments.Add($"/lto-sample-profile:\"{ProfDataFilename}\"");
+						}
+
+						if (Target.WindowsPlatform.Compiler.IsIntel())
+						{
+							Arguments.Add($"-fprofile-sample-use=\"{ProfDataFilename}\"");
+						}
 					}
 					else
 					{
