UnrealBuildTool: Fix typo with plugin validation which causes false positive warnings when validation plugins (opt-in)

#rnx

diff --git a/Engine/Source/Programs/UnrealBuildTool/System/PluginDescriptor.cs b/Engine/Source/Programs/UnrealBuildTool/System/PluginDescriptor.cs
--- a/Engine/Source/Programs/UnrealBuildTool/System/PluginDescriptor.cs
+++ b/Engine/Source/Programs/UnrealBuildTool/System/PluginDescriptor.cs
@@ -791,7 +791,7 @@ namespace UnrealBuildTool
 
 			foreach (ModuleDescriptor ModuleInfo in Descriptors)
 			{
-				if (!ModuleInfo.IsCompiledInConfiguration(Platform, Configuration, Name, TargetType, bBuildDeveloperTools, bBuildRequiresCookedData, out string? InvalidReason))
+				if (!ModuleInfo.IsCompiledInConfiguration(Platform, Configuration, TargetName, TargetType, bBuildDeveloperTools, bBuildRequiresCookedData, out string? InvalidReason))
 				{
 					invalidReasons.Add(InvalidReason);
 				}
