Add support for shared merged libraries

- Overhaul merging strategy to use shared merged libraries for GFPs common to multiple root GFPs
- Add shared library dependencies to the module manifest
- Add runtime support for reading the new dependencies in the module manifest, clean up the reader code for ease of understanding
- Automatically load, unload shared library dependencies when needed
- Improve merged library logging to help debug issues

#tests Tested on console with the current merged modules implementation on a release branch
#rnx

diff --git a/Engine/Source/Programs/UnrealBuildTool/Configuration/ReadOnlyTargetRules.cs b/Engine/Source/Programs/UnrealBuildTool/Configuration/ReadOnlyTargetRules.cs
--- a/Engine/Source/Programs/UnrealBuildTool/Configuration/ReadOnlyTargetRules.cs
+++ b/Engine/Source/Programs/UnrealBuildTool/Configuration/ReadOnlyTargetRules.cs
@@ -575,10 +575,14 @@ namespace UnrealBuildTool
 
 		public IReadOnlyList<string> MergePluginsIgnored => Inner.MergePluginsIgnored;
 
+		public IReadOnlyList<string> MergePluginsExcluded => Inner.MergePluginsExcluded;
+
 		public IReadOnlyDictionary<string, IEnumerable<string>> MergePluginsShared => Inner.MergePluginsShared.ToDictionary(k => k.Key, v => v.Value.AsEnumerable());
 
 		public IReadOnlyList<string> MergePluginsLaunch => Inner.MergePluginsLaunch;
 
+		public bool MergeOnlyGameFeatures => Inner.MergeOnlyGameFeatures;
+
 		public IReadOnlyList<string> GlobalDefinitions => Inner.GlobalDefinitions.AsReadOnly();
 
 		public IReadOnlyList<string> ProjectDefinitions => Inner.ProjectDefinitions.AsReadOnly();
diff --git a/Engine/Source/Programs/UnrealBuildTool/Configuration/TargetRules.cs b/Engine/Source/Programs/UnrealBuildTool/Configuration/TargetRules.cs
--- a/Engine/Source/Programs/UnrealBuildTool/Configuration/TargetRules.cs
+++ b/Engine/Source/Programs/UnrealBuildTool/Configuration/TargetRules.cs
@@ -2637,6 +2637,14 @@ namespace UnrealBuildTool
 		[CommandLine("-MergePluginsIgnored=", ListSeparator = '+')]
 		public List<string> MergePluginsIgnored { get; } = new();
 
+		/// <summary>
+		/// Experimental: List of plugins (and their dependencies) to explicitly exclude when processing merged libraries. Requires bMergeModules to be enabled
+		/// Ignored plugins will still be compiled but forced into the executable.
+		/// </summary>
+		[RequiresUniqueBuildEnvironment]
+		[CommandLine("-MergePluginsExcluded=", ListSeparator = '+')]
+		public List<string> MergePluginsExcluded { get; } = new();
+
 		/// <summary>
 		/// Experimental: List of merged plugins to move common dependencies to a shared library, requires MergePlugins. Can be chained
 		/// </summary>
@@ -2647,6 +2655,11 @@ namespace UnrealBuildTool
 		/// </summary>
 		public List<string> MergePluginsLaunch { get; } = new();
 
+		/// <summary>
+		/// Experimental: Whether to merge only game feature plugins, ignoring regular plugins.
+		/// </summary>
+		public bool MergeOnlyGameFeatures { get; } = true;
+
 		/// <summary>
 		/// Backing storage for the LinkType property.
 		/// </summary>
diff --git a/Engine/Source/Programs/UnrealBuildTool/Configuration/UEBuildTarget.cs b/Engine/Source/Programs/UnrealBuildTool/Configuration/UEBuildTarget.cs
--- a/Engine/Source/Programs/UnrealBuildTool/Configuration/UEBuildTarget.cs
+++ b/Engine/Source/Programs/UnrealBuildTool/Configuration/UEBuildTarget.cs
@@ -2179,6 +2179,36 @@ namespace UnrealBuildTool
 						{
 							Manifest.ModuleNameToFileName[Module.Name] = Binary.OutputFilePath.GetFileName();
 						}
+
+						// In order to correctly lookup shared merged library dependencies at runtime, use decorated filenames
+						string GetFilenameFromLibraryName(string LibraryName)
+						{
+							string FilenameFragment = $"{AppName}{Rules.DecoratedSeparator}{LibraryName}";
+							IEnumerable<UEBuildBinary> Candidates = Binaries.Where(x =>
+								x.Type == UEBuildBinaryType.DynamicLinkLibrary && x.OutputFilePath.GetFileNameWithoutAnyExtensions().Contains(FilenameFragment));
+							return Candidates.FirstOrDefault()?.OutputFilePath.GetFileName().ToString() ?? LibraryName;
+						}
+
+						foreach (string Root in Rules.MergePlugins)
+						{
+							string RootFilename = GetFilenameFromLibraryName(Root);
+							if (!Manifest.LibraryDependencies.ContainsKey(RootFilename))
+							{
+								Manifest.LibraryDependencies.Add(RootFilename, new());
+							}
+						}
+
+						foreach (KeyValuePair<string, IEnumerable<string>> SharedGroupRule in Rules.MergePluginsShared)
+						{
+							foreach (string Root in SharedGroupRule.Value)
+							{
+								string RootFilename = GetFilenameFromLibraryName(Root);
+								if (Manifest.LibraryDependencies.ContainsKey(RootFilename))
+								{
+									Manifest.LibraryDependencies[RootFilename].Add(GetFilenameFromLibraryName(SharedGroupRule.Key));
+								}
+							}
+						}
 					}
 				}
 			}
@@ -3522,7 +3552,7 @@ namespace UnrealBuildTool
 					}
 				}
 
-				// No non-eninge PerModuleInline files found, Fall back to Core.
+				// No non-engine PerModuleInline files found, Fall back to Core.
 				((UEBuildModuleCPP)Modules["Core"]).bCreatePerModuleFile = true;
 			}
 
@@ -3569,10 +3599,17 @@ namespace UnrealBuildTool
 				}
 			}
 
-			bool IsGameFeaturePlugin(UEBuildPlugin plugin)
+			bool CanPluginBeMerged(UEBuildPlugin plugin)
 			{
-				DirectoryReference gfpDir = DirectoryReference.Combine(ProjectDirectory, "Plugins", "GameFeatures");
-				return plugin.File.IsUnderDirectory(gfpDir);
+				if (Rules.MergeOnlyGameFeatures)
+				{
+					DirectoryReference gameFeaturesDir = DirectoryReference.Combine(ProjectDirectory, "Plugins", "GameFeatures");
+					return plugin.File.IsUnderDirectory(gameFeaturesDir);
+				}
+				else
+				{
+					return true;
+				}
 			}
 
 			void GetPluginDependencies(UEBuildPlugin plugin, ref HashSet<UEBuildPlugin> dependencies, bool recursively = false, bool checkAllPlugins = false)
@@ -3585,7 +3622,7 @@ namespace UnrealBuildTool
 						if (dependentPluginCandidates.Any())
 						{
 							UEBuildPlugin dependentPlugin = dependentPluginCandidates.First();
-							if (!dependencies.Contains(dependentPlugin) && (IsGameFeaturePlugin(dependentPlugin) || checkAllPlugins))
+							if (!dependencies.Contains(dependentPlugin) && (CanPluginBeMerged(dependentPlugin) || checkAllPlugins))
 							{
 								dependencies.Add(dependentPlugin);
 								if (recursively)
@@ -3625,79 +3662,6 @@ namespace UnrealBuildTool
 					}
 				}
 			}
-			
-			void HandleSharedLibraries(ref Dictionary<string, HashSet<UEBuildPlugin>> pluginGroups, IReadOnlySet<string> nonRootSharedPlugins)
-			{
-				foreach (KeyValuePair<string, IEnumerable<string>> group in Rules.MergePluginsShared)
-				{
-					// Get the existing groups that will be shared
-					Dictionary<string, HashSet<UEBuildPlugin>> shared = new();
-					foreach (string dependency in group.Value)
-					{
-						if (pluginGroups.TryGetValue(dependency, out HashSet<UEBuildPlugin>? value))
-						{
-							shared.Add(dependency, value);
-						}
-					}
-
-					// Find all plugins that are referenced by more than one group or are forced into the shared set
-					HashSet<UEBuildPlugin> common = new(shared.SelectMany(x => x.Value).Where(x => shared.Count(y => y.Value.Contains(x)) > 1));
-					HashSet<UEBuildPlugin> forced = new();
-					foreach (KeyValuePair<string, HashSet<UEBuildPlugin>> item in shared)
-					{
-						if (nonRootSharedPlugins.Contains(item.Key))
-						{
-							forced.UnionWith(item.Value);
-							pluginGroups.Remove(item.Key);
-						}
-						else
-						{
-							item.Value.ExceptWith(common);
-							pluginGroups.Remove(item.Key);
-							pluginGroups.Add(item.Key, item.Value);
-						}
-					}
-
-					common.UnionWith(forced);
-					pluginGroups.Add(group.Key, common);
-				}
-			}
-
-			void FindCommonPluginDependencies(ref Dictionary<string, HashSet<UEBuildPlugin>> pluginGroups, ref Dictionary<string, HashSet<UEBuildPlugin>> testPluginGroups,
-				ref HashSet<UEBuildPlugin> outCommonPluginDependencies, ref HashSet<UEBuildPlugin> outConflictingPlugins)
-			{
-				if (pluginGroups.Any() && testPluginGroups.Any())
-				{
-					foreach (KeyValuePair<string, HashSet<UEBuildPlugin>> item in pluginGroups)
-					{
-						IEnumerable<UEBuildPlugin> otherPluginDependencies = testPluginGroups
-							.Where(x => x.Key != item.Key)       // get other groups
-							.SelectMany(x => x.Value)            // get plugins
-							.Distinct()                          // unique
-							.Where(x => item.Value.Contains(x)); // that are shared here
-
-						outCommonPluginDependencies.UnionWith(otherPluginDependencies);
-						outConflictingPlugins.UnionWith(otherPluginDependencies);
-					}
-				}
-			}
-
-			void RemoveNonExclusivePlugins(ref Dictionary<string, HashSet<UEBuildPlugin>> pluginGroups, ref Dictionary<string, HashSet<UEBuildPlugin>> additionalPluginGroups, ref HashSet<UEBuildPlugin> outConflictingPlugins)
-			{
-				if (pluginGroups.Any())
-				{
-					HashSet<UEBuildPlugin> commonPluginDependencies = new();
-					FindCommonPluginDependencies(ref pluginGroups, ref pluginGroups, ref commonPluginDependencies, ref outConflictingPlugins);
-					FindCommonPluginDependencies(ref pluginGroups, ref additionalPluginGroups, ref commonPluginDependencies, ref outConflictingPlugins);
-
-					Dictionary<string, HashSet<UEBuildPlugin>> filtered = new();
-					foreach (KeyValuePair<string, HashSet<UEBuildPlugin>> item in pluginGroups)
-					{
-						filtered.Add(item.Key, new(item.Value.Except(commonPluginDependencies)));
-					}
-					pluginGroups = filtered;
-				}
-			}
 
 			UEBuildBinary executableBinary = Binaries[0];
 			executableBinary.bAllowExports = true; // Need to export symbols because other binaries might depend on this. (They will be stripped anyway if not used)
@@ -3708,161 +3672,120 @@ namespace UnrealBuildTool
 			if (Rules.MergePlugins.Any() && BuildPlugins != null)
 			{
 				// Dependency rules:
-				// - Regular (non-shared) merged libraries cannot contain plugins that are dependencies of outside plugins.
+				// - Regular (non-shared) merged libraries cannot contain plugins that are dependencies of other merged libraries.
 				// - Regular (non-shared) merged libraries can depend on any number of shared merged libraries.
-				// - Shared merged libraries cannot have dependencies at all (other than the base executable) to avoid circularity.
-				// - Plugins that are part of no merged library (shared or not) will be pulled into the executable *with their dependencies*, so they should be minimized.
+				// - Shared merged libraries cannot have dependencies other than the base executable to avoid circularity.
+				// - Plugins that are part of no merged library will be pulled into the executable *with their dependencies*, so they should be minimized.
 				// - Some plugins are organizational in nature and should be ignored for dependencies.
 
-				// Syntax
+				// Config setup:
 				// - Rules.MergePlugin represents merged libraries: everything in that plugin, direct plugin dependencies and dependent modules.
 				// - Rules.MergePluginsShared represents shared merged libraries: same as above, except it's a list of plugins, and plugins in a merged library are excluded.
-				// - Rules.MergePluginsIgnored: explicitly ignored plugins.
+				// - Rules.MergePluginsIgnored: plugins ignored during the merging process.
+				// - Rules.MergePluginsExcluded: plugins explicitly removed from all merged libraries.
 
-				// Shared merged libraries hold modules that are common to at least two in a list of root GFPs.
-				// If an entry for a shared merged library isn't also defined as a root, it'll be fully merged into the shared.
-				HashSet<string> nonRootSharedPlugins = Rules.MergePluginsShared.SelectMany(x => x.Value).Except(Rules.MergePlugins).ToHashSet();
-				IEnumerable<string> mergedPlugins = Rules.MergePlugins.Union(nonRootSharedPlugins);
-				IEnumerable<string> sharedPlugins = nonRootSharedPlugins.Where(plugin => Rules.MergePluginsShared.SelectMany(x => x.Value).Where(x => x == plugin).Count() > 1);
-				nonRootSharedPlugins.ExceptWith(sharedPlugins);
-
-				// Find the plugins to merge and gather their direct dependencies
 				Dictionary<string, HashSet<UEBuildPlugin>> pluginGroups = new();
-				foreach (string name in mergedPlugins)
+				Dictionary<string, HashSet<UEBuildPlugin>> pluginGroupDependencies = new();
+
+				// Create plugin groups with only the root plugin
+				foreach (string name in Rules.MergePlugins)
 				{
 					UEBuildPlugin? plugin = BuildPlugins.FirstOrDefault(x => x.Name == name);
 					if (plugin == null)
 					{
-						Logger.LogInformation("Plugin {Name} not found in BuildPlugins and will not be merged", name);
-						continue;
+						Logger.LogInformation("{Name}: setup for merging but not built", name);
 					}
-					HashSet<UEBuildPlugin> needed = new();
-					needed.Add(plugin);
-					GetPluginDependencies(plugin, ref needed);
-					pluginGroups.Add(plugin.Name, needed);
-				}
-
-				// For each group, find out plugin dependencies that would make sense to also merge
-				Dictionary<string, HashSet<UEBuildPlugin>> desiredPluginGroups = new();
-				foreach (KeyValuePair<string, HashSet<UEBuildPlugin>> item in pluginGroups)
-				{
-					HashSet<UEBuildPlugin> desired = new();
-					foreach (UEBuildPlugin plugin in item.Value)
+					else if (CanPluginBeMerged(plugin))
 					{
-						if (!desiredPluginGroups.ContainsKey(plugin.Name))
-						{
-							GetPluginDependencies(plugin, ref desired, true);
-						}
+						HashSet<UEBuildPlugin> self = new HashSet<UEBuildPlugin> { plugin };
+						HashSet<UEBuildPlugin> dependencies = new();
+						GetPluginDependencies(plugin, ref dependencies, true, false);
+
+						pluginGroups.Add(name, self);
+						pluginGroupDependencies.Add(name, dependencies);
+
+						Logger.LogInformation("{Name}: setup for merging with {0} total dependencies", name, dependencies.Count());
 					}
-					desiredPluginGroups.Add(item.Key, desired);
 				}
 
-				// Move common plugins to a new group and then remove duplicates
-				HashSet<UEBuildPlugin> conflictingPlugins = new();
-				Dictionary<string, HashSet<UEBuildPlugin>> empty = new();
-				HandleSharedLibraries(ref pluginGroups, nonRootSharedPlugins);
-				RemoveNonExclusivePlugins(ref pluginGroups, ref empty, ref conflictingPlugins);
-
-				foreach (KeyValuePair<string, HashSet<UEBuildPlugin>> item in pluginGroups)
+				// Set up shared dynamic libraries with an empty group
+				foreach (KeyValuePair<string, IEnumerable<string>> sharedGroupRule in Rules.MergePluginsShared)
 				{
-					foreach (UEBuildPlugin plugin in item.Value)
-					{
-						Logger.LogInformation("{0}: merged into {1} (requested)", plugin.Name, item.Key);
-					}
+					pluginGroups.Add(sharedGroupRule.Key, new());
 				}
 
-				// Apply the same logic to desired dependencies and merge them
-				HandleSharedLibraries(ref desiredPluginGroups, nonRootSharedPlugins);
-				RemoveNonExclusivePlugins(ref desiredPluginGroups, ref pluginGroups, ref conflictingPlugins);
-				foreach (string group in pluginGroups.Keys.ToArray())
+				// Expand dependencies from the roots
+				HashSet<UEBuildPlugin> seenPlugins = new();
+				foreach (KeyValuePair<string, HashSet<UEBuildPlugin>> group in pluginGroups)
 				{
-					if (desiredPluginGroups.TryGetValue(group, out HashSet<UEBuildPlugin>? value))
+					void AddPluginToGroup(UEBuildPlugin plugin, string groupName)
 					{
-						pluginGroups[group].UnionWith(value);
-					}
-				}
+						Logger.LogDebug("{0}: added to plugin group {1}", plugin.Name, groupName);
 
-				// Find remaining unmerged plugins
-				IEnumerable<UEBuildPlugin> ignoredPlugins = BuildPlugins.Where(x => Rules.MergePluginsIgnored.Contains(x.Name));
-				HashSet<UEBuildPlugin> unmergablePlugins = BuildPlugins.Except(ignoredPlugins).Except(pluginGroups.SelectMany(x => x.Value)).ToHashSet();
+						pluginGroups[groupName].Add(plugin);
 
-				// For all unreferenced plugins, attempt to merge them through their dependencies
-				for (int unreferencedPhaseIndex = 0; unreferencedPhaseIndex < 3; unreferencedPhaseIndex++)
-				{
-					Logger.LogInformation("Unreferenced phase {unreferencedPhaseIndex}", unreferencedPhaseIndex);
+						if (plugin.Dependencies != null)
+						{
+							ProcessPluginsRecursively(groupName, plugin.Dependencies);
+						}
+					}
 
-					// For all unreferenced plugins, attempt to merge them through their dependencies
-					foreach (UEBuildPlugin plugin in unmergablePlugins)
+					void ProcessPluginsRecursively(string groupName, HashSet<UEBuildPlugin> plugins)
 					{
-						if (!pluginGroups.SelectMany(x => x.Value).Contains(plugin))
+						foreach (UEBuildPlugin plugin in plugins)
 						{
-							HashSet<UEBuildPlugin> dependencies = new();
-							GetPluginDependencies(plugin, ref dependencies);
-
-							if (dependencies.Count() > 0)
+							if (!seenPlugins.Contains(plugin) && CanPluginBeMerged(plugin))
 							{
-								string? foundPluginGroup = pluginGroups.FirstOrDefault(x => x.Value.Intersect(dependencies).Count() == dependencies.Count()).Key;
+								seenPlugins.Add(plugin);
+
+								IEnumerable<string> groupsWithPlugin = pluginGroupDependencies.Where(x => x.Value.Contains(plugin)).Select(x => x.Key);
 
-								if (foundPluginGroup != null)
+								// Plugin is unique to this group
+								if (groupsWithPlugin.Count() == 1)
 								{
-									Logger.LogInformation("{0}: was unreferenced but will be merged with {1}", plugin.Name, foundPluginGroup);
-									pluginGroups[foundPluginGroup].Add(plugin);
-									pluginGroups[foundPluginGroup].UnionWith(dependencies);
+									AddPluginToGroup(plugin, groupName);
 								}
-							}
-						}
-					}
-				}
 
-				// Avoid merging anything that depends on a plugin already merged elsewhere, as we cannot have lateral dependencies between merged libraries
-				foreach (KeyValuePair<string, HashSet<UEBuildPlugin>> item in pluginGroups)
-				{
-					IEnumerable<UEBuildPlugin> otherMergedPlugins = pluginGroups.Where(x => x.Key != item.Key).SelectMany(x => x.Value);
-					foreach (UEBuildPlugin plugin in item.Value)
-					{
-						bool shouldRemove = false;
-						List<string> conflicts = new();
+								// Plugin is shared
+								else
+								{
+									IEnumerable<KeyValuePair<string, IEnumerable<string>>> sharedGroupCandidates =
+										Rules.MergePluginsShared.Where(x => x.Value.Intersect(groupsWithPlugin).Count() == groupsWithPlugin.Count());
 
-						HashSet<UEBuildPlugin> dependencies = new();
-						GetPluginDependencies(plugin, ref dependencies);
+									if (sharedGroupCandidates.Any())
+									{
+										sharedGroupCandidates.OrderBy(x => x.Value.Count()).Reverse();
+										KeyValuePair<string, IEnumerable<string>> sharedGroupRule = sharedGroupCandidates.First();
 
-						foreach (UEBuildPlugin dependency in dependencies)
-						{
-							if (otherMergedPlugins.Contains(dependency))
-							{
-								shouldRemove = true;
-								conflicts.Add(dependency.Name);
+										AddPluginToGroup(plugin, sharedGroupRule.Key);
+									}
+									else
+									{
+										Logger.LogDebug("{0}: unmergeable", plugin.Name);
+									}
+								}
 							}
 						}
+					}
 
-						if (shouldRemove)
-						{
-							item.Value.Remove(plugin);
-							Logger.LogInformation("{0}: filtered out of {group} because of dependencies: {1}", plugin.Name, item.Key, String.Join(",", conflicts));
-						}
+					// There's only one entry here at this point
+					HashSet<UEBuildPlugin>? dependencies = group.Value.First().Dependencies;
+					if (dependencies != null)
+					{
+						ProcessPluginsRecursively(group.Key, dependencies);
 					}
 				}
 
 				// At this point, we have to exclude any code that falls under remaining unmergeable plugins, dependencies included.
 				// We need that because there is no way the executable can depend on a merged library, and this remaining chunk would always be loaded as its own library.
-				unmergablePlugins.ExceptWith(pluginGroups.SelectMany(x => x.Value));
+				IEnumerable<UEBuildPlugin> ignoredPlugins = BuildPlugins.Where(x => Rules.MergePluginsIgnored.Contains(x.Name));
+				HashSet<UEBuildPlugin> unmergablePlugins = BuildPlugins.Except(ignoredPlugins).Except(pluginGroups.SelectMany(x => x.Value)).ToHashSet();
+				unmergablePlugins.UnionWith(BuildPlugins.Where(x => Rules.MergePluginsExcluded.Contains(x.Name)));
 
 				// Update unmergeable plugins by accounting for their dependencies
 				Dictionary<UEBuildPlugin, HashSet<string>> unmergeablePluginDependencies = new();
 				foreach (UEBuildPlugin plugin in unmergablePlugins)
 				{
-					if (IsGameFeaturePlugin(plugin))
-					{
-						if (conflictingPlugins.Contains(plugin))
-						{
-							Logger.LogInformation("{0}: unmergeable (conflicted plugin)", plugin.Name);
-						}
-						else
-						{
-							Logger.LogInformation("{0}: unmergeable (unreferenced or filtered out)", plugin.Name);
-						}
-					}
-
 					HashSet<UEBuildPlugin> dependencies = new();
 					GetPluginDependencies(plugin, ref dependencies, true, true);
 
@@ -3885,11 +3808,11 @@ namespace UnrealBuildTool
 				unmergablePlugins = new(unmergablePlugins.Union(unmergeablePluginDependencies.Keys));
 				foreach (KeyValuePair<UEBuildPlugin, HashSet<string>> dependency in unmergeablePluginDependencies)
 				{
-					if (IsGameFeaturePlugin(dependency.Key))
+					if (CanPluginBeMerged(dependency.Key))
 					{
 						if (dependency.Key.Modules.Any())
 						{
-							Logger.LogInformation("{0}: unmergeable (dependency of: {1})", dependency.Key, string.Join(",", dependency.Value));
+							Logger.LogDebug("{0}: unmergeable (dependency of: {1})", dependency.Key, string.Join(",", dependency.Value));
 						}
 					}
 				}
diff --git a/Engine/Source/Programs/UnrealBuildTool/System/ModuleManifest.cs b/Engine/Source/Programs/UnrealBuildTool/System/ModuleManifest.cs
--- a/Engine/Source/Programs/UnrealBuildTool/System/ModuleManifest.cs
+++ b/Engine/Source/Programs/UnrealBuildTool/System/ModuleManifest.cs
@@ -29,6 +29,11 @@ namespace UnrealBuildTool
 		/// </summary>
 		public Dictionary<string, string> ModuleNameToFileName = new Dictionary<string, string>();
 
+		/// <summary>
+		/// Map of dependencies for each merged library.
+		/// </summary>
+		public Dictionary<string, HashSet<string>> LibraryDependencies = new Dictionary<string, HashSet<string>>();
+
 		/// <summary>
 		/// Constructs the module map with the given changelist
 		/// </summary>
@@ -151,6 +156,21 @@ namespace UnrealBuildTool
 				}
 				OutputWriter.WriteObjectEnd();
 
+				if (LibraryDependencies.Any())
+				{
+					OutputWriter.WriteObjectStart("LibraryDependencies");
+					foreach (KeyValuePair<string, HashSet<string>> LibraryNameToDependenciesPair in LibraryDependencies.Where(x => x.Value.Any()).OrderBy(x => x.Key))
+					{
+						OutputWriter.WriteArrayStart(LibraryNameToDependenciesPair.Key);
+						foreach (string DependencyName in LibraryNameToDependenciesPair.Value)
+						{
+							OutputWriter.WriteValue(DependencyName);
+						}
+						OutputWriter.WriteArrayEnd();
+					}
+					OutputWriter.WriteObjectEnd();
+				}
+
 				OutputWriter.WriteObjectEnd();
 			}
 		}
diff --git a/Engine/Source/Runtime/Core/Private/Modules/ModuleManager.cpp b/Engine/Source/Runtime/Core/Private/Modules/ModuleManager.cpp
--- a/Engine/Source/Runtime/Core/Private/Modules/ModuleManager.cpp
+++ b/Engine/Source/Runtime/Core/Private/Modules/ModuleManager.cpp
@@ -532,6 +532,24 @@ void* FModuleManager::InternalLoadLibrary(FName ModuleName, const FString& Libra
 		<< FPlatformProcess_GetDllHandle.Name(*LibraryToLoad);
 #endif // CPUPROFILERTRACE_ENABLED
 
+#if UE_MERGED_MODULES
+
+	// The plugin manager normally handles dependencies, but merged libraries introduce indirect dependencies.
+	// If we load plugin A that's merged with B, and B depends on C in another library...
+	// ... then we need to load C first, but the plugin manager cannot know that.
+	TArray<FName> Dependencies = GetAllSharedLibrariesForLibrary(LibraryToLoad);
+	for (FName Dependency : Dependencies)
+	{
+		if (LibraryHandles.Find(Dependency.ToString()) == nullptr)
+		{
+			UE_LOG(LogModuleManager, Verbose, TEXT("InternalLoadLibrary: loading dependency '%s'"), *Dependency.ToString());
+			InternalLoadLibrary(NAME_None, Dependency.ToString());
+		}
+	}
+
+#endif // UE_MERGED_MODULES
+
+	// Do the actual loading
 	const uint64_t InitialMemUsed = FPlatformMemory::GetMemoryUsedFast();
 	Handle = FPlatformProcess::GetDllHandle(*LibraryToLoad);
 	const uint64_t FinalMemUsed = FPlatformMemory::GetMemoryUsedFast();
@@ -562,11 +580,37 @@ void FModuleManager::InternalFreeLibrary(FName ModuleName, void* Handle)
 	//     - Notify the UObject system to remove live classes in every module of that library
 	//     - Wait for garbage collection to be run
 	//     - In FModuleManager::OnObjectCleanup, after GC, actually do the removal using FreeDllHandle
+	//     - Afterwards, unused shared libraries will call InternalFreeLibrary
 
-	ModuleInfoRef ModuleInfo = Modules[ModuleName];
-	const FString LibraryToUnload = FPaths::ConvertRelativePathToFull(ModuleInfo->Filename);
-	const int32 ModuleCount = GetLibraryUsers(LibraryToUnload).Num();
-	UE_LOG(LogModuleManager, Verbose, TEXT("InternalFreeLibrary: library '%s' has %d users"), *LibraryToUnload, ModuleCount);
+	int32 ModuleCount = 0;
+	FString LibraryToUnload;
+
+	// Explicit module: get the path and handle from there
+	if (ModuleName != NAME_None)
+	{
+		ModuleInfoRef ModuleInfo = Modules[ModuleName];
+		LibraryToUnload = FPaths::ConvertRelativePathToFull(ModuleInfo->Filename);
+		ModuleCount = GetLoadedModulesForLibrary(LibraryToUnload).Num();
+		UE_LOG(LogModuleManager, Verbose, TEXT("InternalFreeLibrary: library '%s' has %d users"), *LibraryToUnload, ModuleCount);
+	}
+
+	// Full library unload, when a shared library goes unused
+	else if (Handle)
+	{
+		for (const TPair<FString, void*>& LibraryNameAndHandle : LibraryHandles)
+		{
+			if (LibraryNameAndHandle.Value == Handle)
+			{
+				LibraryToUnload = LibraryNameAndHandle.Key;
+				UE_LOG(LogModuleManager, Verbose, TEXT("InternalFreeLibrary: library '%s' was explicitly requested"), *LibraryToUnload);
+				break;
+			}
+		}
+	}
+	else
+	{
+		UE_LOG(LogModuleManager, Fatal, TEXT("InternalFreeLibrary: no module name or valid handle was passed"))
+	}
 
 	if (ModuleCount == 0)
 	{
@@ -599,6 +643,21 @@ void FModuleManager::InternalFreeLibrary(FName ModuleName, void* Handle)
 		RemoveLoadedObjectsCallback.Broadcast(UnloadedModules);
 	}
 
+	// Identify if any library is now useless
+	for (const TPair<FString, void*>& LibraryNameAndHandle : LibraryHandles)
+	{
+		const FString LibraryName = FPaths::ConvertRelativePathToFull(LibraryNameAndHandle.Key);
+		const TArray<FName> LoadedModules = GetLoadedModulesForLibrary(LibraryName);
+		const TArray<FName> LoadedLibraries = GetLoadedLibrariesForSharedLibrary(LibraryName);
+
+		if (LoadedModules.IsEmpty() && LoadedLibraries.IsEmpty())
+		{
+			UE_LOG(LogModuleManager, Verbose, TEXT("InternalFreeLibrary: '%s' is now unused"), *LibraryName);
+
+			InternalFreeLibrary(NAME_None, LibraryNameAndHandle.Value);
+		}
+	}
+
 #else
 	FPlatformProcess::FreeDllHandle(Handle);
 #endif // UE_MERGED_MODULES
@@ -606,48 +665,119 @@ void FModuleManager::InternalFreeLibrary(FName ModuleName, void* Handle)
 
 #if UE_MERGED_MODULES
 
-TArray<FName> FModuleManager::GetLibraryUsers(const FString& ModuleFile) const
+TArray<FName> FModuleManager::GetLoadedModulesForLibrary(const FString& LibraryFilename) const
 {
-	TArray<FName> MatchingModules;
+	TArray<FName> UsingModules;
 
 	for (const TPair<FName, ModuleInfoRef>& ModuleNameAndInfo : Modules)
 	{
 		const ModuleInfoRef& ModuleInfo = ModuleNameAndInfo.Value;
-		if (ModuleInfo->Module.IsValid() && ModuleFile == FPaths::ConvertRelativePathToFull(ModuleInfo->Filename))
+		if (ModuleInfo->Module.IsValid() && !ModuleInfo->Filename.IsEmpty() && LibraryFilename == FPaths::ConvertRelativePathToFull(ModuleInfo->Filename))
 		{
-			MatchingModules.Add(ModuleNameAndInfo.Key);
+			UsingModules.Add(ModuleNameAndInfo.Key);
 		}
 	}
 
-	return MatchingModules;
+	return UsingModules;
 }
 
-#endif // UE_MERGED_MODULES
+TArray<FName> FModuleManager::GetLoadedLibrariesForSharedLibrary(const FString& LibraryFilename) const
+{
+	TArray<FName> UsingLibraries;
 
-#if UE_MERGED_MODULES && !UE_BUILD_SHIPPING
+	// Walk over all dependency entries to find any loaded library directly depending on this one.
+	for (const TPair<FString, TArray<FString>>& LibraryAndDependenciesPair : LibraryDependencies)
+	{
+		// Find out if there's an active handle for this library.
+		// LibraryHandles use full filenames, but LibraryDependencies only has the filename, so we can't just Find() it.
+		FString OtherLibraryFilename;
+		for (const TPair<FString, void*>& LibraryHandlePair : LibraryHandles)
+		{
+			if (LibraryHandlePair.Key.EndsWith(LibraryAndDependenciesPair.Key))
+			{
+				OtherLibraryFilename = LibraryHandlePair.Key;
+				break;
+			}
+		}
 
-void FModuleManager::LogMergedLibraryUsage()
+		// Find out if this other active library depends on the library we're looking for.
+		if (!OtherLibraryFilename.IsEmpty())
+		{
+			for (const FString& Dependency : LibraryAndDependenciesPair.Value)
+			{
+				if (LibraryFilename.EndsWith(Dependency))
+				{
+					UsingLibraries.Add(*OtherLibraryFilename);
+					break;
+				}
+			}
+		}
+	}
+
+	return UsingLibraries;
+}
+
+TArray<FName> FModuleManager::GetAllSharedLibrariesForLibrary(const FString& LibraryFilename) const
 {
-	TMap<FString, TArray<FName>> LibraryUsageMap;
+	TArray<FName> AllSharedLibraries;
 
-	for (const TPair<FName, ModuleInfoRef>& ModuleNameAndInfo : Modules)
+	for (const TPair<FString, TArray<FString>>& LibraryAndDependenciesPair : LibraryDependencies)
 	{
-		const ModuleInfoRef& ModuleInfo = ModuleNameAndInfo.Value;
-		if (ModuleInfo->Module.IsValid() && !ModuleInfo->Filename.IsEmpty())
+		if (LibraryFilename.EndsWith(LibraryAndDependenciesPair.Key))
 		{
-			const FString LibraryName = FPaths::ConvertRelativePathToFull(ModuleInfo->Filename);
-			TArray<FName>& UsingModules = LibraryUsageMap.FindOrAdd(LibraryName);
-			UsingModules.Add(ModuleNameAndInfo.Key);
+			for (const FString& Dependency : LibraryAndDependenciesPair.Value)
+			{
+				const FString DependencyFilename = FPaths::Combine(FPlatformProcess::GetModulesDirectory(), Dependency);
+				AllSharedLibraries.AddUnique(*FPaths::ConvertRelativePathToFull(DependencyFilename));
+			}
+
+			break;
 		}
 	}
 
-	for (TPair<FString, TArray<FName>>& LibraryAndModules : LibraryUsageMap)
+	return AllSharedLibraries;
+}
+
+#endif // UE_MERGED_MODULES
+
+#if UE_MERGED_MODULES && !UE_BUILD_SHIPPING
+
+void FModuleManager::LogMergedLibraryUsage()
+{
+	for (const TPair<FString, void*>& LibraryNameAndHandle : LibraryHandles)
 	{
-		UE_LOG(LogModuleManager, Log, TEXT("Modules loaded for library '%s' (%d)"), *LibraryAndModules.Key, LibraryAndModules.Value.Num());
+		const FString LibraryName = FPaths::ConvertRelativePathToFull(LibraryNameAndHandle.Key);
+		const TArray<FName> LoadedModules = GetLoadedModulesForLibrary(LibraryName);
+		const TArray<FName> SharedLibraries = GetAllSharedLibrariesForLibrary(LibraryName);
+		const TArray<FName> LoadedLibraries = GetLoadedLibrariesForSharedLibrary(LibraryName);
+
+		UE_LOG(LogModuleManager, Log, TEXT("Merged library usage for '%s'"), *LibraryName);
+
+		if (LoadedModules.Num())
+		{
+			UE_LOG(LogModuleManager, Log, TEXT("  %d modules loaded for this library"), LoadedModules.Num());
+			for (const FName& Module : LoadedModules)
+			{
+				UE_LOG(LogModuleManager, Log, TEXT("    '%s'"), *Module.ToString());
+			}
+		}
 
-		for (const FName& Module : LibraryAndModules.Value)
+		if (SharedLibraries.Num())
 		{
-			UE_LOG(LogModuleManager, Log, TEXT("    '%s'"), *Module.ToString());
+			UE_LOG(LogModuleManager, Log, TEXT("  depends on %d other libraries"), SharedLibraries.Num());
+			for (const FName& Library : SharedLibraries)
+			{
+				UE_LOG(LogModuleManager, Log, TEXT("    '%s'"), *Library.ToString());
+			}
+		}
+
+		if (LoadedLibraries.Num())
+		{
+			UE_LOG(LogModuleManager, Log, TEXT("  %d other libraries depend on this library"), LoadedLibraries.Num());
+			for (const FName& Library : LoadedLibraries)
+			{
+				UE_LOG(LogModuleManager, Log, TEXT("    '%s'"), *Library.ToString());
+			}
 		}
 	}
 }
@@ -700,6 +830,7 @@ void FModuleManager::OnObjectCleanup()
 				const uint64_t FinalMemUsed = FPlatformMemory::GetMemoryUsedFast();
 
 				UE_LOG(LogModuleManager, Verbose, TEXT("OnObjectCleanup: freed about %llu KB"), (InitialMemUsed - FinalMemUsed) / 1024);
+				LibraryHandles.Remove(LibraryAndModulesIterator.Key());
 			}
 
 			LibraryAndModulesIterator.RemoveCurrent();
@@ -1676,6 +1807,10 @@ void FModuleManager::FindModulePathsInDirectory(const FString& InDirectoryName,
 		{
 			if ((Manifest.BuildId == BuildId.GetValue() || SearchDirectoryName.Contains(TEXT("/Engine/Plugins/Bridge/"))))
 			{
+#if UE_MERGED_MODULES
+				LibraryDependencies = Manifest.LibraryDependencies;
+#endif // UE_MERGED_MODULES
+
 				for (const TPair<FString, FString>& Pair : Manifest.ModuleNameToFileName)
 				{
 					OutModulePaths.Add(FName(*Pair.Key), *FPaths::Combine(*SearchDirectoryName, *Pair.Value));
diff --git a/Engine/Source/Runtime/Core/Private/Modules/ModuleManifest.cpp b/Engine/Source/Runtime/Core/Private/Modules/ModuleManifest.cpp
--- a/Engine/Source/Runtime/Core/Private/Modules/ModuleManifest.cpp
+++ b/Engine/Source/Runtime/Core/Private/Modules/ModuleManifest.cpp
@@ -39,7 +39,38 @@ bool FModuleManifest::TryRead(const FString& FileName, FModuleManifest& OutManif
 	}
 
 	const TCHAR* Ptr = *Text;
-	if (!FSimpleParse::MatchZeroOrMoreWhitespace(Ptr) || !FSimpleParse::MatchChar(Ptr, TEXT('{')) || !FSimpleParse::MatchZeroOrMoreWhitespace(Ptr) || FSimpleParse::MatchChar(Ptr, TEXT('}')))
+
+	// Helpers
+	const auto Whitespace = [&Ptr]()
+	{
+		return FSimpleParse::MatchZeroOrMoreWhitespace(Ptr);
+	};
+	const auto ObjectStart = [&Ptr]()
+	{
+		return FSimpleParse::MatchChar(Ptr, TEXT('{'));
+	};
+	const auto ObjectEnd = [&Ptr]()
+	{
+		return FSimpleParse::MatchChar(Ptr, TEXT('}'));
+	};
+	const auto ArrayStart = [&Ptr]()
+	{
+		return FSimpleParse::MatchChar(Ptr, TEXT('['));
+	};
+	const auto ArrayEnd = [&Ptr]()
+	{
+		return FSimpleParse::MatchChar(Ptr, TEXT(']'));
+	};
+	const auto Colon = [&Ptr]()
+	{
+		return FSimpleParse::MatchChar(Ptr, TEXT(':'));
+	};
+	const auto Comma = [&Ptr]()
+	{
+		return FSimpleParse::MatchChar(Ptr, TEXT(','));
+	};
+
+	if (!Whitespace() || !ObjectStart() || !Whitespace() || ObjectEnd())
 	{
 		return false;
 	}
@@ -52,7 +83,7 @@ bool FModuleManifest::TryRead(const FString& FileName, FModuleManifest& OutManif
 			return false;
 		}
 
-		if (!FSimpleParse::MatchZeroOrMoreWhitespace(Ptr) || !FSimpleParse::MatchChar(Ptr, TEXT(':')) || !FSimpleParse::MatchZeroOrMoreWhitespace(Ptr))
+		if (!Whitespace() || !Colon() || !Whitespace())
 		{
 			return false;
 		}
@@ -64,32 +95,98 @@ bool FModuleManifest::TryRead(const FString& FileName, FModuleManifest& OutManif
 				return false;
 			}
 		}
+
+		// Modules is an array of ModuleName:ModulePath pairs
 		else if (Field.ToView() == TEXTVIEW("Modules"))
 		{
-			if (!FSimpleParse::MatchZeroOrMoreWhitespace(Ptr) || !FSimpleParse::MatchChar(Ptr, TEXT('{')) || !FSimpleParse::MatchZeroOrMoreWhitespace(Ptr))
+			if (!Whitespace() || !ObjectStart() || !Whitespace())
 			{
 				return false;
 			}
 
-			if (!FSimpleParse::MatchChar(Ptr, TEXT('}')))
+			if (!ObjectEnd())
 			{
 				for (;;)
 				{
 					TStringBuilder<64> ModuleName;
 					TStringBuilder<80> ModulePath;
-					if (!FSimpleParse::ParseString(Ptr, ModuleName) || !FSimpleParse::MatchZeroOrMoreWhitespace(Ptr) || !FSimpleParse::MatchChar(Ptr, TEXT(':')) || !FSimpleParse::MatchZeroOrMoreWhitespace(Ptr) || !FSimpleParse::ParseString(Ptr, ModulePath) || !FSimpleParse::MatchZeroOrMoreWhitespace(Ptr))
+					if (!FSimpleParse::ParseString(Ptr, ModuleName) || !Whitespace() || !Colon() || !Whitespace() || !FSimpleParse::ParseString(Ptr, ModulePath) || !Whitespace())
 					{
 						return false;
 					}
 
 					FindOrAddHeterogeneous(OutManifest.ModuleNameToFileName, ModuleName.ToView(), ModulePath.ToView());
 
-					if (FSimpleParse::MatchChar(Ptr, TEXT('}')))
+					if (ObjectEnd())
+					{
+						break;
+					}
+
+					if (!Comma() || !Whitespace())
+					{
+						return false;
+					}
+				}
+			}
+		}
+
+		// LibraryDependencies is an array of LibraryName keys to arrays of DependencyName values
+		else if (Field.ToView() == TEXTVIEW("LibraryDependencies"))
+		{
+			if (!Whitespace() || !ObjectStart() || !Whitespace())
+			{
+				return false;
+			}
+
+			if (!ObjectEnd())
+			{
+				for (;;)
+				{
+					TStringBuilder<64> LibraryName;
+					if (!FSimpleParse::ParseString(Ptr, LibraryName) || !Whitespace() || !Colon() || !Whitespace() || !ArrayStart() || !Whitespace())
+					{
+						return false;
+					}
+
+					OutManifest.LibraryDependencies.Add(LibraryName.ToString(), TArray<FString>());
+
+					for (;;)
+					{
+						TStringBuilder<64> DependencyName;
+						if (!FSimpleParse::ParseString(Ptr, DependencyName) || !Whitespace())
+						{
+							return false;
+						}
+
+						OutManifest.LibraryDependencies[LibraryName.ToString()].Add(DependencyName.ToString());
+
+						if (!Whitespace())
+						{
+							return false;
+						}
+
+						if (ArrayEnd())
+						{
+							break;
+						}
+
+						if (!Comma() || !Whitespace())
+						{
+							return false;
+						}
+					}
+
+					if (!Whitespace())
+					{
+						return false;
+					}
+
+					if (ObjectEnd())
 					{
 						break;
 					}
 
-					if (!FSimpleParse::MatchChar(Ptr, TEXT(',')) || !FSimpleParse::MatchZeroOrMoreWhitespace(Ptr))
+					if (!Comma() || !Whitespace())
 					{
 						return false;
 					}
@@ -101,17 +198,17 @@ bool FModuleManifest::TryRead(const FString& FileName, FModuleManifest& OutManif
 			return false;
 		}
 
-		if (!FSimpleParse::MatchZeroOrMoreWhitespace(Ptr))
+		if (!Whitespace())
 		{
 			return false;
 		}
 
-		if (FSimpleParse::MatchChar(Ptr, TEXT('}')))
+		if (ObjectEnd())
 		{
 			return true;
 		}
 
-		if (!FSimpleParse::MatchChar(Ptr, TEXT(',')) || !FSimpleParse::MatchZeroOrMoreWhitespace(Ptr))
+		if (!Comma() || !Whitespace())
 		{
 			return false;
 		}
diff --git a/Engine/Source/Runtime/Core/Public/Modules/ModuleManager.h b/Engine/Source/Runtime/Core/Public/Modules/ModuleManager.h
--- a/Engine/Source/Runtime/Core/Public/Modules/ModuleManager.h
+++ b/Engine/Source/Runtime/Core/Public/Modules/ModuleManager.h
@@ -706,8 +706,16 @@ private:
 	void InternalFreeLibrary(FName ModuleName, void* Handle);
 
 #if UE_MERGED_MODULES
-	/** Get all modules linking to a particular file */
-	TArray<FName> GetLibraryUsers(const FString& ModuleFile) const;
+
+	/** Get all loaded modules linked into a particular merged library */
+	TArray<FName> GetLoadedModulesForLibrary(const FString& LibraryFilename) const;
+
+	/** Get all loaded merged libraries that depend on a particular shared merged library */
+	TArray<FName> GetLoadedLibrariesForSharedLibrary(const FString& LibraryFilename) const;
+
+	/** Get all shared library dependencies, loaded or not, for a particular merged library */
+	TArray<FName> GetAllSharedLibrariesForLibrary(const FString& LibraryFilename) const;
+
 #endif // UE_MERGED_MODULES
 
 
@@ -776,6 +784,9 @@ private:
 	/** Merged libraries that should be unloaded after the next object cleanup, with their associated modules */
 	TMap<FString, TArray<FName>> DelayUnloadLibraries;
 
+	/** Dependency map for merged libraries, for use with shared merged libraries */
+	mutable TMap<FString, TArray<FString>> LibraryDependencies;
+
 #endif // UE_MERGED_MODULES
 };
 
diff --git a/Engine/Source/Runtime/Core/Public/Modules/ModuleManifest.h b/Engine/Source/Runtime/Core/Public/Modules/ModuleManifest.h
--- a/Engine/Source/Runtime/Core/Public/Modules/ModuleManifest.h
+++ b/Engine/Source/Runtime/Core/Public/Modules/ModuleManifest.h
@@ -15,6 +15,7 @@ class FModuleManifest
 public:
 	FString BuildId;
 	TMap<FString, FString> ModuleNameToFileName;
+	TMap<FString, TArray<FString>> LibraryDependencies;
 
 	/**
 	 * Default constructor 
