LiveLinkOpenVR disabled for arm64
-We do not currently feel we need support for this plugin on arm64, in the long run the openvr usage will hopefully be replaced by openxr.
#rnx

diff --git a/Engine/Plugins/Experimental/LiveLinkOpenVR/LiveLinkOpenVR.uplugin b/Engine/Plugins/Experimental/LiveLinkOpenVR/LiveLinkOpenVR.uplugin
--- a/Engine/Plugins/Experimental/LiveLinkOpenVR/LiveLinkOpenVR.uplugin
+++ b/Engine/Plugins/Experimental/LiveLinkOpenVR/LiveLinkOpenVR.uplugin
@@ -3,7 +3,7 @@
 	"Version": 1,
 	"VersionName": "1.0",
 	"FriendlyName": "LiveLinkOpenVR",
-	"Description": "Live Link plugin for OpenVR",
+	"Description": "Live Link plugin for OpenVR (Not supported for native arm64, see UEBuildWindows.cs DisablePlugins.)",
 	"Category": "Virtual Production",
 	"CreatedBy": "Epic Games, Inc.",
 	"CreatedByURL": "https://epicgames.com",
diff --git a/Engine/Source/Programs/UnrealBuildTool/Platform/Windows/UEBuildWindows.cs b/Engine/Source/Programs/UnrealBuildTool/Platform/Windows/UEBuildWindows.cs
--- a/Engine/Source/Programs/UnrealBuildTool/Platform/Windows/UEBuildWindows.cs
+++ b/Engine/Source/Programs/UnrealBuildTool/Platform/Windows/UEBuildWindows.cs
@@ -1266,6 +1266,10 @@ namespace UnrealBuildTool
 							// DVP
 							"MediaIOFramework",
 							"Composure",
+
+							// Permanently disabled meaning we have no plans to add arm64 suppport at this time.
+							// LiveLinkOpenVR is used in virtual production where arm64 machines are unlikely to be used at this time and the OpenVR library is generally being abandoned in favor of OpenXR.
+							"LiveLinkOpenVR",
 						});
 					}
 
