Cleaning up the implementation for the behavior settings in UHT

#rnx

diff --git a/Engine/Source/Programs/Shared/EpicGames.UHT/Parsers/UhtEnumParser.cs b/Engine/Source/Programs/Shared/EpicGames.UHT/Parsers/UhtEnumParser.cs
--- a/Engine/Source/Programs/Shared/EpicGames.UHT/Parsers/UhtEnumParser.cs
+++ b/Engine/Source/Programs/Shared/EpicGames.UHT/Parsers/UhtEnumParser.cs
@@ -127,11 +127,8 @@ namespace EpicGames.UHT.Parsers
 
 					if (enumObject.CppForm != UhtEnumCppForm.EnumClass && enumObject.UnderlyingType == UhtEnumUnderlyingType.Unspecified)
 					{
-						UhtIssueBehavior enumUnderlyingTypeBehavior = enumObject.Module.IsPartOfEngine ? topScope.Session.Config!.EngineEnumUnderlyingTypeNotSet
-							: topScope.Session.Config!.NonEngineEnumUnderlyingTypeNotSet;
-
 						string logMessage = $"Underlying type must be specified.";
-						switch (enumUnderlyingTypeBehavior)
+						switch (topScope.Session.Config!.EnumUnderlyingTypeNotSet.GetBehavior(enumObject.Module))
 						{
 							case UhtIssueBehavior.AllowSilently:
 								break;
diff --git a/Engine/Source/Programs/Shared/EpicGames.UHT/Parsers/UhtHeaderFileParser.cs b/Engine/Source/Programs/Shared/EpicGames.UHT/Parsers/UhtHeaderFileParser.cs
--- a/Engine/Source/Programs/Shared/EpicGames.UHT/Parsers/UhtHeaderFileParser.cs
+++ b/Engine/Source/Programs/Shared/EpicGames.UHT/Parsers/UhtHeaderFileParser.cs
@@ -409,10 +409,7 @@ namespace EpicGames.UHT.Parsers
 
 				if (missingGeneratedHeader)
 				{
-					UhtIssueBehavior missingGeneratedHeaderBehavior = headerFile.Module.IsPartOfEngine ? headerFile.Session.Config!.EngineMissingGeneratedHeaderIncludeBehavior
-						: headerFile.Session.Config!.NonEngineMissingGeneratedHeaderIncludeBehavior;
-
-					switch (missingGeneratedHeaderBehavior) 
+					switch (headerFile.Session.Config!.MissingGeneratedHeaderIncludeBehavior.GetBehavior(headerFile.Module)) 
 					{
 						case UhtIssueBehavior.AllowSilently:
 							break;
diff --git a/Engine/Source/Programs/Shared/EpicGames.UHT/Parsers/UhtPropertyParser.cs b/Engine/Source/Programs/Shared/EpicGames.UHT/Parsers/UhtPropertyParser.cs
--- a/Engine/Source/Programs/Shared/EpicGames.UHT/Parsers/UhtPropertyParser.cs
+++ b/Engine/Source/Programs/Shared/EpicGames.UHT/Parsers/UhtPropertyParser.cs
@@ -1251,8 +1251,7 @@ namespace EpicGames.UHT.Parsers
 				// TVal seems to prevent the Init method from working.
 				if (!propertySettings.PropertyExportFlags.HasAnyFlags(UhtPropertyExportFlags.TNonNullPtr))
 				{
-					args.ConditionalLogPointerUsage(args.Config.EngineNativePointerMemberBehavior, args.Config.EnginePluginNativePointerMemberBehavior,
-						args.Config.NonEngineNativePointerMemberBehavior, "Native pointer", typeStartPos, "TObjectPtr");
+					args.ConditionalLogPointerUsage(args.Config.NativePointerMemberBehavior, "Native pointer", typeStartPos, "TObjectPtr");
 				}
 
 				if (propertySettings.PropertyCategory == UhtPropertyCategory.Member)
diff --git a/Engine/Source/Programs/Shared/EpicGames.UHT/Types/Properties/UhtObjectPropertyBase.cs b/Engine/Source/Programs/Shared/EpicGames.UHT/Types/Properties/UhtObjectPropertyBase.cs
--- a/Engine/Source/Programs/Shared/EpicGames.UHT/Types/Properties/UhtObjectPropertyBase.cs
+++ b/Engine/Source/Programs/Shared/EpicGames.UHT/Types/Properties/UhtObjectPropertyBase.cs
@@ -680,8 +680,7 @@ namespace EpicGames.UHT.Types
 		private static UhtProperty? ObjectPtrProperty(UhtPropertyResolveArgs args)
 		{
 			return args.ParseTemplateObjectProperty(UhtTemplateObjectMode.Normal, (UhtClass referencedClass, int typeStartPos) => {
-				args.ConditionalLogPointerUsage(args.Config.EngineObjectPtrMemberBehavior, args.Config.EnginePluginObjectPtrMemberBehavior,
-					args.Config.NonEngineObjectPtrMemberBehavior, "ObjectPtr", typeStartPos, null);
+				args.ConditionalLogPointerUsage(args.Config.ObjectPtrMemberBehavior, "ObjectPtr", typeStartPos, null);
 				if (referencedClass.IsChildOf(args.Session.UClass))
 				{
 					// UObject specifies that there is no limiter
diff --git a/Engine/Source/Programs/Shared/EpicGames.UHT/Utils/UhtConfig.cs b/Engine/Source/Programs/Shared/EpicGames.UHT/Utils/UhtConfig.cs
--- a/Engine/Source/Programs/Shared/EpicGames.UHT/Utils/UhtConfig.cs
+++ b/Engine/Source/Programs/Shared/EpicGames.UHT/Utils/UhtConfig.cs
@@ -30,80 +30,107 @@ namespace EpicGames.UHT.Utils
 	};
 
 	/// <summary>
-	/// Interface for accessing configuration data.  Since UnrealBuildTool depends on EpicGames.UHT and all 
-	/// of the configuration support exists in UBT, configuration data must be accessed through an interface.
+	/// Collection of the behaviors for engine, engine plugin, and non engine projects
 	/// </summary>
-	public interface IUhtConfig
+	public struct UhtIssueBehaviorSet
 	{
+
 		/// <summary>
-		/// Default version of generated code. Defaults to oldest possible, unless specified otherwise in config.
+		/// Behavior for engine projects that aren't engine plugins
 		/// </summary>
-		public EGeneratedCodeVersion DefaultGeneratedCodeVersion { get; }
+		public UhtIssueBehavior Engine { get; init; }
 
 		/// <summary>
-		/// Pointer warning for native pointers in the engine
+		/// Behavior for engine plugin projects
 		/// </summary>
-		public UhtIssueBehavior EngineNativePointerMemberBehavior { get; }
+		public UhtIssueBehavior EnginePlugin { get; init; }
 
 		/// <summary>
-		/// Pointer warning for object pointers in the engine
+		/// Behavior for non engine projects
 		/// </summary>
-		public UhtIssueBehavior EngineObjectPtrMemberBehavior { get; }
+		public UhtIssueBehavior NonEngine { get; init; }
 
 		/// <summary>
-		/// Pointer warning for native pointers in engine plugins
+		/// Construct a behavior set
 		/// </summary>
-		public UhtIssueBehavior EnginePluginNativePointerMemberBehavior { get; }
+		/// <param name="engine">Behavior for engine projects that aren't engine plugins</param>
+		/// <param name="enginePlugin">Behavior for engine plugin projects</param>
+		/// <param name="nonEngine">Behavior for non engine projects</param>
+		public UhtIssueBehaviorSet(UhtIssueBehavior engine, UhtIssueBehavior enginePlugin, UhtIssueBehavior nonEngine)
+		{
+			Engine = engine;
+			EnginePlugin = enginePlugin;
+			NonEngine = nonEngine;
+		}
 
 		/// <summary>
-		/// Pointer warning for object pointers in engine plugins
+		/// Construct a behavior set where the engine plugin setting is the same as the engine setting
 		/// </summary>
-		public UhtIssueBehavior EnginePluginObjectPtrMemberBehavior { get; }
+		/// <param name="engine">Behavior for engine projects that aren't engine plugins</param>
+		/// <param name="nonEngine">Behavior for non engine projects</param>
+		public UhtIssueBehaviorSet(UhtIssueBehavior engine, UhtIssueBehavior nonEngine)
+		{
+			Engine = engine;
+			EnginePlugin = engine;
+			NonEngine = nonEngine;
+		}
 
 		/// <summary>
-		/// Pointer warning for native pointers outside the engine
+		/// Return the behavior based on the module type
 		/// </summary>
-		public UhtIssueBehavior NonEngineNativePointerMemberBehavior { get; }
+		/// <param name="module"></param>
+		/// <returns></returns>
+		public readonly UhtIssueBehavior GetBehavior(UhtModule module)
+		{
+			return module.IsPartOfEngine ? (module.IsPlugin ? EnginePlugin : Engine) : NonEngine;
+		}
+	}
 
+	/// <summary>
+	/// Interface for accessing configuration data.  Since UnrealBuildTool depends on EpicGames.UHT and all 
+	/// of the configuration support exists in UBT, configuration data must be accessed through an interface.
+	/// </summary>
+	public interface IUhtConfig
+	{
 		/// <summary>
-		/// Pointer warning for object pointers outside the engine
+		/// Default version of generated code. Defaults to oldest possible, unless specified otherwise in config.
 		/// </summary>
-		public UhtIssueBehavior NonEngineObjectPtrMemberBehavior { get; }
+		public EGeneratedCodeVersion DefaultGeneratedCodeVersion { get; }
 
 		/// <summary>
-		/// If true, deprecation warnings should be shown
+		/// Pointer warning for native pointers
 		/// </summary>
-		public bool ShowDeprecations { get; }
+		public UhtIssueBehaviorSet NativePointerMemberBehavior { get; }
 
 		/// <summary>
-		/// If true, UObject properties are enabled in RigVM
+		/// Pointer warning for object pointers
 		/// </summary>
-		public bool AreRigVMUObjectPropertiesEnabled { get; }
+		public UhtIssueBehaviorSet ObjectPtrMemberBehavior { get; }
 
 		/// <summary>
-		/// If true, UInterface properties are enabled in RigVM
+		/// Behavior for when generated headers aren't properly included.
 		/// </summary>
-		public bool AreRigVMUInterfaceProeprtiesEnabled { get; }
+		public UhtIssueBehaviorSet MissingGeneratedHeaderIncludeBehavior { get; }
 
 		/// <summary>
-		/// Behavior for when generated headers aren't properly included in engine code.
+		/// Behavior for when enum underlying types aren't set for engine code.
 		/// </summary>
-		public UhtIssueBehavior EngineMissingGeneratedHeaderIncludeBehavior { get; }
+		public UhtIssueBehaviorSet EnumUnderlyingTypeNotSet { get; }
 
 		/// <summary>
-		/// Behavior for when generated headers aren't properly included in non-engine code.
+		/// If true, deprecation warnings should be shown
 		/// </summary>
-		public UhtIssueBehavior NonEngineMissingGeneratedHeaderIncludeBehavior { get; }
+		public bool ShowDeprecations { get; }
 
 		/// <summary>
-		/// Behavior for when enum underlying types aren't set for engine code.
+		/// If true, UObject properties are enabled in RigVM
 		/// </summary>
-		public UhtIssueBehavior EngineEnumUnderlyingTypeNotSet { get; }
+		public bool AreRigVMUObjectPropertiesEnabled { get; }
 
 		/// <summary>
-		/// Behavior for when enum underlying types aren't set for non engine code.
+		/// If true, UInterface properties are enabled in RigVM
 		/// </summary>
-		public UhtIssueBehavior NonEngineEnumUnderlyingTypeNotSet { get; }
+		public bool AreRigVMUInterfaceProeprtiesEnabled { get; }
 
 		/// <summary>
 		/// Collection of all known documentation policies
diff --git a/Engine/Source/Programs/Shared/EpicGames.UHT/Utils/UhtPropertyResolveArgs.cs b/Engine/Source/Programs/Shared/EpicGames.UHT/Utils/UhtPropertyResolveArgs.cs
--- a/Engine/Source/Programs/Shared/EpicGames.UHT/Utils/UhtPropertyResolveArgs.cs
+++ b/Engine/Source/Programs/Shared/EpicGames.UHT/Utils/UhtPropertyResolveArgs.cs
@@ -344,15 +344,12 @@ namespace EpicGames.UHT.Utils
 		/// <summary>
 		/// Logs message for Object pointers to convert UObject* to TObjectPtr or the reverse
 		/// </summary>
-		/// <param name="engineBehavior">Expected behavior for engine types</param>
-		/// <param name="enginePluginBehavior">Expected behavior for engine plugin types</param>
-		/// <param name="nonEngineBehavior">Expected behavior for non-engine types</param>
+		/// <param name="behaviorSet">Behavior set being tested against</param>
 		/// <param name="pointerTypeDesc">Description of the pointer type</param>
 		/// <param name="typeStartPos">Starting character position of the type</param>
 		/// <param name="alternativeTypeDesc">Suggested alternate declaration</param>
 		/// <exception cref="UhtIceException">Thrown if the behavior type is unexpected</exception>
-		public readonly void ConditionalLogPointerUsage(UhtIssueBehavior engineBehavior, UhtIssueBehavior enginePluginBehavior,
-			UhtIssueBehavior nonEngineBehavior, string pointerTypeDesc, int typeStartPos, string? alternativeTypeDesc)
+		public readonly void ConditionalLogPointerUsage(UhtIssueBehaviorSet behaviorSet, string pointerTypeDesc, int typeStartPos, string? alternativeTypeDesc)
 		{
 			if (PropertySettings.PropertyCategory != UhtPropertyCategory.Member)
 			{
@@ -360,18 +357,7 @@ namespace EpicGames.UHT.Utils
 			}
 
 			UhtModule module = PropertySettings.Outer.Module;
-			UhtIssueBehavior behavior = nonEngineBehavior;
-			if (module.IsPartOfEngine)
-			{
-				if (module.IsPlugin)
-				{
-					behavior = enginePluginBehavior;
-				}
-				else
-				{
-					behavior = engineBehavior;
-				}
-			}
+			UhtIssueBehavior behavior = behaviorSet.GetBehavior(module);
 
 			// HACK for CoreUObject tests - allow raw pointers in non-intrinsic classes there.
 			if (module.ShortName == "CoreUObject" && PropertySettings.Outer.HeaderFile.ModuleRelativeFilePath.StartsWith("Tests", StringComparison.InvariantCultureIgnoreCase))
diff --git a/Engine/Source/Programs/UnrealBuildTool/Modes/UnrealHeaderToolMode.cs b/Engine/Source/Programs/UnrealBuildTool/Modes/UnrealHeaderToolMode.cs
--- a/Engine/Source/Programs/UnrealBuildTool/Modes/UnrealHeaderToolMode.cs
+++ b/Engine/Source/Programs/UnrealBuildTool/Modes/UnrealHeaderToolMode.cs
@@ -63,34 +63,24 @@ namespace UnrealBuildTool.Modes
 		private readonly EGeneratedCodeVersion _defaultGeneratedCodeVersion = EGeneratedCodeVersion.V1;
 
 		/// <summary>
-		/// Internal version of pointer warning for native pointers in the engine
+		/// Internal version of pointer warning for native pointers
 		/// </summary>
-		private readonly UhtIssueBehavior _engineNativePointerMemberBehavior = UhtIssueBehavior.AllowSilently;
+		private readonly UhtIssueBehaviorSet _nativePointerMemberBehavior = new(UhtIssueBehavior.AllowSilently, UhtIssueBehavior.AllowSilently, UhtIssueBehavior.AllowSilently);
 
 		/// <summary>
 		/// Internal version of pointer warning for object pointers in the engine
 		/// </summary>
-		private readonly UhtIssueBehavior _engineObjectPtrMemberBehavior = UhtIssueBehavior.AllowSilently;
+		private readonly UhtIssueBehaviorSet _objectPtrMemberBehavior = new(UhtIssueBehavior.AllowSilently, UhtIssueBehavior.AllowSilently, UhtIssueBehavior.AllowSilently);
 
 		/// <summary>
-		/// Internal version of pointer warning for native pointers in engine plugins
+		/// Internal version of behavior when there is a missing generated header include
 		/// </summary>
-		private readonly UhtIssueBehavior _enginePluginNativePointerMemberBehavior = UhtIssueBehavior.AllowSilently;
+		private readonly UhtIssueBehaviorSet _missingGeneratedHeaderIncludeBehavior = new(UhtIssueBehavior.AllowSilently, UhtIssueBehavior.AllowSilently);
 
 		/// <summary>
-		/// Internal version of pointer warning for object pointers in engine plugins
+		/// Internal version of the behavior set when the underlying type of a regular and namespaced enum isn't set
 		/// </summary>
-		private readonly UhtIssueBehavior _enginePluginObjectPtrMemberBehavior = UhtIssueBehavior.AllowSilently;
-
-		/// <summary>
-		/// Internal version of pointer warning for native pointers outside the engine
-		/// </summary>
-		private readonly UhtIssueBehavior _nonEngineNativePointerMemberBehavior = UhtIssueBehavior.AllowSilently;
-
-		/// <summary>
-		/// Internal version of pointer warning for object pointers outside the engine
-		/// </summary>
-		private readonly UhtIssueBehavior _nonEngineObjectPtrMemberBehavior = UhtIssueBehavior.AllowSilently;
+		private readonly UhtIssueBehaviorSet _enumUnderlyingTypeNotSet = new(UhtIssueBehavior.AllowSilently, UhtIssueBehavior.AllowSilently);
 
 		/// <summary>
 		/// If true, deprecation warnings should be shown
@@ -107,26 +97,6 @@ namespace UnrealBuildTool.Modes
 		/// </summary>
 		private readonly bool _areRigVMUInterfaceProeprtiesEnabled = false;
 
-		/// <summary>
-		/// Internal version of behavior when there is a missing generated header include in the engine code.
-		/// </summary>
-		private readonly UhtIssueBehavior _engineMissingGeneratedHeaderIncludeBehavior = UhtIssueBehavior.AllowSilently;
-
-		/// <summary>
-		/// Internal version of behavior when there is a missing generated header include in non engine code.
-		/// </summary>
-		private readonly UhtIssueBehavior _nonEngineMissingGeneratedHeaderIncludeBehavior = UhtIssueBehavior.AllowSilently;
-
-		/// <summary>
-		/// Internal version of the behavior set when the underlying type of a regular and namespaced enum isn't set in the engine code.
-		/// </summary>
-		private readonly UhtIssueBehavior _engineEnumUnderlyingTypeNotSet = UhtIssueBehavior.AllowSilently;
-
-		/// <summary>
-		/// Internal version of the behavior set when the underlying type of a regular and namespaced enum isn't set in the non engine code.
-		/// </summary>
-		private readonly UhtIssueBehavior _nonEngineEnumUnderlyingTypeNotSet = UhtIssueBehavior.AllowSilently;
-
 		/// <summary>
 		/// Collection of known documentation policies
 		/// </summary>
@@ -162,22 +132,10 @@ namespace UnrealBuildTool.Modes
 		public EGeneratedCodeVersion DefaultGeneratedCodeVersion => _defaultGeneratedCodeVersion;
 
 		/// <inheritdoc/>
-		public UhtIssueBehavior EngineNativePointerMemberBehavior => _engineNativePointerMemberBehavior;
+		public UhtIssueBehaviorSet NativePointerMemberBehavior => _nativePointerMemberBehavior;
 
 		/// <inheritdoc/>
-		public UhtIssueBehavior EngineObjectPtrMemberBehavior => _engineObjectPtrMemberBehavior;
-
-		/// <inheritdoc/>
-		public UhtIssueBehavior EnginePluginNativePointerMemberBehavior => _enginePluginNativePointerMemberBehavior;
-
-		/// <inheritdoc/>
-		public UhtIssueBehavior EnginePluginObjectPtrMemberBehavior => _enginePluginObjectPtrMemberBehavior;
-
-		/// <inheritdoc/>
-		public UhtIssueBehavior NonEngineNativePointerMemberBehavior => _nonEngineNativePointerMemberBehavior;
-
-		/// <inheritdoc/>
-		public UhtIssueBehavior NonEngineObjectPtrMemberBehavior => _nonEngineObjectPtrMemberBehavior;
+		public UhtIssueBehaviorSet ObjectPtrMemberBehavior => _objectPtrMemberBehavior;
 
 		/// <summary>
 		/// If true, UObject properties are enabled in RigVM
@@ -195,16 +153,10 @@ namespace UnrealBuildTool.Modes
 		public bool ShowDeprecations => _showDeprecations;
 
 		/// <inheritdoc/>
-		public UhtIssueBehavior EngineMissingGeneratedHeaderIncludeBehavior => _engineMissingGeneratedHeaderIncludeBehavior;
-
-		/// <inheritdoc/>
-		public UhtIssueBehavior NonEngineMissingGeneratedHeaderIncludeBehavior => _nonEngineMissingGeneratedHeaderIncludeBehavior;
+		public UhtIssueBehaviorSet MissingGeneratedHeaderIncludeBehavior => _missingGeneratedHeaderIncludeBehavior;
 
 		/// <inheritdoc/>
-		public UhtIssueBehavior EngineEnumUnderlyingTypeNotSet => _engineEnumUnderlyingTypeNotSet;
-
-		/// <inheritdoc/>
-		public UhtIssueBehavior NonEngineEnumUnderlyingTypeNotSet => _nonEngineEnumUnderlyingTypeNotSet;
+		public UhtIssueBehaviorSet EnumUnderlyingTypeNotSet => _enumUnderlyingTypeNotSet;
 
 		/// <inheritdoc/>
 		public IReadOnlyDictionary<string, UhtDocumentationPolicy> DocumentationPolicies => _documentationPolicies;
@@ -308,20 +260,13 @@ namespace UnrealBuildTool.Modes
 			_units = GetHashSet("UnrealHeaderTool", "Units", StringViewComparer.OrdinalIgnoreCase);
 			_delegateParameterCountStrings = GetList("UnrealHeaderTool", "DelegateParameterCountStrings");
 			_defaultGeneratedCodeVersion = GetGeneratedCodeVersion("UnrealHeaderTool", "DefaultGeneratedCodeVersion", EGeneratedCodeVersion.V1);
-			_engineNativePointerMemberBehavior = GetIssueBehavior("UnrealHeaderTool", "EngineNativePointerMemberBehavior", UhtIssueBehavior.AllowSilently);
-			_engineObjectPtrMemberBehavior = GetIssueBehavior("UnrealHeaderTool", "EngineObjectPtrMemberBehavior", UhtIssueBehavior.AllowSilently);
-			_enginePluginNativePointerMemberBehavior = GetIssueBehavior("UnrealHeaderTool", "EnginePluginNativePointerMemberBehavior", UhtIssueBehavior.AllowSilently);
-			_enginePluginObjectPtrMemberBehavior = GetIssueBehavior("UnrealHeaderTool", "EnginePluginObjectPtrMemberBehavior", UhtIssueBehavior.AllowSilently);
-			_nonEngineNativePointerMemberBehavior = GetIssueBehavior("UnrealHeaderTool", "NonEngineNativePointerMemberBehavior", UhtIssueBehavior.AllowSilently);
-			_nonEngineObjectPtrMemberBehavior = GetIssueBehavior("UnrealHeaderTool", "NonEngineObjectPtrMemberBehavior", UhtIssueBehavior.AllowSilently);
+			_nativePointerMemberBehavior = GetIssueBehaviorSet3("UnrealHeaderTool", "NativePointerMemberBehavior", _nativePointerMemberBehavior);
+			_objectPtrMemberBehavior = GetIssueBehaviorSet3("UnrealHeaderTool", "ObjectPtrMemberBehavior", _objectPtrMemberBehavior);
+			_missingGeneratedHeaderIncludeBehavior = GetIssueBehaviorSet2("UnrealHeaderTool", "MissingGeneratedHeaderIncludeBehavior", _missingGeneratedHeaderIncludeBehavior);
+			_enumUnderlyingTypeNotSet = GetIssueBehaviorSet2("UnrealHeaderTool", "EnumUnderlyingTypeNotSet", _enumUnderlyingTypeNotSet);
 			_areRigVMUObjectPropertiesEnabled = GetBoolean("UnrealHeaderTool", "AreRigVMUObjectPropertiesEnabled", false);
 			_areRigVMUInterfaceProeprtiesEnabled = GetBoolean("UnrealHeaderTool", "AreRigVMUInterfaceProeprtiesEnabled", false);
 			_showDeprecations = GetBoolean("UnrealHeaderTool", "ShowDeprecations", true);
-			_engineMissingGeneratedHeaderIncludeBehavior = GetIssueBehavior("UnrealHeaderTool", "EngineMissingGeneratedHeaderIncludeBehavior", UhtIssueBehavior.AllowSilently);
-			_nonEngineMissingGeneratedHeaderIncludeBehavior = GetIssueBehavior("UnrealHeaderTool", "NonEngineMissingGeneratedHeaderIncludeBehavior", UhtIssueBehavior.AllowSilently);
-			_engineEnumUnderlyingTypeNotSet = GetIssueBehavior("UnrealHeaderTool", "EngineEnumUnderlyingTypeNotSet", UhtIssueBehavior.AllowSilently);
-			_nonEngineEnumUnderlyingTypeNotSet = GetIssueBehavior("UnrealHeaderTool", "NonEngineEnumUnderlyingTypeNotSet", UhtIssueBehavior.AllowSilently);
-
 			GetDocumentationPolicies("UnrealHeaderTool", "DocumentationPolicies");
 			_defaultDocumentationPolicy = GetString("UnrealHeaderTool", "DefaultDocumentationPolicy", "");
 		}
@@ -386,6 +331,21 @@ namespace UnrealBuildTool.Modes
 			return Default;
 		}
 
+		private UhtIssueBehaviorSet GetIssueBehaviorSet3(string SectionName, string PartialKeyName, UhtIssueBehaviorSet Default)
+		{
+			return new UhtIssueBehaviorSet(
+				GetIssueBehavior(SectionName, $"Engine{PartialKeyName}", Default.Engine),
+				GetIssueBehavior(SectionName, $"EnginePlugin{PartialKeyName}", Default.EnginePlugin),
+				GetIssueBehavior(SectionName, $"NonEngine{PartialKeyName}", Default.NonEngine));
+		}
+
+		private UhtIssueBehaviorSet GetIssueBehaviorSet2(string SectionName, string PartialKeyName, UhtIssueBehaviorSet Default)
+		{
+			return new UhtIssueBehaviorSet(
+				GetIssueBehavior(SectionName, $"Engine{PartialKeyName}", Default.Engine),
+				GetIssueBehavior(SectionName, $"NonEngine{PartialKeyName}", Default.NonEngine));
+		}
+
 		private EGeneratedCodeVersion GetGeneratedCodeVersion(string SectionName, string KeyName, EGeneratedCodeVersion Default)
 		{
 			if (_ini.TryGetValue(SectionName, KeyName, out string? BehaviorStr))
