Fix Android GetAllInstalledSDKVersions() failing when the ndk directory doesn't exist.

diff --git a/Engine/Source/Programs/UnrealBuildTool/Platform/Android/AndroidPlatformSDK.cs b/Engine/Source/Programs/UnrealBuildTool/Platform/Android/AndroidPlatformSDK.cs
--- a/Engine/Source/Programs/UnrealBuildTool/Platform/Android/AndroidPlatformSDK.cs
+++ b/Engine/Source/Programs/UnrealBuildTool/Platform/Android/AndroidPlatformSDK.cs
@@ -226,7 +226,11 @@ namespace UnrealBuildTool
 			
 			if (Directory.Exists(Path.Combine(StudioSdkDir, "platforms", GetPlatformSpecificVersion("platforms"))))
 			{
-				return Directory.EnumerateDirectories(Path.Combine(StudioSdkDir, "ndk")).Select(GetNdkNameFromDirectory).Where(Directory => Directory != null).Distinct().ToArray()!;
+				string NDKPath = Path.Combine(StudioSdkDir, "ndk");
+				if (Directory.Exists(NDKPath))
+				{
+					return Directory.EnumerateDirectories(NDKPath).Select(GetNdkNameFromDirectory).Where(Directory => Directory != null).Distinct().ToArray()!;
+				}
 			}
 			
 			return [];
