Disable iOS editor modules tools on Windows Arm64

#rnx

diff --git a/Engine/Source/Programs/UnrealBuildTool/Platform/IOS/UEBuildIOS.cs b/Engine/Source/Programs/UnrealBuildTool/Platform/IOS/UEBuildIOS.cs
--- a/Engine/Source/Programs/UnrealBuildTool/Platform/IOS/UEBuildIOS.cs
+++ b/Engine/Source/Programs/UnrealBuildTool/Platform/IOS/UEBuildIOS.cs
@@ -1092,6 +1092,13 @@ namespace UnrealBuildTool
 			{
 				return;
 			}
+			
+			// no iOS tools on Windows Arm64
+			bool bIsWindowsArm64 = (Target.Platform == UnrealTargetPlatform.Win64 && Target.Architecture == UnrealArch.Arm64);
+			if (bIsWindowsArm64)
+			{
+				return;
+			}		
 
 			if ((Target.Platform == UnrealTargetPlatform.Win64) || (Target.Platform == UnrealTargetPlatform.Mac))
 			{
