Multiple swift file support
-We now pass all the swift files in a module, if there are any, into the swift compiler together which allows them work together.  The output is ModuleName.swift.o.
-We also generate a single bridging header to expose any [REDACTED_HANDLE] api.  The header is ModuleName-Swift.h.

diff --git a/Engine/Plugins/Online/IOS/OnlineSubsystemIOS/Source/Private/OnlineStoreIOS.cpp b/Engine/Plugins/Online/IOS/OnlineSubsystemIOS/Source/Private/OnlineStoreIOS.cpp
--- a/Engine/Plugins/Online/IOS/OnlineSubsystemIOS/Source/Private/OnlineStoreIOS.cpp
+++ b/Engine/Plugins/Online/IOS/OnlineSubsystemIOS/Source/Private/OnlineStoreIOS.cpp
@@ -11,7 +11,7 @@
 #import <StoreKit/SKProductsRequest.h>
 #include "Misc/ConfigCacheIni.h"
 
-#include "StoreKit2-Swift.h"
+#include "OnlineSubsystemIOS-Swift.h"
 #include "StoreKitSwiftInterop.h"
 #include "IOSAppDelegate.h"
 
diff --git a/Engine/Source/Programs/UnrealBuildTool/ToolChain/AppleToolChain.cs b/Engine/Source/Programs/UnrealBuildTool/ToolChain/AppleToolChain.cs
--- a/Engine/Source/Programs/UnrealBuildTool/ToolChain/AppleToolChain.cs
+++ b/Engine/Source/Programs/UnrealBuildTool/ToolChain/AppleToolChain.cs
@@ -314,11 +314,26 @@ namespace UnrealBuildTool
 
 			CPPOutput Result = new CPPOutput();
 			// Create a compile action for each source file.
+			List<FileItem> SwiftFiles = new List<FileItem>();
+
 			foreach (FileItem SourceFile in InputFiles)
 			{
-				Action CompileAction = CompileCPPFile(CompileEnvironment, SourceFile, OutputDir, ModuleName, Graph, GlobalArguments, Result);
+				if (SourceFile.HasExtension(".swift"))
+				{
+					SwiftFiles.Add(SourceFile);
+				}
+				else
+				{
+					Action CompileAction = CompileCPPFile(CompileEnvironment, SourceFile, OutputDir, ModuleName, Graph, GlobalArguments, Result);
+					CompileAction.PrerequisiteItems.UnionWith(FrameworkTokenFiles);
+				}
+			}
+			if (SwiftFiles.Any())
+			{
+				Action CompileAction = CompileSwiftFiles(CompileEnvironment, SwiftFiles, OutputDir, ModuleName, Graph, Result);
 				CompileAction.PrerequisiteItems.UnionWith(FrameworkTokenFiles);
 			}
+	
 			return Result;
 		}
 
@@ -607,22 +622,12 @@ namespace UnrealBuildTool
 			return CompileEnvironment;
 		}
 
-		private FileItem GetBridgingHeader(string SourceFile, DirectoryReference OutputDir)
+		private FileItem GetBridgingHeader(string ModuleName, DirectoryReference OutputDir)
 		{
-			string Filename = $"{Path.GetFileNameWithoutExtension(SourceFile)}-Swift.h";
+			string Filename = ModuleName + "-Swift.h";
 			return FileItem.GetItemByFileReference(FileReference.Combine(OutputDir, "Bridging", Filename));
 		}
 
-		protected override Action CompileCPPFile(CppCompileEnvironment CompileEnvironment, FileItem SourceFile, DirectoryReference OutputDir, string ModuleName, IActionGraphBuilder Graph, IReadOnlyCollection<string> GlobalArguments, CPPOutput Result)
-		{
-			if (SourceFile.HasExtension(".swift"))
-			{
-				return CompileSwiftFile(CompileEnvironment, SourceFile, OutputDir, Graph, Result);
-			}
-
-			return base.CompileCPPFile(CompileEnvironment, SourceFile, OutputDir, ModuleName, Graph, GlobalArguments, Result);
-		}
-
 		private void GetSharedSwiftArguments(CppCompileEnvironment CompileEnvironment, List<string> Arguments, string ModuleName)
 		{
 			// TODO: Fix so swiftc path supports vfs
@@ -664,13 +669,15 @@ namespace UnrealBuildTool
 			Console.WriteLine("Arguments: {0}", string.Join(", ", Arguments));
 		}
 
-		private Action CompileSwiftFile(CppCompileEnvironment CompileEnvironment, FileItem SourceFile, DirectoryReference OutputDir, IActionGraphBuilder Graph, CPPOutput CompileResult)
+		private Action CompileSwiftFiles(CppCompileEnvironment CompileEnvironment, IEnumerable<FileItem> InputFiles, DirectoryReference OutputDir, string ModuleName, IActionGraphBuilder Graph, CPPOutput CompileResult)
 		{
-			string ModuleName = Path.GetFileNameWithoutExtension(SourceFile.FullName);
-			FileItem OutputFile = FileItem.GetItemByFileReference(FileReference.Combine(OutputDir, GetFileNameFromExtension(SourceFile.Name, ".o")));
+			FileItem OutputFile = FileItem.GetItemByFileReference(FileReference.Combine(OutputDir, GetFileNameFromExtension(ModuleName + ".swift", ".o")));
 
 			List<string> Arguments = new();
-			Arguments.Add($"\"{SourceFile}\"");
+			foreach (FileItem SourceFile in InputFiles)
+			{
+				Arguments.Add($"\"{SourceFile}\"");
+			}
 
 			GetSharedSwiftArguments(CompileEnvironment, Arguments, ModuleName);
 			// output file settings
@@ -694,7 +701,10 @@ namespace UnrealBuildTool
 			CompileAction.Weight = CompileActionWeight;
 			CompileAction.CommandArguments = String.Join(" ", Arguments);
 			CompileAction.CommandPath = FileReference.Combine(ToolChainSettings.Value.ToolchainDir, "swift-frontend");
-			CompileAction.PrerequisiteItems.Add(SourceFile);
+			foreach (FileItem SourceFile in InputFiles)
+			{
+				CompileAction.PrerequisiteItems.Add(SourceFile);
+			}
 			CompileAction.ProducedItems.Add(OutputFile);
 			CompileResult.ObjectFiles.Add(OutputFile);
 			CompileAction.WorkingDirectory = Unreal.EngineSourceDirectory;
@@ -705,7 +715,7 @@ namespace UnrealBuildTool
 				string ReadableArch = ArchConfig.ConvertToReadableArchitecture(CompileEnvironment.Architecture);
 				CompileAction.CommandDescription += $" [{ReadableArch}]";
 			}
-			CompileAction.StatusDescription = Path.GetFileName(SourceFile.AbsolutePath);
+			CompileAction.StatusDescription = string.Join(" ", InputFiles.Select(x => x.Name));
 			CompileAction.bIsGCCCompiler = true;
 			CompileAction.bCanExecuteRemotely = true;
 			CompileAction.bCanExecuteInUBA = false;
@@ -713,11 +723,14 @@ namespace UnrealBuildTool
 
 			if (CompileEnvironment.ForceIncludeFiles.Count() > 0)
 			{
-				FileItem OutputInteropHeader = GetBridgingHeader(SourceFile.FullName, OutputDir);
+				FileItem OutputInteropHeader = GetBridgingHeader(ModuleName, OutputDir);
 
 				// obj-c bridging header settings
 				Arguments.Clear();
-				Arguments.Add($"\"{SourceFile}\"");
+				foreach (FileItem SourceFile in InputFiles)
+				{
+					Arguments.Add($"\"{SourceFile}\"");
+				}
 				GetSharedSwiftArguments(CompileEnvironment, Arguments, ModuleName);
 
 				Arguments.Add("-parse"); // this will allow it to generate the header without writing out any .o/executable
@@ -727,13 +740,15 @@ namespace UnrealBuildTool
 
 				//Arguments.Add("-enable-objc-interop");
 				//Arguments.Add("-cxx-interoperability-mode=default");
-
-
+				
 				// now make an action to export the swift code as a header Obj-C bridging
 				Action HeaderAction = Graph.CreateAction(ActionType.Compile);
 				HeaderAction.CommandPath = FileReference.Combine(ToolChainSettings.Value.ToolchainDir, "swiftc");
 				HeaderAction.CommandArguments = String.Join(" ", Arguments);
-				HeaderAction.PrerequisiteItems.Add(SourceFile);
+				foreach (FileItem SourceFile in InputFiles)
+				{
+					CompileAction.PrerequisiteItems.Add(SourceFile);
+				}
 				HeaderAction.ProducedItems.Add(OutputInteropHeader);
 				// swiftc won't touch the output header if the generated content doesn't change, so always delete the output so that 
 				// the timestamp is correct (it must be newer than the input .swift file, or it will build over and over)
