[UBA]
* Changed UbaExecutor to use Uba's scheduler. Reason is to be able to implement shadow processes and also dynamically add more processes when things like llvm pgo wants to distribute internal processes
* Moved logging out of ImmediateActionQueue and into base class ActionLogger. This made it possible to reuse logging logic for UbaExecutor

diff --git a/Engine/Source/Programs/Shared/EpicGames.UBA/ICacheClient.cs b/Engine/Source/Programs/Shared/EpicGames.UBA/ICacheClient.cs
--- a/Engine/Source/Programs/Shared/EpicGames.UBA/ICacheClient.cs
+++ b/Engine/Source/Programs/Shared/EpicGames.UBA/ICacheClient.cs
@@ -36,13 +36,13 @@ namespace EpicGames.UBA
 		/// Write to cache
 		/// </summary>
 		/// <param name="bucket">Bucket to store cache entry</param>
-		/// <param name="process">Process</param>
+		/// <param name="processHandle">Process</param>
 		/// <param name="inputs">Input files</param>
 		/// <param name="inputsSize">Input files size</param>
 		/// <param name="outputs">Output files</param>
 		/// <param name="outputsSize">Output files size</param>
 		/// <returns>True if successful</returns>
-		public abstract bool WriteToCache(uint bucket, IProcess process, byte[] inputs, uint inputsSize, byte[] outputs, uint outputsSize);
+		public abstract bool WriteToCache(uint bucket, nint processHandle, byte[] inputs, uint inputsSize, byte[] outputs, uint outputsSize);
 
 		/// <summary>
 		/// Fetch from cache
diff --git a/Engine/Source/Programs/Shared/EpicGames.UBA/IConfig.cs b/Engine/Source/Programs/Shared/EpicGames.UBA/IConfig.cs
--- a/Engine/Source/Programs/Shared/EpicGames.UBA/IConfig.cs
+++ b/Engine/Source/Programs/Shared/EpicGames.UBA/IConfig.cs
@@ -7,7 +7,7 @@ namespace EpicGames.UBA
 	/// <summary>
 	/// Base interface for uba config file
 	/// </summary>
-	public interface IConfig
+	public interface IConfig : IBaseInterface
 	{
 		/// <summary>
 		/// Load a config file
diff --git a/Engine/Source/Programs/Shared/EpicGames.UBA/IScheduler.cs b/Engine/Source/Programs/Shared/EpicGames.UBA/IScheduler.cs
new file mode 100644
--- /dev/null
+++ b/Engine/Source/Programs/Shared/EpicGames.UBA/IScheduler.cs
@@ -0,0 +1,156 @@
+// Copyright Epic Games, Inc. All Rights Reserved.
+
+using System;
+using System.Collections.Generic;
+using EpicGames.UBA.Impl;
+
+namespace EpicGames.UBA
+{
+	/// <summary>
+	/// Type of execution used for process
+	/// </summary>
+	public enum ProcessExecutionType
+	{
+		/// <summary>
+		/// Process executed locally without detouring
+		/// </summary>
+		Native,
+
+		/// <summary>
+		/// Process executed locally with detouring enabled
+		/// </summary>
+		Local,
+
+		/// <summary>
+		/// Process executed on a remote session
+		/// </summary>
+		Remote,
+
+		/// <summary>
+		/// Process was never executed and instead downloaded from cache
+		/// </summary>
+		Cache,
+
+		/// <summary>
+		/// Process was skipped and never executed
+		/// </summary>
+		Skip,
+	}
+
+	/// <summary>
+	/// </summary>
+	public interface IProcessFinishedInfo
+	{
+		/// <summary>
+		/// Type of execution
+		/// </summary>
+		ProcessExecutionType ExecutionType { get; }
+
+		/// <summary>
+		/// Process exit code
+		/// </summary>
+		int ExitCode { get; }
+
+		/// <summary>
+		/// Captured output lines
+		/// </summary>
+		List<string> LogLines { get; }
+
+		/// <summary>
+		/// The remote host that ran the process, if run remotely
+		/// </summary>
+		string? ExecutingHost { get; }
+
+		/// <summary>
+		/// Total time spent for the processor
+		/// </summary>
+		TimeSpan TotalProcessorTime { get; }
+
+		/// <summary>
+		/// Total wall time spent
+		/// </summary>
+		TimeSpan TotalWallTime { get; }
+
+		/// <summary>
+		/// UserData that was provided in EnqueueProcess
+		/// </summary>
+		object UserData { get; }
+
+		/// <summary>
+		/// Native uba handle to process.
+		/// </summary>
+		nint ProcessHandle { get; }
+	}
+
+	/// <summary>
+	/// </summary>
+	public enum ProcessFinishedResponse
+	{
+		/// <summary>
+		/// None means that nothing should be done
+		/// </summary>
+		None,
+		/// <summary>
+		/// RerunLocal means that we want to re-run the process locally with detouring enabled
+		/// </summary>
+		RerunLocal,
+		/// <summary>
+		/// RerunNative means that we want to re-run the process locally without detouring
+		/// </summary>
+		RerunNative,
+	}
+
+	/// <summary>
+	/// Base interface for uba config file
+	/// </summary>
+	public interface IScheduler : IBaseInterface
+	{
+
+		/// <summary>
+		/// Start the scheduler. It will start processing enqueued processes straight away
+		/// </summary>
+		void Start();
+
+		/// <summary>
+		/// Queue process.
+		/// </summary>
+		uint EnqueueProcess(ProcessStartInfo info, double weight, bool canDetour, bool canExecuteRemotely, int[]? dependencies, byte[]? knownInputs, uint knownInputsCount, uint cacheBucket, object userData);
+
+		/// <summary>
+		/// Cancel all active processes and skip queued ones
+		/// </summary>
+		void Cancel();
+
+		/// <summary>
+		/// Returns true if no processes are running or queued
+		/// </summary>
+		bool IsEmpty { get; }
+
+		/// <summary>
+		/// Accumulated weight of all processes that are Queued and can run right now (dependencies are done)
+		/// </summary>
+		double GetProcessWeightThatCanRunRemotelyNow();
+
+		/// <summary>
+		/// Set callback for when process has finished
+		/// </summary>
+		void SetProcessFinishedCallback(Func<IProcessFinishedInfo, ProcessFinishedResponse> processFinished);
+
+		/// <summary>
+		/// Allows uba to disable remote execution if running out of processes that can execute remotely
+		/// </summary>
+		void SetAllowDisableRemoteExecution();
+
+		/// <summary>
+		/// Create a scheduler
+		/// </summary>
+		/// <param name="session">Session</param>
+		/// <param name="cacheClients">List of cache clients</param>
+		/// <param name="maxLocalProcessors">Max number of local processes scheduler can run in parallel</param>
+		/// <param name="forceRemote">Force all processes that can to run remote</param>
+		public static IScheduler CreateScheduler(ISessionServer session, IEnumerable<ICacheClient> cacheClients, int maxLocalProcessors, bool forceRemote)
+		{
+			return new SchedulerImpl(session, cacheClients, maxLocalProcessors, forceRemote);
+		}
+	}
+}
\ No newline at end of file
diff --git a/Engine/Source/Programs/Shared/EpicGames.UBA/Impl/CacheClientImpl.cs b/Engine/Source/Programs/Shared/EpicGames.UBA/Impl/CacheClientImpl.cs
--- a/Engine/Source/Programs/Shared/EpicGames.UBA/Impl/CacheClientImpl.cs
+++ b/Engine/Source/Programs/Shared/EpicGames.UBA/Impl/CacheClientImpl.cs
@@ -101,9 +101,9 @@ namespace EpicGames.UBA.Impl
 			return CacheClient_RegisterPathHash(_handle, path, hash) != 0;
 		}
 
-		public bool WriteToCache(uint bucket, IProcess process, byte[] inputs, uint inputsSize, byte[] outputs, uint outputsSize)
+		public bool WriteToCache(uint bucket, nint processHandle, byte[] inputs, uint inputsSize, byte[] outputs, uint outputsSize)
 		{
-			return CacheClient_WriteToCache2(_handle, bucket, process.GetHandle(), inputs, inputsSize, outputs, outputsSize) != 0;
+			return CacheClient_WriteToCache2(_handle, bucket, processHandle, inputs, inputsSize, outputs, outputsSize) != 0;
 		}
 
 		public FetchFromCacheResult FetchFromCache(ulong rootsHandle, uint bucket, ProcessStartInfo info)
diff --git a/Engine/Source/Programs/Shared/EpicGames.UBA/Impl/ConfigImpl.cs b/Engine/Source/Programs/Shared/EpicGames.UBA/Impl/ConfigImpl.cs
--- a/Engine/Source/Programs/Shared/EpicGames.UBA/Impl/ConfigImpl.cs
+++ b/Engine/Source/Programs/Shared/EpicGames.UBA/Impl/ConfigImpl.cs
@@ -28,5 +28,26 @@ namespace EpicGames.UBA.Impl
 			nint tablePtr = Config_AddTable(_handle, table);
 			ConfigTable_AddValueString(tablePtr, name, value);
 		}
+
+		#region IDisposable
+		~ConfigImpl() => Dispose(false);
+
+		public void Dispose()
+		{
+			Dispose(true);
+			GC.SuppressFinalize(this);
+		}
+
+		protected virtual void Dispose(bool disposing)
+		{
+			if (_handle != IntPtr.Zero)
+			{
+				//Config_Destroy(_handle); // instance acquired by Config_Load is owned by native side
+				_handle = IntPtr.Zero;
+			}
+		}
+		#endregion
+
+		public nint GetHandle() => _handle;
 	}
 }
\ No newline at end of file
diff --git a/Engine/Source/Programs/Shared/EpicGames.UBA/Impl/SchedulerImpl.cs b/Engine/Source/Programs/Shared/EpicGames.UBA/Impl/SchedulerImpl.cs
new file mode 100644
--- /dev/null
+++ b/Engine/Source/Programs/Shared/EpicGames.UBA/Impl/SchedulerImpl.cs
@@ -0,0 +1,201 @@
+// Copyright Epic Games, Inc. All Rights Reserved.
+
+using System;
+using System.Collections.Generic;
+using System.Linq;
+using System.Runtime.InteropServices;
+using System.Threading;
+
+namespace EpicGames.UBA.Impl
+{
+	internal class ProcessFinishedInfoImpl : IProcessFinishedInfo
+	{
+		#region DllImport
+		[DllImport("UbaHost", CharSet = CharSet.Auto)]
+		static extern uint ProcessHandle_GetExitCode(nint process);
+
+		[DllImport("UbaHost", CharSet = CharSet.Auto)]
+		static extern byte ProcessHandle_GetExecutionType(nint process);
+
+		[DllImport("UbaHost", CharSet = CharSet.Auto)]
+		static extern nint ProcessHandle_GetExecutingHost(nint process);
+
+		[DllImport("UbaHost", CharSet = CharSet.Auto)]
+		static extern nint ProcessHandle_GetLogLine(nint process, uint index);
+
+		[DllImport("UbaHost", CharSet = CharSet.Auto)]
+		static extern ulong ProcessHandle_GetTotalProcessorTime(nint process);
+
+		[DllImport("UbaHost", CharSet = CharSet.Auto)]
+		static extern ulong ProcessHandle_GetTotalWallTime(nint process);
+		#endregion
+
+		internal ProcessFinishedInfoImpl(object userData, nint processHandle) { _userData = userData;  _processHandle = processHandle; }
+		object _userData;
+		nint _processHandle;
+
+		public ProcessExecutionType ExecutionType => (ProcessExecutionType)ProcessHandle_GetExecutionType(_processHandle);
+		public int ExitCode => (int)ProcessHandle_GetExitCode(_processHandle);
+		public List<string> LogLines
+		{
+			get
+			{
+				List<string> logLines = [];
+				string? line = Marshal.PtrToStringAuto(ProcessHandle_GetLogLine(_processHandle, 0));
+				while (line != null)
+				{
+					logLines.Add(line);
+					line = Marshal.PtrToStringAuto(ProcessHandle_GetLogLine(_processHandle, (uint)logLines.Count));
+				}
+				return logLines;
+			}
+		}
+		public string? ExecutingHost => Marshal.PtrToStringAuto(ProcessHandle_GetExecutingHost(_processHandle));
+		public TimeSpan TotalProcessorTime => new((long)ProcessHandle_GetTotalProcessorTime(_processHandle));
+		public TimeSpan TotalWallTime => new((long)ProcessHandle_GetTotalWallTime(_processHandle));
+		public object UserData => _userData;
+		public nint ProcessHandle => _processHandle;
+	}
+
+	internal class SchedulerProcess
+	{
+		#region DllImport
+		[DllImport("UbaHost", CharSet = CharSet.Auto)]
+		static extern nint ProcessStartInfo_Create2(string application, string arguments, string workingDir, string description, uint priorityClass, ulong rootsHandle, byte trackInputs, string logFile, ExitCallback? exit);
+		#endregion
+
+		public SchedulerProcess(SchedulerImpl scheduler, ProcessStartInfo info, object userData)
+		{
+			_scheduler = scheduler;
+			_userData = userData;
+			_exitCallbackDelegate = RaiseExited;
+			_startInfoHandle = ProcessStartInfo_Create2(info.Application, info.Arguments, info.WorkingDirectory, info.Description, (uint)info.Priority, info.RootsHandle, (byte)(info.TrackInputs ? 1 : 0), info.LogFile ?? String.Empty, _exitCallbackDelegate);
+		}
+
+		void RaiseExited(nint userData, nint processHandle, ref byte exitedResponse)
+		{
+			while (_index == uint.MaxValue) { Thread.Sleep(1); } // This to prevent a (most likely impossible) race condition (Index is set after process is enqueued)
+
+			//ProcessImpl process = new(handle);
+			ProcessFinishedInfoImpl info = new(_userData, processHandle);
+
+			exitedResponse = (byte)_scheduler._processFinishedCallback!.Invoke(info);
+		}
+
+		public delegate void ExitCallback(nint userData, nint processHandle, ref byte exitedResponse);
+		SchedulerImpl _scheduler;
+		object _userData;
+		event ExitCallback _exitCallbackDelegate;
+		public nint _startInfoHandle;
+		public uint _index = uint.MaxValue;
+	}
+
+	internal class SchedulerImpl : IScheduler
+	{
+		#region DllImport
+		[DllImport("UbaHost", CharSet = CharSet.Auto)]
+		static extern nint Scheduler_Create3(nint session, nint[] cacheClients, uint cacheClientCount, nint config);
+
+		[DllImport("UbaHost", CharSet = CharSet.Auto)]
+		static extern void Scheduler_Destroy(nint scheduler);
+
+		[DllImport("UbaHost", CharSet = CharSet.Auto)]
+		static extern void Scheduler_Start(nint scheduler);
+
+		[DllImport("UbaHost", CharSet = CharSet.Auto)]
+		static extern void Scheduler_Cancel(nint scheduler);
+
+		[DllImport("UbaHost", CharSet = CharSet.Auto)]
+		static extern byte Scheduler_IsEmpty(nint scheduler);
+
+		[DllImport("UbaHost", CharSet = CharSet.Auto)]
+		static extern uint Scheduler_EnqueueProcess2(nint scheduler, nint info, float weight, byte canDetour, byte canExecuteRemotely, int[]? dependencies, int dependencyCount, byte[]? knownInputs, int knownInputsBytes, uint knownInputsCount, uint cacheBucket);
+
+		[DllImport("UbaHost", CharSet = CharSet.Auto)]
+		static extern float Scheduler_GetProcessWeightThatCanRunRemotelyNow(nint scheduler);
+
+		[DllImport("UbaHost", CharSet = CharSet.Auto)]
+		static extern float Scheduler_SetAllowDisableRemoteExecution(nint scheduler, byte allow);
+		#endregion
+
+		nint _schedulerHandle = IntPtr.Zero;
+		internal Func<IProcessFinishedInfo, ProcessFinishedResponse>? _processFinishedCallback;
+		List<SchedulerProcess> _processes = new();
+
+		internal SchedulerImpl(ISessionServer server, IEnumerable<ICacheClient> cacheClients, int maxLocalProcessors, bool forceRemote)
+		{
+			nint[] cc = new nint[cacheClients.Count()];
+			uint cacheClientCount = 0;
+			foreach (var client in cacheClients)
+			{
+				cc[cacheClientCount++] = client.GetHandle();
+			}
+
+			IConfig config = IConfig.LoadConfig("");
+			config.AddValue("Scheduler", "MaxLocalProcessors", maxLocalProcessors.ToString());
+			if (forceRemote)
+			{
+				config.AddValue("Scheduler", "ForceRemote", "True");
+			}
+			_schedulerHandle = Scheduler_Create3(server.GetHandle(), cc, cacheClientCount, config.GetHandle());
+		}
+
+		public void Start()
+		{
+			Scheduler_Start(_schedulerHandle);
+		}
+
+		public uint EnqueueProcess(ProcessStartInfo info, double weight, bool canDetour, bool canExecuteRemotely, int[]? dependencies, byte[]? knownInputs, uint knownInputsCount, uint cacheBucket, object userData)
+		{
+			SchedulerProcess process = new(this, info, userData);
+			uint index = Scheduler_EnqueueProcess2(_schedulerHandle, process._startInfoHandle, (float)weight, (byte)(canDetour?1:0), (byte)(canExecuteRemotely?1:0), dependencies, dependencies != null ? dependencies.Length : 0, knownInputs, knownInputs != null ? knownInputs.Length : 0, knownInputsCount, cacheBucket);
+			process._index = index;
+			_processes.Add(process);
+			return index;
+		}
+
+		public void Cancel()
+		{
+			Scheduler_Cancel(_schedulerHandle);
+		}
+
+		public void SetProcessFinishedCallback(Func<IProcessFinishedInfo, ProcessFinishedResponse> processFinished)
+		{
+			_processFinishedCallback = processFinished;
+		}
+
+		public void SetAllowDisableRemoteExecution()
+		{
+			Scheduler_SetAllowDisableRemoteExecution(_schedulerHandle, 1);
+		}
+
+		public bool IsEmpty => Scheduler_IsEmpty(_schedulerHandle) != 0;
+
+		public double GetProcessWeightThatCanRunRemotelyNow() => (double)Scheduler_GetProcessWeightThatCanRunRemotelyNow(_schedulerHandle);
+
+		public nint GetHandle() => _schedulerHandle;
+
+		#region IDisposable
+		~SchedulerImpl() => Dispose(false);
+
+		public void Dispose()
+		{
+			Dispose(true);
+			GC.SuppressFinalize(this);
+		}
+
+		protected virtual void Dispose(bool disposing)
+		{
+			if (disposing)
+			{
+			}
+
+			if (_schedulerHandle != IntPtr.Zero)
+			{
+				Scheduler_Destroy(_schedulerHandle);
+				_schedulerHandle = IntPtr.Zero;
+			}
+		}
+		#endregion
+	}
+}
\ No newline at end of file
diff --git a/Engine/Source/Programs/UnrealBuildTool/Executors/ActionLogger.cs b/Engine/Source/Programs/UnrealBuildTool/Executors/ActionLogger.cs
new file mode 100644
--- /dev/null
+++ b/Engine/Source/Programs/UnrealBuildTool/Executors/ActionLogger.cs
@@ -0,0 +1,398 @@
+// Copyright Epic Games, Inc. All Rights Reserved.
+
+using EpicGames.Core;
+using Microsoft.Extensions.Logging;
+using System;
+using System.Collections.Generic;
+using System.Linq;
+using System.Threading;
+using System.Threading.Tasks;
+
+namespace UnrealBuildTool
+{
+	/// <summary>
+	/// Results from a run action
+	/// </summary>
+	/// <param name="LogLines">Console log lines</param>
+	/// <param name="ExitCode">Process return code.  Zero is assumed to be success and all other values an error.</param>
+	/// <param name="ExecutionTime">Wall time</param>
+	/// <param name="ProcessorTime">CPU time</param>
+	/// <param name="Action">Action</param>
+	/// <param name="AdditionalDescription">Additional description of action</param>
+	record ExecuteResults(List<string> LogLines, int ExitCode, TimeSpan ExecutionTime, TimeSpan ProcessorTime, LinkedAction Action, string? AdditionalDescription = null);
+
+	class ActionLogger
+	{
+		/// <summary>
+		/// Add target names for each action executed
+		/// </summary>
+		public bool PrintActionTargetNames = false;
+
+		/// <summary>
+		/// Whether to log command lines for actions being executed
+		/// </summary>
+		public bool LogActionCommandLines = false;
+
+		/// <summary>
+		/// Whether to show compilation times for each executed action
+		/// </summary>
+		public bool ShowPerActionCompilationTimes = false;
+
+		/// <summary>
+		/// Collapse non-error output lines
+		/// </summary>
+		public bool CompactOutput = false;
+
+		/// <summary>
+		/// Whether to show compilation times along with worst offenders or not.
+		/// </summary>
+		public bool ShowCompilationTimes = false;
+
+		/// <summary>
+		/// Output logging
+		/// </summary>
+		public readonly ILogger Logger;
+
+		/// <summary>
+		/// Progress writer
+		/// </summary>
+		public readonly ProgressWriter ProgressWriter;
+
+		/// <summary>
+		/// Process group
+		/// </summary>		
+		public readonly ManagedProcessGroup ProcessGroup = new();
+
+		/// <summary>
+		/// The last action group printed in multi-target builds
+		/// </summary>
+		private string? LastGroupPrefix = null;
+
+		/// <summary>
+		/// Whether to show CPU utilization after the work is complete.
+		/// </summary>
+		public bool ShowCPUUtilization = false;
+
+		/// <summary>
+		/// Action to invoke when writing tool output
+		/// </summary>
+		private readonly Action<string> _writeToolOutput;
+
+		/// <summary>
+		/// Flush the tool output after logging has completed
+		/// </summary>
+		protected readonly System.Action _flushToolOutput;
+
+		/// <summary>
+		/// Used only by the logger to track the [x,total] output
+		/// </summary>
+		private int _loggedCompletedActions = 0;
+
+		/// <summary>
+		/// Collection of all actions remaining to be logged
+		/// </summary>
+		protected readonly List<int> _actionsToLog = new();
+
+		/// <summary>
+		/// Task waiting to process logging
+		/// </summary>
+		protected Task? _actionsToLogTask = null;
+
+		/// <summary>
+		/// Timer used to collect CPU utilization
+		/// </summary>
+		protected Timer? _cpuUtilizationTimer;
+
+		/// <summary>
+		/// Per-second logging of cpu utilization
+		/// </summary>
+		protected readonly List<float> _cpuUtilization = new();
+
+		/// <summary>
+		/// Optional execution results (ordered like actions)
+		/// </summary>
+		protected readonly ExecuteResults?[] Results;
+
+		public ActionLogger(int actionCount, string progressWriterText, Action<string> writeToolOutput, System.Action flushToolOutput, ILogger logger)
+		{
+			Results = new ExecuteResults[actionCount];
+
+			Logger = logger;
+			ProgressWriter = new(progressWriterText, false, logger);
+
+			_writeToolOutput = writeToolOutput;
+			_flushToolOutput = flushToolOutput;
+		}
+
+		public void AddActionToLog(ExecuteResults executeTaskResult)
+		{
+			int index = executeTaskResult.Action.SortIndex;
+			Results[index] = executeTaskResult;
+			lock (_actionsToLog)
+			{
+				_actionsToLog.Add(index);
+				_actionsToLogTask ??= Task.Run(LogActions2);
+			}
+		}
+
+
+		private static int s_previousLineLength = -1;
+
+		/// <summary>
+		/// Log an action that has completed
+		/// </summary>
+		/// <param name="action">Action that has completed</param>
+		/// <param name="executeTaskResult">Results of the action</param>
+		/// <param name="totalActions">Number of total actions</param>
+		public void LogAction(LinkedAction action, ExecuteResults? executeTaskResult, int totalActions)
+		{
+			List<string>? logLines = null;
+			int exitCode = Int32.MaxValue;
+			TimeSpan executionTime = TimeSpan.Zero;
+			TimeSpan processorTime = TimeSpan.Zero;
+			string? additionalDescription = null;
+			if (executeTaskResult != null)
+			{
+				logLines = executeTaskResult.LogLines;
+				exitCode = executeTaskResult.ExitCode;
+				executionTime = executeTaskResult.ExecutionTime;
+				processorTime = executeTaskResult.ProcessorTime;
+				additionalDescription = executeTaskResult.AdditionalDescription;
+			}
+
+			// Write it to the log
+			string description = String.Empty;
+			if (action.bShouldOutputStatusDescription || (logLines != null && logLines.Count == 0))
+			{
+				description = $"{(action.CommandDescription ?? action.CommandPath.GetFileNameWithoutExtension())} {action.StatusDescription}".Trim();
+			}
+			else if (logLines != null && logLines.Count > 0)
+			{
+				description = $"{(action.CommandDescription ?? action.CommandPath.GetFileNameWithoutExtension())} {logLines[0]}".Trim();
+			}
+			if (!String.IsNullOrEmpty(additionalDescription))
+			{
+				description = $"{description} {additionalDescription}";
+			}
+
+			lock (ProgressWriter)
+			{
+				int completedActions = Interlocked.Increment(ref _loggedCompletedActions);
+				ProgressWriter.Write(completedActions, totalActions);
+
+				// Canceled
+				if (exitCode == Int32.MaxValue)
+				{
+					//Logger.LogInformation("[{CompletedActions}/{TotalActions}] {Description} canceled", completedActions, totalActions, description);
+					return;
+				}
+
+				string targetDetails = "";
+				TargetDescriptor? target = action.Target;
+				if (PrintActionTargetNames && target != null)
+				{
+					targetDetails = $"[{target.Name} {target.Platform} {target.Configuration}]";
+				}
+
+				if (LogActionCommandLines)
+				{
+					Logger.LogDebug("[{CompletedActions}/{TotalActions}]{TargetDetails} Command: {CommandPath} {CommandArguments}", completedActions, totalActions, targetDetails, action.CommandPath, action.CommandArguments);
+				}
+
+				string compilationTimes = "";
+
+				if (ShowPerActionCompilationTimes)
+				{
+					if (processorTime.Ticks > 0)
+					{
+						compilationTimes = $" (Wall: {executionTime.TotalSeconds:0.00}s CPU: {processorTime.TotalSeconds:0.00}s)";
+					}
+					else
+					{
+						compilationTimes = $" (Wall: {executionTime.TotalSeconds:0.00}s)";
+					}
+				}
+
+				string message = ($"[{completedActions}/{totalActions}]{targetDetails}{compilationTimes} {description}");
+
+				if (CompactOutput)
+				{
+					if (s_previousLineLength > 0)
+					{
+						// move the cursor to the far left position, one line back
+						Console.CursorLeft = 0;
+						Console.CursorTop -= 1;
+						// clear the line
+						Console.Write("".PadRight(s_previousLineLength));
+						// move the cursor back to the left, so output is written to the desired location
+						Console.CursorLeft = 0;
+					}
+				}
+				else
+				{
+					// If the action group has changed for a multi target build, write it to the log
+					if (action.GroupNames.Count > 0)
+					{
+						string ActionGroup = $"** For {String.Join(" + ", action.GroupNames)} **";
+						if (!ActionGroup.Equals(LastGroupPrefix, StringComparison.Ordinal))
+						{
+							LastGroupPrefix = ActionGroup;
+							_writeToolOutput(ActionGroup);
+						}
+					}
+				}
+
+				s_previousLineLength = message.Length;
+
+				_writeToolOutput(message);
+				if (logLines != null && action.bShouldOutputLog)
+				{
+					foreach (string Line in logLines.Skip(action.bShouldOutputStatusDescription ? 0 : 1))
+					{
+						// suppress library creation messages when writing compact output
+						if (CompactOutput && Line.StartsWith("   Creating library ", StringComparison.OrdinalIgnoreCase) && Line.EndsWith(".exp", StringComparison.OrdinalIgnoreCase))
+						{
+							continue;
+						}
+
+						_writeToolOutput(Line);
+
+						// Prevent overwriting of logged lines
+						s_previousLineLength = -1;
+					}
+				}
+
+				if (exitCode != 0)
+				{
+					string exitCodeStr = String.Empty;
+					if ((uint)exitCode == 0xC0000005)
+					{
+						exitCodeStr = "(Access violation)";
+					}
+					else if ((uint)exitCode == 0xC0000409)
+					{
+						exitCodeStr = "(Stack buffer overflow)";
+					}
+
+					// If we have an error code but no output, chances are the tool crashed.  Generate more detailed information to let the
+					// user know something went wrong.
+					if (logLines == null || logLines.Count <= (action.bShouldOutputStatusDescription ? 0 : 1))
+					{
+						Logger.LogError("{TargetDetails} {Description}: Exited with error code {ExitCode} {ExitCodeStr}. The build will fail.", targetDetails, description, exitCode, exitCodeStr);
+						Logger.LogInformation("{TargetDetails} {Description}: WorkingDirectory {WorkingDirectory}", targetDetails, description, action.WorkingDirectory);
+						Logger.LogInformation("{TargetDetails} {Description}: {CommandPath} {CommandArguments}", targetDetails, description, action.CommandPath, action.CommandArguments);
+					}
+					// Always print error details to to the log file
+					else
+					{
+						Logger.LogDebug("{TargetDetails} {Description}: Exited with error code {ExitCode} {ExitCodeStr}. The build will fail.", targetDetails, description, exitCode, exitCodeStr);
+						Logger.LogDebug("{TargetDetails} {Description}: WorkingDirectory {WorkingDirectory}", targetDetails, description, action.WorkingDirectory);
+						Logger.LogDebug("{TargetDetails} {Description}: {CommandPath} {CommandArguments}", targetDetails, description, action.CommandPath, action.CommandArguments);
+					}
+
+					// prevent overwriting of error text
+					s_previousLineLength = -1;
+				}
+			}
+		}
+
+		/// <summary>
+		/// Generate the final summary display
+		/// </summary>
+		public void TraceSummary(bool success)
+		{
+
+			// Wait for logging to complete
+			Task? loggingTask = null;
+			lock (_actionsToLog)
+			{
+				loggingTask = _actionsToLogTask;
+			}
+			loggingTask?.Wait();
+			_flushToolOutput();
+
+			LogLevel LogLevel = success ? LogLevel.Information : LogLevel.Debug;
+			if (ShowCPUUtilization)
+			{
+				lock (_cpuUtilization)
+				{
+					if (_cpuUtilization.Count > 0)
+					{
+						Logger.Log(LogLevel, "");
+						Logger.Log(LogLevel, "Average CPU Utilization: {CPUPercentage}%", (int)(_cpuUtilization.Average()));
+					}
+				}
+			}
+
+			if (!ShowCompilationTimes)
+			{
+				return;
+			}
+
+			Logger.Log(LogLevel, "");
+			if (ProcessGroup.TotalProcessorTime.Ticks > 0)
+			{
+				Logger.Log(LogLevel, "Total CPU Time: {TotalSeconds} s", ProcessGroup.TotalProcessorTime.TotalSeconds);
+				Logger.Log(LogLevel, "");
+			}
+
+			IEnumerable<int> CompletedActions = Enumerable.Range(0, Results.Length)
+				.Where(x => Results[x] != null && Results[x]!.ExecutionTime > TimeSpan.Zero && Results[x]!.ExitCode != Int32.MaxValue)
+				.OrderByDescending(x => Results[x]!.ExecutionTime)
+				.Take(20);
+
+			if (CompletedActions.Any())
+			{
+				Logger.Log(LogLevel, "Compilation Time Top {CompletedTaskCount}", CompletedActions.Count());
+				Logger.Log(LogLevel, "");
+				foreach (int Index in CompletedActions)
+				{
+					ExecuteResults Result = Results[Index]!;
+					IExternalAction Action = Result.Action.Inner;
+
+					string Description = $"{(Action.CommandDescription ?? Action.CommandPath.GetFileNameWithoutExtension())} {Action.StatusDescription}".Trim();
+					if (Result.ProcessorTime.Ticks > 0)
+					{
+						Logger.Log(LogLevel, "{Description} [ Wall Time {ExecutionTime:0.00} s / CPU Time {ProcessorTime:0.00} s ]", Description, Result.ExecutionTime.TotalSeconds, Result.ProcessorTime.TotalSeconds);
+					}
+					else
+					{
+						Logger.Log(LogLevel, "{Description} [ Time {ExecutionTime:0.00} s ]", Description, Result.ExecutionTime.TotalSeconds);
+					}
+				}
+				Logger.Log(LogLevel, "");
+			}
+		}
+
+		private void LogActions2()
+		{
+			for (; ; )
+			{
+				int[]? actionsToLog = null;
+				lock (_actionsToLog)
+				{
+					if (_actionsToLog.Count == 0)
+					{
+						_actionsToLogTask = null;
+					}
+					else
+					{
+						actionsToLog = _actionsToLog.ToArray();
+						_actionsToLog.Clear();
+					}
+				}
+
+				if (actionsToLog == null)
+				{
+					return;
+				}
+
+				foreach (int index in actionsToLog)
+				{
+					ExecuteResults result = Results[index]!;
+					LogAction(result.Action, result, Results.Length);
+				}
+			}
+		}
+	}
+}
\ No newline at end of file
diff --git a/Engine/Source/Programs/UnrealBuildTool/Executors/ImmediateActionQueue.cs b/Engine/Source/Programs/UnrealBuildTool/Executors/ImmediateActionQueue.cs
--- a/Engine/Source/Programs/UnrealBuildTool/Executors/ImmediateActionQueue.cs
+++ b/Engine/Source/Programs/UnrealBuildTool/Executors/ImmediateActionQueue.cs
@@ -13,16 +13,6 @@ using UnrealBuildTool.Artifacts;
 
 namespace UnrealBuildTool
 {
-	/// <summary>
-	/// Results from a run action
-	/// </summary>
-	/// <param name="LogLines">Console log lines</param>
-	/// <param name="ExitCode">Process return code.  Zero is assumed to be success and all other values an error.</param>
-	/// <param name="ExecutionTime">Wall time</param>
-	/// <param name="ProcessorTime">CPU time</param>
-	/// <param name="AdditionalDescription">Additional description of action</param>
-	record ExecuteResults(List<string> LogLines, int ExitCode, TimeSpan ExecutionTime, TimeSpan ProcessorTime, string? AdditionalDescription = null);
-
 	/// <summary>
 	/// Defines the phase of execution
 	/// </summary>
@@ -104,7 +94,7 @@ namespace UnrealBuildTool
 	///		UBAExecutor uses an automatic runner to run jobs locally and a manual runner to run jobs remotely as processes 
 	///			become available.
 	/// </summary>
-	class ImmediateActionQueue : IDisposable
+	class ImmediateActionQueue : ActionLogger, IDisposable
 	{
 
 		/// <summary>
@@ -174,11 +164,6 @@ namespace UnrealBuildTool
 			/// </summary>
 			public ImmediateActionQueueRunner? Runner;
 
-			/// <summary>
-			/// Optional execution results
-			/// </summary>
-			public ExecuteResults? Results;
-
 			/// <summary>
 			/// Indices to prereq actions
 			/// </summary>
@@ -190,21 +175,11 @@ namespace UnrealBuildTool
 		/// </summary>
 		private readonly ActionState[] Actions;
 
-		/// <summary>
-		/// Output logging
-		/// </summary>
-		public readonly ILogger Logger;
-
 		/// <summary>
 		/// Total number of actions
 		/// </summary>
 		public int TotalActions => Actions.Length;
 
-		/// <summary>
-		/// Process group
-		/// </summary>		
-		public readonly ManagedProcessGroup ProcessGroup = new();
-
 		/// <summary>
 		/// Source for the cancellation token
 		/// </summary>
@@ -215,46 +190,11 @@ namespace UnrealBuildTool
 		/// </summary>
 		public CancellationToken CancellationToken => CancellationTokenSource.Token;
 
-		/// <summary>
-		/// Progress writer
-		/// </summary>
-		public readonly ProgressWriter ProgressWriter;
-
 		/// <summary>
 		/// Overall success of the action queue
 		/// </summary>
 		public bool Success = true;
 
-		/// <summary>
-		/// Whether to show compilation times along with worst offenders or not.
-		/// </summary>
-		public bool ShowCompilationTimes = false;
-
-		/// <summary>
-		/// Whether to show CPU utilization after the work is complete.
-		/// </summary>
-		public bool ShowCPUUtilization = false;
-
-		/// <summary>
-		/// Add target names for each action executed
-		/// </summary>
-		public bool PrintActionTargetNames = false;
-
-		/// <summary>
-		/// Whether to log command lines for actions being executed
-		/// </summary>
-		public bool LogActionCommandLines = false;
-
-		/// <summary>
-		/// Whether to show compilation times for each executed action
-		/// </summary>
-		public bool ShowPerActionCompilationTimes = false;
-
-		/// <summary>
-		/// Collapse non-error output lines
-		/// </summary>
-		public bool CompactOutput = false;
-
 		/// <summary>
 		/// When enabled, will stop compiling targets after a compile error occurs.
 		/// </summary>
@@ -289,41 +229,6 @@ namespace UnrealBuildTool
 		/// </summary>
 		private int _firstPendingAction;
 
-		/// <summary>
-		/// Action to invoke when writing tool output
-		/// </summary>
-		private readonly Action<string> _writeToolOutput;
-
-		/// <summary>
-		/// Flush the tool output after logging has completed
-		/// </summary>
-		private readonly System.Action _flushToolOutput;
-
-		/// <summary>
-		/// Timer used to collect CPU utilization
-		/// </summary>
-		private Timer? _cpuUtilizationTimer;
-
-		/// <summary>
-		/// Per-second logging of cpu utilization
-		/// </summary>
-		private readonly List<float> _cpuUtilization = new();
-
-		/// <summary>
-		/// Collection of all actions remaining to be logged
-		/// </summary>
-		private readonly List<int> _actionsToLog = new();
-
-		/// <summary>
-		/// Task waiting to process logging
-		/// </summary>
-		private Task? _actionsToLogTask = null;
-
-		/// <summary>
-		/// Used only by the logger to track the [x,total] output
-		/// </summary>
-		private int _loggedCompletedActions = 0;
-
 		/// <summary>
 		/// Tracks the number of completed actions.
 		/// </summary>
@@ -374,11 +279,6 @@ namespace UnrealBuildTool
 		/// </summary>
 		private readonly TaskCompletionSource _doneTaskSource = new();
 
-		/// <summary>
-		/// The last action group printed in multi-target builds
-		/// </summary>
-		private string? LastGroupPrefix = null;
-
 		/// <summary>
 		/// If set, artifact cache used to retrieve previously compiled results and save new results
 		/// </summary>
@@ -395,6 +295,7 @@ namespace UnrealBuildTool
 		/// <param name="flushToolOutput">Action to invoke when flushing tool output</param>
 		/// <param name="logger">Logging interface</param>
 		public ImmediateActionQueue(IEnumerable<LinkedAction> actions, IActionArtifactCache? actionArtifactCache, int maxActionArtifactCacheTasks, string progressWriterText, Action<string> writeToolOutput, System.Action flushToolOutput, ILogger logger)
+		:	base(actions.Count(), progressWriterText, writeToolOutput, flushToolOutput, logger)
 		{
 			CommandLine.ParseArguments(Environment.GetCommandLineArgs(), this, logger);
 			XmlConfig.ApplyTo(this);
@@ -402,11 +303,7 @@ namespace UnrealBuildTool
 			int count = actions.Count();
 			Actions = new ActionState[count];
 
-			Logger = logger;
-			ProgressWriter = new(progressWriterText, false, logger);
 			_actionArtifactCache = actionArtifactCache;
-			_writeToolOutput = writeToolOutput;
-			_flushToolOutput = flushToolOutput;
 
 			bool readArtifacts = _actionArtifactCache != null && _actionArtifactCache.EnableReads;
 			ActionPhase initialPhase = readArtifacts ? ActionPhase.ArtifactCheck : ActionPhase.Compile;
@@ -419,7 +316,6 @@ namespace UnrealBuildTool
 					Action = action,
 					Status = ActionStatus.Queued,
 					Phase = action.ArtifactMode.HasFlag(ArtifactMode.Enabled) ? initialPhase : ActionPhase.Compile,
-					Results = null,
 					PrereqActionsSortIndex = action.PrerequisiteActions.Select(a => a.SortIndex).ToArray()
 				};
 			}
@@ -433,7 +329,7 @@ namespace UnrealBuildTool
 						ActionArtifactResult result = await _actionArtifactCache!.CompleteActionFromCacheAsync(action, CancellationToken);
 						if (result.Success)
 						{
-							ExecuteResults results = new(result.LogLines, result.ExitCode, TimeSpan.Zero, TimeSpan.Zero, "[Cache]");
+							ExecuteResults results = new(result.LogLines, result.ExitCode, TimeSpan.Zero, TimeSpan.Zero, action, "[Cache]");
 							Interlocked.Increment(ref _cacheHitActions);
 							OnActionCompleted(action, results.ExitCode == 0, results);
 							OnArtifactsRead?.Invoke(action);
@@ -576,7 +472,7 @@ namespace UnrealBuildTool
 		public async Task<bool> RunTillDone()
 		{
 			await _doneTaskSource.Task;
-			TraceSummary();
+			TraceSummary(Success);
 			return Success;
 		}
 
@@ -591,8 +487,8 @@ namespace UnrealBuildTool
 		public void GetActionResultCounts(out int totalActions, out int succeededActions, out int failedActions, out int cacheHitActions, out int cacheMissActions)
 		{
 			totalActions = Actions.Length;
-			succeededActions = Actions.Where(x => x.Results?.ExitCode == 0).Count();
-			failedActions = Actions.Where(x => x.Results != null && x.Results.ExitCode != 0).Count();
+			succeededActions = Results.Where(x => x?.ExitCode == 0).Count();
+			failedActions = Results.Where(x => x != null && x.ExitCode != 0).Count();
 			cacheHitActions = _cacheHitActions;
 			cacheMissActions = _cacheMissActions;
 		}
@@ -732,13 +628,13 @@ namespace UnrealBuildTool
 				}
 				else
 				{
-					ExecuteResults results = new(new List<string>(), Int32.MaxValue, TimeSpan.Zero, TimeSpan.Zero);
+					ExecuteResults results = new(new List<string>(), Int32.MaxValue, TimeSpan.Zero, TimeSpan.Zero, action);
 					OnActionCompleted(action, false, results);
 				}
 			}
 			else if (ex is OperationCanceledException)
 			{
-				ExecuteResults results = new(new List<string>(), Int32.MaxValue, TimeSpan.Zero, TimeSpan.Zero);
+				ExecuteResults results = new(new List<string>(), Int32.MaxValue, TimeSpan.Zero, TimeSpan.Zero, action);
 				OnActionCompleted(action, false, results);
 			}
 			else
@@ -748,7 +644,7 @@ namespace UnrealBuildTool
 					ExceptionUtils.FormatException(ex),
 					ExceptionUtils.FormatExceptionDetails(ex),
 				};
-				ExecuteResults results = new(text, Int32.MaxValue, TimeSpan.Zero, TimeSpan.Zero);
+				ExecuteResults results = new(text, Int32.MaxValue, TimeSpan.Zero, TimeSpan.Zero, action);
 				OnActionCompleted(action, false, results);
 			}
 		}
@@ -965,7 +861,7 @@ namespace UnrealBuildTool
 				Actions[actionIndex].Runner = null;
 				if (results != null)
 				{
-					Actions[actionIndex].Results = results;
+					Results[actionIndex] = results;
 				}
 
 				// Add the action to the logging queue
@@ -1095,170 +991,11 @@ namespace UnrealBuildTool
 
 				foreach (int index in actionsToLog)
 				{
-					LogAction(Actions[index].Action, Actions[index].Results);
-				}
-			}
-		}
-
-		private static int s_previousLineLength = -1;
-
-		/// <summary>
-		/// Log an action that has completed
-		/// </summary>
-		/// <param name="action">Action that has completed</param>
-		/// <param name="executeTaskResult">Results of the action</param>
-		private void LogAction(LinkedAction action, ExecuteResults? executeTaskResult)
-		{
-			List<string>? logLines = null;
-			int exitCode = Int32.MaxValue;
-			TimeSpan executionTime = TimeSpan.Zero;
-			TimeSpan processorTime = TimeSpan.Zero;
-			string? additionalDescription = null;
-			if (executeTaskResult != null)
-			{
-				logLines = executeTaskResult.LogLines;
-				exitCode = executeTaskResult.ExitCode;
-				executionTime = executeTaskResult.ExecutionTime;
-				processorTime = executeTaskResult.ProcessorTime;
-				additionalDescription = executeTaskResult.AdditionalDescription;
-			}
-
-			// Write it to the log
-			string description = String.Empty;
-			if (action.bShouldOutputStatusDescription || (logLines != null && logLines.Count == 0))
-			{
-				description = $"{(action.CommandDescription ?? action.CommandPath.GetFileNameWithoutExtension())} {action.StatusDescription}".Trim();
-			}
-			else if (logLines != null && logLines.Count > 0)
-			{
-				description = $"{(action.CommandDescription ?? action.CommandPath.GetFileNameWithoutExtension())} {logLines[0]}".Trim();
-			}
-			if (!String.IsNullOrEmpty(additionalDescription))
-			{
-				description = $"{description} {additionalDescription}";
-			}
-
-			lock (ProgressWriter)
-			{
-				int totalActions = Actions.Length;
-				int completedActions = Interlocked.Increment(ref _loggedCompletedActions);
-				ProgressWriter.Write(completedActions, Actions.Length);
-
-				// Canceled
-				if (exitCode == Int32.MaxValue)
-				{
-					//Logger.LogInformation("[{CompletedActions}/{TotalActions}] {Description} canceled", completedActions, totalActions, description);
-					return;
-				}
-
-				string targetDetails = "";
-				TargetDescriptor? target = action.Target;
-				if (PrintActionTargetNames && target != null)
-				{
-					targetDetails = $"[{target.Name} {target.Platform} {target.Configuration}]";
-				}
-
-				if (LogActionCommandLines)
-				{
-					Logger.LogDebug("[{CompletedActions}/{TotalActions}]{TargetDetails} Command: {CommandPath} {CommandArguments}", completedActions, totalActions, targetDetails, action.CommandPath, action.CommandArguments);
-				}
-
-				string compilationTimes = "";
-
-				if (ShowPerActionCompilationTimes)
-				{
-					if (processorTime.Ticks > 0)
-					{
-						compilationTimes = $" (Wall: {executionTime.TotalSeconds:0.00}s CPU: {processorTime.TotalSeconds:0.00}s)";
-					}
-					else
-					{
-						compilationTimes = $" (Wall: {executionTime.TotalSeconds:0.00}s)";
-					}
-				}
-
-				string message = ($"[{completedActions}/{totalActions}]{targetDetails}{compilationTimes} {description}");
-
-				if (CompactOutput)
-				{
-					if (s_previousLineLength > 0)
-					{
-						// move the cursor to the far left position, one line back
-						Console.CursorLeft = 0;
-						Console.CursorTop -= 1;
-						// clear the line
-						Console.Write("".PadRight(s_previousLineLength));
-						// move the cursor back to the left, so output is written to the desired location
-						Console.CursorLeft = 0;
-					}
-				}
-				else
-				{
-					// If the action group has changed for a multi target build, write it to the log
-					if (action.GroupNames.Count > 0)
-					{
-						string ActionGroup = $"** For {String.Join(" + ", action.GroupNames)} **";
-						if (!ActionGroup.Equals(LastGroupPrefix, StringComparison.Ordinal))
-						{
-							LastGroupPrefix = ActionGroup;
-							_writeToolOutput(ActionGroup);
-						}
-					}
-				}
-
-				s_previousLineLength = message.Length;
-
-				_writeToolOutput(message);
-				if (logLines != null && action.bShouldOutputLog)
-				{
-					foreach (string Line in logLines.Skip(action.bShouldOutputStatusDescription ? 0 : 1))
-					{
-						// suppress library creation messages when writing compact output
-						if (CompactOutput && Line.StartsWith("   Creating library ", StringComparison.OrdinalIgnoreCase) && Line.EndsWith(".exp", StringComparison.OrdinalIgnoreCase))
-						{
-							continue;
-						}
-
-						_writeToolOutput(Line);
-
-						// Prevent overwriting of logged lines
-						s_previousLineLength = -1;
-					}
-				}
-
-				if (exitCode != 0)
-				{
-					string exitCodeStr = String.Empty;
-					if ((uint)exitCode == 0xC0000005)
-					{
-						exitCodeStr = "(Access violation)";
-					}
-					else if ((uint)exitCode == 0xC0000409)
-					{
-						exitCodeStr = "(Stack buffer overflow)";
-					}
-
-					// If we have an error code but no output, chances are the tool crashed.  Generate more detailed information to let the
-					// user know something went wrong.
-					if (logLines == null || logLines.Count <= (action.bShouldOutputStatusDescription ? 0 : 1))
-					{
-						Logger.LogError("{TargetDetails} {Description}: Exited with error code {ExitCode} {ExitCodeStr}. The build will fail.", targetDetails, description, exitCode, exitCodeStr);
-						Logger.LogInformation("{TargetDetails} {Description}: WorkingDirectory {WorkingDirectory}", targetDetails, description, action.WorkingDirectory);
-						Logger.LogInformation("{TargetDetails} {Description}: {CommandPath} {CommandArguments}", targetDetails, description, action.CommandPath, action.CommandArguments);
-					}
-					// Always print error details to to the log file
-					else
-					{
-						Logger.LogDebug("{TargetDetails} {Description}: Exited with error code {ExitCode} {ExitCodeStr}. The build will fail.", targetDetails, description, exitCode, exitCodeStr);
-						Logger.LogDebug("{TargetDetails} {Description}: WorkingDirectory {WorkingDirectory}", targetDetails, description, action.WorkingDirectory);
-						Logger.LogDebug("{TargetDetails} {Description}: {CommandPath} {CommandArguments}", targetDetails, description, action.CommandPath, action.CommandArguments);
-					}
-
-					// prevent overwriting of error text
-					s_previousLineLength = -1;
+					ExecuteResults? result = Results[index];
+					LogAction(Actions[index].Action, result, Actions.Length);
 
 					// Cancel all other pending tasks
-					if (StopCompilationAfterErrors)
+					if (StopCompilationAfterErrors && (result == null || result.ExitCode != 0))
 					{
 						CancellationTokenSource.Cancel();
 					}
@@ -1266,74 +1003,6 @@ namespace UnrealBuildTool
 			}
 		}
 
-		/// <summary>
-		/// Generate the final summary display
-		/// </summary>
-		private void TraceSummary()
-		{
-
-			// Wait for logging to complete
-			Task? loggingTask = null;
-			lock (_actionsToLog)
-			{
-				loggingTask = _actionsToLogTask;
-			}
-			loggingTask?.Wait();
-			_flushToolOutput();
-
-			LogLevel LogLevel = Success ?  LogLevel.Information : LogLevel.Debug;
-			if (ShowCPUUtilization)
-			{
-				lock (_cpuUtilization)
-				{
-					if (_cpuUtilization.Count > 0)
-					{
-						Logger.Log(LogLevel, "");
-						Logger.Log(LogLevel, "Average CPU Utilization: {CPUPercentage}%", (int)(_cpuUtilization.Average()));
-					}
-				}
-			}
-
-			if (!ShowCompilationTimes)
-			{
-				return;
-			}
-
-			Logger.Log(LogLevel, "");
-			if (ProcessGroup.TotalProcessorTime.Ticks > 0)
-			{
-				Logger.Log(LogLevel, "Total CPU Time: {TotalSeconds} s", ProcessGroup.TotalProcessorTime.TotalSeconds);
-				Logger.Log(LogLevel, "");
-			}
-
-			IEnumerable<int> CompletedActions = Enumerable.Range(0, Actions.Length)
-				.Where(x => Actions[x].Results != null && Actions[x].Results!.ExecutionTime > TimeSpan.Zero && Actions[x].Results!.ExitCode != Int32.MaxValue)
-				.OrderByDescending(x => Actions[x].Results!.ExecutionTime)
-				.Take(20);
-
-			if (CompletedActions.Any())
-			{
-				Logger.Log(LogLevel, "Compilation Time Top {CompletedTaskCount}", CompletedActions.Count());
-				Logger.Log(LogLevel, "");
-				foreach (int Index in CompletedActions)
-				{
-					IExternalAction Action = Actions[Index].Action.Inner;
-					ExecuteResults Result = Actions[Index].Results!;
-
-					string Description = $"{(Action.CommandDescription ?? Action.CommandPath.GetFileNameWithoutExtension())} {Action.StatusDescription}".Trim();
-					if (Result.ProcessorTime.Ticks > 0)
-					{
-						Logger.Log(LogLevel, "{Description} [ Wall Time {ExecutionTime:0.00} s / CPU Time {ProcessorTime:0.00} s ]", Description, Result.ExecutionTime.TotalSeconds, Result.ProcessorTime.TotalSeconds);
-					}
-					else
-					{
-						Logger.Log(LogLevel, "{Description} [ Time {ExecutionTime:0.00} s ]", Description, Result.ExecutionTime.TotalSeconds);
-					}
-				}
-				Logger.Log(LogLevel, "");
-			}
-		}
-
 		private enum ActionReadyState
 		{
 			NotReady,
diff --git a/Engine/Source/Programs/UnrealBuildTool/Executors/ParallelExecutor.cs b/Engine/Source/Programs/UnrealBuildTool/Executors/ParallelExecutor.cs
--- a/Engine/Source/Programs/UnrealBuildTool/Executors/ParallelExecutor.cs
+++ b/Engine/Source/Programs/UnrealBuildTool/Executors/ParallelExecutor.cs
@@ -57,31 +57,31 @@ namespace UnrealBuildTool
 		/// When enabled, will stop compiling targets after a compile error occurs.
 		/// </summary>
 		[XmlConfigFile]
-		private static bool bStopCompilationAfterErrors = false;
+		protected static bool bStopCompilationAfterErrors = false;
 
 		/// <summary>
 		/// Whether to show compilation times along with worst offenders or not.
 		/// </summary>
 		[XmlConfigFile]
-		private static bool bShowCompilationTimes = Unreal.IsBuildMachine();
+		protected static bool bShowCompilationTimes = Unreal.IsBuildMachine();
 
 		/// <summary>
 		/// Whether to show compilation times for each executed action
 		/// </summary>
 		[XmlConfigFile]
-		private static bool bShowPerActionCompilationTimes = Unreal.IsBuildMachine();
+		protected static bool bShowPerActionCompilationTimes = Unreal.IsBuildMachine();
 
 		/// <summary>
 		/// Whether to log command lines for actions being executed
 		/// </summary>
 		[XmlConfigFile]
-		private static bool bLogActionCommandLines = false;
+		protected static bool bLogActionCommandLines = false;
 
 		/// <summary>
 		/// Add target names for each action executed
 		/// </summary>
 		[XmlConfigFile]
-		private static bool bPrintActionTargetNames = false;
+		protected static bool bPrintActionTargetNames = false;
 
 		/// <summary>
 		/// Whether to take into account the Action's weight when determining to do more work or not.
@@ -98,7 +98,7 @@ namespace UnrealBuildTool
 		/// <summary>
 		/// Collapse non-error output lines
 		/// </summary>
-		private bool bCompactOutput = false;
+		protected bool bCompactOutput = false;
 
 		/// <summary>
 		/// How many processes that will be executed in parallel
@@ -255,7 +255,7 @@ namespace UnrealBuildTool
 				LogLines = LogLines.Select(x => x.Replace("): warning: ", "): error: ", StringComparison.OrdinalIgnoreCase)).ToList();
 			}
 
-			return new ExecuteResults(LogLines, ExitCode, ExecutionTime, ProcessorTime, AdditionalDescription);
+			return new ExecuteResults(LogLines, ExitCode, ExecutionTime, ProcessorTime, Action, AdditionalDescription);
 		}
 	}
 
diff --git a/Engine/Source/Programs/UnrealBuildTool/Executors/UnrealBuildAccelerator/UBAAgentCoordinator.cs b/Engine/Source/Programs/UnrealBuildTool/Executors/UnrealBuildAccelerator/UBAAgentCoordinator.cs
--- a/Engine/Source/Programs/UnrealBuildTool/Executors/UnrealBuildAccelerator/UBAAgentCoordinator.cs
+++ b/Engine/Source/Programs/UnrealBuildTool/Executors/UnrealBuildAccelerator/UBAAgentCoordinator.cs
@@ -10,13 +10,20 @@ namespace UnrealBuildTool
 	// statusRow, statusColumn, statusText, statusType, statusLink
 	using StatusUpdateAction = Action<uint, uint, string, LogEntryType, string?>;
 
+	interface IUbaAgentCoordinatorScheduler
+	{
+		bool IsEmpty { get; }
+		double GetProcessWeightThatCanRunRemotelyNow();
+		bool AddClient(string ip, int port, string crypto = "");
+	}
+
 	interface IUBAAgentCoordinator
 	{
 		DirectoryReference? GetUBARootDir();
 
 		Task InitAsync(UBAExecutor executor);
 
-		void Start(ImmediateActionQueue queue, Func<LinkedAction, bool> canRunRemotely, StatusUpdateAction updateStatus);
+		void Start(IUbaAgentCoordinatorScheduler scheduler, Func<LinkedAction, bool> canRunRemotely, StatusUpdateAction updateStatus);
 
 		void Stop();
 
diff --git a/Engine/Source/Programs/UnrealBuildTool/Executors/UnrealBuildAccelerator/UBAAgentCoordinatorHorde.cs b/Engine/Source/Programs/UnrealBuildTool/Executors/UnrealBuildAccelerator/UBAAgentCoordinatorHorde.cs
--- a/Engine/Source/Programs/UnrealBuildTool/Executors/UnrealBuildAccelerator/UBAAgentCoordinatorHorde.cs
+++ b/Engine/Source/Programs/UnrealBuildTool/Executors/UnrealBuildAccelerator/UBAAgentCoordinatorHorde.cs
@@ -48,8 +48,9 @@ namespace UnrealBuildTool
 		readonly CancellationTokenSource _cancellationTokenSource = new();
 		readonly ILogger _logger;
 
-		readonly UBAExecutor _owner;
+		readonly UnrealBuildAcceleratorConfig _config;
 		internal StatusUpdateAction? _updateStatus;
+		internal IUbaAgentCoordinatorScheduler? _scheduler;
 
 		readonly BundleStorageNamespace _storage = BundleStorageNamespace.CreateInMemory(NullLogger.Instance);
 		BlobLocator _ubaAgentLocator;
@@ -73,9 +74,9 @@ namespace UnrealBuildTool
 
 		readonly List<Worker> _workers = new();
 
-		public UBAHordeSession(UBAExecutor owner, Uri? serverUrl, string? accessToken, string? pool, bool allowWine, int maxCores, bool strict, ConnectionMode? connectionMode, Encryption? encryption, ILogger logger)
+		public UBAHordeSession(UnrealBuildAcceleratorConfig config, string crypto, Uri? serverUrl, string? accessToken, string? pool, bool allowWine, int maxCores, bool strict, ConnectionMode? connectionMode, Encryption? encryption, ILogger logger)
 		{
-			_owner = owner;
+			_config = config;
 			_pool = pool;
 			_allowWine = allowWine;
 			_maxCores = maxCores;
@@ -83,7 +84,7 @@ namespace UnrealBuildTool
 			_connectionMode = connectionMode;
 			_encryption = encryption;
 			_logger = logger;
-			_crypto = owner.Crypto;
+			_crypto = crypto;
 
 			ServiceCollection services = new();
 			services.AddLogging(builder => builder.AddEpicDefault());
@@ -146,6 +147,7 @@ namespace UnrealBuildTool
 			if (OperatingSystem.IsWindows())
 			{
 				agentFiles.Add("UbaAgent.exe");
+				agentFiles.Add("UbaWine.dll.so"); // Not needed but used by linux/wine helpers for more detailed logging or optimizations
 				if (useSentry)
 				{
 					agentFiles.Add("crashpad_handler.exe");
@@ -424,7 +426,7 @@ namespace UnrealBuildTool
 			}
 		}
 
-		public static async Task<(UBAHordeSession?, string)> TryCreateHordeSessionAsync(UnrealBuildAcceleratorHordeConfig hordeConfig, UBAExecutor executor, bool bStrictErrors, ILogger logger, CancellationToken cancellationToken = default)
+		public static async Task<(UBAHordeSession?, string)> TryCreateHordeSessionAsync(UnrealBuildAcceleratorHordeConfig hordeConfig, UnrealBuildAcceleratorConfig ubaConfig, string crypto, bool bStrictErrors, ILogger logger, CancellationToken cancellationToken = default)
 		{
 			if (hordeConfig.bDisableHorde)
 			{
@@ -456,7 +458,7 @@ namespace UnrealBuildTool
 			{
 				bool allowWine = hordeConfig.bHordeAllowWine && OperatingSystem.IsWindows();
 
-				UBAHordeSession session = new(executor, server, token, hordeConfig.HordePool, allowWine, hordeConfig.HordeMaxCores, bStrictErrors, connectionMode, encryption, logger);
+				UBAHordeSession session = new(ubaConfig, crypto, server, token, hordeConfig.HordePool, allowWine, hordeConfig.HordeMaxCores, bStrictErrors, connectionMode, encryption, logger);
 				await session.InitAsync(useSentry: !String.IsNullOrEmpty(hordeConfig.UBASentryUrl), cancellationToken);
 				return (session, "");
 			}
@@ -526,7 +528,7 @@ namespace UnrealBuildTool
 						logger.LogDebug("Uploading files...");
 						await channel.UploadFilesAsync("", tool, _storage.Backend, cancellationToken);
 
-						string hordeHost = _owner.UBAConfig.Host;
+						string hordeHost = _config.Host;
 						if (!String.IsNullOrEmpty(hordeConfig.HordeHost))
 						{
 							hordeHost = hordeConfig.HordeHost;
@@ -542,7 +544,7 @@ namespace UnrealBuildTool
 
 						if (useListen)
 						{
-							arguments.Add($"-Host={hordeHost}:{_owner.UBAConfig.Port}");
+							arguments.Add($"-Host={hordeHost}:{_config.Port}");
 						}
 						else
 						{
@@ -562,7 +564,7 @@ namespace UnrealBuildTool
 							arguments.Add($"-Sentry=\"{hordeConfig.UBASentryUrl}\"");
 						}
 						arguments.Add("-ProxyPort=" + self.ProxyPort!.AgentPort);
-						if (_owner.UBAConfig.bUseQuic)
+						if (_config.bUseQuic)
 						{
 							arguments.Add("-quic");
 						}
@@ -586,17 +588,18 @@ namespace UnrealBuildTool
 						arguments.Add("-MaxIdle=15");
 						arguments.Add("-UseCrawler");
 						arguments.Add("-ProxyUseLocalStorage"); // If set then agent's cas storage will be used by proxy
+						//arguments.Add("-UseIocp=4");
 						//arguments.Add("-NoStore"); // This means that no cas files will be written to disk.. so no state is stored between runs
 
 						if (!String.IsNullOrEmpty(desc))
 							arguments.Add($"-Description=\"{desc}\"");
 
-						if (_owner.UBAConfig.bLogEnabled)
+						if (_config.bLogEnabled)
 						{
 							arguments.Add("-Log");
 						}
 
-						LogLevel logLevel = _owner.UBAConfig.bDetailedLog ? LogLevel.Information : LogLevel.Debug;
+						LogLevel logLevel = _config.bDetailedLog ? LogLevel.Information : LogLevel.Debug;
 
 						logger.Log(logLevel, "Executing child process: {Executable} {Arguments}", executable, CommandLineArguments.Join(arguments));
 
@@ -616,7 +619,7 @@ namespace UnrealBuildTool
 								logger.LogInformation("Connecting to UbaAgent on {Ip}:{Port} (local agent port {AgentPort}) {Seconds}.{Milliseconds} seconds after assigned",
 									self.Ip, self.Port!.Port, self.Port.AgentPort, totalMs / 1000, totalMs % 1000);
 
-								if (!_owner.Server!.AddClient(self.Ip!, self.Port.Port, _crypto))
+								if (!_scheduler!.AddClient(self.Ip!, self.Port.Port, _crypto))
 								{
 									break;
 								}
@@ -710,7 +713,7 @@ namespace UnrealBuildTool
 
 		public static string ProviderPrefix = "Uba.Provider.Horde";
 
-		public static IEnumerable<IUBAAgentCoordinator> Init(ILogger logger, UBAExecutor executor, IEnumerable<TargetDescriptor> targetDescriptors, UnrealBuildAcceleratorConfig ubaConfig, ref string remoteConnectionMode)
+		public static IEnumerable<IUBAAgentCoordinator> Init(ILogger logger, IEnumerable<TargetDescriptor> targetDescriptors, UnrealBuildAcceleratorConfig ubaConfig, ref string remoteConnectionMode)
 		{
 			IEnumerable<string> providers = Unreal.IsBuildMachine()
 				? [.. ubaConfig.BuildMachineProviders, .. ubaConfig.IniBuildMachineProviders]
@@ -791,14 +794,13 @@ namespace UnrealBuildTool
 				return;
 			}
 
-			_executor = executor;
 			_cancellationSource = new CancellationTokenSource();
-			_hordeSessionTask = UBAHordeSession.TryCreateHordeSessionAsync(_hordeConfig, executor, _ubaConfig.bStrict, _logger, _cancellationSource.Token);
+			_hordeSessionTask = UBAHordeSession.TryCreateHordeSessionAsync(_hordeConfig, _ubaConfig, executor.Crypto, _ubaConfig.bStrict, _logger, _cancellationSource.Token);
 			await _hordeSessionTask;
 			executor.AgentCoordinatorInitialized(this, _hordeSessionTask.IsCompletedSuccessfully && _hordeSessionTask.Result.Item1 != null);
 		}
 
-		public void Start(ImmediateActionQueue queue, Func<LinkedAction, bool> canRunRemotely, StatusUpdateAction updateStatus)
+		public void Start(IUbaAgentCoordinatorScheduler scheduler, Func<LinkedAction, bool> canRunRemotely, StatusUpdateAction updateStatus)
 		{
 			_updateStatus = updateStatus;
 
@@ -829,6 +831,7 @@ namespace UnrealBuildTool
 				return;
 			}
 
+			hordeSession._scheduler = scheduler;
 			hordeSession._updateStatus = _updateStatus;
 			hordeSession.UpdateHordeStatus(null);
 
@@ -850,7 +853,7 @@ namespace UnrealBuildTool
 					return;
 				}
 
-				if (queue.IsDone || _cancellationSource!.IsCancellationRequested)
+				if (scheduler.IsEmpty || _cancellationSource!.IsCancellationRequested)
 				{
 					hordeSession.UpdateHordeStatus(" - Requests stopped");
 					return;
@@ -871,7 +874,7 @@ namespace UnrealBuildTool
 
 						hordeSession.RemoveCompleteWorkersAsync();
 
-						double queueWeight = queue.EnumerateReadyToCompileActions().Where(x => canRunRemotely(x)).Sum(x => x.Weight);
+						double queueWeight = scheduler.GetProcessWeightThatCanRunRemotelyNow();
 
 						int activeCores;
 						int queuedCores;
@@ -1019,7 +1022,6 @@ namespace UnrealBuildTool
 		readonly UnrealBuildAcceleratorConfig _ubaConfig;
 		readonly UnrealBuildAcceleratorHordeConfig _hordeConfig;
 
-		UBAExecutor? _executor;
 		CancellationTokenSource? _cancellationSource;
 		Task<(UBAHordeSession?, string)>? _hordeSessionTask;
 		StatusUpdateAction? _updateStatus;
diff --git a/Engine/Source/Programs/UnrealBuildTool/Executors/UnrealBuildAccelerator/UBAExecutor.cs b/Engine/Source/Programs/UnrealBuildTool/Executors/UnrealBuildAccelerator/UBAExecutor.cs
--- a/Engine/Source/Programs/UnrealBuildTool/Executors/UnrealBuildAccelerator/UBAExecutor.cs
+++ b/Engine/Source/Programs/UnrealBuildTool/Executors/UnrealBuildAccelerator/UBAExecutor.cs
@@ -24,7 +24,6 @@ namespace UnrealBuildTool
 		public UnrealBuildAcceleratorConfig UBAConfig { get; init; } = new();
 
 		public string Crypto { get; private set; } = String.Empty;
-		public IServer? Server { get; private set; }
 		ISessionServer? _session;
 		object _sessionLock = new();
 		IEnumerable<UnrealBuildAcceleratorCacheConfig> _cacheConfigs = [];
@@ -33,31 +32,16 @@ namespace UnrealBuildTool
 		readonly List<IUBAAgentCoordinator> _agentCoordinators = new();
 		DirectoryReference? _rootDirRef;
 		bool _bIsCancelled;
-		bool _bIsRemoteActionsAllowed = true;
-		readonly object _actionsChangedLock = new();
-		bool _bActionsChanged = true;
-		uint _errorCount = 0;
-		uint _actionsQueuedThatCanRunRemotely = UInt32.MaxValue;
 		readonly ThreadedLogger _threadedLogger;
 
 		DateTime _ubaStartTimeUtc = DateTime.UtcNow;
 		TimeSpan _ubaDurationWaitingForRemote = TimeSpan.Zero;
 
-		// Tracking for LinkedActions that failed remotely that should be retried locally
-		readonly ConcurrentDictionary<LinkedAction, bool> _localRetryActions = new();
-		// Tracking for LinkedActions that failed locally that should be retried without UBA
-		readonly ConcurrentDictionary<LinkedAction, bool> _forcedRetryActions = new();
-
 		// Tracking for successful coordinator connections
 		int _successfulCoordinatorConnections = 0;
 		// Tracking for failed coordinator connections
 		int _failedCoordinatorConnections = 0;
 
-		// Tracking for all actions processed locally
-		int _localProcessedActions = 0;
-		// Tracking for all actions processed remotely
-		int _remoteProcessedActions = 0;
-
 		int _cacheableActionsLeft = 0;
 
 		// Tracking for remote connection mode
@@ -149,7 +133,7 @@ namespace UnrealBuildTool
 				Crypto = CreateCrypto();
 			}
 
-			_agentCoordinators.AddRange(UBAAgentCoordinatorHorde.Init(logger, this, targetDescriptors, UBAConfig, ref _remoteConnectionMode));
+			_agentCoordinators.AddRange(UBAAgentCoordinatorHorde.Init(logger, targetDescriptors, UBAConfig, ref _remoteConnectionMode));
 			_threadedLogger = new ThreadedLogger(logger);
 		}
 
@@ -277,102 +261,6 @@ namespace UnrealBuildTool
 			}
 		}
 
-		class UBAArtifactCache : IArtifactCache
-		{
-			public ArtifactCacheState State => ArtifactCacheState.Available;
-			public Task<ArtifactCacheState> WaitForReadyAsync() => Task.FromResult(ArtifactCacheState.Available);
-			public Task<ArtifactAction[]> QueryArtifactActionsAsync(IoHash[] partialKeys, CancellationToken cancellationToken) => Task.FromResult(Array.Empty<ArtifactAction>());
-			public Task<bool[]?> QueryArtifactOutputsAsync(ArtifactAction[] artifactActions, CancellationToken cancellationToken) => Task.FromResult<bool[]?>(null);
-			public Task SaveArtifactActionsAsync(ArtifactAction[] artifactActions, CancellationToken cancellationToken) => Task.CompletedTask;
-			public Task FlushChangesAsync(CancellationToken cancellationToken) => Task.CompletedTask;
-		}
-
-		private static void RemoveLogLineSpam(List<string> logLines)
-		{
-			logLines.RemoveAll(line => 
-				line.StartsWith("   Creating library ", StringComparison.Ordinal) ||
-				line.StartsWith("   Creating object ", StringComparison.Ordinal) ||
-				line.EndsWith("file(s) copied.", StringComparison.Ordinal));
-		}
-
-		class UBAActionArtifactCache : IActionArtifactCache
-		{
-			readonly UBAExecutor _executor;
-			readonly UBAArtifactCache _cache = new UBAArtifactCache();
-			public UBAActionArtifactCache(UBAExecutor executor) { _executor = executor; }
-			public IArtifactCache ArtifactCache => _cache;
-			public bool EnableReads { get => true; set { } }
-			public bool EnableWrites { get => true; set { } }
-			public bool LogCacheMisses { get => true; set { } }
-			public DirectoryReference? EngineRoot { get => null; set { } }
-			public DirectoryReference[]? DirectoryRoots { get => null; set { } }
-			public Task<ActionArtifactResult> CompleteActionFromCacheAsync(LinkedAction action, CancellationToken cancellationToken)
-			{
-				return Task.Factory.StartNew(() =>
-				{
-					ProcessStartInfo startInfo = _executor.GetActionStartInfo(action);
-					uint bucket = _executor.GetCacheBucket(action);
-					ulong rootsHandle = _executor.GetActionRootsHandle(action);
-
-					IEnumerable<KeyValuePair<UnrealBuildAcceleratorCacheConfig, ICacheClient>> cacheClients;
-					if (_executor.UBAConfig.bCacheShuffle && _executor._cacheClients.Count > 1)
-					{
-						KeyValuePair<UnrealBuildAcceleratorCacheConfig, ICacheClient>[] array = [.. _executor._cacheClients];
-						Random r = new Random();
-						r.Shuffle(array);
-						cacheClients = array;
-					}
-					else
-					{
-						cacheClients = _executor._cacheClients;
-					}
-
-					foreach (KeyValuePair<UnrealBuildAcceleratorCacheConfig, ICacheClient> item in cacheClients)
-					{
-						if (item.Key.bRequireVfs && !action.RootPaths.bUseVfs)
-						{
-							continue;
-						}
-						FetchFromCacheResult result = item.Value.FetchFromCache(rootsHandle, bucket, startInfo);
-						int exitCode = 0;
-						RemoveLogLineSpam(result.LogLines);
-						if (result.Success && action.bForceWarningsAsError && result.LogLines.Any(x => x.Contains("): warning: ")))
-						{
-							result = new(true, result.LogLines.Select(x => x.Replace("): warning: ", "): error: ", StringComparison.OrdinalIgnoreCase)).ToList());
-							exitCode = 1;
-							++_executor._errorCount;
-						}
-						if (result.Success)
-						{
-							_executor.DecrementCacheableAction();
-							return new ActionArtifactResult(result.Success, result.LogLines, exitCode);
-						}
-					}
-					return new ActionArtifactResult(false, [], 0);
-				}, cancellationToken, TaskCreationOptions.LongRunning | TaskCreationOptions.PreferFairness, TaskScheduler.Default);
-			}
-
-			public Task ActionCompleteAsync(LinkedAction action, CancellationToken cancellationToken) => Task.CompletedTask;
-			public Task FlushChangesAsync(CancellationToken cancellationToken) => Task.CompletedTask;
-		}
-
-		private void ActionQueueCanceled(IStorageServer? ubaStorage)
-		{
-			_bIsCancelled = true;
-			Server?.StopServer(); // Make sure all remove processes are returned. We can't have any callbacks after this
-			_session?.CancelAll(); // Cancel all processes native side
-			ubaStorage?.SaveCasTable();
-
-			// We need the lock here since things are happening in parallel.
-			lock (_sessionLock)
-			{
-				foreach (IUBAAgentCoordinator coordinator in _agentCoordinators)
-				{
-					_agentCoordinators.ForEach(ac => ac.CloseAsync().Wait(2000)); // Give coordinators some time to close (this makes coordinators like horde return resources faster)
-				}
-			}
-		}
-
 		public override async Task<bool> ExecuteActionsAsync(IEnumerable<LinkedAction> inputActions, Microsoft.Extensions.Logging.ILogger logger, IActionArtifactCache? actionArtifactCache)
 		{
 			if (!inputActions.Any())
@@ -380,6 +268,7 @@ namespace UnrealBuildTool
 				return true;
 			}
 
+			// TODO: We still want to connect to coordinators because it might be that that input action is an action that spawns tons of internal processes (distributed thinlto linking)
 			if (inputActions.Count() < NumParallelProcesses && !UBAConfig.bForceBuildAllRemote)
 			{
 				UBAConfig.bDisableRemote = true;
@@ -483,7 +372,7 @@ namespace UnrealBuildTool
 				}
 
 				using EpicGames.UBA.ILogger ubaLogger = EpicGames.UBA.ILogger.CreateLogger(logger);
-				using (Server = IServer.CreateServer(UBAConfig.MaxWorkers, UBAConfig.SendSize, ubaLogger, UBAConfig.bUseQuic))
+				using (IServer Server = IServer.CreateServer(UBAConfig.MaxWorkers, UBAConfig.SendSize, ubaLogger, UBAConfig.bUseQuic))
 				{
 					_ubaLogger = ubaLogger;
 					using IStorageServer ubaStorageServer = IStorageServer.CreateStorageServer(Server, ubaLogger, new StorageServerCreateInfo(_rootDirRef.FullName, ((ulong)UBAConfig.StoreCapacityGb) * 1000 * 1000 * 1000, !UBAConfig.bStoreRaw, UBAConfig.Zone));
@@ -491,18 +380,19 @@ namespace UnrealBuildTool
 					{
 						try
 						{
-							_session = ISessionServer.CreateSessionServer(serverCreateInfo);
+							ISessionServer session = ISessionServer.CreateSessionServer(serverCreateInfo);
+							_session = session;
 
-							OnSessionCreated?.Invoke(_session);
+							OnSessionCreated?.Invoke(session);
 
 							if (_cacheableActionsLeft != 0 && _cacheConfigs.Any())
 							{
-								_session.UpdateStatus(1, 1, "Cache", LogEntryType.Info, null);
-								_session.UpdateStatus(1, 6, "Connecting...", LogEntryType.Info);
+								session.UpdateStatus(1, 1, "Cache", LogEntryType.Info, null);
+								session.UpdateStatus(1, 6, "Connecting...", LogEntryType.Info);
 
 								foreach (UnrealBuildAcceleratorCacheConfig cacheConfig in _cacheConfigs)
 								{
-									ICacheClient cacheClient = ICacheClient.CreateCacheClient(_session, UBAConfig.bReportCacheMissReason, "", _cacheableActionsLeft != 0 ? GetCacheClientHint() : "");
+									ICacheClient cacheClient = ICacheClient.CreateCacheClient(session, UBAConfig.bReportCacheMissReason, "", _cacheableActionsLeft != 0 ? GetCacheClientHint() : "");
 									string[] nameAndPort = cacheConfig.CacheServer.Split(':');
 									int port = 1347;
 									if (nameAndPort.Length > 1)
@@ -518,14 +408,13 @@ namespace UnrealBuildTool
 									if (cacheSuccess)
 									{
 										_cacheClients.Add(cacheConfig, cacheClient);
-										actionArtifactCache ??= new UBAActionArtifactCache(this);
 									}
 									else
 									{
 										cacheClient.Dispose();
 									}
 								}
-								_session.UpdateStatus(1, 6, _cacheClients.Any() ? "Connected" : "Failed to connect", LogEntryType.Info);
+								session.UpdateStatus(1, 6, _cacheClients.Any() ? "Connected" : "Failed to connect", LogEntryType.Info);
 							}
 
 							ubaStorage = ubaStorageServer;
@@ -536,10 +425,10 @@ namespace UnrealBuildTool
 							}
 							else
 							{
-								_session.DisableRemoteExecution();
+								session.DisableRemoteExecution();
 							}
 
-							bool success = ExecuteActionsInternal(inputActions, _session, logger, actionArtifactCache, () => ActionQueueCanceled(ubaStorage));
+							bool success = ExecuteActionsInternal(inputActions, Server, ubaStorage, session, _cacheClients.Values);
 
 							if (!UBAConfig.bDisableRemote)
 							{
@@ -548,7 +437,7 @@ namespace UnrealBuildTool
 
 							if (UBAConfig.bPrintSummary)
 							{
-								_session.PrintSummary();
+								session.PrintSummary();
 							}
 
 							return success && !_bIsCancelled;
@@ -626,94 +515,372 @@ namespace UnrealBuildTool
 
 		public override bool VerifyOutputs => UBAConfig.bWriteToDisk;
 
+
+		(byte[]?, uint) BuildKnownInputs(LinkedAction action)
+		{
+			if (!UBAConfig.bUseKnownInputs)
+			{
+				return (null, 0);
+			}
+			
+			uint knownInputsCount = 0;
+			int sizeOfChar = System.OperatingSystem.IsWindows() ? 2 : 1;
+
+			int byteCount = 0;
+			foreach (FileItem item in action.PrerequisiteItems)
+			{
+				byteCount += (item.FullName.Length + 1) * sizeOfChar;
+				++knownInputsCount;
+			}
+
+			byte[] knownInputs = new byte[byteCount + sizeOfChar];
+
+			int byteOffset = 0;
+			foreach (FileItem item in action.PrerequisiteItems)
+			{
+				string str = item.FullName;
+				int strBytes = str.Length * sizeOfChar;
+				if (sizeOfChar == 1) // Unmanaged size uses ascii
+				{
+					System.Buffer.BlockCopy(System.Text.Encoding.ASCII.GetBytes(str.ToCharArray()), 0, knownInputs, byteOffset, strBytes);
+				}
+				else
+				{
+					System.Buffer.BlockCopy(str.ToCharArray(), 0, knownInputs, byteOffset, strBytes);
+				}
+
+				byteOffset += strBytes + sizeOfChar;
+			}
+			return (knownInputs, knownInputsCount);
+		}
+
+		class AgentCoordinatorScheduler : IUbaAgentCoordinatorScheduler
+		{
+			internal AgentCoordinatorScheduler(IScheduler scheduler, IServer server) { _scheduler = scheduler; _server = server; }
+			internal IScheduler? _scheduler;
+			internal IServer? _server;
+			public void Reset() { lock (this) { _scheduler = null; _server = null; } }
+			public bool IsEmpty { get { lock (this) return _scheduler != null ? _scheduler.IsEmpty : true; } }
+			public double GetProcessWeightThatCanRunRemotelyNow() { lock (this) return _scheduler != null ? _scheduler.GetProcessWeightThatCanRunRemotelyNow() : 0; }
+			public bool AddClient(string ip, int port, string crypto = "")
+			{
+				lock (this)
+				{
+					if (_server != null)
+					{
+						return _server.AddClient(ip, port, crypto);
+					}
+				}
+				return false;
+			}
+		}
+
 		/// <summary>
 		/// Executes the provided actions
 		/// </summary>
 		/// <returns>True if all the tasks successfully executed, or false if any of them failed.</returns>
-		bool ExecuteActionsInternal(IEnumerable<LinkedAction> inputActions, ISessionServer session, Microsoft.Extensions.Logging.ILogger logger, IActionArtifactCache? actionArtifactCache, System.Action onCancel)
+		bool ExecuteActionsInternal(IEnumerable<LinkedAction> inputActions, IServer server, IStorageServer storage, ISessionServer session, IEnumerable<ICacheClient> cache)
 		{
 			_ubaStartTimeUtc = DateTime.UtcNow;
-			using ImmediateActionQueue queue = CreateActionQueue(inputActions, actionArtifactCache, UBAConfig.CacheMaxWorkers, logger);
-			int actionLimit = Math.Min(NumParallelProcesses, queue.TotalActions);
-			queue.CreateAutomaticRunner(action => RunActionLocal(queue, action), bUseActionWeights, actionLimit, NumParallelProcesses);
-			ImmediateActionQueueRunner remoteRunner = queue.CreateManualRunner(action => RunActionRemote(queue, action, false));
-			ImmediateActionQueueRunner crossArchitectureRemoteRunner = queue.CreateManualRunner(action => RunActionRemote(queue, action, true));
-			queue.CancellationToken.Register(onCancel);
-
-			if (actionArtifactCache != null)
+
+			ActionLogger actionLogger = new(inputActions.Count(), "Compiling C++ source code...", x => WriteToolOutput(x), () => FlushToolOutput(), _threadedLogger)
 			{
-				queue.OnArtifactsRead += (action) => { UpdateCacheProgress(queue); UpdateProgress(queue); };
-				queue.OnArtifactsMiss += (action) => UpdateCacheProgress(queue);
-			}
+				ShowCompilationTimes = bShowCompilationTimes,
+				ShowCPUUtilization = bShowCPUUtilization,
+				PrintActionTargetNames = bPrintActionTargetNames,
+				LogActionCommandLines = bLogActionCommandLines,
+				ShowPerActionCompilationTimes = bShowPerActionCompilationTimes,
+				CompactOutput = bCompactOutput,
+			};
+
+			using IScheduler scheduler = IScheduler.CreateScheduler(session, cache, NumParallelProcesses, UBAConfig.bForceBuildAllRemote);
+
+			var cancelBuild = () =>
+			{
+				if (_bIsCancelled)
+					return;
+				_bIsCancelled = true;
+				scheduler.Cancel();
+				server?.StopServer(); // Make sure all remove processes are returned. We can't have any callbacks after this
+				session?.CancelAll(); // Cancel all processes native side
+				storage?.SaveCasTable();
+
+				// We need the lock here since things are happening in parallel.
+				lock (_sessionLock)
+				{
+					foreach (IUBAAgentCoordinator coordinator in _agentCoordinators)
+					{
+						_agentCoordinators.ForEach(ac => ac.CloseAsync().Wait(2000)); // Give coordinators some time to close (this makes coordinators like horde return resources faster)
+					}
+				}
+			};
 
-			UpdateProgress(queue);
 
-			// Start the queue
-			queue.Start();
 
-			// Handle process available from remote
-			session!.RemoteProcessSlotAvailable += (sender, args) =>
+			int totalActions = 0;
+			int cacheTotalActions = 0;
+			int cacheHitActions = 0;
+			int succeededActions = 0;
+			int failedActions = 0;
+			int localProcessedActions = 0;
+			int remoteProcessedActions = 0;
+			int localRetryActions = 0;
+			int forcedRetryActions = 0;
+
+			scheduler.SetProcessFinishedCallback((info) =>
 			{
-				ImmediateActionQueueRunner runner = args.IsCrossArchitecture ? crossArchitectureRemoteRunner : remoteRunner;
-				if (!queue.TryStartOneAction(runner) && (_bIsRemoteActionsAllowed && !UBAConfig.bForceBuildAllRemote))
+				int exitCode = info.ExitCode;
+
+				if (_bIsCancelled || exitCode == 99999) // Process was cancelled by executor
 				{
-					uint count = ActionsLeftThatCanRunRemotely(queue);
+					return ProcessFinishedResponse.None;
+				}
+
+				List<string> logLines = info.LogLines;
+
+				LinkedAction action = (LinkedAction)info.UserData;
+
+				string additionalDescription = string.Empty;
+
+				ProcessExecutionType executionType = info.ExecutionType;
+
+				switch (executionType)
+				{
+					case ProcessExecutionType.Skip:
+						{
+							return ProcessFinishedResponse.None;
+						}
+
+					case ProcessExecutionType.Remote:
+						{
+							if (_ubaDurationWaitingForRemote == TimeSpan.Zero)
+							{
+								_ubaDurationWaitingForRemote = DateTime.UtcNow - _ubaStartTimeUtc;
+							}
+
+							string executingHost = info.ExecutingHost ?? "Unknown";
+
+							var rerun = (string error) =>
+							{
+								_threadedLogger.LogInformation("{Description} {StatusDescription} [RemoteExecutor: {ExecutingHost}]: Exited with error code {ExitCode} ({Error}). This action will retry locally", action.CommandDescription, action.StatusDescription, executingHost, exitCode, error);
+								foreach (string line in logLines)
+								{
+									_threadedLogger.LogInformation(line);
+								}
+								++localRetryActions;
+								return ProcessFinishedResponse.RerunLocal;
+							};
+
+							if (exitCode != 0 && !logLines.Any())
+							{
+								_threadedLogger.LogInformation("{Description} {StatusDescription} [RemoteExecutor: {ExecutingHost}]: Exited with error code {ExitCode} with no output. This action will retry locally", action.CommandDescription, action.StatusDescription, executingHost, exitCode);
+								return ProcessFinishedResponse.RerunLocal;
+							}
+							else if ((uint)exitCode == 0xC0000005)
+							{
+								return rerun("Access violation");
+							}
+							else if ((uint)exitCode == 0xC0000409)
+							{
+								return rerun("Stack buffer overflow");
+							}
+							else if ((uint)exitCode == 0xC0000602)
+							{
+								return rerun("Fail Fast Exception");
+							}
+							else if (exitCode != 0 && logLines.Any(x => x.Contains(" C1001: ", StringComparison.Ordinal)))
+							{
+								return rerun("C1001");
+							}
+							else if (exitCode >= 9000 && exitCode < 10000)
+							{
+								return rerun("UBA error");
+							}
+							else if (exitCode != 0 && UBAConfig.bForcedRetryRemote)
+							{
+								return rerun("Force local retry");
+							}
+							additionalDescription = $"[RemoteExecutor: {executingHost}]";
+							++remoteProcessedActions;
+							break;
+						}
+
+					case ProcessExecutionType.Local:
+						{
+							if (UBAConfig.bAllowRetry && (exitCode != 0 && UBAConfig.bForcedRetry || (exitCode >= 9000 && exitCode < 10000)))
+							{
+								if (UBAConfig.bStoreObjFilesCompressed)
+								{
+									_threadedLogger.LogInformation("{Description} {StatusDescription}: Exited with error code {ExitCode}. This action will retry", action.CommandDescription, action.StatusDescription, exitCode);
+									++forcedRetryActions;
+									return ProcessFinishedResponse.RerunLocal;
+								}
+								else
+								{
+									_threadedLogger.LogInformation("{Description} {StatusDescription}: Exited with error code {ExitCode}. This action will retry without UBA", action.CommandDescription, action.StatusDescription, exitCode);
+									++forcedRetryActions;
+									return ProcessFinishedResponse.RerunNative;
+								}
+							}
+							++localProcessedActions;
+							break;
+						}
+
+					case ProcessExecutionType.Native:
+						{
+							// If not detoured we manually have to report created files
+							if (exitCode == 0)
+							{
+								session!.RegisterNewFiles(action.ProducedItems.Where(x => FileReference.Exists(x.Location)).Select(x => x.FullName).ToArray());
+							}
+							additionalDescription = "[NoUba]";
+							break;
+						}
+
+					case ProcessExecutionType.Cache:
+						{
+							++cacheHitActions;
+							additionalDescription = "[Cache]";
+							break;
+						}
+				}
 
-					// We didn't find an action to start, let's check how many queued items are left.. if there are less than NumParallelProcesses, then disconnect
-					if (count <= NumParallelProcesses)
+				TimeSpan processorTime = TimeSpan.Zero;
+				TimeSpan executionTime = TimeSpan.Zero;
+				if (executionType != ProcessExecutionType.Cache)
+				{
+					processorTime = info.TotalProcessorTime;
+					executionTime = info.TotalWallTime;
+				}
+
+				logLines.RemoveAll(line =>
+					line.StartsWith("   Creating library ", StringComparison.Ordinal) ||
+					line.StartsWith("   Creating object ", StringComparison.Ordinal) ||
+					line.EndsWith("file(s) copied.", StringComparison.Ordinal));
+
+				if (exitCode == 0)
+				{
+					if (action.bForceWarningsAsError && logLines.Any(x => x.Contains("): warning: ")))
 					{
-						_bIsRemoteActionsAllowed = false;
-						_agentCoordinators.ForEach(ac => ac.Stop());
-						session!.DisableRemoteExecution();
+						exitCode = 1;
+						logLines = logLines.Select(x => x.Replace("): warning: ", "): error: ", StringComparison.OrdinalIgnoreCase)).ToList();
 					}
 					else
 					{
-						// Tell UBA max number of remote processes left. Providing this information makes it possible for UBA to start disconnecting clients
-						session!.SetMaxRemoteProcessCount(count);
+						if (executionType != ProcessExecutionType.Cache)
+						{
+							WriteToCache(action, exitCode, info.ProcessHandle);
+						}
+						++succeededActions;
 					}
 				}
-			};
 
-			session!.RemoteProcessReturned += (sender, args) =>
+				if (exitCode != 0)
+				{
+					++failedActions;
+
+					// Delete produced items on error if requested
+					if (action.bDeleteProducedItemsOnError)
+					{
+						foreach (FileItem output in action.ProducedItems.Distinct().Where(x => FileReference.Exists(x.Location)))
+						{
+							FileReference.Delete(output.Location);
+						}
+					}
+
+					if (bStopCompilationAfterErrors)
+					{
+						cancelBuild();
+					}
+				}
+
+				actionLogger.AddActionToLog(new(logLines, exitCode, executionTime, processorTime, action, additionalDescription));
+
+				return ProcessFinishedResponse.None;
+			});
+
+			scheduler.Start();
+
+			ConsoleCancelEventHandler? cancelHandler = null;
+			cancelHandler = (object? sender, ConsoleCancelEventArgs e) =>
 			{
-				args.Process.Cancel(true);
-				LinkedAction action = (LinkedAction)args.Process.UserData!;
-				//logger.LogInformation("REQUEUING " + action.ProducedItems.FirstOrDefault()!.Name);
-				queue.RequeueAction(action);
+				Console.CancelKeyPress -= cancelHandler;
+				cancelHandler = null;
+				e.Cancel = true;
+				cancelBuild();
 			};
+			Console.CancelKeyPress += cancelHandler;
+
+			int index = 0;
+			foreach (LinkedAction action in inputActions)
+			{
+				action.SortIndex = index++;
+				ProcessStartInfo startInfo = GetActionStartInfo(action);
+				startInfo.RootsHandle = GetActionRootsHandle(action);
+
+				var prereqActionsSortIndex = action.PrerequisiteActions.Select(a => a.SortIndex).ToArray();
+				bool enableDetour = !ForceLocalNoDetour(action) && action.bCanExecuteInUBA;
+				bool canRunRemotely = CanRunRemotely(action);
 
-			// Add all actions we can add
-			queue.StartManyActions();
+				var (knownInputs, knownInputsCount) = BuildKnownInputs(action);
+
+				uint bucket = GetCacheBucket(action);
+				scheduler.EnqueueProcess(startInfo, action.Weight, enableDetour, canRunRemotely, prereqActionsSortIndex, knownInputs, knownInputsCount, bucket, action);
+
+				++totalActions;
+				cacheTotalActions += bucket != 0 ? 0 : 1;
+			}
+
+			if (!UBAConfig.bForceBuildAllRemote)
+			{
+				scheduler.SetAllowDisableRemoteExecution();
+			}
+
+			_threadedLogger.LogDebug("All processes queued");
 
 			try
 			{
+				AgentCoordinatorScheduler coordinatorScheduler = new(scheduler, server);
 				uint statusUpdateCounter = 10;
 				foreach (IUBAAgentCoordinator coordinator in _agentCoordinators)
 				{
 					uint statusUpdateIndex = statusUpdateCounter;
-					coordinator.Start(queue, CanRunRemotely, (sr, sc, st, t, sl) => UpdateStatus(statusUpdateIndex + sr, sc, st, t, sl));
+					coordinator.Start(coordinatorScheduler, CanRunRemotely, (sr, sc, st, t, sl) => UpdateStatus(statusUpdateIndex + sr, sc, st, t, sl));
 					statusUpdateCounter += 10;
 				}
 
-				bool res = queue.RunTillDone().Result; // Using inline wait to avoid possible thread switch
+				while (!scheduler.IsEmpty)
+					Thread.Sleep(100);
+
+				bool res = failedActions == 0;
+
+				actionLogger.TraceSummary(res);
 
-				queue.GetActionResultCounts(out int totalActions, out int succeededActions, out int failedActions, out int cacheHitActions, out int cacheMissActions);
-				telemetryEvent = new TelemetryExecutorUBAEvent(Name, _ubaStartTimeUtc, res, totalActions, succeededActions, failedActions, cacheHitActions, cacheMissActions,
-					_localProcessedActions, _remoteProcessedActions,
-					_localRetryActions.Count, _forcedRetryActions.Count,
-					!UBAConfig.bDisableRemote ? _agentCoordinators.Count : 0, !UBAConfig.bDisableRemote ? _successfulCoordinatorConnections : 0, !UBAConfig.bDisableRemote ? _failedCoordinatorConnections : 0,
-					!UBAConfig.bDisableRemote ? _ubaDurationWaitingForRemote : TimeSpan.Zero,
-					!UBAConfig.bDisableRemote || _agentCoordinators.Count == 0 ? "Local" : _remoteConnectionMode,
-					DateTime.UtcNow);
+				coordinatorScheduler.Reset();
 
+				int cacheMissActions = cacheTotalActions - cacheHitActions;
+
+				if (!_bIsCancelled)
+				{
+					telemetryEvent = new TelemetryExecutorUBAEvent(Name, _ubaStartTimeUtc, res, totalActions, succeededActions, failedActions, cacheHitActions, cacheMissActions,
+						localProcessedActions, remoteProcessedActions,
+						localRetryActions, forcedRetryActions,
+						!UBAConfig.bDisableRemote ? _agentCoordinators.Count : 0, !UBAConfig.bDisableRemote ? _successfulCoordinatorConnections : 0, !UBAConfig.bDisableRemote ? _failedCoordinatorConnections : 0,
+						!UBAConfig.bDisableRemote ? _ubaDurationWaitingForRemote : TimeSpan.Zero,
+						!UBAConfig.bDisableRemote || _agentCoordinators.Count == 0 ? "Local" : _remoteConnectionMode,
+						DateTime.UtcNow);
+				}
 				return res;
 			}
 			finally
 			{
-				if (_bIsRemoteActionsAllowed)
+				if (cancelHandler != null)
 				{
-					_agentCoordinators.ForEach(ac => ac.Stop());
+					Console.CancelKeyPress -= cancelHandler;
+					cancelHandler = null;
 				}
+
+				_agentCoordinators.ForEach(ac => ac.Stop());
 			}
 		}
 
@@ -750,8 +917,6 @@ namespace UnrealBuildTool
 			action.bCanExecuteInUBA &&
 			action.bCanExecuteRemotely &&
 			(UBAConfig.bLinkRemote || action.ActionType != ActionType.Link) &&
-			!_localRetryActions.ContainsKey(action) &&
-			!_forcedRetryActions.ContainsKey(action) &&
 			!ForceLocalNoDetour(action);
 
 		ProcessStartInfo GetActionStartInfo(LinkedAction action)
@@ -816,6 +981,11 @@ namespace UnrealBuildTool
 
 		uint GetCacheBucket(LinkedAction action)
 		{
+			if (!action.ArtifactMode.HasFlag(ArtifactMode.Enabled))
+			{
+				return 0;
+			}
+
 			uint bucket = 0;
 			if (action.CacheBucket != 0)
 			{
@@ -893,9 +1063,9 @@ namespace UnrealBuildTool
 			public DepsData? Data { get; init; }
 		}
 
-		bool WriteToCache(LinkedAction action, IProcess process)
+		bool WriteToCache(LinkedAction action, int exitCode, nint processHandle)
 		{
-			if (process.ExitCode != 0 || !action.ArtifactMode.HasFlag(ArtifactMode.Enabled) || !_cacheClients.Any(x => x.Key.CanWrite))
+			if (exitCode != 0 || !action.ArtifactMode.HasFlag(ArtifactMode.Enabled) || !_cacheClients.Any(x => x.Key.CanWrite))
 			{
 				return true;
 			}
@@ -963,290 +1133,12 @@ namespace UnrealBuildTool
 					continue;
 				}
 
-				success = item.Value.WriteToCache(bucket, process, inputsMemory.GetBuffer(), (uint)inputsMemory.Position, outputsMemory.GetBuffer(), (uint)outputsMemory.Position) || success;
+				success = item.Value.WriteToCache(bucket, processHandle, inputsMemory.GetBuffer(), (uint)inputsMemory.Position, outputsMemory.GetBuffer(), (uint)outputsMemory.Position) || success;
 			}
 
 			DecrementCacheableAction();
 
 			return success;
 		}
-
-		Func<Task>? RunActionLocal(ImmediateActionQueue queue, LinkedAction action)
-		{
-			if (UBAConfig.bForceBuildAllRemote && CanRunRemotely(action))
-			{
-				return null;
-			}
-
-			return () =>
-			{
-				if (_bIsCancelled)
-				{
-					HandleActionCancelled(queue, null, action);
-					return Task.CompletedTask;
-				}
-				bool enableDetour = !ForceLocalNoDetour(action) && action.bCanExecuteInUBA && !_forcedRetryActions.ContainsKey(action);
-
-				ProcessStartInfo startInfo = GetActionStartInfo(action);
-				startInfo.RootsHandle = GetActionRootsHandle(action);
-
-				using (IProcess process = _session!.RunProcess(startInfo, false, null, enableDetour))
-				{
-					Interlocked.Add(ref _localProcessedActions, 1);
-					if (!_bIsCancelled && (UBAConfig.bAllowRetry && (process.ExitCode != 0 && UBAConfig.bForcedRetry || (process.ExitCode >= 9000 && process.ExitCode < 10000))))
-					{
-						if (UBAConfig.bStoreObjFilesCompressed)
-						{
-							_threadedLogger.LogInformation("{Description} {StatusDescription}: Exited with error code {ExitCode}. This action will retry", action.CommandDescription, action.StatusDescription, process.ExitCode);
-						}
-						else
-						{
-							_threadedLogger.LogInformation("{Description} {StatusDescription}: Exited with error code {ExitCode}. This action will retry without UBA", action.CommandDescription, action.StatusDescription, process.ExitCode);
-							_forcedRetryActions.AddOrUpdate(action, false, (k, v) => false);
-						}
-						queue.RequeueAction(action);
-						return Task.CompletedTask;
-					}
-
-					// If not detoured we manually have to report created files
-					if (!enableDetour && process.ExitCode == 0)
-					{
-						_session!.RegisterNewFiles(action.ProducedItems.Where(x => FileReference.Exists(x.Location)).Select(x => x.FullName).ToArray());
-					}
-
-					if (enableDetour)
-					{
-						WriteToCache(action, process);
-					}
-
-					TimeSpan processorTime = process.TotalProcessorTime;
-					TimeSpan executionTime = process.TotalWallTime;
-					List<string> logLines = process.LogLines;
-					RemoveLogLineSpam(logLines);
-					string? additionalDescription = !enableDetour ? "(UBA disabled)" : null;
-					ActionFinished(queue, new ExecuteResults(logLines, process.ExitCode, executionTime, processorTime, additionalDescription), action, process);
-				}
-				return Task.CompletedTask;
-			};
-		}
-
-		Func<Task>? RunActionRemote(ImmediateActionQueue queue, LinkedAction action, bool isCrossArchitecture)
-		{
-			if (!CanRunRemotely(action))
-			{
-				return null;
-			}
-
-			if (isCrossArchitecture && !action.bCanExecuteInUBACrossArchitecture)
-			{
-				return null;
-			}
-
-			return () =>
-			{
-				if (_ubaDurationWaitingForRemote == TimeSpan.Zero)
-				{
-					_ubaDurationWaitingForRemote = DateTime.UtcNow - _ubaStartTimeUtc;
-				}
-				if (_bIsCancelled)
-				{
-					HandleActionCancelled(queue, null, action);
-					return Task.CompletedTask;
-				}
-
-				uint knownInputsCount = 0;
-				byte[]? knownInputs = null;
-				if (UBAConfig.bUseKnownInputs)
-				{
-					int sizeOfChar = System.OperatingSystem.IsWindows() ? 2 : 1;
-
-					int byteCount = 0;
-					foreach (FileItem item in action.PrerequisiteItems)
-					{
-						byteCount += (item.FullName.Length + 1) * sizeOfChar;
-						++knownInputsCount;
-					}
-
-					knownInputs = new byte[byteCount + sizeOfChar];
-
-					int byteOffset = 0;
-					foreach (FileItem item in action.PrerequisiteItems)
-					{
-						string str = item.FullName;
-						int strBytes = str.Length * sizeOfChar;
-						if (sizeOfChar == 1) // Unmanaged size uses ascii
-						{
-							System.Buffer.BlockCopy(System.Text.Encoding.ASCII.GetBytes(str.ToCharArray()), 0, knownInputs, byteOffset, strBytes);
-						}
-						else
-						{
-							System.Buffer.BlockCopy(str.ToCharArray(), 0, knownInputs, byteOffset, strBytes);
-						}
-
-						byteOffset += strBytes + sizeOfChar;
-					}
-				}
-
-				ProcessStartInfo startInfo = GetActionStartInfo(action);
-				startInfo.RootsHandle = GetActionRootsHandle(action);
-
-				_session!.RunProcessRemote(startInfo, (s, e) =>
-				{
-					if (e.ExitCode == 99999) // Process was cancelled by executor
-					{
-						return;
-					}
-
-					Interlocked.Add(ref _remoteProcessedActions, 1);
-					if (e.ExitCode != 0 && !e.LogLines.Any())
-					{
-						RemoteActionFailedNoOutput(queue, action, e.ExitCode, e.ExecutingHost ?? "Unknown");
-						return;
-					}
-					else if ((uint)e.ExitCode == 0xC0000005)
-					{
-						RemoteActionFailedCrash(queue, action, e.ExitCode, e.ExecutingHost ?? "Unknown", "Access violation", e.LogLines);
-						return;
-					}
-					else if ((uint)e.ExitCode == 0xC0000409)
-					{
-						RemoteActionFailedCrash(queue, action, e.ExitCode, e.ExecutingHost ?? "Unknown", "Stack buffer overflow", e.LogLines);
-						return;
-					}
-					else if ((uint)e.ExitCode == 0xC0000602)
-					{
-						RemoteActionFailedCrash(queue, action, e.ExitCode, e.ExecutingHost ?? "Unknown", "Fail Fast Exception", e.LogLines);
-						return;
-					}
-					else if (e.ExitCode != 0 && e.LogLines.Any(x => x.Contains(" C1001: ", StringComparison.Ordinal)))
-					{
-						RemoteActionFailedCrash(queue, action, e.ExitCode, e.ExecutingHost ?? "Unknown", "C1001", e.LogLines);
-						return;
-					}
-					else if (e.ExitCode >= 9000 && e.ExitCode < 10000)
-					{
-						RemoteActionFailedCrash(queue, action, e.ExitCode, e.ExecutingHost ?? "Unknown", "UBA error", e.LogLines);
-						return;
-					}
-					else if (e.ExitCode != 0 && UBAConfig.bForcedRetryRemote)
-					{
-						RemoteActionFailedCrash(queue, action, e.ExitCode, e.ExecutingHost ?? "Unknown", "Force local retry", e.LogLines);
-						return;
-					}
-
-					IProcess process = (IProcess)s;
-
-					WriteToCache(action, process);
-
-					string additionalDescription = $"[RemoteExecutor: {e.ExecutingHost}]";
-					TimeSpan processorTime = e.TotalProcessorTime;
-					TimeSpan executionTime = e.TotalWallTime;
-					List<string> logLines = e.LogLines;
-					RemoveLogLineSpam(logLines);
-					ActionFinished(queue, new ExecuteResults(logLines, e.ExitCode, executionTime, processorTime, additionalDescription), action, process);
-				}, action.Weight, knownInputs, knownInputsCount, action.bCanExecuteInUBACrossArchitecture);
-				return Task.CompletedTask;
-			};
-		}
-
-		uint ActionsLeftThatCanRunRemotely(ImmediateActionQueue queue)
-		{
-			lock (_actionsChangedLock)
-			{
-				if (_bActionsChanged)
-				{
-					_actionsQueuedThatCanRunRemotely = queue.GetQueuedActionsCount(CanRunRemotely);
-					_bActionsChanged = false;
-				}
-				return _actionsQueuedThatCanRunRemotely;
-			}
-		}
-
-		protected void HandleActionCancelled(ImmediateActionQueue queue, ExecuteResults? results, LinkedAction action)
-		{
-			ExecuteResults cancelResults = new(results?.LogLines ?? new(), Int32.MaxValue, results?.ExecutionTime ?? TimeSpan.Zero, results?.ProcessorTime ?? TimeSpan.Zero, results?.AdditionalDescription);
-			queue.OnActionCompleted(action, false, cancelResults);
-		}
-
-		protected void ActionFinished(ImmediateActionQueue queue, ExecuteResults results, LinkedAction action, IProcess? process = null)
-		{
-			if (_bIsCancelled)
-			{
-				HandleActionCancelled(queue, results, action);
-				return;
-			}
-
-			bool success = results.ExitCode == 0;
-
-			if (results.ExitCode == 0 && action.bForceWarningsAsError && results.LogLines.Any(x => x.Contains("): warning: ")))
-			{
-				results = new(results.LogLines.Select(x => x.Replace("): warning: ", "): error: ", StringComparison.OrdinalIgnoreCase)).ToList(), 1, results.ExecutionTime, results.ProcessorTime, results.AdditionalDescription);
-				success = false;
-			}
-
-			queue.OnActionCompleted(action, success, results);
-
-			if (!success)
-			{
-				++_errorCount;
-			}
-			UpdateProgress(queue);
-
-			lock (_actionsChangedLock)
-			{
-				_bActionsChanged = true;
-			}
-		}
-		void UpdateProgress(ImmediateActionQueue queue)
-		{
-			int completedActions = queue.CompletedActions;
-			int totalActions = queue.TotalActions;
-			int percent = completedActions * 100 / totalActions;
-			_session!.UpdateProgress((uint)totalActions, (uint)completedActions, _errorCount);
-		}
-		void UpdateCacheProgress(ImmediateActionQueue queue)
-		{
-			_session!.UpdateStatus(1, 6, $"Hits {queue.CacheHitActions} Misses {queue.CacheMissActions}", LogEntryType.Info);
-		}
-		void RemoteActionFailedNoOutput(ImmediateActionQueue queue, LinkedAction action, int exitCode, string executingHost)
-		{
-			if (_bIsCancelled)
-			{
-				HandleActionCancelled(queue, null, action);
-				return;
-			}
-
-			_threadedLogger.LogInformation("{Description} {StatusDescription} [RemoteExecutor: {ExecutingHost}]: Exited with error code {ExitCode} with no output. This action will retry locally", action.CommandDescription, action.StatusDescription, executingHost, exitCode);
-			_localRetryActions.AddOrUpdate(action, false, (k, v) => false);
-			queue.RequeueAction(action);
-
-			lock (_actionsChangedLock)
-			{
-				_bActionsChanged = true;
-			}
-		}
-
-		void RemoteActionFailedCrash(ImmediateActionQueue queue, LinkedAction action, int exitCode, string executingHost, string error, List<string> logLines)
-		{
-			if (_bIsCancelled)
-			{
-				HandleActionCancelled(queue, null, action);
-				return;
-			}
-
-			_threadedLogger.LogInformation("{Description} {StatusDescription} [RemoteExecutor: {ExecutingHost}]: Exited with error code {ExitCode} ({Error}). This action will retry locally", action.CommandDescription, action.StatusDescription, executingHost, exitCode, error);
-			foreach (string line in logLines)
-			{
-				_threadedLogger.LogInformation(line);
-			}
-
-			_localRetryActions.AddOrUpdate(action, false, (k, v) => false);
-			queue.RequeueAction(action);
-
-			lock (_actionsChangedLock)
-			{
-				_bActionsChanged = true;
-			}
-		}
 	}
 }
