UnrealBuildTool: Merge modules MergeAdditionalCommonModules to force engine modules into the common group

#rnx

diff --git a/Engine/Source/Programs/UnrealBuildTool/Configuration/ReadOnlyTargetRules.cs b/Engine/Source/Programs/UnrealBuildTool/Configuration/ReadOnlyTargetRules.cs
--- a/Engine/Source/Programs/UnrealBuildTool/Configuration/ReadOnlyTargetRules.cs
+++ b/Engine/Source/Programs/UnrealBuildTool/Configuration/ReadOnlyTargetRules.cs
@@ -571,6 +571,8 @@ namespace UnrealBuildTool
 
 		public bool bMergeModules => Inner.bMergeModules;
 
+		public IReadOnlySet<string> MergeAdditionalCommonModules => Inner.MergeAdditionalCommonModules;
+
 		public IReadOnlyList<string> MergePlugins => Inner.MergePlugins;
 
 		public IReadOnlyList<string> MergePluginsIgnored => Inner.MergePluginsIgnored;
diff --git a/Engine/Source/Programs/UnrealBuildTool/Configuration/TargetRules.cs b/Engine/Source/Programs/UnrealBuildTool/Configuration/TargetRules.cs
--- a/Engine/Source/Programs/UnrealBuildTool/Configuration/TargetRules.cs
+++ b/Engine/Source/Programs/UnrealBuildTool/Configuration/TargetRules.cs
@@ -2621,6 +2621,13 @@ namespace UnrealBuildTool
 		[CommandLine("-MergeModules")]
 		public bool bMergeModules { get; set; }
 
+		/// <summary>
+		/// Experimental: List of modules to exclude from Engine to be placed in Common
+		/// </summary>
+		[RequiresUniqueBuildEnvironment]
+		[CommandLine("-MergeAdditionalCommonModules=", ListSeparator = '+')]
+		public HashSet<string> MergeAdditionalCommonModules { get; } = new();
+
 		/// <summary>
 		/// Experimental: List of plugins (and their dependencies) to each merge into separate libraries. Requires bMergeModules to be enabled
 		/// </summary>
@@ -2658,7 +2665,7 @@ namespace UnrealBuildTool
 		/// <summary>
 		/// Experimental: Whether to merge only game feature plugins, ignoring regular plugins.
 		/// </summary>
-		public bool MergeOnlyGameFeatures { get; } = true;
+		public bool MergeOnlyGameFeatures { get; set; } = true;
 
 		/// <summary>
 		/// Backing storage for the LinkType property.
diff --git a/Engine/Source/Programs/UnrealBuildTool/Configuration/UEBuildTarget.cs b/Engine/Source/Programs/UnrealBuildTool/Configuration/UEBuildTarget.cs
--- a/Engine/Source/Programs/UnrealBuildTool/Configuration/UEBuildTarget.cs
+++ b/Engine/Source/Programs/UnrealBuildTool/Configuration/UEBuildTarget.cs
@@ -3769,7 +3769,7 @@ namespace UnrealBuildTool
 					}
 
 					// There's only one entry here at this point
-					HashSet<UEBuildPlugin>? dependencies = group.Value.First().Dependencies;
+					HashSet<UEBuildPlugin>? dependencies = group.Value.FirstOrDefault()?.Dependencies;
 					if (dependencies != null)
 					{
 						ProcessPluginsRecursively(group.Key, dependencies);
@@ -3852,9 +3852,9 @@ namespace UnrealBuildTool
 					ConsumeModules(item.Key, pluginModules);
 				}
 			}
-			
+
 			// Merge remaining engine modules
-			IEnumerable<UEBuildModuleCPP> engineModules = allModules.Where(x => x.RulesFile.IsUnderDirectory(Unreal.EngineDirectory));
+			IEnumerable<UEBuildModuleCPP> engineModules = allModules.Where(x => x.RulesFile.IsUnderDirectory(Unreal.EngineDirectory) && !Rules.MergeAdditionalCommonModules.Contains(x.Name));
 			ConsumeModules("Engine", engineModules);
 
 			// Merge remaining modules to common set
