UnrealBuildTool: Add TargetRules.bIdenticalCodeFolding setting (--icf=all)

diff --git a/Engine/Source/Programs/UnrealBuildTool/Configuration/ReadOnlyTargetRules.cs b/Engine/Source/Programs/UnrealBuildTool/Configuration/ReadOnlyTargetRules.cs
--- a/Engine/Source/Programs/UnrealBuildTool/Configuration/ReadOnlyTargetRules.cs
+++ b/Engine/Source/Programs/UnrealBuildTool/Configuration/ReadOnlyTargetRules.cs
@@ -480,8 +480,6 @@ namespace UnrealBuildTool
 
 		public bool bEnableCppModules => Inner.bEnableCppModules;
 
-		public bool bEnableCppCoroutinesForEvaluation => Inner.bEnableCppCoroutinesForEvaluation;
-
 		public bool bEnableProcessPriorityControl => Inner.bEnableProcessPriorityControl;
 
 		[Obsolete("Deprecated in UE5.3 - No longer used as MallocProfiler is removed in favor of UnrealInsights.")]
@@ -539,6 +537,8 @@ namespace UnrealBuildTool
 
 		public bool bPublicSymbolsByDefault => Inner.bPublicSymbolsByDefault;
 
+		public bool bIdenticalCodeFolding => Inner.bIdenticalCodeFolding;
+
 		public string? ToolChainName => Inner.ToolChainName;
 
 		public float MSVCCompileActionWeight => Inner.MSVCCompileActionWeight;
diff --git a/Engine/Source/Programs/UnrealBuildTool/Configuration/TargetRules.cs b/Engine/Source/Programs/UnrealBuildTool/Configuration/TargetRules.cs
--- a/Engine/Source/Programs/UnrealBuildTool/Configuration/TargetRules.cs
+++ b/Engine/Source/Programs/UnrealBuildTool/Configuration/TargetRules.cs
@@ -2295,6 +2295,7 @@ namespace UnrealBuildTool
 		/// This option is provided to facilitate evaluation of the feature.
 		/// Expect the name of the option to change. Use of coroutines in with UE is untested and unsupported.
 		/// </summary>
+		[Obsolete("CppCoroutines are not supported")]
 		public bool bEnableCppCoroutinesForEvaluation { get; set; }
 
 		/// <summary>
@@ -2472,6 +2473,17 @@ namespace UnrealBuildTool
 		[XmlConfigFile(Category = "BuildConfiguration")]
 		public bool bPublicSymbolsByDefault { get; set; }
 
+		/// <summary>
+		/// Enables Identical Code Folding (Clang toolchains)
+		/// </summary>
+		[RequiresUniqueBuildEnvironment]
+		public bool bIdenticalCodeFolding
+		{
+			get => bIdenticalCodeFoldingPrivate ?? Configuration == UnrealTargetConfiguration.Shipping;
+			set => bIdenticalCodeFoldingPrivate = value;
+		}
+		private bool? bIdenticalCodeFoldingPrivate;
+
 		/// <summary>
 		/// Disable supports for inlining gen.cpps
 		/// </summary>
diff --git a/Engine/Source/Programs/UnrealBuildTool/Configuration/UEBuildTarget.cs b/Engine/Source/Programs/UnrealBuildTool/Configuration/UEBuildTarget.cs
--- a/Engine/Source/Programs/UnrealBuildTool/Configuration/UEBuildTarget.cs
+++ b/Engine/Source/Programs/UnrealBuildTool/Configuration/UEBuildTarget.cs
@@ -5869,12 +5869,6 @@ namespace UnrealBuildTool
 				GlobalCompileEnvironment.Definitions.Add("WITH_CPP_MODULES=0");
 			}
 
-			// Whether C++20 coroutines are enabled
-			if (Rules.bEnableCppCoroutinesForEvaluation)
-			{
-				Log.TraceInformationOnce($"NOTE: C++ coroutine experimental support is no longer available");
-			}
-
 			if (Rules.bEnableProcessPriorityControl)
 			{
 				GlobalCompileEnvironment.Definitions.Add("WITH_PROCESS_PRIORITY_CONTROL=1");
diff --git a/Engine/Source/Programs/UnrealBuildTool/Platform/Android/AndroidToolChain.cs b/Engine/Source/Programs/UnrealBuildTool/Platform/Android/AndroidToolChain.cs
--- a/Engine/Source/Programs/UnrealBuildTool/Platform/Android/AndroidToolChain.cs
+++ b/Engine/Source/Programs/UnrealBuildTool/Platform/Android/AndroidToolChain.cs
@@ -1,13 +1,13 @@
 // Copyright Epic Games, Inc. All Rights Reserved.
 
+using EpicGames.Core;
+using Microsoft.Extensions.Logging;
 using System;
 using System.Collections.Generic;
 using System.Diagnostics;
 using System.IO;
 using System.Linq;
 using System.Text.RegularExpressions;
-using EpicGames.Core;
-using Microsoft.Extensions.Logging;
 using UnrealBuildBase;
 
 namespace UnrealBuildTool
@@ -770,9 +770,13 @@ namespace UnrealBuildTool
 				Result += " -march=armv8-a";
 			}
 
+			if (TargetRules?.bIdenticalCodeFolding == true)
+			{
+				Result += "-Wl,--icf=all";
+			}
+
 			if (LinkEnvironment.Configuration == CppConfiguration.Shipping)
 			{
-				Result += " -Wl,--icf=all"; // Enables ICF (Identical Code Folding). [all, safe] safe == fold functions that can be proven not to have their address taken.
 				Result += " -Wl,-O3";
 			}
 
diff --git a/Engine/Source/Programs/UnrealBuildTool/Platform/Linux/LinuxToolChain.cs b/Engine/Source/Programs/UnrealBuildTool/Platform/Linux/LinuxToolChain.cs
--- a/Engine/Source/Programs/UnrealBuildTool/Platform/Linux/LinuxToolChain.cs
+++ b/Engine/Source/Programs/UnrealBuildTool/Platform/Linux/LinuxToolChain.cs
@@ -848,9 +848,13 @@ namespace UnrealBuildTool
 				}
 			}
 
+			if (TargetRules?.bIdenticalCodeFolding == true)
+			{
+				Arguments.Add("-Wl,--icf=all");
+			}
+
 			if (LinkEnvironment.Configuration == CppConfiguration.Shipping)
 			{
-				Arguments.Add("-Wl,--icf=all"); // Enables ICF (Identical Code Folding). [all, safe] safe == fold functions that can be proven not to have their address taken.
 				Arguments.Add("-Wl,-O3");
 			}
 		}
