UnrealBuildTool: Move obsolete TargetRules and ModuleRules properties to separate files

#rnx

diff --git a/Engine/Source/Programs/UnrealBuildTool/Configuration/ModuleRules.Obsolete.cs b/Engine/Source/Programs/UnrealBuildTool/Configuration/ModuleRules.Obsolete.cs
new file mode 100644
--- /dev/null
+++ b/Engine/Source/Programs/UnrealBuildTool/Configuration/ModuleRules.Obsolete.cs
@@ -0,0 +1,118 @@
+// Copyright Epic Games, Inc. All Rights Reserved.
+
+using System;
+
+namespace UnrealBuildTool
+{
+	public partial class ModuleRules
+	{
+		/// <summary>
+		/// Obsolete: Direct the compiler to generate AVX instructions wherever SSE or AVX intrinsics are used, on the platforms that support it.
+		/// Note that by enabling this you are changing the minspec for the PC platform, and the resultant executable will crash on machines without AVX support.
+		/// </summary>
+		[Obsolete("bUseAVX is obsolete and will be removed in UE5.4, please replace with MinCpuArchX64")]
+		public bool bUseAVX
+		{
+			get => MinCpuArchX64 >= MinimumCpuArchitectureX64.AVX;
+			set => MinCpuArchX64 = value ? MinimumCpuArchitectureX64.AVX : MinimumCpuArchitectureX64.None;
+		}
+
+		/// <summary>
+		/// Enforce "include what you use" rules when PCHUsage is set to ExplicitOrSharedPCH; warns when monolithic headers (Engine.h, UnrealEd.h, etc...) 
+		/// are used, and checks that source files include their matching header first.
+		/// </summary>
+		[Obsolete("Deprecated in UE5.2 - Use IWYUSupport instead.")]
+		public bool bEnforceIWYU
+		{
+			set
+			{
+				if (!value)
+				{
+					IWYUSupport = IWYUSupport.None;
+				}
+			}
+		}
+
+		/// <summary>
+		/// How to treat deterministic warnings (experimental).
+		/// </summary>
+		[Obsolete("Deprecated in UE5.6 - Replace with ModuleRules.CppCompileWarningSettings.DeterministicWarningLevel")]
+		public WarningLevel DeterministicWarningLevel
+		{
+			get => CppCompileWarningSettings.DeterministicWarningLevel;
+			set => CppCompileWarningSettings.DeterministicWarningLevel = value;
+		}
+
+		/// <summary>
+		/// How to treat shadow variable warnings
+		/// </summary>
+		[Obsolete("Deprecated in UE5.6 - Replace with ModuleRules.CppCompileWarningSettings.ShadowVariableWarningLevel")]
+		public WarningLevel ShadowVariableWarningLevel
+		{
+			get => CppCompileWarningSettings.ShadowVariableWarningLevel;
+			set => CppCompileWarningSettings.ShadowVariableWarningLevel = value;
+		}
+
+		/// <summary>
+		/// How to treat unsafe implicit type cast warnings (e.g., double->float or int64->int32)
+		/// </summary>
+		[Obsolete("Deprecated in UE5.6 - Replace with ModuleRules.CppCompileWarningSettings.UnsafeTypeCastWarningLevel")]
+		public WarningLevel UnsafeTypeCastWarningLevel
+		{
+			get => CppCompileWarningSettings.UnsafeTypeCastWarningLevel;
+			set => CppCompileWarningSettings.UnsafeTypeCastWarningLevel = value;
+		}
+
+		/// <summary>
+		/// Indicates what warning/error level to treat undefined identifiers in conditional expressions.
+		/// </summary>
+		[Obsolete("Deprecated in UE5.6 - Replace with ModuleRules.CppCompileWarningSettings.UndefinedIdentifierWarningLevel")]
+		public WarningLevel UndefinedIdentifierWarningLevel
+		{
+			get => CppCompileWarningSettings.UndefinedIdentifierWarningLevel;
+			set => CppCompileWarningSettings.UndefinedIdentifierWarningLevel = value;
+		}
+
+		/// <summary>
+		/// Enable warnings for using undefined identifiers in #if expressions
+		/// </summary>
+		[Obsolete("Deprecated in UE5.5 - Replace with ModuleRules.UndefinedIdentifierWarningLevel")]
+		public bool bEnableUndefinedIdentifierWarnings { get; set; } = true;
+
+		/// <summary>
+		/// How to treat general module include path validation messages
+		/// </summary>
+		[Obsolete("Deprecated in UE5.6 - Replace with ModuleRules.CppCompileWarningSettings.ModuleIncludePathWarningLevel")]
+		public WarningLevel ModuleIncludePathWarningLevel
+		{
+			get => CppCompileWarningSettings.ModuleIncludePathWarningLevel;
+			set => CppCompileWarningSettings.ModuleIncludePathWarningLevel = value;
+		}
+
+		/// <summary>
+		/// How to treat private module include path validation messages, where a module is adding an include path that exposes private headers
+		/// </summary>
+		[Obsolete("Deprecated in UE5.6 - Replace with ModuleRules.CppCompileWarningSettings.ModuleIncludePrivateWarningLevel")]
+		public WarningLevel ModuleIncludePrivateWarningLevel
+		{
+			get => CppCompileWarningSettings.ModuleIncludePrivateWarningLevel;
+			set => CppCompileWarningSettings.ModuleIncludePrivateWarningLevel = value;
+		}
+
+		/// <summary>
+		/// How to treat unnecessary module sub-directory include path validation messages
+		/// </summary>
+		[Obsolete("Deprecated in UE5.6 - Replace with ModuleRules.CppCompileWarningSettings.ModuleIncludeSubdirectoryWarningLevel")]
+		public WarningLevel ModuleIncludeSubdirectoryWarningLevel
+		{
+			get => CppCompileWarningSettings.ModuleIncludeSubdirectoryWarningLevel;
+			set => CppCompileWarningSettings.ModuleIncludeSubdirectoryWarningLevel = value;
+		}
+
+		/// <summary>
+		/// Enable warnings for when there are .gen.cpp files that could be inlined in a matching handwritten cpp file
+		/// </summary>
+		[Obsolete("Deprecated in UE5.6 - Replace with ModuleRules.CppCompileWarningSettings.NonInlinedGenCppWarningLevel")]
+		public bool bEnableNonInlinedGenCppWarnings { get; set; }
+	}
+}
\ No newline at end of file
diff --git a/Engine/Source/Programs/UnrealBuildTool/Configuration/ModuleRules.cs b/Engine/Source/Programs/UnrealBuildTool/Configuration/ModuleRules.cs
--- a/Engine/Source/Programs/UnrealBuildTool/Configuration/ModuleRules.cs
+++ b/Engine/Source/Programs/UnrealBuildTool/Configuration/ModuleRules.cs
@@ -100,7 +100,7 @@ namespace UnrealBuildTool
 	/// <summary>
 	/// ModuleRules is a data structure that contains the rules for defining a module
 	/// </summary>
-	public class ModuleRules
+	public partial class ModuleRules
 	{
 		/// <summary>
 		/// Type of module
@@ -927,17 +927,6 @@ namespace UnrealBuildTool
 		/// </summary>
 		public bool bCodeCoverage { get; set; }
 
-		/// <summary>
-		/// Obsolete: Direct the compiler to generate AVX instructions wherever SSE or AVX intrinsics are used, on the platforms that support it.
-		/// Note that by enabling this you are changing the minspec for the PC platform, and the resultant executable will crash on machines without AVX support.
-		/// </summary>
-		[Obsolete("bUseAVX is obsolete and will be removed in UE5.4, please replace with MinCpuArchX64")]
-		public bool bUseAVX
-		{
-			get => MinCpuArchX64 >= MinimumCpuArchitectureX64.AVX;
-			set => MinCpuArchX64 = value ? MinimumCpuArchitectureX64.AVX : MinimumCpuArchitectureX64.None;
-		}
-
 		/// <summary>
 		/// Direct the compiler to generate AVX instructions wherever SSE or AVX intrinsics are used, on the x64 platforms that support it.
 		/// Note that by enabling this you are changing the minspec for the PC platform, and the resultant executable will crash on machines without AVX support.
@@ -966,93 +955,16 @@ namespace UnrealBuildTool
 		/// </summary>
 		public bool bEnableObjCAutomaticReferenceCounting { get; set; }
 
-		/// <summary>
-		/// How to treat deterministic warnings (experimental).
-		/// </summary>
-		[Obsolete("Deprecated in UE5.6 - Replace with ModuleRules.CppCompileWarningSettings.DeterministicWarningLevel")]
-		public WarningLevel DeterministicWarningLevel
-		{
-			get => CppCompileWarningSettings.DeterministicWarningLevel;
-			set => CppCompileWarningSettings.DeterministicWarningLevel = value;
-		}
-
 		/// <summary>
 		/// Settings object used by the module.
 		/// </summary>
 		public CppCompileWarnings CppCompileWarningSettings { get; init; }
 
-		/// <summary>
-		/// How to treat shadow variable warnings
-		/// </summary>
-		[Obsolete("Deprecated in UE5.6 - Replace with ModuleRules.CppCompileWarningSettings.ShadowVariableWarningLevel")]
-		public WarningLevel ShadowVariableWarningLevel
-		{
-			get => CppCompileWarningSettings.ShadowVariableWarningLevel;
-			set => CppCompileWarningSettings.ShadowVariableWarningLevel = value;
-		}
-
 		/// <summary>
 		/// Whether to enable all warnings as errors. UE enables most warnings as errors already, but disables a few (such as deprecation warnings).
 		/// </summary>`
 		public bool bWarningsAsErrors { get; set; }
 
-		/// <summary>
-		/// How to treat unsafe implicit type cast warnings (e.g., double->float or int64->int32)
-		/// </summary>
-		[Obsolete("Deprecated in UE5.6 - Replace with ModuleRules.CppCompileWarningSettings.UnsafeTypeCastWarningLevel")]
-		public WarningLevel UnsafeTypeCastWarningLevel
-		{
-			get => CppCompileWarningSettings.UnsafeTypeCastWarningLevel;
-			set => CppCompileWarningSettings.UnsafeTypeCastWarningLevel = value;
-		}
-
-
-		/// <summary>
-		/// Indicates what warning/error level to treat undefined identifiers in conditional expressions.
-		/// </summary>
-		[Obsolete("Deprecated in UE5.6 - Replace with ModuleRules.CppCompileWarningSettings.UndefinedIdentifierWarningLevel")]
-		public WarningLevel UndefinedIdentifierWarningLevel
-		{
-			get => CppCompileWarningSettings.UndefinedIdentifierWarningLevel;
-			set => CppCompileWarningSettings.UndefinedIdentifierWarningLevel = value;
-		}
-
-		/// <summary>
-		/// Enable warnings for using undefined identifiers in #if expressions
-		/// </summary>
-		[Obsolete("Deprecated in UE5.5 - Replace with ModuleRules.UndefinedIdentifierWarningLevel")]
-		public bool bEnableUndefinedIdentifierWarnings { get; set; } = true;
-
-		/// <summary>
-		/// How to treat general module include path validation messages
-		/// </summary>
-		[Obsolete("Deprecated in UE5.6 - Replace with ModuleRules.CppCompileWarningSettings.ModuleIncludePathWarningLevel")]
-		public WarningLevel ModuleIncludePathWarningLevel
-		{
-			get => CppCompileWarningSettings.ModuleIncludePathWarningLevel;
-			set => CppCompileWarningSettings.ModuleIncludePathWarningLevel = value;
-		}
-
-		/// <summary>
-		/// How to treat private module include path validation messages, where a module is adding an include path that exposes private headers
-		/// </summary>
-		[Obsolete("Deprecated in UE5.6 - Replace with ModuleRules.CppCompileWarningSettings.ModuleIncludePrivateWarningLevel")]
-		public WarningLevel ModuleIncludePrivateWarningLevel
-		{
-			get => CppCompileWarningSettings.ModuleIncludePrivateWarningLevel;
-			set => CppCompileWarningSettings.ModuleIncludePrivateWarningLevel = value;
-		}
-
-		/// <summary>
-		/// How to treat unnecessary module sub-directory include path validation messages
-		/// </summary>
-		[Obsolete("Deprecated in UE5.6 - Replace with ModuleRules.CppCompileWarningSettings.ModuleIncludeSubdirectoryWarningLevel")]
-		public WarningLevel ModuleIncludeSubdirectoryWarningLevel
-		{
-			get => CppCompileWarningSettings.ModuleIncludeSubdirectoryWarningLevel;
-			set => CppCompileWarningSettings.ModuleIncludeSubdirectoryWarningLevel = value;
-		}
-
 		/// <summary>
 		/// Disable all static analysis - clang, msvc, pvs-studio.
 		/// </summary>
@@ -1171,12 +1083,6 @@ namespace UnrealBuildTool
 		/// </summary>
 		public bool bBuildLocallyWithSNDBS { get; set; }
 
-		/// <summary>
-		/// Enable warnings for when there are .gen.cpp files that could be inlined in a matching handwritten cpp file
-		/// </summary>
-		[Obsolete("Deprecated in UE5.6 - Replace with ModuleRules.CppCompileWarningSettings.NonInlinedGenCppWarningLevel")]
-		public bool bEnableNonInlinedGenCppWarnings { get; set; }
-
 		/// <summary>
 		/// Redistribution override flag for this module.
 		/// </summary>
@@ -1199,22 +1105,6 @@ namespace UnrealBuildTool
 		/// </summary>
 		public Dictionary<string, string> AliasRestrictedFolders = new();
 
-		/// <summary>
-		/// Enforce "include what you use" rules when PCHUsage is set to ExplicitOrSharedPCH; warns when monolithic headers (Engine.h, UnrealEd.h, etc...) 
-		/// are used, and checks that source files include their matching header first.
-		/// </summary>
-		[Obsolete("Deprecated in UE5.2 - Use IWYUSupport instead.")]
-		public bool bEnforceIWYU
-		{
-			set
-			{
-				if (!value)
-				{
-					IWYUSupport = IWYUSupport.None;
-				}
-			}
-		}
-
 		/// <summary>
 		/// Allows "include what you use" to modify the source code when run. bEnforceIWYU must be true for this variable to matter.
 		/// 
diff --git a/Engine/Source/Programs/UnrealBuildTool/Configuration/ReadOnlyTargetRules.Obsolete.cs b/Engine/Source/Programs/UnrealBuildTool/Configuration/ReadOnlyTargetRules.Obsolete.cs
new file mode 100644
--- /dev/null
+++ b/Engine/Source/Programs/UnrealBuildTool/Configuration/ReadOnlyTargetRules.Obsolete.cs
@@ -0,0 +1,77 @@
+// Copyright Epic Games, Inc. All Rights Reserved.
+
+using System;
+
+namespace UnrealBuildTool
+{
+	public partial class ReadOnlyTargetRules
+	{
+#pragma warning disable CS1591
+		[Obsolete("Deprecated in UE5.5 - No longer used in engine.")]
+		public bool bUsesSteam => Inner.bUsesSteam;
+
+		[Obsolete("Deprecated in UE5.1 - No longer used as Chaos is always enabled.")]
+		public bool bCompileChaos => Inner.bCompileChaos;
+
+		[Obsolete("Deprecated in UE5.1 - No longer used as Chaos is always enabled.")]
+		public bool bUseChaos => Inner.bUseChaos;
+
+		[Obsolete("Deprecated in UE5.1 - No longer used in engine.")]
+		public bool bCustomSceneQueryStructure => Inner.bCustomSceneQueryStructure;
+
+		[Obsolete("Deprecated in UE5.1 - No longer used as Chaos is always enabled.")]
+		public bool bCompilePhysX => Inner.bCompilePhysX;
+
+		[Obsolete("Deprecated in UE5.1 - No longer used as Chaos is always enabled.")]
+		public bool bCompileAPEX => Inner.bCompileAPEX;
+
+		[Obsolete("Deprecated in UE5.1 - No longer used as Chaos is always enabled.")]
+		public bool bCompileNvCloth => Inner.bCompileNvCloth;
+
+		[Obsolete("Deprecated in UE5.4 - No longer used.")]
+		public bool bCompileIntelMetricsDiscovery => Inner.bCompileIntelMetricsDiscovery;
+
+		[Obsolete("Deprecated in UE 5.6 - DirectXMath path was incompatible with LWC and unused since 5.0")]
+		public bool bWithDirectXMath => Inner.bWithDirectXMath;
+
+		[Obsolete("Deprecated in UE5.1 - Please use OptimizationLevel instead.")]
+		public bool bCompileForSize => Inner.bCompileForSize;
+
+		[Obsolete("Deprecated in UE5.6")]
+		public bool bDetailedUnityFiles => Inner.bDetailedUnityFiles;
+
+		[Obsolete("Deprecated in UE5.4 - Replace with ReadOnlyTargetRules.DebugInfo")]
+		public bool bDisableDebugInfo => Inner.bDisableDebugInfo;
+
+		[Obsolete("Deprecated in UE5.3 - No longer used as MallocProfiler is removed in favor of UnrealInsights.")]
+		public bool bUseMallocProfiler => Inner.bUseMallocProfiler;
+
+		[Obsolete("Deprecated in UE5.6 - Replace with ReadOnlyTargetRules.CppCompileWarningSettings.DeprecationWarningLevel")]
+		public WarningLevel DeprecationWarningLevel => Inner.CppCompileWarningSettings.DeprecationWarningLevel;
+
+		[Obsolete("Deprecated in UE5.6 - Replace with ReadOnlyTargetRules.CppCompileWarningSettings.ShadowVariableWarningLevel")]
+		public WarningLevel ShadowVariableWarningLevel => Inner.CppCompileWarningSettings.ShadowVariableWarningLevel;
+
+		[Obsolete("Deprecated in UE5.6 - Replace with ReadOnlyTargetRules.CppCompileWarningSettings.UnsafeTypeCastWarningLevel")]
+		public WarningLevel UnsafeTypeCastWarningLevel => Inner.CppCompileWarningSettings.UnsafeTypeCastWarningLevel;
+
+		[Obsolete("Deprecated in UE5.6 - Replace with ReadOnlyTargetRules.CppCompileWarningSettings.UndefinedIdentifierWarningLevel")]
+		public WarningLevel UndefinedIdentifierWarningLevel => Inner.CppCompileWarningSettings.UndefinedIdentifierWarningLevel;
+
+		[Obsolete("Deprecated in UE5.5 - Replace with ModuleRules.UndefinedIdentifierWarningLevel")]
+		public bool bUndefinedIdentifierErrors => Inner.bUndefinedIdentifierErrors;
+
+		[Obsolete("Deprecated in UE5.6 - Replace with ReadOnlyTargetRules.CppCompileWarningSettings.PCHPerformanceIssueWarningLevel")]
+		public WarningLevel PCHPerformanceIssueWarningLevel => Inner.CppCompileWarningSettings.PCHPerformanceIssueWarningLevel;
+
+		[Obsolete("Deprecated in UE5.6 - Replace with ReadOnlyTargetRules.CppCompileWarningSettings.ModuleIncludePathWarningLevel")]
+		public WarningLevel ModuleIncludePathWarningLevel => Inner.CppCompileWarningSettings.ModuleIncludePathWarningLevel;
+
+		[Obsolete("Deprecated in UE5.6 - Replace with ReadOnlyTargetRules.CppCompileWarningSettings.ModuleIncludePrivateWarningLevel")]
+		public WarningLevel ModuleIncludePrivateWarningLevel => Inner.CppCompileWarningSettings.ModuleIncludePrivateWarningLevel;
+
+		[Obsolete("Deprecated in UE5.6 - Replace with ReadOnlyTargetRules.CppCompileWarningSettings.ModuleIncludeSubdirectoryWarningLevel")]
+		public WarningLevel ModuleIncludeSubdirectoryWarningLevel => Inner.CppCompileWarningSettings.ModuleIncludeSubdirectoryWarningLevel;
+#pragma warning restore CS1591
+	}
+}
\ No newline at end of file
diff --git a/Engine/Source/Programs/UnrealBuildTool/Configuration/ReadOnlyTargetRules.cs b/Engine/Source/Programs/UnrealBuildTool/Configuration/ReadOnlyTargetRules.cs
--- a/Engine/Source/Programs/UnrealBuildTool/Configuration/ReadOnlyTargetRules.cs
+++ b/Engine/Source/Programs/UnrealBuildTool/Configuration/ReadOnlyTargetRules.cs
@@ -82,9 +82,6 @@ namespace UnrealBuildTool
 
 		public string? OutputFile => Inner.OutputFile;
 
-		[Obsolete("Deprecated in UE5.5 - No longer used in engine.")]
-		public bool bUsesSteam => Inner.bUsesSteam;
-
 		public bool bUsesCEF3 => Inner.bUsesCEF3;
 
 		public bool bUsesSlate => Inner.bUsesSlate;
@@ -142,30 +139,12 @@ namespace UnrealBuildTool
 		public bool bAutoRTFMVerify => Inner.bAutoRTFMVerify;
 		public bool bAutoRTFMClosedStaticLinkage => Inner.bAutoRTFMClosedStaticLinkage;
 
-		[Obsolete("Deprecated in UE5.1 - No longer used as Chaos is always enabled.")]
-		public bool bCompileChaos => Inner.bCompileChaos;
-
-		[Obsolete("Deprecated in UE5.1 - No longer used as Chaos is always enabled.")]
-		public bool bUseChaos => Inner.bUseChaos;
-
 		public bool bUseChaosMemoryTracking => Inner.bUseChaosMemoryTracking;
 
 		public bool bCompileChaosVisualDebuggerSupport => Inner.bCompileChaosVisualDebuggerSupport;
 
 		public bool bUseChaosChecked => Inner.bUseChaosChecked;
 
-		[Obsolete("Deprecated in UE5.1 - No longer used in engine.")]
-		public bool bCustomSceneQueryStructure => Inner.bCustomSceneQueryStructure;
-
-		[Obsolete("Deprecated in UE5.1 - No longer used as Chaos is always enabled.")]
-		public bool bCompilePhysX => Inner.bCompilePhysX;
-
-		[Obsolete("Deprecated in UE5.1 - No longer used as Chaos is always enabled.")]
-		public bool bCompileAPEX => Inner.bCompileAPEX;
-
-		[Obsolete("Deprecated in UE5.1 - No longer used as Chaos is always enabled.")]
-		public bool bCompileNvCloth => Inner.bCompileNvCloth;
-
 		public bool bCompileICU => Inner.bCompileICU;
 
 		public bool bCompileCEF3 => Inner.bCompileCEF3;
@@ -182,9 +161,6 @@ namespace UnrealBuildTool
 
 		public bool bTrackRHIResourceInfoForTest => Inner.bTrackRHIResourceInfoForTest;
 
-		[Obsolete("Deprecated in UE5.4 - No longer used.")]
-		public bool bCompileIntelMetricsDiscovery => Inner.bCompileIntelMetricsDiscovery;
-
 		public bool bCompilePython => Inner.bCompilePython;
 
 		public bool bBuildEditor => Inner.bBuildEditor;
@@ -263,9 +239,6 @@ namespace UnrealBuildTool
 
 		public bool bUseDebugLiveCodingConsole => Inner.bUseDebugLiveCodingConsole;
 
-		[Obsolete("Deprecated in UE 5.6 - DirectXMath path was incompatible with LWC and unused since 5.0")]
-		public bool bWithDirectXMath => Inner.bWithDirectXMath;
-
 		public bool bWithFixedTimeStepSupport => Inner.bWithFixedTimeStepSupport;
 
 		public bool bUseLoggingInShipping => Inner.bUseLoggingInShipping;
@@ -290,9 +263,6 @@ namespace UnrealBuildTool
 
 		public bool bUseExecCommandsInShipping => Inner.bUseExecCommandsInShipping;
 
-		[Obsolete("Deprecated in UE5.1 - Please use OptimizationLevel instead.")]
-		public bool bCompileForSize => Inner.bCompileForSize;
-
 		public OptimizationMode OptimizationLevel => Inner.OptimizationLevel;
 
 		public FPSemanticsMode FPSemantics => Inner.FPSemantics;
@@ -361,35 +331,8 @@ namespace UnrealBuildTool
 
 		public WarningLevel DefaultWarningLevel => Inner.DefaultWarningLevel;
 
-		[Obsolete("Deprecated in UE5.6 - Replace with ReadOnlyTargetRules.CppCompileWarningSettings.DeprecationWarningLevel")]
-		public WarningLevel DeprecationWarningLevel => Inner.CppCompileWarningSettings.DeprecationWarningLevel;
-
 		public ReadOnlyCppCompileWarnings CppCompileWarningSettings { get; init; }
 
-		[Obsolete("Deprecated in UE5.6 - Replace with ReadOnlyTargetRules.CppCompileWarningSettings.ShadowVariableWarningLevel")]
-		public WarningLevel ShadowVariableWarningLevel => Inner.CppCompileWarningSettings.ShadowVariableWarningLevel;
-
-		[Obsolete("Deprecated in UE5.6 - Replace with ReadOnlyTargetRules.CppCompileWarningSettings.UnsafeTypeCastWarningLevel")]
-		public WarningLevel UnsafeTypeCastWarningLevel => Inner.CppCompileWarningSettings.UnsafeTypeCastWarningLevel;
-
-		[Obsolete("Deprecated in UE5.6 - Replace with ReadOnlyTargetRules.CppCompileWarningSettings.UndefinedIdentifierWarningLevel")]
-		public WarningLevel UndefinedIdentifierWarningLevel => Inner.CppCompileWarningSettings.UndefinedIdentifierWarningLevel;
-
-		[Obsolete("Deprecated in UE5.5 - Replace with ModuleRules.UndefinedIdentifierWarningLevel")]
-		public bool bUndefinedIdentifierErrors => Inner.bUndefinedIdentifierErrors;
-
-		[Obsolete("Deprecated in UE5.6 - Replace with ReadOnlyTargetRules.CppCompileWarningSettings.PCHPerformanceIssueWarningLevel")]
-		public WarningLevel PCHPerformanceIssueWarningLevel => Inner.CppCompileWarningSettings.PCHPerformanceIssueWarningLevel;
-
-		[Obsolete("Deprecated in UE5.6 - Replace with ReadOnlyTargetRules.CppCompileWarningSettings.ModuleIncludePathWarningLevel")]
-		public WarningLevel ModuleIncludePathWarningLevel => Inner.CppCompileWarningSettings.ModuleIncludePathWarningLevel;
-
-		[Obsolete("Deprecated in UE5.6 - Replace with ReadOnlyTargetRules.CppCompileWarningSettings.ModuleIncludePrivateWarningLevel")]
-		public WarningLevel ModuleIncludePrivateWarningLevel => Inner.CppCompileWarningSettings.ModuleIncludePrivateWarningLevel;
-
-		[Obsolete("Deprecated in UE5.6 - Replace with ReadOnlyTargetRules.CppCompileWarningSettings.ModuleIncludeSubdirectoryWarningLevel")]
-		public WarningLevel ModuleIncludeSubdirectoryWarningLevel => Inner.CppCompileWarningSettings.ModuleIncludeSubdirectoryWarningLevel;
-
 		public bool bShowIncludes => Inner.bShowIncludes;
 
 		public bool bWarningsAsErrors => Inner.bWarningsAsErrors;
@@ -402,12 +345,6 @@ namespace UnrealBuildTool
 
 		public bool bStressTestUnity => Inner.bStressTestUnity;
 
-		[Obsolete("Deprecated in UE5.6")]
-		public bool bDetailedUnityFiles => Inner.bDetailedUnityFiles;
-
-		[Obsolete("Deprecated in UE5.4 - Replace with ReadOnlyTargetRules.DebugInfo")]
-		public bool bDisableDebugInfo => Inner.bDisableDebugInfo;
-
 		public DebugInfoMode DebugInfo => Inner.DebugInfo;
 
 		public IReadOnlySet<string> DisableDebugInfoModules => Inner.DisableDebugInfoModules;
@@ -482,9 +419,6 @@ namespace UnrealBuildTool
 
 		public bool bEnableProcessPriorityControl => Inner.bEnableProcessPriorityControl;
 
-		[Obsolete("Deprecated in UE5.3 - No longer used as MallocProfiler is removed in favor of UnrealInsights.")]
-		public bool bUseMallocProfiler => Inner.bUseMallocProfiler;
-
 		public bool bShaderCompilerWorkerTrace => Inner.bShaderCompilerWorkerTrace;
 
 		public bool bUseSharedPCHs => Inner.bUseSharedPCHs;
diff --git a/Engine/Source/Programs/UnrealBuildTool/Configuration/TargetRules.Obsolete.cs b/Engine/Source/Programs/UnrealBuildTool/Configuration/TargetRules.Obsolete.cs
new file mode 100644
--- /dev/null
+++ b/Engine/Source/Programs/UnrealBuildTool/Configuration/TargetRules.Obsolete.cs
@@ -0,0 +1,212 @@
+// Copyright Epic Games, Inc. All Rights Reserved.
+
+using System;
+
+namespace UnrealBuildTool
+{
+	public abstract partial class TargetRules
+	{
+		/// <summary>
+		/// Whether the target uses Steam.
+		/// </summary>
+		[Obsolete("Deprecated in UE5.5 - No longer used in engine.")]
+		public bool bUsesSteam { get; set; }
+
+		/// <summary>
+		/// Whether to compile the Chaos physics plugin.
+		/// </summary>
+		[Obsolete("Deprecated in UE5.1 - No longer used as Chaos is always enabled.")]
+		public bool bCompileChaos { get; set; }
+
+		/// <summary>
+		/// Whether to use the Chaos physics interface. This overrides the physx flags to disable APEX and NvCloth
+		/// </summary>
+		[Obsolete("Deprecated in UE5.1 - No longer used as Chaos is always enabled.")]
+		public bool bUseChaos { get; set; }
+
+		/// <summary>
+		/// Whether scene query acceleration is done by UE. The physx scene query structure is still created, but we do not use it.
+		/// </summary>
+		[Obsolete("Deprecated in UE5.1 - No longer used in engine.")]
+		public bool bCustomSceneQueryStructure { get; set; }
+
+		/// <summary>
+		/// Whether to include PhysX support.
+		/// </summary>
+		[Obsolete("Deprecated in UE5.1 - No longer used as Chaos is always enabled.")]
+		public bool bCompilePhysX { get; set; }
+
+		/// <summary>
+		/// Whether to include PhysX APEX support.
+		/// </summary>
+		[Obsolete("Deprecated in UE5.1 - No longer used as Chaos is always enabled.")]
+		public bool bCompileAPEX { get; set; }
+
+		/// <summary>
+		/// Whether to include NvCloth.
+		/// </summary>
+		[Obsolete("Deprecated in UE5.1 - No longer used as Chaos is always enabled.")]
+		public bool bCompileNvCloth { get; set; }
+
+		/// <summary>
+		/// Whether to compile IntelMetricsDiscovery.
+		/// </summary>
+		[Obsolete("Deprecated in UE5.4 - No longer used.")]
+		public bool bCompileIntelMetricsDiscovery { get; set; } = false;
+
+		/// <summary>
+		/// Obsolete: whether to enable support for DirectX Math
+		/// Not compatible with LWC so has been disabled since UE 5.
+		/// </summary>
+		[RequiresUniqueBuildEnvironment]
+		[Obsolete("Deprecated in UE 5.6 - DirectXMath path was incompatible with LWC and unused since 5.0")]
+		public bool bWithDirectXMath { get; set; }
+
+		/// <summary>
+		/// True if we want to favor optimizing size over speed.
+		/// </summary>
+		[Obsolete("Deprecated in UE5.1 - Please use OptimizationLevel instead.")]
+		public bool bCompileForSize { get; set; }
+
+		/// <summary>
+		/// Whether to merge module and generated unity files for faster compilation.
+		/// </summary>
+		[Obsolete("Deprecated in UE5.3; use DisableMergingModuleAndGeneratedFilesInUnityFiles instead")]
+		public bool bMergeModuleAndGeneratedUnityFiles { get; set; }
+
+		/// <summary>
+		/// Whether to add additional information to the unity files, such as '_of_X' in the file name. Not recommended.
+		/// </summary>
+		[XmlConfigFile(Category = "BuildConfiguration")]
+		[Obsolete("Deprecated in UE5.6")]
+		public bool bDetailedUnityFiles { get; set; }
+
+		/// <summary>
+		/// Whether to globally disable debug info generation; Obsolete, please use TargetRules.DebugInfoMode instead
+		/// </summary>
+		[XmlConfigFile(Category = "BuildConfiguration", Deprecated = true)]
+		[Obsolete("Deprecated in UE5.4 - Replace with TargetRules.DebugInfo")]
+		public bool bDisableDebugInfo
+		{
+			get => DebugInfo == DebugInfoMode.None;
+			set => DebugInfo = value ? DebugInfo = DebugInfoMode.None : DebugInfo = DebugInfoMode.Full;
+		}
+
+		/// <summary>
+		/// Whether to enable support for C++20 coroutines
+		/// This option is provided to facilitate evaluation of the feature.
+		/// Expect the name of the option to change. Use of coroutines in with UE is untested and unsupported.
+		/// </summary>
+		[Obsolete("CppCoroutines are not supported")]
+		public bool bEnableCppCoroutinesForEvaluation { get; set; }
+
+		/// <summary>
+		/// If true, then enable memory profiling in the build (defines USE_MALLOC_PROFILER=1 and forces bOmitFramePointers=false).
+		/// </summary>
+		[Obsolete("Deprecated in UE5.3 - No longer used as MallocProfiler is removed in favor of UnrealInsights.")]
+		public bool bUseMallocProfiler { get; set; }
+
+		/// <summary>
+		/// Disable supports for inlining gen.cpps
+		/// </summary>
+		[Obsolete("Deprecated in UE5.3; the engine code relies on inlining generated cpp files to compile")]
+		public bool bDisableInliningGenCpps { get; set; }
+
+		/// <summary>
+		/// How to treat conflicts when a disabled plugin is being enabled by another plugin referencing it
+		/// </summary>
+		[Obsolete("Deprecated in UE5.6 - Replace with TargetRules.CppCompileWarningSettings.DisablePluginsConflictWarningLevel")]
+		public WarningLevel DisablePluginsConflictWarningLevel
+		{
+			get => CppCompileWarningSettings.DisablePluginsConflictWarningLevel;
+			set => CppCompileWarningSettings.DisablePluginsConflictWarningLevel = value;
+		}
+
+		/// <summary>
+		/// Level to report deprecation warnings as errors
+		/// </summary>
+		[Obsolete("Deprecated in UE5.6 - Replace with TargetRules.CppCompileWarningSettings.DeprecationWarningLevel")]
+		public WarningLevel DeprecationWarningLevel
+		{
+			get => CppCompileWarningSettings.DeprecationWarningLevel;
+			set => CppCompileWarningSettings.DeprecationWarningLevel = value;
+		}
+
+		/// <summary>
+		/// Forces shadow variable warnings to be treated as errors on platforms that support it.
+		/// </summary>
+		[Obsolete("Deprecated in UE5.6 - Replace with TargetRules.CppCompileWarningSettings.ShadowVariableWarningLevel")]
+		public WarningLevel ShadowVariableWarningLevel
+		{
+			get => CppCompileWarningSettings.ShadowVariableWarningLevel;
+			set => CppCompileWarningSettings.ShadowVariableWarningLevel = value;
+		}
+
+		/// <summary>
+		/// Indicates what warning/error level to treat unsafe type casts as on platforms that support it (e.g., double->float or int64->int32)
+		/// </summary>
+		[Obsolete("Deprecated in UE5.6 - Replace with TargetRules.CppCompileWarningSettings.UnsafeTypeCastWarningLevel")]
+		public WarningLevel UnsafeTypeCastWarningLevel
+		{
+			get => CppCompileWarningSettings.UnsafeTypeCastWarningLevel;
+			set => CppCompileWarningSettings.UnsafeTypeCastWarningLevel = value;
+		}
+
+		/// <summary>
+		/// Forces the use of undefined identifiers in conditional expressions to be treated as errors.
+		/// </summary>
+		[XmlConfigFile(Category = "BuildConfiguration", Deprecated = true)]
+		[Obsolete("Deprecated in UE5.5 - Replace with TargetRules.UndefinedIdentifierWarningLevel")]
+		public bool bUndefinedIdentifierErrors { get; set; } = true;
+
+		/// <summary>
+		/// Indicates what warning/error level to treat undefined identifiers in conditional expressions.
+		/// </summary>
+		[Obsolete("Deprecated in UE5.6 - Replace with TargetRules.CppCompileWarningSettings.UndefinedIdentifierWarningLevel")]
+		public WarningLevel UndefinedIdentifierWarningLevel
+		{
+			get => CppCompileWarningSettings.UndefinedIdentifierWarningLevel;
+			set => CppCompileWarningSettings.UndefinedIdentifierWarningLevel = value;
+		}
+
+		/// <summary>
+		/// Indicates what warning/error level to treat potential PCH performance issues.
+		/// </summary>
+		[Obsolete("Deprecated in UE5.6 - Replace with TargetRules.CppCompileWarningSettings.PCHPerformanceIssueWarningLevel")]
+		public WarningLevel PCHPerformanceIssueWarningLevel
+		{
+			get => CppCompileWarningSettings.PCHPerformanceIssueWarningLevel;
+			set => CppCompileWarningSettings.PCHPerformanceIssueWarningLevel = value;
+		}
+
+		/// <summary>
+		/// How to treat general module include path validation messages
+		/// </summary>
+		[Obsolete("Deprecated in UE5.6 - Replace with TargetRules.CppCompileWarningSettings.ModuleIncludePathWarningLevel")]
+		public WarningLevel ModuleIncludePathWarningLevel
+		{
+			get => CppCompileWarningSettings.ModuleIncludePathWarningLevel;
+			set => CppCompileWarningSettings.ModuleIncludePathWarningLevel = value;
+		}
+
+		/// <summary>
+		/// How to treat private module include path validation messages, where a module is adding an include path that exposes private headers
+		/// </summary>
+		[Obsolete("Deprecated in UE5.6 - Replace with TargetRules.CppCompileWarningSettings.ModuleIncludePrivateWarningLevel")]
+		public WarningLevel ModuleIncludePrivateWarningLevel
+		{
+			get => CppCompileWarningSettings.ModuleIncludePrivateWarningLevel;
+			set => CppCompileWarningSettings.ModuleIncludePrivateWarningLevel = value;
+		}
+
+		/// <summary>
+		/// How to treat unnecessary module sub-directory include path validation messages
+		/// </summary>
+		[Obsolete("Deprecated in UE5.6 - Replace with TargetRules.CppCompileWarningSettings.ModuleIncludeSubdirectoryWarningLevel")]
+		public WarningLevel ModuleIncludeSubdirectoryWarningLevel
+		{
+			get => CppCompileWarningSettings.ModuleIncludeSubdirectoryWarningLevel;
+			set => CppCompileWarningSettings.ModuleIncludeSubdirectoryWarningLevel = value;
+		}
+	}
+}
\ No newline at end of file
diff --git a/Engine/Source/Programs/UnrealBuildTool/Configuration/TargetRules.cs b/Engine/Source/Programs/UnrealBuildTool/Configuration/TargetRules.cs
--- a/Engine/Source/Programs/UnrealBuildTool/Configuration/TargetRules.cs
+++ b/Engine/Source/Programs/UnrealBuildTool/Configuration/TargetRules.cs
@@ -693,12 +693,6 @@ namespace UnrealBuildTool
 		/// </summary>
 		internal readonly ConfigValueTracker ConfigValueTracker;
 
-		/// <summary>
-		/// Whether the target uses Steam.
-		/// </summary>
-		[Obsolete("Deprecated in UE5.5 - No longer used in engine.")]
-		public bool bUsesSteam { get; set; }
-
 		/// <summary>
 		/// Whether the target uses CEF3.
 		/// </summary>
@@ -771,7 +765,7 @@ namespace UnrealBuildTool
 		/// Additional plugins that are built for this target type but not enabled.
 		/// </summary>
 		[CommandLine("-BuildPlugin=", ListSeparator = '+')]
-		public List<string> BuildPlugins = new();
+		public List<string> BuildPlugins = [];
 
 		/// <summary>
 		/// If this is true, then the BuildPlugins list will be used to populate RuntimeDependencies, rather than EnablePlugins
@@ -784,46 +778,36 @@ namespace UnrealBuildTool
 		/// in ModuleHostType.
 		/// </summary>
 		[CommandLine("-AdditionalPlugins=", ListSeparator = '+')]
-		public List<string> AdditionalPlugins = new();
+		public List<string> AdditionalPlugins = [];
 
 		/// <summary>
 		/// Additional plugins that should be included for this target.
 		/// </summary>
 		[CommandLine("-EnablePlugin=", ListSeparator = '+')]
-		public List<string> EnablePlugins = new();
+		public List<string> EnablePlugins = [];
 
 		/// <summary>
 		/// List of plugins to be disabled for this target. Note that the project file may still reference them, so they should be marked
 		/// as optional to avoid failing to find them at runtime.
 		/// </summary>
 		[CommandLine("-DisablePlugin=", ListSeparator = '+')]
-		public List<string> DisablePlugins = new();
+		public List<string> DisablePlugins = [];
 
 		/// <summary>
 		/// Additional plugins that should be included for this target if they are found.
 		/// </summary>
 		[CommandLine("-OptionalPlugins=", ListSeparator = '+')]
-		public List<string> OptionalPlugins = new();
+		public List<string> OptionalPlugins = [];
 
 		/// <summary>
 		/// If false, suppress loading "enabled by default" engine plugins not explicitly enabled by the project or target files.
 		/// </summary>
 		public bool bAllowEnginePluginsEnabledByDefault { get; set; } = true;
 
-		/// <summary>
-		/// How to treat conflicts when a disabled plugin is being enabled by another plugin referencing it
-		/// </summary>
-		[Obsolete("Deprecated in UE5.6 - Replace with TargetRules.CppCompileWarningSettings.DisablePluginsConflictWarningLevel")]
-		public WarningLevel DisablePluginsConflictWarningLevel
-		{
-			get => CppCompileWarningSettings.DisablePluginsConflictWarningLevel;
-			set => CppCompileWarningSettings.DisablePluginsConflictWarningLevel = value; 
-		}
-
 		/// <summary>
 		/// A list of Plugin names that are allowed to exist as dependencies without being defined in the uplugin descriptor
 		/// </summary>
-		public List<string> InternalPluginDependencies = new();
+		public List<string> InternalPluginDependencies = [];
 
 		/// <summary>
 		/// Path to the set of pak signing keys to embed in the executable.
@@ -937,18 +921,6 @@ namespace UnrealBuildTool
 		[XmlConfigFile(Category = "BuildConfiguration")]
 		public bool bAutoRTFMClosedStaticLinkage { get; set; } = false;
 
-		/// <summary>
-		/// Whether to compile the Chaos physics plugin.
-		/// </summary>
-		[Obsolete("Deprecated in UE5.1 - No longer used as Chaos is always enabled.")]
-		public bool bCompileChaos { get; set; }
-
-		/// <summary>
-		/// Whether to use the Chaos physics interface. This overrides the physx flags to disable APEX and NvCloth
-		/// </summary>
-		[Obsolete("Deprecated in UE5.1 - No longer used as Chaos is always enabled.")]
-		public bool bUseChaos { get; set; }
-
 		/// <summary>
 		/// Whether to compile in checked chaos features for debugging
 		/// </summary>
@@ -969,30 +941,6 @@ namespace UnrealBuildTool
 		[RequiresUniqueBuildEnvironment]
 		public bool bCompileChaosVisualDebuggerSupport { get; set; } = true;
 
-		/// <summary>
-		/// Whether scene query acceleration is done by UE. The physx scene query structure is still created, but we do not use it.
-		/// </summary>
-		[Obsolete("Deprecated in UE5.1 - No longer used in engine.")]
-		public bool bCustomSceneQueryStructure { get; set; }
-
-		/// <summary>
-		/// Whether to include PhysX support.
-		/// </summary>
-		[Obsolete("Deprecated in UE5.1 - No longer used as Chaos is always enabled.")]
-		public bool bCompilePhysX { get; set; }
-
-		/// <summary>
-		/// Whether to include PhysX APEX support.
-		/// </summary>
-		[Obsolete("Deprecated in UE5.1 - No longer used as Chaos is always enabled.")]
-		public bool bCompileAPEX { get; set; }
-
-		/// <summary>
-		/// Whether to include NvCloth.
-		/// </summary>
-		[Obsolete("Deprecated in UE5.1 - No longer used as Chaos is always enabled.")]
-		public bool bCompileNvCloth { get; set; }
-
 		/// <summary>
 		/// Whether to include ICU unicode/i18n support in Core.
 		/// </summary>
@@ -1013,12 +961,6 @@ namespace UnrealBuildTool
 		[RequiresUniqueBuildEnvironment]
 		public bool bCompileISPC { get; set; }
 
-		/// <summary>
-		/// Whether to compile IntelMetricsDiscovery.
-		/// </summary>
-		[Obsolete("Deprecated in UE5.4 - No longer used.")]
-		public bool bCompileIntelMetricsDiscovery { get; set; } = false;
-
 		/// <summary>
 		/// Whether to compile in python support
 		/// </summary>
@@ -1126,11 +1068,6 @@ namespace UnrealBuildTool
 		}
 		private bool? bBuildWithEditorOnlyDataOverride;
 
-		/// <summary>
-		/// Manually specified value for bBuildDeveloperTools.
-		/// </summary>
-		bool? bBuildDeveloperToolsOverride;
-
 		/// <summary>
 		/// Whether to compile the developer tools.
 		/// </summary>
@@ -1140,11 +1077,7 @@ namespace UnrealBuildTool
 			set => bBuildDeveloperToolsOverride = value;
 			get => bBuildDeveloperToolsOverride ?? (bCompileAgainstEngine && (Type == TargetType.Editor || Type == TargetType.Program || (Configuration != UnrealTargetConfiguration.Test && Configuration != UnrealTargetConfiguration.Shipping)));
 		}
-
-		/// <summary>
-		/// Manually specified value for bBuildTargetDeveloperTools.
-		/// </summary>
-		bool? bBuildTargetDeveloperToolsOverride;
+		private bool? bBuildDeveloperToolsOverride;
 
 		/// <summary>
 		/// Whether to compile the developer tools that are for target platforms or connected devices (defaults to bBuildDeveloperTools)
@@ -1155,6 +1088,7 @@ namespace UnrealBuildTool
 			set => bBuildTargetDeveloperToolsOverride = value;
 			get => bBuildTargetDeveloperToolsOverride ?? bBuildDeveloperTools;
 		}
+		private bool? bBuildTargetDeveloperToolsOverride;
 
 		/// <summary>
 		/// Whether to force compiling the target platform modules, even if they wouldn't normally be built.
@@ -1280,11 +1214,6 @@ namespace UnrealBuildTool
 		}
 		private bool? bForceEnableTracePrivate;
 
-		/// <summary>
-		/// Manually specified value for bCompileAgainstEditor.
-		/// </summary>
-		bool? bCompileAgainstEditorOverride;
-
 		/// <summary>
 		/// Enabled for editor builds (TargetType.Editor). Can be overridden for programs (TargetType.Program) that would need to compile against editor code. Not available for other target types.
 		/// Mainly drives the value of WITH_EDITOR.
@@ -1295,6 +1224,7 @@ namespace UnrealBuildTool
 			set => bCompileAgainstEditorOverride = value;
 			get => bCompileAgainstEditorOverride ?? (Type == TargetType.Editor);
 		}
+		private bool? bCompileAgainstEditorOverride;
 
 		/// <summary>
 		/// Whether to compile Recast navmesh generation.
@@ -1484,14 +1414,6 @@ namespace UnrealBuildTool
 		[XmlConfigFile(Category = "BuildConfiguration")]
 		public bool bUseDebugLiveCodingConsole { get; set; }
 
-		/// <summary>
-		/// Obsolete: whether to enable support for DirectX Math
-		/// Not compatible with LWC so has been disabled since UE 5.
-		/// </summary>
-		[RequiresUniqueBuildEnvironment]
-		[Obsolete("Deprecated in UE 5.6 - DirectXMath path was incompatible with LWC and unused since 5.0")]
-		public bool bWithDirectXMath { get; set; }
-
 		/// <summary>
 		/// Whether to enable the support for FixedTimeStep in the engine
 		/// </summary>
@@ -1565,12 +1487,6 @@ namespace UnrealBuildTool
 		[RequiresUniqueBuildEnvironment]
 		public bool bUseExecCommandsInShipping { get; set; } = true;
 
-		/// <summary>
-		/// True if we want to favor optimizing size over speed.
-		/// </summary>
-		[Obsolete("Deprecated in UE5.1 - Please use OptimizationLevel instead.")]
-		public bool bCompileForSize { get; set; }
-
 		/// <summary>
 		/// Allows to fine tune optimizations level for speed and\or code size
 		/// </summary>
@@ -1724,12 +1640,6 @@ namespace UnrealBuildTool
 		[XmlConfigFile(Category = "BuildConfiguration")]
 		public bool bAlwaysUseUnityForGeneratedFiles { get; set; } = true;
 
-		/// <summary>
-		/// Whether to merge module and generated unity files for faster compilation.
-		/// </summary>
-		[Obsolete("Deprecated in UE5.3; use DisableMergingModuleAndGeneratedFilesInUnityFiles instead")]
-		public bool bMergeModuleAndGeneratedUnityFiles { get; set; }
-
 		/// <summary>
 		/// List of modules that disables merging module and generated cpp files in the same unity files.
 		/// </summary>
@@ -1836,32 +1746,12 @@ namespace UnrealBuildTool
 		}
 		private bool? bValidateFormatStringsPrivate;
 
-		/// <summary>
-		/// Level to report deprecation warnings as errors
-		/// </summary>
-		[Obsolete("Deprecated in UE5.6 - Replace with TargetRules.CppCompileWarningSettings.DeprecationWarningLevel")]
-		public WarningLevel DeprecationWarningLevel
-		{
-			get => CppCompileWarningSettings.DeprecationWarningLevel;
-			set => CppCompileWarningSettings.DeprecationWarningLevel = value;
-		}
-
 		/// <summary>
 		/// C++ Warnings settings object used by the target
 		/// </summary>
 		[ConfigSubObject]
 		public CppCompileWarnings CppCompileWarningSettings { get; init; }
 
-		/// <summary>
-		/// Forces shadow variable warnings to be treated as errors on platforms that support it.
-		/// </summary>
-		[Obsolete("Deprecated in UE5.6 - Replace with TargetRules.CppCompileWarningSettings.ShadowVariableWarningLevel")]
-		public WarningLevel ShadowVariableWarningLevel
-		{
-			get => CppCompileWarningSettings.ShadowVariableWarningLevel;
-			set => CppCompileWarningSettings.ShadowVariableWarningLevel = value;
-		}
-
 		/// <summary>
 		/// Whether to enable all warnings as errors. UE enables most warnings as errors already, but disables a few (such as deprecation warnings).
 		/// </summary>
@@ -1870,73 +1760,6 @@ namespace UnrealBuildTool
 		[RequiresUniqueBuildEnvironment]
 		public bool bWarningsAsErrors { get; set; }
 
-		/// <summary>
-		/// Indicates what warning/error level to treat unsafe type casts as on platforms that support it (e.g., double->float or int64->int32)
-		/// </summary>
-		[Obsolete("Deprecated in UE5.6 - Replace with TargetRules.CppCompileWarningSettings.UnsafeTypeCastWarningLevel")]
-		public WarningLevel UnsafeTypeCastWarningLevel 
-		{
-			get => CppCompileWarningSettings.UnsafeTypeCastWarningLevel;
-			set => CppCompileWarningSettings.UnsafeTypeCastWarningLevel = value;
-		}
-
-		/// <summary>
-		/// Forces the use of undefined identifiers in conditional expressions to be treated as errors.
-		/// </summary>
-		[XmlConfigFile(Category = "BuildConfiguration", Deprecated = true)]
-		[Obsolete("Deprecated in UE5.5 - Replace with TargetRules.UndefinedIdentifierWarningLevel")]
-		public bool bUndefinedIdentifierErrors { get; set; } = true;
-
-		/// <summary>
-		/// Indicates what warning/error level to treat undefined identifiers in conditional expressions.
-		/// </summary>
-		[Obsolete("Deprecated in UE5.6 - Replace with TargetRules.CppCompileWarningSettings.UndefinedIdentifierWarningLevel")]
-		public WarningLevel UndefinedIdentifierWarningLevel 
-		{
-			get => CppCompileWarningSettings.UndefinedIdentifierWarningLevel;
-			set => CppCompileWarningSettings.UndefinedIdentifierWarningLevel = value;
-		}
-
-		/// <summary>
-		/// Indicates what warning/error level to treat potential PCH performance issues.
-		/// </summary>
-		[Obsolete("Deprecated in UE5.6 - Replace with TargetRules.CppCompileWarningSettings.PCHPerformanceIssueWarningLevel")]
-		public WarningLevel PCHPerformanceIssueWarningLevel
-		{
-			get => CppCompileWarningSettings.PCHPerformanceIssueWarningLevel;
-			set => CppCompileWarningSettings.PCHPerformanceIssueWarningLevel = value;
-		}
-
-		/// <summary>
-		/// How to treat general module include path validation messages
-		/// </summary>
-		[Obsolete("Deprecated in UE5.6 - Replace with TargetRules.CppCompileWarningSettings.ModuleIncludePathWarningLevel")]
-		public WarningLevel ModuleIncludePathWarningLevel
-		{
-			get => CppCompileWarningSettings.ModuleIncludePathWarningLevel;
-			set => CppCompileWarningSettings.ModuleIncludePathWarningLevel = value;
-		}
-
-		/// <summary>
-		/// How to treat private module include path validation messages, where a module is adding an include path that exposes private headers
-		/// </summary>
-		[Obsolete("Deprecated in UE5.6 - Replace with TargetRules.CppCompileWarningSettings.ModuleIncludePrivateWarningLevel")]
-		public WarningLevel ModuleIncludePrivateWarningLevel
-		{
-			get => CppCompileWarningSettings.ModuleIncludePrivateWarningLevel;
-			set => CppCompileWarningSettings.ModuleIncludePrivateWarningLevel = value;
-		}
-
-		/// <summary>
-		/// How to treat unnecessary module sub-directory include path validation messages
-		/// </summary>
-		[Obsolete("Deprecated in UE5.6 - Replace with TargetRules.CppCompileWarningSettings.ModuleIncludeSubdirectoryWarningLevel")]
-		public WarningLevel ModuleIncludeSubdirectoryWarningLevel
-		{
-			get => CppCompileWarningSettings.ModuleIncludeSubdirectoryWarningLevel;
-			set => CppCompileWarningSettings.ModuleIncludeSubdirectoryWarningLevel = value;
-		}
-
 		/// <summary>
 		/// Print out files that are included by each source file
 		/// </summary>
@@ -1992,30 +1815,12 @@ namespace UnrealBuildTool
 		[XmlConfigFile(Category = "BuildConfiguration")]
 		public bool bStressTestUnity { get; set; }
 
-		/// <summary>
-		/// Whether to add additional information to the unity files, such as '_of_X' in the file name. Not recommended.
-		/// </summary>
-		[XmlConfigFile(Category = "BuildConfiguration")]
-		[Obsolete("Deprecated in UE5.6")]
-		public bool bDetailedUnityFiles { get; set; }
-
 		/// <summary>
 		/// Whether to force debug info to be generated.
 		/// </summary>
 		[CommandLine("-ForceDebugInfo")]
 		public bool bForceDebugInfo { get; set; }
 
-		/// <summary>
-		/// Whether to globally disable debug info generation; Obsolete, please use TargetRules.DebugInfoMode instead
-		/// </summary>
-		[XmlConfigFile(Category = "BuildConfiguration", Deprecated = true)]
-		[Obsolete("Deprecated in UE5.4 - Replace with TargetRules.DebugInfo")]
-		public bool bDisableDebugInfo
-		{
-			get => DebugInfo == DebugInfoMode.None;
-			set => DebugInfo = value ? DebugInfo = DebugInfoMode.None : DebugInfo = DebugInfoMode.Full;
-		}
-
 		/// <summary>
 		/// How much debug info should be generated. See DebugInfoMode enum for more details
 		/// </summary>
@@ -2028,13 +1833,13 @@ namespace UnrealBuildTool
 		/// Modules that should have debug info disabled
 		/// </summary>
 		[CommandLine("-DisableDebugInfoModules=", ListSeparator = '+')]
-		public HashSet<string> DisableDebugInfoModules { get; } = new();
+		public HashSet<string> DisableDebugInfoModules { get; } = [];
 
 		/// <summary>
 		/// Plugins that should have debug info disabled
 		/// </summary>
 		[CommandLine("-DisableDebugInfoPlugins=", ListSeparator = '+')]
-		public HashSet<string> DisableDebugInfoPlugins { get; } = new();
+		public HashSet<string> DisableDebugInfoPlugins { get; } = [];
 
 		/// <summary>
 		/// True if only debug line number tables should be emitted in debug information for compilers that support doing so. Overrides TargetRules.DebugInfo
@@ -2048,13 +1853,13 @@ namespace UnrealBuildTool
 		/// Modules that should emit line number tables instead of full debug information for compilers that support doing so. Overrides DisableDebugInfoModules
 		/// </summary>
 		[CommandLine("-DebugInfoLineTablesOnlyModules=", ListSeparator = '+')]
-		public HashSet<string> DebugInfoLineTablesOnlyModules { get; } = new();
+		public HashSet<string> DebugInfoLineTablesOnlyModules { get; } = [];
 
 		/// <summary>
 		/// Plugins that should emit line number tables instead of full debug information for compilers that support doing so. Overrides DisableDebugInfoPlugins
 		/// </summary>
 		[CommandLine("-DebugInfoLineTablesOnlyPlugins=", ListSeparator = '+')]
-		public HashSet<string> DebugInfoLineTablesOnlyPlugins { get; } = new();
+		public HashSet<string> DebugInfoLineTablesOnlyPlugins { get; } = [];
 
 		/// <summary>
 		/// Whether to disable debug info generation for generated files. This improves link times and reduces pdb size for modules that have a lot of generated glue code.
@@ -2290,14 +2095,6 @@ namespace UnrealBuildTool
 		/// </summary>
 		public bool bEnableCppModules { get; set; }
 
-		/// <summary>
-		/// Whether to enable support for C++20 coroutines
-		/// This option is provided to facilitate evaluation of the feature.
-		/// Expect the name of the option to change. Use of coroutines in with UE is untested and unsupported.
-		/// </summary>
-		[Obsolete("CppCoroutines are not supported")]
-		public bool bEnableCppCoroutinesForEvaluation { get; set; }
-
 		/// <summary>
 		/// Whether to enable engine's ability to set process priority on runtime.
 		/// This option requires some environment setup on Linux, that's why it's disabled by default.
@@ -2305,12 +2102,6 @@ namespace UnrealBuildTool
 		/// </summary>
 		public bool bEnableProcessPriorityControl { get; set; }
 
-		/// <summary>
-		/// If true, then enable memory profiling in the build (defines USE_MALLOC_PROFILER=1 and forces bOmitFramePointers=false).
-		/// </summary>
-		[Obsolete("Deprecated in UE5.3 - No longer used as MallocProfiler is removed in favor of UnrealInsights.")]
-		public bool bUseMallocProfiler { get; set; }
-
 		/// <summary>
 		/// If true, then enable Unreal Insights (utrace) profiling in the build for the Shader Compiler Worker (defines USE_SHADER_COMPILER_WORKER_TRACE=1).
 		/// </summary>
@@ -2484,12 +2275,6 @@ namespace UnrealBuildTool
 		}
 		private bool? bIdenticalCodeFoldingPrivate;
 
-		/// <summary>
-		/// Disable supports for inlining gen.cpps
-		/// </summary>
-		[Obsolete("Deprecated in UE5.3; the engine code relies on inlining generated cpp files to compile")]
-		public bool bDisableInliningGenCpps { get; set; }
-
 		/// <summary>
 		/// Allows overriding the toolchain to be created for this target. This must match the name of a class declared in the UnrealBuildTool assembly.
 		/// </summary>
@@ -2638,14 +2423,14 @@ namespace UnrealBuildTool
 		/// </summary>
 		[RequiresUniqueBuildEnvironment]
 		[CommandLine("-MergeAdditionalCommonModules=", ListSeparator = '+')]
-		public HashSet<string> MergeAdditionalCommonModules { get; } = new();
+		public HashSet<string> MergeAdditionalCommonModules { get; } = [];
 
 		/// <summary>
 		/// Experimental: List of plugins (and their dependencies) to each merge into separate libraries. Requires bMergeModules to be enabled
 		/// </summary>
 		[RequiresUniqueBuildEnvironment]
 		[CommandLine("-MergePlugins=", ListSeparator = '+')]
-		public List<string> MergePlugins { get; } = new();
+		public List<string> MergePlugins { get; } = [];
 
 		/// <summary>
 		/// Experimental: List of plugins (and their dependencies) to ignore when processing merged libraries. Requires bMergeModules to be enabled
@@ -2654,7 +2439,7 @@ namespace UnrealBuildTool
 		/// </summary>
 		[RequiresUniqueBuildEnvironment]
 		[CommandLine("-MergePluginsIgnored=", ListSeparator = '+')]
-		public List<string> MergePluginsIgnored { get; } = new();
+		public List<string> MergePluginsIgnored { get; } = [];
 
 		/// <summary>
 		/// Experimental: List of plugins (and their dependencies) to explicitly exclude when processing merged libraries. Requires bMergeModules to be enabled
@@ -2662,17 +2447,17 @@ namespace UnrealBuildTool
 		/// </summary>
 		[RequiresUniqueBuildEnvironment]
 		[CommandLine("-MergePluginsExcluded=", ListSeparator = '+')]
-		public List<string> MergePluginsExcluded { get; } = new();
+		public List<string> MergePluginsExcluded { get; } = [];
 
 		/// <summary>
 		/// Experimental: List of merged plugins to move common dependencies to a shared library, requires MergePlugins. Can be chained
 		/// </summary>
-		public Dictionary<string, HashSet<string>> MergePluginsShared { get; } = new();
+		public Dictionary<string, HashSet<string>> MergePluginsShared { get; } = [];
 
 		/// <summary>
 		/// Experimental: List of merged plugins to move to the main executable, requires MergePlugins. "Engine" for all shared engine modules and "Common" for shared project modules.
 		/// </summary>
-		public List<string> MergePluginsLaunch { get; } = new();
+		public List<string> MergePluginsLaunch { get; } = [];
 
 		/// <summary>
 		/// Experimental: Whether to merge only game feature plugins, ignoring regular plugins.
@@ -2692,13 +2477,13 @@ namespace UnrealBuildTool
 		/// </summary>
 		[RequiresUniqueBuildEnvironment]
 		[CommandLine("-Define:")]
-		public List<string> GlobalDefinitions { get; } = new();
+		public List<string> GlobalDefinitions { get; } = [];
 
 		/// <summary>
 		/// Macros to define across all macros in the project.
 		/// </summary>
 		[CommandLine("-ProjectDefine:")]
-		public List<string> ProjectDefinitions { get; } = new();
+		public List<string> ProjectDefinitions { get; } = [];
 
 		/// <summary>
 		/// Specifies the name of the launch module. For modular builds, this is the module that is compiled into the target's executable.
@@ -2722,19 +2507,19 @@ namespace UnrealBuildTool
 		/// <summary>
 		/// List of additional modules to be compiled into the target.
 		/// </summary>
-		public List<string> ExtraModuleNames { get; } = new();
+		public List<string> ExtraModuleNames { get; } = [];
 
 		/// <summary>
 		/// Path to a manifest to output for this target
 		/// </summary>
 		[CommandLine("-Manifest")]
-		public List<FileReference> ManifestFileNames { get; } = new();
+		public List<FileReference> ManifestFileNames { get; } = [];
 
 		/// <summary>
 		/// Path to a list of dependencies for this target, when precompiling
 		/// </summary>
 		[CommandLine("-DependencyList")]
-		public List<FileReference> DependencyListFileNames { get; } = new();
+		public List<FileReference> DependencyListFileNames { get; } = [];
 
 		/// <summary>
 		/// Backing storage for the BuildEnvironment property
@@ -2791,26 +2576,26 @@ namespace UnrealBuildTool
 		/// <summary>
 		/// Specifies a list of targets which should be built before this target is built.
 		/// </summary>
-		public List<TargetInfo> PreBuildTargets { get; } = new();
+		public List<TargetInfo> PreBuildTargets { get; } = [];
 
 		/// <summary>
 		/// Specifies a list of steps which should be executed before this target is built, in the context of the host platform's shell.
 		/// The following variables will be expanded before execution:
 		/// $(EngineDir), $(ProjectDir), $(TargetName), $(TargetPlatform), $(TargetConfiguration), $(TargetType), $(ProjectFile).
 		/// </summary>
-		public List<string> PreBuildSteps { get; } = new();
+		public List<string> PreBuildSteps { get; } = [];
 
 		/// <summary>
 		/// Specifies a list of steps which should be executed after this target is built, in the context of the host platform's shell.
 		/// The following variables will be expanded before execution:
 		/// $(EngineDir), $(ProjectDir), $(TargetName), $(TargetPlatform), $(TargetConfiguration), $(TargetType), $(ProjectFile).
 		/// </summary>
-		public List<string> PostBuildSteps { get; } = new();
+		public List<string> PostBuildSteps { get; } = [];
 
 		/// <summary>
 		/// Specifies additional build products produced as part of this target.
 		/// </summary>
-		public List<string> AdditionalBuildProducts { get; } = new();
+		public List<string> AdditionalBuildProducts { get; } = [];
 
 		/// <summary>
 		/// Additional arguments to pass to the compiler
@@ -2872,7 +2657,7 @@ namespace UnrealBuildTool
 		/// <summary>
 		/// If this is non-null, then any platforms NOT listed will not be allowed to have modules in their directories be created
 		/// </summary>
-		public UnrealTargetPlatform[]? OptedInModulePlatforms = null;
+		public UnrealTargetPlatform[]? OptedInModulePlatforms { get; set; } = null;
 
 		/// <summary>
 		/// Android-specific target settings.
@@ -3101,7 +2886,7 @@ namespace UnrealBuildTool
 			// if we didn't set it above, check the properties
 			if (BuildEnvironmentOverride.Value == TargetBuildEnvironment.UniqueIfNeeded)
 			{
-				Dictionary<string, (string?, string?)> propNames = new();
+				Dictionary<string, (string?, string?)> propNames = [];
 				string? baseTargetName;
 				if (RequiresUniqueEnvironment(rulesAssembly, arguments, propNames, out baseTargetName))
 				{
@@ -3163,7 +2948,7 @@ namespace UnrealBuildTool
 			rules.TargetSourceFile = platformFile;
 
 			// All target files for this target that could cause invalidation
-			rules.TargetFiles = targetFiles?.ToHashSet() ?? new();
+			rules.TargetFiles = targetFiles?.ToHashSet() ?? [];
 
 			// Initialize the logger
 			rules.Logger = logger;
@@ -3229,7 +3014,7 @@ namespace UnrealBuildTool
 			}
 
 			// Read settings from config files
-			Dictionary<ConfigDependencyKey, IReadOnlyList<string>?> configValues = new Dictionary<ConfigDependencyKey, IReadOnlyList<string>?>();
+			Dictionary<ConfigDependencyKey, IReadOnlyList<string>?> configValues = [];
 			foreach (object configurableObject in GetConfigurableObjects())
 			{
 				ConfigCache.ReadSettings(DirectoryReference.FromFile(ProjectFile), Platform, configurableObject, configValues, target.Arguments);
@@ -3375,7 +3160,7 @@ namespace UnrealBuildTool
 
 			ConfigHierarchy projectGameIni = ConfigCache.ReadHierarchy(ConfigHierarchyType.Game, ProjectFile?.Directory, Platform);
 			List<string>? disabledPlugins;
-			List<string> allDisabledPlugins = new List<string>();
+			List<string> allDisabledPlugins = [];
 
 			if (projectGameIni.GetArray("CookedEditorSettings", "DisabledPlugins", out disabledPlugins))
 			{
@@ -3546,7 +3331,7 @@ namespace UnrealBuildTool
 				diagnostics.Add("[Upgrade]");
 				diagnostics.Add("[Upgrade] Using backward-compatible build settings. The latest version of UE sets the following values by default, which may require code changes:");
 
-				List<Tuple<string, string>> modifiedSettings = new List<Tuple<string, string>>();
+				List<Tuple<string, string>> modifiedSettings = [];
 				if (BuildSettingsVersion.V2 <= latestVersion && DefaultBuildSettings < BuildSettingsVersion.V2)
 				{
 					modifiedSettings.Add(Tuple.Create(String.Format("{0} = false", nameof(bLegacyPublicIncludePaths)), "Omits subfolders from public include paths to reduce compiler command line length. (Previously: true)."));
@@ -3635,7 +3420,7 @@ namespace UnrealBuildTool
 		/// </summary>
 		public void PrintBuildSettingsInfoWarnings()
 		{
-			List<string> diagnostics = new();
+			List<string> diagnostics = [];
 			GetBuildSettingsInfo(diagnostics);
 			diagnostics.ForEach(x => Logger.LogWarning("Warning: {Message}", x));
 		}
