[UBT]
* Fixed so symbols are defaulting to visible on linux when building with merged modules since we will use version script to control exports anyway

diff --git a/Engine/Source/Programs/UnrealBuildTool/Platform/Linux/LinuxToolChain.cs b/Engine/Source/Programs/UnrealBuildTool/Platform/Linux/LinuxToolChain.cs
--- a/Engine/Source/Programs/UnrealBuildTool/Platform/Linux/LinuxToolChain.cs
+++ b/Engine/Source/Programs/UnrealBuildTool/Platform/Linux/LinuxToolChain.cs
@@ -457,7 +457,11 @@ namespace UnrealBuildTool
 
 			if (CompileEnvironment.bHideSymbolsByDefault)
 			{
-				Arguments.Add("-fvisibility-ms-compat");
+				if (!bMergeModules) // With merge modules we use version-script to control visibility instead.
+				{
+					Arguments.Add("-fvisibility-ms-compat");
+				}
+
 				Arguments.Add("-fvisibility-inlines-hidden");
 			}
 		}
