Add an option in the module rules to enable UE types appearing in namespaces.

#rnx

diff --git a/Engine/Source/Programs/Shared/EpicGames.Core/UHTTypes.cs b/Engine/Source/Programs/Shared/EpicGames.Core/UHTTypes.cs
--- a/Engine/Source/Programs/Shared/EpicGames.Core/UHTTypes.cs
+++ b/Engine/Source/Programs/Shared/EpicGames.Core/UHTTypes.cs
@@ -128,6 +128,7 @@ namespace EpicGames.Core
 			public string VerseMountPoint { get; set; } = "";
 			public bool AlwaysExportStructs { get; set; } = true;
 			public bool AlwaysExportEnums { get; set; } = true;
+			public bool AllowUETypesInNamespaces { get; set; } = false;
 
 			public override string ToString()
 			{
diff --git a/Engine/Source/Programs/UnrealBuildTool/Configuration/ModuleRules.cs b/Engine/Source/Programs/UnrealBuildTool/Configuration/ModuleRules.cs
--- a/Engine/Source/Programs/UnrealBuildTool/Configuration/ModuleRules.cs
+++ b/Engine/Source/Programs/UnrealBuildTool/Configuration/ModuleRules.cs
@@ -1151,6 +1151,11 @@ namespace UnrealBuildTool
 		/// If true, autogenerated functions of UENUMS are always exported
 		/// </summary>
 		public bool bAlwaysExportEnums = true;
+
+		/// <summary>
+		/// If true, this module can contain UE types (i.e. UCLASS) in namespaces
+		/// </summary>
+		public bool bAllowUETypesInNamespaces = false;
 		
 		/// <summary>
 		/// List of modules names (no path needed) with header files that our module's public headers needs access to, but we don't need to "import" or link against.
diff --git a/Engine/Source/Programs/UnrealBuildTool/System/ExternalExecution.cs b/Engine/Source/Programs/UnrealBuildTool/System/ExternalExecution.cs
--- a/Engine/Source/Programs/UnrealBuildTool/System/ExternalExecution.cs
+++ b/Engine/Source/Programs/UnrealBuildTool/System/ExternalExecution.cs
@@ -187,6 +187,11 @@ namespace UnrealBuildTool
 		/// If true, autogenerated functions of UENUMS are always exported
 		/// </summary>
 		public bool AlwaysExportEnums;
+
+		/// <summary>
+		/// If true, UE types (i.e. UCLASS) will be allowed in namespaces
+		/// </summary>
+		public bool AllowUETypesInNamespaces;
 		
 		public UHTModuleInfo(
 			string ModuleName, 
@@ -203,7 +208,8 @@ namespace UnrealBuildTool
 			bool HasVerse,
 			string VerseMountPoint,
 			bool AlwaysExportStructs,
-			bool AlwaysExportEnums)
+			bool AlwaysExportEnums,
+			bool AllowUETypesInNamespaces)
 		{
 			this.ModuleName = ModuleName;
 			this.ModuleRulesFile = ModuleRulesFile;
@@ -225,6 +231,7 @@ namespace UnrealBuildTool
 			this.VerseMountPoint = VerseMountPoint;
 			this.AlwaysExportStructs = AlwaysExportStructs;
 			this.AlwaysExportEnums = AlwaysExportEnums;
+			this.AllowUETypesInNamespaces = AllowUETypesInNamespaces;
 		}
 
 		public UHTModuleInfo(BinaryArchiveReader Reader)
@@ -250,6 +257,7 @@ namespace UnrealBuildTool
 			VerseMountPoint = Reader.ReadString()!;
 			AlwaysExportStructs = Reader.ReadBool();
 			AlwaysExportEnums = Reader.ReadBool();
+			AllowUETypesInNamespaces = Reader.ReadBool();
 		}
 
 		public void Write(BinaryArchiveWriter Writer)
@@ -275,6 +283,7 @@ namespace UnrealBuildTool
 			Writer.WriteString(VerseMountPoint);
 			Writer.WriteBool(AlwaysExportStructs);
 			Writer.WriteBool(AlwaysExportEnums);
+			Writer.WriteBool(AllowUETypesInNamespaces);
 		}
 
 		public override string ToString()
@@ -497,7 +506,8 @@ namespace UnrealBuildTool
 						Module.bHasVerse,
 						Module.bHasVerse ? GetVerseMountPointForModule(Module) : string.Empty,
 						Module.Rules.bAlwaysExportStructs,
-						Module.Rules.bAlwaysExportEnums);
+						Module.Rules.bAlwaysExportEnums,
+						Module.Rules.bAllowUETypesInNamespaces);
 					ModuleInfoArray[Idx] = Info;
 
 					Queue.Enqueue(() => SetupUObjectModule(Info, ExcludedFolders, MetadataCache, Queue));
@@ -1107,6 +1117,7 @@ namespace UnrealBuildTool
 						VerseMountPoint = UObjectModule.VerseMountPoint,
 						AlwaysExportStructs = UObjectModule.AlwaysExportStructs,
 						AlwaysExportEnums = UObjectModule.AlwaysExportEnums,
+						AllowUETypesInNamespaces = UObjectModule.AllowUETypesInNamespaces,
 					}); ;
 			}
 
