Plugin modules can be included & excluded on a per-architecture basis.

For example:
"PlatformArchitectureDenyList" : [ "Win64:arm64" ]

diff --git a/Engine/Plugins/Compositing/Composure/Composure.uplugin b/Engine/Plugins/Compositing/Composure/Composure.uplugin
--- a/Engine/Plugins/Compositing/Composure/Composure.uplugin
+++ b/Engine/Plugins/Compositing/Composure/Composure.uplugin
@@ -17,19 +17,22 @@
 	"Modules" :
 	[
 		{
-			"Name" : "Composure",
-			"Type" : "Runtime",
-			"LoadingPhase" : "Default"
+			"Name": "Composure",
+			"Type": "Runtime",
+			"LoadingPhase": "Default",
+			"PlatformArchitectureDenyList": [ "Win64:arm64" ]
 		},
 		{
-			"Name" : "ComposureEditor",
-			"Type" : "Editor",
-			"LoadingPhase" : "Default"
+			"Name": "ComposureEditor",
+			"Type": "Editor",
+			"LoadingPhase": "Default",
+			"PlatformArchitectureDenyList": [ "Win64:arm64" ]
 		},
 		{
-			"Name" : "ComposureLayersEditor",
-			"Type" : "Editor",
-			"LoadingPhase" : "PostEngineInit"
+			"Name": "ComposureLayersEditor",
+			"Type": "Editor",
+			"LoadingPhase": "PostEngineInit",
+			"PlatformArchitectureDenyList": [ "Win64:arm64" ]
 		}
 	],
 	"Plugins": [
diff --git a/Engine/Plugins/Experimental/LiveLinkOpenVR/LiveLinkOpenVR.uplugin b/Engine/Plugins/Experimental/LiveLinkOpenVR/LiveLinkOpenVR.uplugin
--- a/Engine/Plugins/Experimental/LiveLinkOpenVR/LiveLinkOpenVR.uplugin
+++ b/Engine/Plugins/Experimental/LiveLinkOpenVR/LiveLinkOpenVR.uplugin
@@ -3,7 +3,7 @@
 	"Version": 1,
 	"VersionName": "1.0",
 	"FriendlyName": "LiveLinkOpenVR",
-	"Description": "Live Link plugin for OpenVR (Not supported for native arm64, see UEBuildWindows.cs DisablePlugins.)",
+	"Description": "Live Link plugin for OpenVR (Not supported for native arm64.)",
 	"Category": "Virtual Production",
 	"CreatedBy": "Epic Games, Inc.",
 	"CreatedByURL": "https://epicgames.com",
@@ -25,6 +25,9 @@
 			"LoadingPhase": "Default",
 			"PlatformAllowList": [
 				"Win64"
+			],
+			"PlatformArchitectureDenyList": [
+				"Win64:arm64"
 			]
 		}
 	]
diff --git a/Engine/Plugins/Experimental/OpenImageDenoise/OpenImageDenoise.uplugin b/Engine/Plugins/Experimental/OpenImageDenoise/OpenImageDenoise.uplugin
--- a/Engine/Plugins/Experimental/OpenImageDenoise/OpenImageDenoise.uplugin
+++ b/Engine/Plugins/Experimental/OpenImageDenoise/OpenImageDenoise.uplugin
@@ -22,6 +22,10 @@
 			"LoadingPhase" : "Default",
 			"PlatformAllowList": [
 				"Win64"
+			],
+			"PlatformArchitectureDenyList": [
+				"Win64:arm64",
+				"Win64:arm64ec"
 			]
 		}
 	]
diff --git a/Engine/Plugins/Media/MediaIOFramework/MediaIOFramework.uplugin b/Engine/Plugins/Media/MediaIOFramework/MediaIOFramework.uplugin
--- a/Engine/Plugins/Media/MediaIOFramework/MediaIOFramework.uplugin
+++ b/Engine/Plugins/Media/MediaIOFramework/MediaIOFramework.uplugin
@@ -17,17 +17,20 @@
 		{
 			"Name": "MediaIOCore",
 			"Type": "Runtime",
-			"LoadingPhase": "PostConfigInit"
+			"LoadingPhase": "PostConfigInit",
+			"PlatformArchitectureDenyList": [ "Win64:arm64" ]
 		},
 		{
 			"Name": "MediaIOEditor",
 			"Type": "Editor",
-			"LoadingPhase": "PostEngineInit"
+			"LoadingPhase": "PostEngineInit",
+			"PlatformArchitectureDenyList": [ "Win64:arm64" ]
 		},
 		{
 			"Name": "GPUTextureTransfer",
 			"Type": "Runtime",
-			"LoadingPhase": "Default"
+			"LoadingPhase": "Default",
+			"PlatformArchitectureDenyList": [ "Win64:arm64" ]
 		}
 	],
 	"Plugins": [
diff --git a/Engine/Plugins/Online/OnlineSubsystemSteam/OnlineSubsystemSteam.uplugin b/Engine/Plugins/Online/OnlineSubsystemSteam/OnlineSubsystemSteam.uplugin
--- a/Engine/Plugins/Online/OnlineSubsystemSteam/OnlineSubsystemSteam.uplugin
+++ b/Engine/Plugins/Online/OnlineSubsystemSteam/OnlineSubsystemSteam.uplugin
@@ -20,6 +20,10 @@
                 "Win64",
                 "Mac",
                 "Linux"
+			],
+			"PlatformArchitectureDenyList" :
+			[
+				"Win64:arm64"
 			]
         }
     ],
diff --git a/Engine/Plugins/Online/SocketSubsystemSteamIP/SocketSubsystemSteamIP.uplugin b/Engine/Plugins/Online/SocketSubsystemSteamIP/SocketSubsystemSteamIP.uplugin
--- a/Engine/Plugins/Online/SocketSubsystemSteamIP/SocketSubsystemSteamIP.uplugin
+++ b/Engine/Plugins/Online/SocketSubsystemSteamIP/SocketSubsystemSteamIP.uplugin
@@ -18,6 +18,9 @@
 			"PlatformAllowList": [
 				"Win64",
 				"Mac"
+			],
+			"PlatformArchitectureDenyList": [
+				"Win64:arm64"
 			]
 		}
 	],
diff --git a/Engine/Plugins/Runtime/Nvidia/Reflex/Reflex.uplugin b/Engine/Plugins/Runtime/Nvidia/Reflex/Reflex.uplugin
--- a/Engine/Plugins/Runtime/Nvidia/Reflex/Reflex.uplugin
+++ b/Engine/Plugins/Runtime/Nvidia/Reflex/Reflex.uplugin
@@ -15,7 +15,12 @@
 		{
 			"Name" : "Reflex",
 			"Type" : "ClientOnly",
-			"PlatformAllowList" : [ "Win64" ]
+			"PlatformAllowList" : [ "Win64" ],
+			"PlatformArchitectureDenyList" :
+			[
+				"Win64:arm64",
+				"Win64:arm64ec"
+			]
 		}
 	]
 }
diff --git a/Engine/Plugins/Runtime/Steam/SteamController/SteamController.uplugin b/Engine/Plugins/Runtime/Steam/SteamController/SteamController.uplugin
--- a/Engine/Plugins/Runtime/Steam/SteamController/SteamController.uplugin
+++ b/Engine/Plugins/Runtime/Steam/SteamController/SteamController.uplugin
@@ -24,6 +24,10 @@
 			[
 				"Win64",
 				"Linux"
+			],
+			"PlatformArchitectureDenyList" :
+			[
+				"Win64:arm64"
 			]
 		}
 	],
diff --git a/Engine/Plugins/Runtime/Steam/SteamShared/SteamShared.uplugin b/Engine/Plugins/Runtime/Steam/SteamShared/SteamShared.uplugin
--- a/Engine/Plugins/Runtime/Steam/SteamShared/SteamShared.uplugin
+++ b/Engine/Plugins/Runtime/Steam/SteamShared/SteamShared.uplugin
@@ -18,6 +18,9 @@
 				"Win64",
 				"Mac",
 				"Linux"
+			],
+			"PlatformArchitectureDenyList": [
+				"Win64:arm64"
 			]
 		}
 	],
diff --git a/Engine/Plugins/Runtime/Steam/SteamSockets/SteamSockets.uplugin b/Engine/Plugins/Runtime/Steam/SteamSockets/SteamSockets.uplugin
--- a/Engine/Plugins/Runtime/Steam/SteamSockets/SteamSockets.uplugin
+++ b/Engine/Plugins/Runtime/Steam/SteamSockets/SteamSockets.uplugin
@@ -22,6 +22,10 @@
 				"Win64",
 				"Mac",
 				"Linux"
+			],
+			"PlatformArchitectureDenyList" :
+			[
+				"Win64:arm64"
 			]
 		}
 	],
diff --git a/Engine/Source/Programs/AutomationTool/Scripts/BuildPluginCommand.Automation.cs b/Engine/Source/Programs/AutomationTool/Scripts/BuildPluginCommand.Automation.cs
--- a/Engine/Source/Programs/AutomationTool/Scripts/BuildPluginCommand.Automation.cs
+++ b/Engine/Source/Programs/AutomationTool/Scripts/BuildPluginCommand.Automation.cs
@@ -307,10 +307,11 @@ public sealed class BuildPlugin : BuildCommand
 		{
 			bool bBuildDeveloperTools = (TargetType == TargetType.Editor || TargetType == TargetType.Program || (Configuration != UnrealTargetConfiguration.Test && Configuration != UnrealTargetConfiguration.Shipping));
 			bool bBuildRequiresCookedData = (TargetType != TargetType.Editor && TargetType != TargetType.Program);
+			UnrealArchitectures Architectures = PlatformToArchitectureMap.TryGetValue(Platform, out string SpecifiedArchitecture) ? UnrealArchitectures.FromString(SpecifiedArchitecture, null) : null;
 
 			foreach (ModuleDescriptor Module in Plugin.Modules)
 			{
-				if (Module.IsCompiledInConfiguration(Platform, Configuration, TargetName, TargetType, bBuildDeveloperTools, bBuildRequiresCookedData))
+				if (Module.IsCompiledInConfiguration(Platform, Configuration, TargetName, TargetType, bBuildDeveloperTools, bBuildRequiresCookedData, Architectures))
 				{
 					bCompilePlatform = true;
 				}
@@ -350,10 +351,11 @@ public sealed class BuildPlugin : BuildCommand
 		{
 			bool bBuildDeveloperTools = (TargetType == TargetType.Editor || TargetType == TargetType.Program || (Configuration != UnrealTargetConfiguration.Test && Configuration != UnrealTargetConfiguration.Shipping));
 			bool bBuildRequiresCookedData = (TargetType != TargetType.Editor && TargetType != TargetType.Program);
+			UnrealArchitectures Architectures = PlatformToArchitectureMap.TryGetValue(Platform, out string SpecifiedArchitecture) ? UnrealArchitectures.FromString(SpecifiedArchitecture, null) : null;
 
 			foreach (ModuleDescriptor Module in Plugin.Modules)
 			{
-				if (Module.IsCompiledInConfiguration(Platform, Configuration, TargetName, TargetType, bBuildDeveloperTools, bBuildRequiresCookedData))
+				if (Module.IsCompiledInConfiguration(Platform, Configuration, TargetName, TargetType, bBuildDeveloperTools, bBuildRequiresCookedData, Architectures))
 				{
 					bCompilePlatform = true;
 				}
diff --git a/Engine/Source/Programs/AutomationTool/Scripts/CreatePlatformExtension.cs b/Engine/Source/Programs/AutomationTool/Scripts/CreatePlatformExtension.cs
--- a/Engine/Source/Programs/AutomationTool/Scripts/CreatePlatformExtension.cs
+++ b/Engine/Source/Programs/AutomationTool/Scripts/CreatePlatformExtension.cs
@@ -340,6 +340,14 @@ public class CreatePlatformExtension : BuildCommand
 						{
 							ChildPlugin.WriteStringArrayField("PlatformDenyList", new string[] { Platform.ToString() } );
 						}
+						if (ParentModuleDesc.PlatformArchitectureAllowList != null && ParentModuleDesc.PlatformArchitectureAllowList.ContainsKey(Platform))
+						{
+							ChildPlugin.WriteStringArrayField("PlatformArchitectureAllowList", ParentModuleDesc.PlatformArchitectureAllowList.SelectMany( Pair => Pair.Value.Select( Arch => $"{Pair.Key}:{Arch}" ) ));
+						}
+						if (ParentModuleDesc.PlatformArchitectureDenyList != null && ParentModuleDesc.PlatformArchitectureDenyList.ContainsKey(Platform))
+						{
+							ChildPlugin.WriteStringArrayField("PlatformArchitectureDenyList", ParentModuleDesc.PlatformArchitectureDenyList.SelectMany( Pair => Pair.Value.Select( Arch => $"{Pair.Key}:{Arch}" ) ));
+						}
 						ChildPlugin.WriteObjectEnd();
 
 						// see if there is a module rule file too & generate the rules file for this platform
@@ -356,6 +364,14 @@ public class CreatePlatformExtension : BuildCommand
 							ParentModuleDesc.bHasExplicitPlatforms |= (ParentModuleDesc.PlatformAllowList.Count == 0); // an empty list is interpreted as 'all platforms are allowed' otherwise
 							bParentPluginDirty = true;
 						}
+						if (bHasPlatform && ParentModuleDesc.PlatformArchitectureAllowList != null)
+						{
+							bParentPluginDirty |= ParentModuleDesc.PlatformArchitectureAllowList.Remove(Platform);
+						}
+						if (bHasPlatform && ParentModuleDesc.PlatformArchitectureDenyList != null)
+						{
+							bParentPluginDirty |= ParentModuleDesc.PlatformArchitectureDenyList.Remove(Platform);
+						}
 						if (bHasPlatform && ParentModuleDesc.PlatformDenyList != null)
 						{
 							bParentPluginDirty |= ParentModuleDesc.PlatformDenyList.Remove(Platform);
diff --git a/Engine/Source/Programs/UnrealBuildTool/Configuration/UEBuildModule.cs b/Engine/Source/Programs/UnrealBuildTool/Configuration/UEBuildModule.cs
--- a/Engine/Source/Programs/UnrealBuildTool/Configuration/UEBuildModule.cs
+++ b/Engine/Source/Programs/UnrealBuildTool/Configuration/UEBuildModule.cs
@@ -1175,7 +1175,7 @@ namespace UnrealBuildTool
 			// Check if this module is invalid for the current target, based on Plugin attributes
 			if (Rules.CppCompileWarningSettings.PluginModuleUnsupportedWarningLevel != WarningLevel.Off && Rules.Plugin != null)
 			{
-				if (!Rules.Plugin.Descriptor.IsModuleCompiledInConfiguration(Rules.Name, Rules.Target.Platform, Rules.Target.Configuration, Rules.Target.Name, Rules.Target.Type, Rules.Target.bBuildDeveloperTools, Rules.Target.bBuildRequiresCookedData, out string? InvalidReason))
+				if (!Rules.Plugin.Descriptor.IsModuleCompiledInConfiguration(Rules.Name, Rules.Target.Platform, Rules.Target.Configuration, Rules.Target.Name, Rules.Target.Type, Rules.Target.bBuildDeveloperTools, Rules.Target.bBuildRequiresCookedData, Rules.Target.Architectures, out string? InvalidReason))
 				{
 					if (Rules.CppCompileWarningSettings.PluginModuleUnsupportedWarningLevel == WarningLevel.Warning)
 					{
@@ -1267,7 +1267,7 @@ namespace UnrealBuildTool
 			}
 			if (Rules.CppCompileWarningSettings.PluginModuleUnsupportedWarningLevel == WarningLevel.Error && Rules.Plugin != null)
 			{
-				if (!Rules.Plugin.Descriptor.IsModuleCompiledInConfiguration(Rules.Name, Rules.Target.Platform, Rules.Target.Configuration, Rules.Target.Name, Rules.Target.Type, Rules.Target.bBuildDeveloperTools, Rules.Target.bBuildRequiresCookedData, out string? _))
+				if (!Rules.Plugin.Descriptor.IsModuleCompiledInConfiguration(Rules.Name, Rules.Target.Platform, Rules.Target.Configuration, Rules.Target.Name, Rules.Target.Type, Rules.Target.bBuildDeveloperTools, Rules.Target.bBuildRequiresCookedData, Rules.Target.Architectures, out string? _))
 				{
 					anyErrors = true;
 				}
diff --git a/Engine/Source/Programs/UnrealBuildTool/Configuration/UEBuildTarget.cs b/Engine/Source/Programs/UnrealBuildTool/Configuration/UEBuildTarget.cs
--- a/Engine/Source/Programs/UnrealBuildTool/Configuration/UEBuildTarget.cs
+++ b/Engine/Source/Programs/UnrealBuildTool/Configuration/UEBuildTarget.cs
@@ -4465,7 +4465,7 @@ namespace UnrealBuildTool
 				// Add all the modules
 				foreach (ModuleDescriptor ModuleDescriptor in Plugin.Descriptor.Modules)
 				{
-					if (ModuleDescriptor.IsCompiledInConfiguration(Platform, Configuration, TargetName, TargetType, Rules.bBuildDeveloperTools, Rules.bBuildRequiresCookedData))
+					if (ModuleDescriptor.IsCompiledInConfiguration(Platform, Configuration, TargetName, TargetType, Rules.bBuildDeveloperTools, Rules.bBuildRequiresCookedData, Architectures))
 					{
 						FileReference? ModuleFileName = RulesAssembly.GetModuleFileName(ModuleDescriptor.Name);
 						if (ModuleFileName == null)
@@ -4846,7 +4846,7 @@ namespace UnrealBuildTool
 				return;
 			}
 
-			foreach (ModuleDescriptor Descriptor in ProjectDescriptor.Modules.Where(x => x.IsCompiledInConfiguration(Platform, Configuration, TargetName, TargetType, Rules.bBuildDeveloperTools, Rules.bBuildRequiresCookedData)))
+			foreach (ModuleDescriptor Descriptor in ProjectDescriptor.Modules.Where(x => x.IsCompiledInConfiguration(Platform, Configuration, TargetName, TargetType, Rules.bBuildDeveloperTools, Rules.bBuildRequiresCookedData, Architectures)))
 			{
 				// To maintain historical behavior any project modules that do not specifically list the program in the ProgramAllowList will not be added even if IsCompiledInConfiguration returns true
 				if (Rules.Type == TargetType.Program && (Descriptor.ProgramAllowList == null || !Descriptor.ProgramAllowList.Contains(AppName)))
@@ -5228,7 +5228,7 @@ namespace UnrealBuildTool
 				{
 					foreach (ModuleDescriptor ModuleInfo in Info.Descriptor.Modules)
 					{
-						if (ModuleInfo.IsCompiledInConfiguration(Platform, Configuration, TargetName, TargetType, Rules.bBuildDeveloperTools, Rules.bBuildRequiresCookedData))
+						if (ModuleInfo.IsCompiledInConfiguration(Platform, Configuration, TargetName, TargetType, Rules.bBuildDeveloperTools, Rules.bBuildRequiresCookedData, Architectures))
 						{
 							UEBuildModuleCPP Module = FindOrCreateCppModuleByName(ModuleInfo.Name, PluginReferenceChain, Logger);
 							if (!Instance.Modules.Contains(Module))
@@ -5898,6 +5898,7 @@ namespace UnrealBuildTool
 			// tell the compiled code the name of the UBT platform (this affects folder on disk, etc that the game may need to know)
 			GlobalCompileEnvironment.Definitions.Add(String.Format("UBT_COMPILED_PLATFORM=" + Platform.ToString()));
 			GlobalCompileEnvironment.Definitions.Add(String.Format("UBT_COMPILED_TARGET=" + TargetType.ToString()));
+			GlobalCompileEnvironment.Definitions.Add(String.Format("UBT_COMPILED_ARCHITECTURE=" + (Architectures.bIsMultiArch ? "MULTI" : Architectures.SingleArchitecture.ToString().ToLower()) ));
 
 			// Set the global app name
 			GlobalCompileEnvironment.Definitions.Add(String.Format("UE_APP_NAME=\"{0}\"", AppName));
diff --git a/Engine/Source/Programs/UnrealBuildTool/Platform/Windows/UEBuildWindows.cs b/Engine/Source/Programs/UnrealBuildTool/Platform/Windows/UEBuildWindows.cs
--- a/Engine/Source/Programs/UnrealBuildTool/Platform/Windows/UEBuildWindows.cs
+++ b/Engine/Source/Programs/UnrealBuildTool/Platform/Windows/UEBuildWindows.cs
@@ -1231,53 +1231,11 @@ namespace UnrealBuildTool
 
 			if (Platform == UnrealTargetPlatform.Win64)
 			{
-				Target.WindowsPlatform.Architecture = Target.Architecture;// == UnrealArch.Default ? UnrealArch.X64 : Target.Architecture;
+				Target.WindowsPlatform.Architecture = Target.Architecture;
 
-				// Add names of plugins here that are incompatible with arm64ec or arm64
 				bool bCompilingForArm = !Target.Architecture.bIsX64;
 				if (bCompilingForArm && Target.Name != "UnrealHeaderTool")
 				{
-					// TODO: Force unique build environment due to disabled plugins
-					Target.BuildEnvironment = TargetBuildEnvironment.Unique;
-					if (Target.WindowsPlatform.Architecture == UnrealArch.Arm64ec)
-					{
-						Target.DisablePlugins.AddRange(new string[]
-						{
-							"Reflex", // nvapi
-						});
-						Target.GlobalDefinitions.Add("UE_EXTERNAL_PROFILING_ENABLED=0");
-					}
-
-					// disabling some plugins until we get arm64 libs
-					if (Target.WindowsPlatform.Architecture == UnrealArch.Arm64)
-					{
-						Target.DisablePlugins.AddRange(new string[]
-						{
-							// steamapi
-							"OnlineSubsystemSteam",
-							"SteamSockets",
-							"SteamShared",
-							"SteamController",
-							"SocketSubsystemSteamIP",
-
-							// nvapi
-							"Reflex",
-
-							// DVP
-							"MediaIOFramework",
-							"Composure",
-
-							// Permanently disabled meaning we have no plans to add arm64 suppport at this time.
-							// LiveLinkOpenVR is used in virtual production where arm64 machines are unlikely to be used at this time and the OpenVR library is generally being abandoned in favor of OpenXR.
-							"LiveLinkOpenVR",
-						});
-					}
-
-					Target.DisablePlugins.AddRange(new string[]
-					{
-						"OpenImageDenoise",
-					});
-
 					// VTune does not support ARM
 					Target.GlobalDefinitions.Add("UE_EXTERNAL_PROFILING_ENABLED=0");
 				}
diff --git a/Engine/Source/Programs/UnrealBuildTool/System/ModuleDescriptor.cs b/Engine/Source/Programs/UnrealBuildTool/System/ModuleDescriptor.cs
--- a/Engine/Source/Programs/UnrealBuildTool/System/ModuleDescriptor.cs
+++ b/Engine/Source/Programs/UnrealBuildTool/System/ModuleDescriptor.cs
@@ -1,12 +1,14 @@
 // Copyright Epic Games, Inc. All Rights Reserved.
 
+using EpicGames.Core;
+using EpicGames.Serialization;
+using JsonExtensions;
 using System;
 using System.Collections.Generic;
 using System.Diagnostics;
 using System.Diagnostics.CodeAnalysis;
 using System.Linq;
-using EpicGames.Core;
-using JsonExtensions;
+using System.Runtime.InteropServices;
 
 namespace UnrealBuildTool
 {
@@ -182,6 +184,16 @@ namespace UnrealBuildTool
 		/// List of disallowed platforms
 		/// </summary>
 		public List<UnrealTargetPlatform>? PlatformDenyList;
+		
+		/// <summary>
+		/// Collection of allowed architectures per platform
+		/// </summary>
+		public Dictionary<UnrealTargetPlatform,List<UnrealArch>>? PlatformArchitectureAllowList;
+
+		/// <summary>
+		/// Collection of disallowed architectures per platform
+		/// </summary>
+		public Dictionary<UnrealTargetPlatform,List<UnrealArch>>? PlatformArchitectureDenyList;
 
 		/// <summary>
 		/// List of allowed targets
@@ -300,6 +312,62 @@ namespace UnrealBuildTool
 						}
 					}
 				}
+
+				
+				Dictionary<UnrealTargetPlatform,List<UnrealArch>> ParsePlatformArchitectureList( string[] PlatformArchitectureList )
+				{
+					Dictionary<UnrealTargetPlatform,List<UnrealArch>> Result = [];
+
+					foreach (string PlatformArchitecture in PlatformArchitectureList)
+					{
+						string[] Pair = PlatformArchitecture.Split(':');
+						if (Pair.Length != 2)
+						{
+							Log.TraceWarningTask(JsonFilePath, $"Malformed Platform:Architecture pair : {PlatformArchitecture} while parsing platform architecture list for module descriptor {Module.Name}");
+							continue;
+						}
+
+						UnrealTargetPlatform Platform;
+						if (!UnrealTargetPlatform.TryParse(Pair[0], out Platform))
+						{
+							if ( !PluginDescriptor.IsAllowableMissingPlatform(Pair[0], JsonFilePath.Directory) )
+							{
+								Log.TraceWarningTask(JsonFilePath, $"Unknown platform {Pair[0]} while parsing platform architecture list for module descriptor {Module.Name}");
+							}
+							continue;
+						}
+
+						UnrealArch Architecture;
+						if (!UnrealArch.TryParse(Pair[1], out Architecture))
+						{
+							Log.TraceWarningTask(JsonFilePath, $"Unknown architecture {Pair[1]} while parsing platform architecture list for module descriptor {Module.Name}");
+							continue;
+						}
+
+						if (Result.ContainsKey(Platform))
+						{
+							Result[Platform].Add(Architecture);
+						}
+						else
+						{
+							Result.Add(Platform, [Architecture]);
+						}
+					}
+
+					return Result;
+
+				}
+
+
+				if (InObject.TryGetStringArrayField("PlatformArchitectureAllowList", out string[]? PlatformArchitectureAllowList ))
+				{
+					Module.PlatformArchitectureAllowList = ParsePlatformArchitectureList(PlatformArchitectureAllowList);
+				}
+
+				if (InObject.TryGetStringArrayField("PlatformArchitectureDenyList", out string[]? PlatformArchitectureDenyList ))
+				{
+					Module.PlatformArchitectureDenyList = ParsePlatformArchitectureList(PlatformArchitectureDenyList);
+				}
 			}
 			catch (BuildException Ex)
 			{
@@ -307,6 +375,8 @@ namespace UnrealBuildTool
 				throw;
 			}
 
+
+
 			TargetType[]? TargetAllowList;
 			if (InObject.TryGetEnumArrayFieldWithDeprecatedFallback<TargetType>("TargetAllowList", "WhitelistTargets", out TargetAllowList))
 			{
@@ -401,6 +471,24 @@ namespace UnrealBuildTool
 				}
 				Writer.WriteArrayEnd();
 			}
+			if (PlatformArchitectureAllowList != null && PlatformArchitectureAllowList.Count > 0)
+			{
+				Writer.WriteArrayStart("PlatformArchitectureAllowList");
+				foreach (string PlatformArchitecture in PlatformArchitectureAllowList.SelectMany( Pair => Pair.Value.Select( Arch => $"{Pair.Key}:{Arch}" ) ))
+				{
+					Writer.WriteValue(PlatformArchitecture);
+				}
+				Writer.WriteArrayEnd();
+			}
+			if (PlatformArchitectureDenyList != null && PlatformArchitectureDenyList.Count > 0)
+			{
+				Writer.WriteArrayStart("PlatformArchitectureDenyList");
+				foreach (string PlatformArchitecture in PlatformArchitectureDenyList.SelectMany( Pair => Pair.Value.Select( Arch => $"{Pair.Key}:{Arch}" ) ))
+				{
+					Writer.WriteValue(PlatformArchitecture);
+				}
+				Writer.WriteArrayEnd();
+			}
 			if (TargetAllowList != null && TargetAllowList.Length > 0)
 			{
 				Writer.WriteArrayStart("TargetAllowList");
@@ -488,6 +576,16 @@ namespace UnrealBuildTool
 				string[] PlatformDenyListStringArray = PlatformDenyList.Select(X => X.ToString()).ToArray();
 				ModuleObject.AddOrSetFieldValue("PlatformDenyList", PlatformDenyListStringArray);
 			}
+			if (PlatformArchitectureAllowList != null && PlatformArchitectureAllowList.Count > 0)
+			{
+				string[] PlatformArchitectureStringArray = PlatformArchitectureAllowList.SelectMany( Pair => Pair.Value.Select( Arch => $"{Pair.Key}:{Arch}" ) ).ToArray();
+				ModuleObject.AddOrSetFieldValue("PlatformArchitectureAllowList", PlatformArchitectureStringArray);
+			}
+			if (PlatformArchitectureDenyList != null)
+			{
+				string[] PlatformArchitectureStringArray = PlatformArchitectureDenyList.SelectMany( Pair => Pair.Value.Select( Arch => $"{Pair.Key}:{Arch}" ) ).ToArray();
+				ModuleObject.AddOrSetFieldValue("PlatformArchitectureDenyList", PlatformArchitectureStringArray);
+			}
 			if (TargetAllowList != null && TargetAllowList.Length > 0)
 			{
 				string[] TargetAllowListStringArray = TargetAllowList.Select(X => X.ToString()).ToArray();
@@ -583,9 +681,10 @@ namespace UnrealBuildTool
 		/// <param name="TargetType">The type of the target being compiled</param>
 		/// <param name="bBuildDeveloperTools">Whether the configuration includes developer tools (typically UEBuildConfiguration.bBuildDeveloperTools for UBT callers)</param>
 		/// <param name="bBuildRequiresCookedData">Whether the configuration requires cooked content (typically UEBuildConfiguration.bBuildRequiresCookedData for UBT callers)</param>
-		public bool IsCompiledInConfiguration(UnrealTargetPlatform Platform, UnrealTargetConfiguration Configuration, string TargetName, TargetType TargetType, bool bBuildDeveloperTools, bool bBuildRequiresCookedData)
+		/// <param name="Architectures">The architectures being compiled</param>
+		public bool IsCompiledInConfiguration(UnrealTargetPlatform Platform, UnrealTargetConfiguration Configuration, string TargetName, TargetType TargetType, bool bBuildDeveloperTools, bool bBuildRequiresCookedData, UnrealArchitectures? Architectures = null)
 		{
-			return IsCompiledInConfiguration(Platform, Configuration, TargetName, TargetType, bBuildDeveloperTools, bBuildRequiresCookedData, out string? _);
+			return IsCompiledInConfiguration(Platform, Configuration, TargetName, TargetType, bBuildDeveloperTools, bBuildRequiresCookedData, Architectures, out string? _);
 		}
 
 		/// <summary>
@@ -597,8 +696,9 @@ namespace UnrealBuildTool
 		/// <param name="TargetType">The type of the target being compiled</param>
 		/// <param name="bBuildDeveloperTools">Whether the configuration includes developer tools (typically UEBuildConfiguration.bBuildDeveloperTools for UBT callers)</param>
 		/// <param name="bBuildRequiresCookedData">Whether the configuration requires cooked content (typically UEBuildConfiguration.bBuildRequiresCookedData for UBT callers)</param>
+		/// <param name="Architectures">The architectures being compiled</param>
 		/// <param name="invalidReason">Out parameter, reason why this plugin module is invalid</param>
-		public bool IsCompiledInConfiguration(UnrealTargetPlatform Platform, UnrealTargetConfiguration Configuration, string TargetName, TargetType TargetType, bool bBuildDeveloperTools, bool bBuildRequiresCookedData, [NotNullWhen(false)] out string? invalidReason)
+		public bool IsCompiledInConfiguration(UnrealTargetPlatform Platform, UnrealTargetConfiguration Configuration, string TargetName, TargetType TargetType, bool bBuildDeveloperTools, bool bBuildRequiresCookedData, UnrealArchitectures? Architectures, [NotNullWhen(false)] out string? invalidReason)
 		{
 			// Check the platform is allowed
 			// important note: we don't check the length of the platform allow list, because if an unknown platform was read in, but was not valid, the 
@@ -618,6 +718,24 @@ namespace UnrealBuildTool
 				return false;
 			}
 
+			// Check the architecture is allowed
+			if (Architectures != null && !Architectures.bIsMultiArch && PlatformArchitectureAllowList != null && PlatformArchitectureAllowList.TryGetValue(Platform, out List<UnrealArch>? AllowedArchitectures) && !AllowedArchitectures.Contains(Architectures.SingleArchitecture))
+			{
+				invalidReason = $"PlatformArchitectureAllowList does not include {Platform}:{Architectures.SingleArchitecture}";
+				return false;
+			}
+
+			// check the architecture is not denied
+			if (Architectures != null && Architectures.bIsMultiArch && PlatformArchitectureDenyList != null && PlatformArchitectureDenyList.Any())
+			{
+				throw new Exception($"PlatformArchitectureDenyList does not support {Platform} Multi-architecture builds ({Name})");
+			}
+			if (Architectures != null && !Architectures.bIsMultiArch && PlatformArchitectureDenyList != null && PlatformArchitectureDenyList.TryGetValue(Platform, out List<UnrealArch>? DeniedArchitectures) && DeniedArchitectures.Contains(Architectures.SingleArchitecture))
+			{
+				invalidReason = $"PlatformArchitectureDenyList includes {Platform}:{Architectures.SingleArchitecture}";
+				return false;
+			}
+
 			// Check the target is allowed
 			if (TargetAllowList != null && TargetAllowList.Length > 0 && !TargetAllowList.Contains(TargetType))
 			{
diff --git a/Engine/Source/Programs/UnrealBuildTool/System/PluginDescriptor.cs b/Engine/Source/Programs/UnrealBuildTool/System/PluginDescriptor.cs
--- a/Engine/Source/Programs/UnrealBuildTool/System/PluginDescriptor.cs
+++ b/Engine/Source/Programs/UnrealBuildTool/System/PluginDescriptor.cs
@@ -777,8 +777,9 @@ namespace UnrealBuildTool
 		/// <param name="TargetType">The type of the target being compiled</param>
 		/// <param name="bBuildDeveloperTools">Whether the configuration includes developer tools (typically UEBuildConfiguration.bBuildDeveloperTools for UBT callers)</param>
 		/// <param name="bBuildRequiresCookedData">Whether the configuration requires cooked content (typically UEBuildConfiguration.bBuildRequiresCookedData for UBT callers)</param>
+		/// <param name="Architectures">The architectures being built</param>
 		/// <param name="invalidReason">Out parameter, reason why this plugin module is invalid</param>
-		public bool IsModuleCompiledInConfiguration(string Name, UnrealTargetPlatform Platform, UnrealTargetConfiguration Configuration, string TargetName, TargetType TargetType, bool bBuildDeveloperTools, bool bBuildRequiresCookedData, [NotNullWhen(false)] out string? invalidReason)
+		public bool IsModuleCompiledInConfiguration(string Name, UnrealTargetPlatform Platform, UnrealTargetConfiguration Configuration, string TargetName, TargetType TargetType, bool bBuildDeveloperTools, bool bBuildRequiresCookedData, UnrealArchitectures? Architectures, [NotNullWhen(false)] out string? invalidReason)
 		{
 			IEnumerable<ModuleDescriptor>? Descriptors = Modules?.Where(x => x.Name == Name);
 			if (Descriptors == null)
@@ -791,7 +792,7 @@ namespace UnrealBuildTool
 
 			foreach (ModuleDescriptor ModuleInfo in Descriptors)
 			{
-				if (!ModuleInfo.IsCompiledInConfiguration(Platform, Configuration, TargetName, TargetType, bBuildDeveloperTools, bBuildRequiresCookedData, out string? InvalidReason))
+				if (!ModuleInfo.IsCompiledInConfiguration(Platform, Configuration, TargetName, TargetType, bBuildDeveloperTools, bBuildRequiresCookedData, Architectures, out string? InvalidReason))
 				{
 					invalidReasons.Add(InvalidReason);
 				}
diff --git a/Engine/Source/Programs/UnrealBuildTool/System/Plugins.cs b/Engine/Source/Programs/UnrealBuildTool/System/Plugins.cs
--- a/Engine/Source/Programs/UnrealBuildTool/System/Plugins.cs
+++ b/Engine/Source/Programs/UnrealBuildTool/System/Plugins.cs
@@ -1,12 +1,13 @@
 // Copyright Epic Games, Inc. All Rights Reserved.
 
+using EpicGames.Core;
+using Microsoft.Extensions.Logging;
 using System;
 using System.Collections.Generic;
 using System.Diagnostics;
 using System.IO;
 using System.Linq;
-using EpicGames.Core;
-using Microsoft.Extensions.Logging;
+using System.Runtime.InteropServices;
 using UnrealBuildBase;
 
 namespace UnrealBuildTool
@@ -505,6 +506,53 @@ namespace UnrealBuildTool
 									ParentModule.PlatformDenyList = ParentModule.PlatformDenyList.Union(ChildModule.PlatformDenyList).ToList();
 								}
 							}
+
+							// merge platform architecture allow/deny lists
+							static void MergePlatformArchitectureLists( Dictionary<UnrealTargetPlatform,List<UnrealArch>> Parent, Dictionary<UnrealTargetPlatform,List<UnrealArch>> Child)
+							{
+								foreach (KeyValuePair<UnrealTargetPlatform,List<UnrealArch>> Pair in Child)
+								{
+									if (!Parent.TryGetValue(Pair.Key, out List<UnrealArch>? ArchList))
+									{
+										Parent[Pair.Key] = Pair.Value;
+									}
+									else
+									{
+										foreach (UnrealArch Arch in Pair.Value)
+										{
+											if (!ArchList.Contains(Arch))
+											{
+												ArchList.Add(Arch);
+											}
+										}
+									}
+								}
+							}
+
+							if (ChildModule.PlatformArchitectureAllowList != null)
+							{
+								if (ParentModule.PlatformArchitectureAllowList == null)
+								{
+									ParentModule.PlatformArchitectureAllowList = ChildModule.PlatformArchitectureAllowList;
+								}
+								else
+								{
+									MergePlatformArchitectureLists(ParentModule.PlatformArchitectureAllowList, ChildModule.PlatformArchitectureAllowList);
+								}
+							}
+
+							if (ChildModule.PlatformArchitectureDenyList != null)
+							{
+								if (ParentModule.PlatformArchitectureDenyList == null)
+								{
+									ParentModule.PlatformArchitectureDenyList = ChildModule.PlatformArchitectureDenyList;
+								}
+								else
+								{
+									MergePlatformArchitectureLists(ParentModule.PlatformArchitectureDenyList, ChildModule.PlatformArchitectureDenyList);
+								}
+							}
+
 						}
 						else
 						{
@@ -656,8 +704,9 @@ namespace UnrealBuildTool
 		/// <param name="Configuration">The target configuration</param>
 		/// <param name="TargetType">The type of target being built</param>
 		/// <param name="bRequiresCookedData">Whether the target requires cooked data</param>
+		/// <param name="Architectures">The architectures being built</param>
 		/// <returns>True if the plugin should be enabled for this project</returns>
-		public static bool IsPluginCompiledForTarget(PluginInfo Plugin, ProjectDescriptor Project, UnrealTargetPlatform Platform, UnrealTargetConfiguration Configuration, TargetType TargetType, bool bRequiresCookedData)
+		public static bool IsPluginCompiledForTarget(PluginInfo Plugin, ProjectDescriptor Project, UnrealTargetPlatform Platform, UnrealTargetConfiguration Configuration, TargetType TargetType, bool bRequiresCookedData, UnrealArchitectures? Architectures = null)
 		{
 			bool bCompiledForTarget = false;
 			if (IsPluginEnabledForTarget(Plugin, Project, Platform, Configuration, TargetType) && Plugin.Descriptor.Modules != null)
@@ -665,7 +714,7 @@ namespace UnrealBuildTool
 				bool bBuildDeveloperTools = (TargetType == TargetType.Editor || TargetType == TargetType.Program || (Configuration != UnrealTargetConfiguration.Test && Configuration != UnrealTargetConfiguration.Shipping));
 				foreach (ModuleDescriptor Module in Plugin.Descriptor.Modules)
 				{
-					if (Module.IsCompiledInConfiguration(Platform, Configuration, "", TargetType, bBuildDeveloperTools, bRequiresCookedData))
+					if (Module.IsCompiledInConfiguration(Platform, Configuration, "", TargetType, bBuildDeveloperTools, bRequiresCookedData, Architectures))
 					{
 						bCompiledForTarget = true;
 						break;
diff --git a/Engine/Source/Runtime/Core/Private/GenericPlatform/GenericPlatformMisc.cpp b/Engine/Source/Runtime/Core/Private/GenericPlatform/GenericPlatformMisc.cpp
--- a/Engine/Source/Runtime/Core/Private/GenericPlatform/GenericPlatformMisc.cpp
+++ b/Engine/Source/Runtime/Core/Private/GenericPlatform/GenericPlatformMisc.cpp
@@ -1512,6 +1512,12 @@ const TCHAR* FGenericPlatformMisc::GetUBTTarget()
 	return TEXT(PREPROCESSOR_TO_STRING(UBT_COMPILED_TARGET));
 }
 
+const TCHAR* FGenericPlatformMisc::GetUBTArchitecture()
+{
+	return TEXT(PREPROCESSOR_TO_STRING(UBT_COMPILED_ARCHITECTURE));
+}
+
+
 using FUBTTargetNameArrayType = TCHAR[128];
 
 #if PLATFORM_TCHAR_IS_UTF8CHAR
diff --git a/Engine/Source/Runtime/Core/Public/GenericPlatform/GenericPlatformMisc.h b/Engine/Source/Runtime/Core/Public/GenericPlatform/GenericPlatformMisc.h
--- a/Engine/Source/Runtime/Core/Public/GenericPlatform/GenericPlatformMisc.h
+++ b/Engine/Source/Runtime/Core/Public/GenericPlatform/GenericPlatformMisc.h
@@ -1265,6 +1265,8 @@ public:
 
 	static CORE_API const TCHAR* GetUBTTarget();
 
+	static CORE_API const TCHAR* GetUBTArchitecture();
+
 	static CORE_API void SetUBTTargetName(const TCHAR* InTargetName);
 	static CORE_API const TCHAR* GetUBTTargetName();
 
diff --git a/Engine/Source/Runtime/Projects/Private/ModuleDescriptor.cpp b/Engine/Source/Runtime/Projects/Private/ModuleDescriptor.cpp
--- a/Engine/Source/Runtime/Projects/Private/ModuleDescriptor.cpp
+++ b/Engine/Source/Runtime/Projects/Private/ModuleDescriptor.cpp
@@ -230,6 +230,37 @@ bool FModuleDescriptor::Read(const FJsonObject& Object, FText* OutFailReason /*=
 	Object.TryGetStringArrayField(TEXT("GameTargetAllowList"), GameTargetAllowList);
 	Object.TryGetStringArrayField(TEXT("GameTargetDenyList"), GameTargetDenyList);
 	
+	// Read the platform architecture allow/deny lists
+	auto ReadPlatformArchitectureList = [this, Object, OutFailReason]( const TCHAR* FieldName, TMap<FString,TArray<FString>>& OutPlatformArchitectureList )
+	{
+		TArray<FString> PlatformArchitecturePairs;
+		Object.TryGetStringArrayField(FieldName, PlatformArchitecturePairs);
+		for (const FString& PlatformArchitecturePair : PlatformArchitecturePairs)
+		{
+			FString Platform, Architecture;
+			if (!PlatformArchitecturePair.Split(TEXT(":"), &Platform, &Architecture))
+			{
+				if (OutFailReason)
+				{
+					*OutFailReason = FText::Format( LOCTEXT( "ModuleWithInvalidPlatformArchitecture", "Module entry '{0}' specified an unrecognized Platform:Architecture pair '{1}'" ), FText::FromName(Name), FText::FromString(PlatformArchitecturePair) );
+				}
+				return false;
+			}
+
+			OutPlatformArchitectureList.FindOrAdd(Platform).AddUnique(Architecture);
+		}
+		return true;
+
+	};
+	if (!ReadPlatformArchitectureList(TEXT("PlatformArchitectureAllowList"), PlatformArchitectureAllowList))
+	{
+		return false;
+	}
+	if (!ReadPlatformArchitectureList(TEXT("PlatformArchitectureDenyList"), PlatformArchitectureDenyList))
+	{
+		return false;
+	}
+
 	return true;
 }
 
@@ -322,6 +353,41 @@ void FModuleDescriptor::UpdateJson(FJsonObject& JsonObject) const
 		JsonObject.RemoveField(TEXT("PlatformDenyList"));
 	}
 
+
+	auto GetPlatformArchitectureListValues = []( const TMap<FString,TArray<FString>>& PlatformArchitectureList )
+	{
+		TArray<TSharedPtr<FJsonValue>> ListValues;
+		for (const TPair<FString,TArray<FString>>& Pairs : PlatformArchitectureList)
+		{
+			for (const FString& Value : Pairs.Value)
+			{
+				FString PlatformArchitectureValue = Pairs.Key + TEXT(":") + Value;
+				ListValues.Add(MakeShareable(new FJsonValueString(PlatformArchitectureValue)));
+			}
+		}
+
+		return ListValues;
+	};
+
+	if (PlatformArchitectureAllowList.Num() > 0)
+	{
+		JsonObject.SetArrayField(TEXT("PlatformArchitectureAllowList"), GetPlatformArchitectureListValues(PlatformArchitectureAllowList));
+	}
+	else
+	{
+		JsonObject.RemoveField(TEXT("PlatformArchitectureAllowList"));
+	}
+		
+	if (PlatformArchitectureDenyList.Num() > 0)
+	{
+		JsonObject.SetArrayField(TEXT("PlatformArchitectureDenyList"), GetPlatformArchitectureListValues(PlatformArchitectureDenyList));
+	}
+	else
+	{
+		JsonObject.RemoveField(TEXT("PlatformArchitectureDenyList"));
+	}
+
+
 	if (TargetAllowList.Num() > 0)
 	{
 		TArray<TSharedPtr<FJsonValue>> TargetAllowListValues;
@@ -491,7 +557,7 @@ void FModuleDescriptor::UpdateArray(FJsonObject& JsonObject, const TCHAR* ArrayN
 		FModuleJsonArrayUpdater::FUpdateJsonObject::CreateStatic(ModuleDescriptor::UpdateModuleJsonObject));
 }
 
-bool FModuleDescriptor::IsCompiledInConfiguration(const FString& Platform, EBuildConfiguration Configuration, const FString& TargetName, EBuildTargetType TargetType, bool bBuildDeveloperTools, bool bBuildRequiresCookedData) const
+bool FModuleDescriptor::IsCompiledInConfiguration(const FString& Platform, EBuildConfiguration Configuration, const FString& TargetName, EBuildTargetType TargetType, bool bBuildDeveloperTools, bool bBuildRequiresCookedData, const FString& Architecture) const
 {
 	// Check the platform is allowed
 	if ((bHasExplicitPlatforms || PlatformAllowList.Num() > 0) && !PlatformAllowList.Contains(Platform))
@@ -505,6 +571,19 @@ bool FModuleDescriptor::IsCompiledInConfiguration(const FString& Platform, EBuil
 		return false;
 	}
 
+	// Check the platform architecture is allowed
+	if (!Architecture.IsEmpty() && PlatformArchitectureAllowList.Contains(Platform) && !PlatformArchitectureAllowList[Platform].Contains(Architecture))
+	{
+		return false;
+	}
+
+	// Check the platform architecture is not denied
+	checkf(!(PlatformArchitectureDenyList.Contains(Platform) && Architecture == TEXT("MULTI")), TEXT("PlatformArchitectureDenyList does not support %s Multi-architecture builds (%s)"), *Platform, *Name.ToString());
+	if (!Architecture.IsEmpty() && PlatformArchitectureDenyList.Contains(Platform) && PlatformArchitectureDenyList[Platform].Contains(Architecture))
+	{
+		return false;
+	}
+
 	// Check the target is allowed
 	if (TargetAllowList.Num() > 0 && !TargetAllowList.Contains(TargetType))
 	{
@@ -592,7 +671,7 @@ bool FModuleDescriptor::IsCompiledInConfiguration(const FString& Platform, EBuil
 
 bool FModuleDescriptor::IsCompiledInCurrentConfiguration() const
 {
-	return IsCompiledInConfiguration(FPlatformMisc::GetUBTPlatform(), FApp::GetBuildConfiguration(), UE_APP_NAME, FApp::GetBuildTargetType(), !!WITH_UNREAL_DEVELOPER_TOOLS, FPlatformProperties::RequiresCookedData());
+	return IsCompiledInConfiguration(FPlatformMisc::GetUBTPlatform(), FApp::GetBuildConfiguration(), UE_APP_NAME, FApp::GetBuildTargetType(), !!WITH_UNREAL_DEVELOPER_TOOLS, FPlatformProperties::RequiresCookedData(), FPlatformMisc::GetUBTArchitecture());
 }
 
 bool FModuleDescriptor::IsLoadedInCurrentConfiguration() const
diff --git a/Engine/Source/Runtime/Projects/Private/PluginManager.cpp b/Engine/Source/Runtime/Projects/Private/PluginManager.cpp
--- a/Engine/Source/Runtime/Projects/Private/PluginManager.cpp
+++ b/Engine/Source/Runtime/Projects/Private/PluginManager.cpp
@@ -974,6 +974,31 @@ void FPluginManager::ReadAllPlugins(FDiscoveredPluginMap& Plugins, const TSet<FS
 
 						// if we want to deny a platform, add it even if the parent didn't have a deny list. this won't cause problems with other platforms
 						ParentModule.PlatformDenyList.Append(ChildModule.PlatformDenyList);
+
+						// merge platform architecture allow/deny lists
+						auto MergePlatformArchitectureLists = []( TMap<FString,TArray<FString>>& Parent, const TMap<FString,TArray<FString>>& Child )
+						{
+							for (TPair<FString,TArray<FString>> Pair : Child)
+							{
+								TArray<FString>* ArchList = Parent.Find(Pair.Key);
+								if (ArchList == nullptr)
+								{
+									Parent.Add(Pair.Key, Pair.Value);
+								}
+								else
+								{
+									for (const FString& Arch : Pair.Value)
+									{
+										if (!ArchList->Contains(Arch))
+										{
+											ArchList->Add(Arch);
+										}
+									}
+								}
+							}							
+						};
+						MergePlatformArchitectureLists( ParentModule.PlatformArchitectureAllowList, ChildModule.PlatformArchitectureAllowList);
+						MergePlatformArchitectureLists( ParentModule.PlatformArchitectureDenyList, ChildModule.PlatformArchitectureDenyList);
 					}
 				}
 				
diff --git a/Engine/Source/Runtime/Projects/Public/ModuleDescriptor.h b/Engine/Source/Runtime/Projects/Public/ModuleDescriptor.h
--- a/Engine/Source/Runtime/Projects/Public/ModuleDescriptor.h
+++ b/Engine/Source/Runtime/Projects/Public/ModuleDescriptor.h
@@ -168,6 +168,12 @@ struct FModuleDescriptor
 	/** List of disallowed platforms */
 	TArray<FString> PlatformDenyList;
 
+	/** List of allowed platform architectures */
+	TMap<FString,TArray<FString>> PlatformArchitectureAllowList;
+
+	/** List of disallowed platform architectures */
+	TMap<FString,TArray<FString>> PlatformArchitectureDenyList;
+
 	/** List of allowed targets */
 	TArray<EBuildTargetType> TargetAllowList;
 
@@ -227,7 +233,7 @@ struct FModuleDescriptor
 	static PROJECTS_API void UpdateArray(FJsonObject& JsonObject, const TCHAR* ArrayName, const TArray<FModuleDescriptor>& Modules);
 
 	/** Tests whether the module should be built for the given target */
-	PROJECTS_API bool IsCompiledInConfiguration(const FString& Platform, EBuildConfiguration Configuration, const FString& TargetName, EBuildTargetType TargetType, bool bBuildDeveloperTools, bool bBuildRequiresCookedData) const;
+	PROJECTS_API bool IsCompiledInConfiguration(const FString& Platform, EBuildConfiguration Configuration, const FString& TargetName, EBuildTargetType TargetType, bool bBuildDeveloperTools, bool bBuildRequiresCookedData, const FString& Architecture = FString()) const;
 
 	/** Tests whether the module should be built for the current engine configuration */
 	PROJECTS_API bool IsCompiledInCurrentConfiguration() const;
