UnrealBuildTool: Fixed so modules in MergeAdditionalCommonModules also pull the engine modules depending on them into common

#rnx

diff --git a/Engine/Source/Programs/UnrealBuildTool/Configuration/UEBuildTarget.cs b/Engine/Source/Programs/UnrealBuildTool/Configuration/UEBuildTarget.cs
--- a/Engine/Source/Programs/UnrealBuildTool/Configuration/UEBuildTarget.cs
+++ b/Engine/Source/Programs/UnrealBuildTool/Configuration/UEBuildTarget.cs
@@ -3552,8 +3552,9 @@ namespace UnrealBuildTool
 					}
 				}
 
-				// No non-engine PerModuleInline files found, Fall back to Core.
-				((UEBuildModuleCPP)Modules["Core"]).bCreatePerModuleFile = true;
+				// No non-engine PerModuleInline files found, fall back to Core or primary module
+				UEBuildModuleCPP module = binary.Modules.OfType<UEBuildModuleCPP>().FirstOrDefault(x => x.Name == "Core") ?? binary.PrimaryModule;
+				module.bCreatePerModuleFile = true;
 			}
 
 			void MergeModulesToDLL(string name, IEnumerable<UEBuildModuleCPP> modules)
@@ -3853,8 +3854,45 @@ namespace UnrealBuildTool
 				}
 			}
 
+			List<UEBuildModuleCPP> engineModules = [];
+			HashSet<UEBuildModuleCPP> excludedModules = [];
+
+			// Collect engine modules and find the ones that are explicitly moved to common
+			foreach (UEBuildModuleCPP m in allModules)
+			{
+				if (!m.RulesFile.IsUnderDirectory(Unreal.EngineDirectory))
+				{
+					continue;
+				}
+				if (Rules.MergeAdditionalCommonModules.Contains(m.Name))
+				{
+					excludedModules.Add(m);
+				}
+				else
+				{
+					engineModules.Add(m);
+				}
+			}
+
+			// Find all the engine modules that depends on the ones that was moved to common and move them too
+			if (excludedModules.Count != 0)
+			{
+				while (engineModules.RemoveAll((m) =>
+				{
+					foreach (UEBuildModule m2 in m.PublicDependencyModules!.Union(m.PrivateDependencyModules!))
+					{
+						if (!excludedModules.Contains(m2))
+						{
+							continue;
+						}
+						excludedModules.Add(m);
+						return true;
+					}
+					return false;
+				}) != 0) ;
+			}
+
 			// Merge remaining engine modules
-			IEnumerable<UEBuildModuleCPP> engineModules = allModules.Where(x => x.RulesFile.IsUnderDirectory(Unreal.EngineDirectory) && !Rules.MergeAdditionalCommonModules.Contains(x.Name));
 			ConsumeModules("Engine", engineModules);
 
 			// Merge remaining modules to common set
@@ -3866,7 +3904,7 @@ namespace UnrealBuildTool
 			Binaries.RemoveAll(x => !x.Modules.Any());
 
 			// Set module that should pull in PerModuleInline file
-			Binaries.ForEach(binary => SetPerModuleInline(binary));
+			Binaries.ForEach(SetPerModuleInline);
 		}
 
 		/// <summary>
