Adding capabilities to the warning defaults, cleaning up tests, adding more warning+error strings for easier long term adoption & toggling.

diff --git a/Engine/Source/Programs/Shared/UnrealBuildTool.Tests/CppCompileWarningsTests.cs b/Engine/Source/Programs/Shared/UnrealBuildTool.Tests/CppCompileWarningsTests.cs
--- a/Engine/Source/Programs/Shared/UnrealBuildTool.Tests/CppCompileWarningsTests.cs
+++ b/Engine/Source/Programs/Shared/UnrealBuildTool.Tests/CppCompileWarningsTests.cs
@@ -5,7 +5,6 @@
 using System.IO;
 using System.Linq;
 using EpicGames.Core;
-using Microsoft.Extensions.Logging;
 using Microsoft.VisualStudio.TestTools.UnitTesting;
 using UnrealBuildBase;
 using UnrealBuildTool.Tests.TestUtilities;
@@ -305,38 +304,6 @@
 					Assert.IsFalse(args.Any());
 				}
 
-				// VC toolchain check;
-				{
-					baseWarnings.UnsafeTypeCastWarningLevel = WarningLevel.Warning;
-					args = baseWarnings.GenerateWarningCommandLineArgs(newCompileEnvironment, typeof(VCToolChain));
-					Assert.IsTrue(args.Count() == 2);
-					Assert.IsTrue(args.Contains("/w44244"));
-					Assert.IsTrue(args.Contains("/w44838"));
-
-					baseWarnings.UnsafeTypeCastWarningLevel = WarningLevel.Default;
-				}
-
-				// MSVC compiler check;
-				{
-					baseWarnings.ShadowVariableWarningLevel = WarningLevel.Error;
-					args = baseWarnings.GenerateWarningCommandLineArgs(newCompileEnvironment, typeof(VCToolChain));
-
-					Assert.IsTrue(args.Count() == 3);
-					Assert.IsTrue(args.Contains("/we4456"));
-					Assert.IsTrue(args.Contains("/we4458"));
-					Assert.IsTrue(args.Contains("/we4459"));
-
-					WindowsCompiler previousCompiler = testTargetRules.WindowsPlatform.Compiler;
-					testTargetRules.WindowsPlatform.Compiler = WindowsCompiler.Intel;
-
-					args = baseWarnings.GenerateWarningCommandLineArgs(newCompileEnvironment, typeof(VCToolChain));
-
-					Assert.IsFalse(args.Any());
-
-					testTargetRules.WindowsPlatform.Compiler = previousCompiler;
-					baseWarnings.ShadowVariableWarningLevel = WarningLevel.Default;
-				}
-
 				// Specialized VC toolchain check; UndefinedIdentifier
 				{
 					baseWarnings.UndefinedIdentifierWarningLevel = WarningLevel.Warning;
@@ -389,92 +356,207 @@
 					Assert.AreEqual(1, args.Count());
 					Assert.IsTrue(args.Contains("/we5048"));
 				}
+			}
+		}
 
-				// Specialized VC toolchain check;  UnsafeTypeCast
+		[TestMethod]
+		public void TestSpecializedClangCommandLineArgs()
+		{
+			DirectoryReference programsDirectory = DirectoryReference.Combine(Unreal.EngineSourceDirectory, "Programs");
+			FileReference testStubsDirectory = FileReference.Combine(programsDirectory, "Shared", "UnrealBuildTool.Tests", "UBT");
+
+			// Generate ephemeral TargetRules & uproject files
+			using (SafeTestDirectory testRootDirectory = SafeTestDirectory.CreateTestDirectory(Path.Combine(StaticTargetRulesInitializer.UBTTestFolderRoot, $"{StaticTargetRulesInitializer.UBTTestFolderPrefix}{Guid.NewGuid().ToString()}")))
+			using (SafeTestDirectory testSourceDirectory = SafeTestDirectory.CreateTestDirectory(Path.Combine(testRootDirectory.TemporaryDirectory, "Source")))
+			using (SafeTestFile testEditorTargetTestFile = new SafeTestFile(File.ReadAllText(Path.Combine(testStubsDirectory.FullName, "TargetStubs", "TestStubEditor.Target.ubttest")), "TestStubEditor.Target.cs", testSourceDirectory.TemporaryDirectory))
+			using (SafeTestFile testTargetTestFile = new SafeTestFile(File.ReadAllText(Path.Combine(testStubsDirectory.FullName, "TargetStubs", "TestStub.Target.ubttest")), "TestStub.Target.cs", testSourceDirectory.TemporaryDirectory))
+			using (SafeTestFile testUprojectTestFile = new SafeTestFile(File.ReadAllText(Path.Combine(testStubsDirectory.FullName, "UprojectStubs", "TestStub.uproject.ubttest")), "TestStub.uproject", testRootDirectory.TemporaryDirectory))
+			{
+				CppCompileEnvironment newCompileEnvironment = new CppCompileEnvironment(UnrealTargetPlatform.Win64, CppConfiguration.Development, new UnrealArchitectures(UnrealArch.X64), null!);
+
+				FileReference testUproject = new FileReference(testUprojectTestFile.TemporaryFile);
+
+				// Generate rules assembly & subsequent targets
+				RulesAssembly generatedTestRulesAssembly = RulesCompiler.CreateTargetRulesAssembly(testUproject, "TestStubEditor", false, false, false, null, false, new TestLogger());
+				TargetRules testTargetRules = generatedTestRulesAssembly.CreateTargetRules("TestStubEditor", UnrealTargetPlatform.Win64, UnrealTargetConfiguration.Development, null, testUproject, null, new TestLogger());
+				testTargetRules.WindowsPlatform.Compiler = WindowsCompiler.VisualStudio2022;
+				CppCompileWarnings baseWarnings = new CppCompileWarnings(testTargetRules, null);
+
+				IEnumerable<string> args = null!;
 				{
-					RulesAssembly generatedTestGameRulesAssembly = RulesCompiler.CreateTargetRulesAssembly(testUproject, "TestStub", false, false, false, null, false, new TestLogger());
-					TargetRules testTargetRulesNonWin64 = generatedTestGameRulesAssembly.CreateTargetRules("TestStub", UnrealTargetPlatform.Android, UnrealTargetConfiguration.Development, null, testUproject, null, new TestLogger());
-					CppCompileWarnings specializedWarnings = new CppCompileWarnings(testTargetRulesNonWin64, null);
+					// DeterministcFlagSetFilter
+					{
+						testTargetRules.WindowsPlatform.Compiler = WindowsCompiler.Clang;
+						baseWarnings.DeterministicWarningLevel = WarningLevel.Off;
+						VersionNumber newVersionNumber = new VersionNumber(18, 9, 9);
 
-					specializedWarnings.UnsafeTypeCastWarningLevel = WarningLevel.Warning;
-					args = specializedWarnings.GenerateWarningCommandLineArgs(newCompileEnvironment, typeof(VCToolChain));
+						CppCompileEnvironment deterministicCompileEnvironment = new CppCompileEnvironment(UnrealTargetPlatform.Win64, CppConfiguration.Debug, new UnrealArchitectures(UnrealArch.X64), null!);
+						deterministicCompileEnvironment.bDeterministic = false;
 
-					Assert.IsTrue(args.Count() == 2);
-					Assert.IsTrue(args.Contains("/wd4244"));
-					Assert.IsTrue(args.Contains("/wd4838"));
+						args = baseWarnings.GenerateWarningCommandLineArgs(deterministicCompileEnvironment, typeof(ClangToolChain), newVersionNumber);
+						Assert.IsFalse(args.Any());
 
-					baseWarnings.UnsafeTypeCastWarningLevel = WarningLevel.Warning;
-					args = baseWarnings.GenerateWarningCommandLineArgs(newCompileEnvironment, typeof(VCToolChain));
+						deterministicCompileEnvironment.bDeterministic = true;
 
-					Assert.IsTrue(args.Count() == 2);
-					Assert.IsTrue(args.Contains("/w44244"));
-					Assert.IsTrue(args.Contains("/w44838"));
+						args = baseWarnings.GenerateWarningCommandLineArgs(deterministicCompileEnvironment, typeof(ClangToolChain), newVersionNumber);
+						Assert.IsTrue(args.Count() == 1);
+						Assert.IsTrue(args.Contains("-Wno-date-time"));
 
-					baseWarnings.UnsafeTypeCastWarningLevel = WarningLevel.Default;
+						deterministicCompileEnvironment.bDeterministic = false;
+						baseWarnings.DeterministicWarningLevel = WarningLevel.Warning;
+						args = baseWarnings.GenerateWarningCommandLineArgs(deterministicCompileEnvironment, typeof(ClangToolChain), newVersionNumber);
+						Assert.IsFalse(args.Any());
+
+						deterministicCompileEnvironment.bDeterministic = true;
+						args = baseWarnings.GenerateWarningCommandLineArgs(deterministicCompileEnvironment, typeof(ClangToolChain), newVersionNumber);
+
+						Assert.IsTrue(args.Count() == 2);
+						Assert.IsTrue(args.Contains("-Wdate-time"));
+						Assert.IsTrue(args.Contains("-Wno-error=date-time"));
+
+						baseWarnings.DeterministicWarningLevel = WarningLevel.Error;
+						args = baseWarnings.GenerateWarningCommandLineArgs(deterministicCompileEnvironment, typeof(ClangToolChain), newVersionNumber);
+						Assert.IsTrue(args.Count() == 1);
+						Assert.IsTrue(args.Contains("-Wdate-time"));
+
+						baseWarnings.DeterministicWarningLevel = WarningLevel.Default;
+					}
+
+					// PGOOptimizedFilter
+					{
+						testTargetRules.WindowsPlatform.Compiler = WindowsCompiler.Clang;
+						baseWarnings.ProfileInstructWarningLevel = WarningLevel.Off;
+						VersionNumber newVersionNumber = new VersionNumber(18, 9, 9);
+
+						CppCompileEnvironment pgoOptimizedCompileEnvironment = new CppCompileEnvironment(UnrealTargetPlatform.Win64, CppConfiguration.Debug, new UnrealArchitectures(UnrealArch.X64), null!);
+						pgoOptimizedCompileEnvironment.bPGOOptimize = false;
+
+						args = baseWarnings.GenerateWarningCommandLineArgs(pgoOptimizedCompileEnvironment, typeof(ClangToolChain), newVersionNumber);
+						Assert.IsFalse(args.Any());
+
+						pgoOptimizedCompileEnvironment.bPGOOptimize = true;
+
+						args = baseWarnings.GenerateWarningCommandLineArgs(pgoOptimizedCompileEnvironment, typeof(ClangToolChain), newVersionNumber);
+
+						Assert.IsTrue(args.Count() == 2);
+						Assert.IsTrue(args.Contains("-Wno-profile-instr-out-of-date"));
+						Assert.IsTrue(args.Contains("-Wno-profile-instr-unprofiled"));
+
+						baseWarnings.ProfileInstructWarningLevel = WarningLevel.Off;
+					}
+
+					// AndroidNDKR28ToolChainVersionExclusion
+					{
+						testTargetRules.WindowsPlatform.Compiler = WindowsCompiler.Clang;
+						baseWarnings.CastFunctionTypeMismatchWarningLevel = WarningLevel.Off;
+
+						VersionNumber newVersionNumber = new VersionNumber(19, 0, 0);
+						CppCompileEnvironment androidCompileEnvironment = new CppCompileEnvironment(UnrealTargetPlatform.Android, CppConfiguration.Development, new UnrealArchitectures(UnrealArch.X64), null!);
+						args = baseWarnings.GenerateWarningCommandLineArgs(androidCompileEnvironment, typeof(ClangToolChain), newVersionNumber);
+
+						Assert.IsFalse(args.Any());
+
+						args = baseWarnings.GenerateWarningCommandLineArgs(newCompileEnvironment, typeof(ClangToolChain), newVersionNumber);
+
+						Assert.IsTrue(args.Count() == 1);
+						Assert.IsTrue(args.Contains("-Wno-cast-function-type-mismatch"));
+
+						newVersionNumber = new VersionNumber(19, 0, 1);
+						args = baseWarnings.GenerateWarningCommandLineArgs(androidCompileEnvironment, typeof(ClangToolChain), newVersionNumber);
+
+						Assert.IsTrue(args.Count() == 1);
+						Assert.IsTrue(args.Contains("-Wno-cast-function-type-mismatch"));
+
+						baseWarnings.CastFunctionTypeMismatchWarningLevel = WarningLevel.Default;
+					}
+
+					// AndroidNDKR26ToolChainVersionExclusion
+					{
+						testTargetRules.WindowsPlatform.Compiler = WindowsCompiler.Clang;
+						baseWarnings.InvalidUnevaluatedStringWarningLevel = WarningLevel.Off;
+						CppCompileEnvironment androidCompileEnvironment = new CppCompileEnvironment(UnrealTargetPlatform.Android, CppConfiguration.Development, new UnrealArchitectures(UnrealArch.X64), null!);
+
+						// Exclusion cases
+						{
+							VersionNumber newVersionNumber = new VersionNumber(17, 0, 2);
+							args = baseWarnings.GenerateWarningCommandLineArgs(androidCompileEnvironment, typeof(ClangToolChain), newVersionNumber);
+
+							Assert.IsTrue(args.Count() == 1 && args.Contains("-Wno-shadow"));
+
+							newVersionNumber = new VersionNumber(16);
+							args = baseWarnings.GenerateWarningCommandLineArgs(androidCompileEnvironment, typeof(ClangToolChain), newVersionNumber);
+
+							Assert.IsFalse(args.Any());
+
+							newVersionNumber = new VersionNumber(17, 1);
+							androidCompileEnvironment.CppStandard = CppStandardVersion.Latest;
+							args = baseWarnings.GenerateWarningCommandLineArgs(androidCompileEnvironment, typeof(ClangToolChain), newVersionNumber);
+
+							Assert.IsTrue(args.Count() == 1 && args.Contains("-Wno-shadow"));
+						}
+
+						// Inclusion cases
+						{
+							VersionNumber newVersionNumber = new VersionNumber(17, 1);
+							androidCompileEnvironment.CppStandard = CppStandardVersion.Cpp20;
+
+							args = baseWarnings.GenerateWarningCommandLineArgs(androidCompileEnvironment, typeof(ClangToolChain), newVersionNumber);
+
+							Assert.IsTrue(args.Contains("-Wno-invalid-unevaluated-string"));
+
+							newVersionNumber = new VersionNumber(17, 0, 2);
+
+							args = baseWarnings.GenerateWarningCommandLineArgs(newCompileEnvironment, typeof(ClangToolChain), newVersionNumber);
+
+							Assert.IsTrue(args.Contains("-Wno-invalid-unevaluated-string"));
+						}
+
+						baseWarnings.InvalidUnevaluatedStringWarningLevel = WarningLevel.Default;
+					}
+
+					// UnusedValueClangToolChainAttribute
+					{
+						testTargetRules.WindowsPlatform.Compiler = WindowsCompiler.Clang;
+						baseWarnings.UnusedValueWarningLevel = WarningLevel.Off;
+						VersionNumber newVersionNumber = new VersionNumber(18, 9, 9);
+
+						CppCompileEnvironment invalidCompileEnvironment = new CppCompileEnvironment(UnrealTargetPlatform.Win64, CppConfiguration.Debug, new UnrealArchitectures(UnrealArch.X64), null!);
+						CppCompileEnvironment validCompileEnvironment = new CppCompileEnvironment(UnrealTargetPlatform.Win64, CppConfiguration.Shipping, new UnrealArchitectures(UnrealArch.X64), null!);
+
+						args = baseWarnings.GenerateWarningCommandLineArgs(invalidCompileEnvironment, typeof(ClangToolChain), newVersionNumber, StaticAnalyzer.None);
+						Assert.IsFalse(args.Any());
+
+						args = baseWarnings.GenerateWarningCommandLineArgs(validCompileEnvironment, typeof(ClangToolChain), newVersionNumber, StaticAnalyzer.None);
+						Assert.IsTrue(args.Count() == 2 && args.Contains("-Wno-unused-value"));
+
+						args = baseWarnings.GenerateWarningCommandLineArgs(validCompileEnvironment, typeof(ClangToolChain), newVersionNumber, StaticAnalyzer.Clang);
+						Assert.IsTrue(args.Count() == 2 && args.Contains("-Wno-unused-value"));
+					}
 				}
+			}
+		}
 
-				// Clang toolchain version check - no version; should disable Shadow
-				{
-					testTargetRules.WindowsPlatform.Compiler = WindowsCompiler.Clang;
-					baseWarnings.ShadowVariableWarningLevel = WarningLevel.Error;
+		[TestMethod]
+		public void TestConfigurationFilterCommandLineArgs()
+		{
+			DirectoryReference programsDirectory = DirectoryReference.Combine(Unreal.EngineSourceDirectory, "Programs");
+			FileReference testStubsDirectory = FileReference.Combine(programsDirectory, "Shared", "UnrealBuildTool.Tests", "UBT");
 
-					args = baseWarnings.GenerateWarningCommandLineArgs(newCompileEnvironment, typeof(ClangToolChain));
-					Assert.IsTrue(args.Count() == 1);
-					Assert.IsTrue(args.Contains("-Wno-shadow"));
+			// Generate ephemeral TargetRules & uproject files
+			using (SafeTestDirectory testRootDirectory = SafeTestDirectory.CreateTestDirectory(Path.Combine(StaticTargetRulesInitializer.UBTTestFolderRoot, $"{StaticTargetRulesInitializer.UBTTestFolderPrefix}{Guid.NewGuid().ToString()}")))
+			using (SafeTestDirectory testSourceDirectory = SafeTestDirectory.CreateTestDirectory(Path.Combine(testRootDirectory.TemporaryDirectory, "Source")))
+			using (SafeTestFile testEditorTargetTestFile = new SafeTestFile(File.ReadAllText(Path.Combine(testStubsDirectory.FullName, "TargetStubs", "TestStubEditor.Target.ubttest")), "TestStubEditor.Target.cs", testSourceDirectory.TemporaryDirectory))
+			using (SafeTestFile testTargetTestFile = new SafeTestFile(File.ReadAllText(Path.Combine(testStubsDirectory.FullName, "TargetStubs", "TestStub.Target.ubttest")), "TestStub.Target.cs", testSourceDirectory.TemporaryDirectory))
+			using (SafeTestFile testUprojectTestFile = new SafeTestFile(File.ReadAllText(Path.Combine(testStubsDirectory.FullName, "UprojectStubs", "TestStub.uproject.ubttest")), "TestStub.uproject", testRootDirectory.TemporaryDirectory))
+			{
+				FileReference testUproject = new FileReference(testUprojectTestFile.TemporaryFile);
+				RulesAssembly generatedTestRulesAssembly = RulesCompiler.CreateTargetRulesAssembly(testUproject, "TestStubEditor", false, false, false, null, false, new TestLogger());
+				TargetRules testTargetRules = generatedTestRulesAssembly.CreateTargetRules("TestStubEditor", UnrealTargetPlatform.Win64, UnrealTargetConfiguration.Development, null, testUproject, null, new TestLogger());
+				testTargetRules.WindowsPlatform.Compiler = WindowsCompiler.VisualStudio2022;
+				CppCompileWarnings baseWarnings = new CppCompileWarnings(testTargetRules, null);
 
-					baseWarnings.ShadowVariableWarningLevel = WarningLevel.Default;
-				}
+				CppCompileEnvironment newCompileEnvironment = new CppCompileEnvironment(UnrealTargetPlatform.Win64, CppConfiguration.Development, new UnrealArchitectures(UnrealArch.X64), null!);
+				IEnumerable<string> args = null!;
 
-				// Clang toolchain version check - valid versions;
-				{
-					testTargetRules.WindowsPlatform.Compiler = WindowsCompiler.Clang;
-					baseWarnings.ShadowVariableWarningLevel = WarningLevel.Error;
-
-					VersionNumber newVersionNumber = new VersionNumber(16, 9, 5);
-
-					args = baseWarnings.GenerateWarningCommandLineArgs(newCompileEnvironment, typeof(ClangToolChain), newVersionNumber);
-					Assert.IsTrue(args.Count() == 1);
-					Assert.IsTrue(args.Contains("-Wshadow"));
-
-					newVersionNumber = new VersionNumber(17);
-					args = baseWarnings.GenerateWarningCommandLineArgs(newCompileEnvironment, typeof(ClangToolChain), newVersionNumber);
-					Assert.IsTrue(args.Count() == 1);
-					Assert.IsTrue(args.Contains("-Wno-shadow"));
-
-					newVersionNumber = new VersionNumber(17, 1);
-					args = baseWarnings.GenerateWarningCommandLineArgs(newCompileEnvironment, typeof(ClangToolChain), newVersionNumber);
-					Assert.IsTrue(args.Count() == 1);
-					Assert.IsTrue(args.Contains("-Wno-shadow"));
-
-					newVersionNumber = new VersionNumber(18, 1, 2);
-					args = baseWarnings.GenerateWarningCommandLineArgs(newCompileEnvironment, typeof(ClangToolChain), newVersionNumber);
-					Assert.IsTrue(args.Count() == 1);
-					Assert.IsTrue(args.Contains("-Wno-shadow"));
-
-					newVersionNumber = new VersionNumber(18, 1, 3);
-					args = baseWarnings.GenerateWarningCommandLineArgs(newCompileEnvironment, typeof(ClangToolChain), newVersionNumber);
-					Assert.IsTrue(args.Count() == 1);
-					Assert.IsTrue(args.Contains("-Wshadow"));
-
-					baseWarnings.ShadowVariableWarningLevel = WarningLevel.Default;
-				}
-
-				// ToolChainVersionConstrainedFilterAttribute
-				{
-					testTargetRules.WindowsPlatform.Compiler = WindowsCompiler.Clang;
-					baseWarnings.Shorten64To32WarningLevel = WarningLevel.Off;
-
-					VersionNumber newVersionNumber = new VersionNumber(18, 9, 9);
-					args = baseWarnings.GenerateWarningCommandLineArgs(newCompileEnvironment, typeof(ClangToolChain), newVersionNumber);
-					Assert.IsFalse(args.Any());
-
-					//
-					newVersionNumber = new VersionNumber(19, 0, 0);
-					args = baseWarnings.GenerateWarningCommandLineArgs(newCompileEnvironment, typeof(ClangToolChain), newVersionNumber);
-					Assert.IsTrue(args.Count() == 1);
-					Assert.IsTrue(args.Contains("-Wno-shorten-64-to-32"));
-					baseWarnings.Shorten64To32WarningLevel = WarningLevel.Default;
-				}
-
 				// ConfigurationFilterAttribute
 				{
@@ -494,6 +576,30 @@
 
 					baseWarnings.UnusedValueWarningLevel = WarningLevel.Default;
 				}
+			}
+		}
 
+		[TestMethod]
+		public void TestCppStandardConstraintCommandLineArgs()
+		{
+			DirectoryReference programsDirectory = DirectoryReference.Combine(Unreal.EngineSourceDirectory, "Programs");
+			FileReference testStubsDirectory = FileReference.Combine(programsDirectory, "Shared", "UnrealBuildTool.Tests", "UBT");
+
+			// Generate ephemeral TargetRules & uproject files
+			using (SafeTestDirectory testRootDirectory = SafeTestDirectory.CreateTestDirectory(Path.Combine(StaticTargetRulesInitializer.UBTTestFolderRoot, $"{StaticTargetRulesInitializer.UBTTestFolderPrefix}{Guid.NewGuid().ToString()}")))
+			using (SafeTestDirectory testSourceDirectory = SafeTestDirectory.CreateTestDirectory(Path.Combine(testRootDirectory.TemporaryDirectory, "Source")))
+			using (SafeTestFile testEditorTargetTestFile = new SafeTestFile(File.ReadAllText(Path.Combine(testStubsDirectory.FullName, "TargetStubs", "TestStubEditor.Target.ubttest")), "TestStubEditor.Target.cs", testSourceDirectory.TemporaryDirectory))
+			using (SafeTestFile testTargetTestFile = new SafeTestFile(File.ReadAllText(Path.Combine(testStubsDirectory.FullName, "TargetStubs", "TestStub.Target.ubttest")), "TestStub.Target.cs", testSourceDirectory.TemporaryDirectory))
+			using (SafeTestFile testUprojectTestFile = new SafeTestFile(File.ReadAllText(Path.Combine(testStubsDirectory.FullName, "UprojectStubs", "TestStub.uproject.ubttest")), "TestStub.uproject", testRootDirectory.TemporaryDirectory))
+			{
+				FileReference testUproject = new FileReference(testUprojectTestFile.TemporaryFile);
+				RulesAssembly generatedTestRulesAssembly = RulesCompiler.CreateTargetRulesAssembly(testUproject, "TestStubEditor", false, false, false, null, false, new TestLogger());
+				TargetRules testTargetRules = generatedTestRulesAssembly.CreateTargetRules("TestStubEditor", UnrealTargetPlatform.Win64, UnrealTargetConfiguration.Development, null, testUproject, null, new TestLogger());
+				testTargetRules.WindowsPlatform.Compiler = WindowsCompiler.VisualStudio2022;
+				CppCompileWarnings baseWarnings = new CppCompileWarnings(testTargetRules, null);
+
+				CppCompileEnvironment newCompileEnvironment = new CppCompileEnvironment(UnrealTargetPlatform.Win64, CppConfiguration.Development, new UnrealArchitectures(UnrealArch.X64), null!);
+				IEnumerable<string> args = null!;
+
 				// CppStandardConstrainedFilterAttribute
 				{
@@ -521,154 +627,171 @@
 
 					baseWarnings.AmbiguousReversedOperatorWarningLevel = WarningLevel.Default;
 				}
+			}
+		}
 
-				// DeterministcFlagSetFilter
+		[TestMethod]
+		public void TestVersionConstraintCommandLineArgs()
+		{
+			DirectoryReference programsDirectory = DirectoryReference.Combine(Unreal.EngineSourceDirectory, "Programs");
+			FileReference testStubsDirectory = FileReference.Combine(programsDirectory, "Shared", "UnrealBuildTool.Tests", "UBT");
+
+			// Generate ephemeral TargetRules & uproject files
+			using (SafeTestDirectory testRootDirectory = SafeTestDirectory.CreateTestDirectory(Path.Combine(StaticTargetRulesInitializer.UBTTestFolderRoot, $"{StaticTargetRulesInitializer.UBTTestFolderPrefix}{Guid.NewGuid().ToString()}")))
+			using (SafeTestDirectory testSourceDirectory = SafeTestDirectory.CreateTestDirectory(Path.Combine(testRootDirectory.TemporaryDirectory, "Source")))
+			using (SafeTestFile testEditorTargetTestFile = new SafeTestFile(File.ReadAllText(Path.Combine(testStubsDirectory.FullName, "TargetStubs", "TestStubEditor.Target.ubttest")), "TestStubEditor.Target.cs", testSourceDirectory.TemporaryDirectory))
+			using (SafeTestFile testTargetTestFile = new SafeTestFile(File.ReadAllText(Path.Combine(testStubsDirectory.FullName, "TargetStubs", "TestStub.Target.ubttest")), "TestStub.Target.cs", testSourceDirectory.TemporaryDirectory))
+			using (SafeTestFile testUprojectTestFile = new SafeTestFile(File.ReadAllText(Path.Combine(testStubsDirectory.FullName, "UprojectStubs", "TestStub.uproject.ubttest")), "TestStub.uproject", testRootDirectory.TemporaryDirectory))
+			{
+				FileReference testUproject = new FileReference(testUprojectTestFile.TemporaryFile);
+				RulesAssembly generatedTestRulesAssembly = RulesCompiler.CreateTargetRulesAssembly(testUproject, "TestStubEditor", false, false, false, null, false, new TestLogger());
+				TargetRules testTargetRules = generatedTestRulesAssembly.CreateTargetRules("TestStubEditor", UnrealTargetPlatform.Win64, UnrealTargetConfiguration.Development, null, testUproject, null, new TestLogger());
+				testTargetRules.WindowsPlatform.Compiler = WindowsCompiler.VisualStudio2022;
+				CppCompileWarnings baseWarnings = new CppCompileWarnings(testTargetRules, null);
+
+				CppCompileEnvironment newCompileEnvironment = new CppCompileEnvironment(UnrealTargetPlatform.Win64, CppConfiguration.Development, new UnrealArchitectures(UnrealArch.X64), null!);
+				IEnumerable<string> args = null!;
+
+				// Specialized VC toolchain check;  UnsafeTypeCast
+				{
+					RulesAssembly generatedTestGameRulesAssembly = RulesCompiler.CreateTargetRulesAssembly(testUproject, "TestStub", false, false, false, null, false, new TestLogger());
+					TargetRules testTargetRulesNonWin64 = generatedTestGameRulesAssembly.CreateTargetRules("TestStub", UnrealTargetPlatform.Android, UnrealTargetConfiguration.Development, null, testUproject, null, new TestLogger());
+					CppCompileWarnings specializedWarnings = new CppCompileWarnings(testTargetRulesNonWin64, null);
+
+					specializedWarnings.UnsafeTypeCastWarningLevel = WarningLevel.Warning;
+					args = specializedWarnings.GenerateWarningCommandLineArgs(newCompileEnvironment, typeof(VCToolChain));
+
+					Assert.IsTrue(args.Count() == 2);
+					Assert.IsTrue(args.Contains("/wd4244"));
+					Assert.IsTrue(args.Contains("/wd4838"));
+
+					baseWarnings.UnsafeTypeCastWarningLevel = WarningLevel.Warning;
+					args = baseWarnings.GenerateWarningCommandLineArgs(newCompileEnvironment, typeof(VCToolChain));
+
+					Assert.IsTrue(args.Count() == 2);
+					Assert.IsTrue(args.Contains("/w44244"));
+					Assert.IsTrue(args.Contains("/w44838"));
+
+					baseWarnings.UnsafeTypeCastWarningLevel = WarningLevel.Default;
+				}
+
+				// Clang toolchain version check - no version; should disable Shadow
 				{
 					testTargetRules.WindowsPlatform.Compiler = WindowsCompiler.Clang;
-					baseWarnings.DeterministicWarningLevel = WarningLevel.Off;
+					baseWarnings.ShadowVariableWarningLevel = WarningLevel.Error;
+
+					args = baseWarnings.GenerateWarningCommandLineArgs(newCompileEnvironment, typeof(ClangToolChain));
+					Assert.IsTrue(args.Count() == 1);
+					Assert.IsTrue(args.Contains("-Wno-shadow"));
+
+					baseWarnings.ShadowVariableWarningLevel = WarningLevel.Default;
+				}
+
+				// Clang toolchain version check - valid versions;
+				{
+					testTargetRules.WindowsPlatform.Compiler = WindowsCompiler.Clang;
+					baseWarnings.ShadowVariableWarningLevel = WarningLevel.Error;
+
+					VersionNumber newVersionNumber = new VersionNumber(16, 9, 5);
+
+					args = baseWarnings.GenerateWarningCommandLineArgs(newCompileEnvironment, typeof(ClangToolChain), newVersionNumber);
+					Assert.IsTrue(args.Count() == 1);
+					Assert.IsTrue(args.Contains("-Wshadow"));
+
+					newVersionNumber = new VersionNumber(17);
+					args = baseWarnings.GenerateWarningCommandLineArgs(newCompileEnvironment, typeof(ClangToolChain), newVersionNumber);
+					Assert.IsTrue(args.Count() == 1);
+					Assert.IsTrue(args.Contains("-Wno-shadow"));
+
+					newVersionNumber = new VersionNumber(17, 1);
+					args = baseWarnings.GenerateWarningCommandLineArgs(newCompileEnvironment, typeof(ClangToolChain), newVersionNumber);
+					Assert.IsTrue(args.Count() == 1);
+					Assert.IsTrue(args.Contains("-Wno-shadow"));
+
+					newVersionNumber = new VersionNumber(18, 1, 2);
+					args = baseWarnings.GenerateWarningCommandLineArgs(newCompileEnvironment, typeof(ClangToolChain), newVersionNumber);
+					Assert.IsTrue(args.Count() == 1);
+					Assert.IsTrue(args.Contains("-Wno-shadow"));
+
+					newVersionNumber = new VersionNumber(18, 1, 3);
+					args = baseWarnings.GenerateWarningCommandLineArgs(newCompileEnvironment, typeof(ClangToolChain), newVersionNumber);
+					Assert.IsTrue(args.Count() == 1);
+					Assert.IsTrue(args.Contains("-Wshadow"));
+
+					baseWarnings.ShadowVariableWarningLevel = WarningLevel.Default;
+				}
+
+				// ToolChainVersionConstrainedFilterAttribute
+				{
+					testTargetRules.WindowsPlatform.Compiler = WindowsCompiler.Clang;
+					baseWarnings.Shorten64To32WarningLevel = WarningLevel.Off;
+
 					VersionNumber newVersionNumber = new VersionNumber(18, 9, 9);
-
-					CppCompileEnvironment deterministicCompileEnvironment = new CppCompileEnvironment(UnrealTargetPlatform.Win64, CppConfiguration.Debug, new UnrealArchitectures(UnrealArch.X64), null!);
-					deterministicCompileEnvironment.bDeterministic = false;
-
-					args = baseWarnings.GenerateWarningCommandLineArgs(deterministicCompileEnvironment, typeof(ClangToolChain), newVersionNumber);
+					args = baseWarnings.GenerateWarningCommandLineArgs(newCompileEnvironment, typeof(ClangToolChain), newVersionNumber);
 					Assert.IsFalse(args.Any());
 
-					deterministicCompileEnvironment.bDeterministic = true;
+					//
+					newVersionNumber = new VersionNumber(19, 0, 0);
+					args = baseWarnings.GenerateWarningCommandLineArgs(newCompileEnvironment, typeof(ClangToolChain), newVersionNumber);
+					Assert.IsTrue(args.Count() == 1);
+					Assert.IsTrue(args.Contains("-Wno-shorten-64-to-32"));
+					baseWarnings.Shorten64To32WarningLevel = WarningLevel.Default;
+				}
+			}
+		}
 
-					args = baseWarnings.GenerateWarningCommandLineArgs(deterministicCompileEnvironment, typeof(ClangToolChain), newVersionNumber);
-					Assert.IsTrue(args.Count() == 1);
-					Assert.IsTrue(args.Contains("-Wno-date-time"));
+		[TestMethod]
+		public void TestMSVCAndVCClangGenerateWarningCommandLineArgs()
+		{
+			DirectoryReference programsDirectory = DirectoryReference.Combine(Unreal.EngineSourceDirectory, "Programs");
+			FileReference testStubsDirectory = FileReference.Combine(programsDirectory, "Shared", "UnrealBuildTool.Tests", "UBT");
 
-					deterministicCompileEnvironment.bDeterministic = false;
-					baseWarnings.DeterministicWarningLevel = WarningLevel.Warning;
-					args = baseWarnings.GenerateWarningCommandLineArgs(deterministicCompileEnvironment, typeof(ClangToolChain), newVersionNumber);
-					Assert.IsFalse(args.Any());
+			// Generate ephemeral TargetRules & uproject files
+			using (SafeTestDirectory testRootDirectory = SafeTestDirectory.CreateTestDirectory(Path.Combine(StaticTargetRulesInitializer.UBTTestFolderRoot, $"{StaticTargetRulesInitializer.UBTTestFolderPrefix}{Guid.NewGuid().ToString()}")))
+			using (SafeTestDirectory testSourceDirectory = SafeTestDirectory.CreateTestDirectory(Path.Combine(testRootDirectory.TemporaryDirectory, "Source")))
+			using (SafeTestFile testEditorTargetTestFile = new SafeTestFile(File.ReadAllText(Path.Combine(testStubsDirectory.FullName, "TargetStubs", "TestStubEditor.Target.ubttest")), "TestStubEditor.Target.cs", testSourceDirectory.TemporaryDirectory))
+			using (SafeTestFile testTargetTestFile = new SafeTestFile(File.ReadAllText(Path.Combine(testStubsDirectory.FullName, "TargetStubs", "TestStub.Target.ubttest")), "TestStub.Target.cs", testSourceDirectory.TemporaryDirectory))
+			using (SafeTestFile testUprojectTestFile = new SafeTestFile(File.ReadAllText(Path.Combine(testStubsDirectory.FullName, "UprojectStubs", "TestStub.uproject.ubttest")), "TestStub.uproject", testRootDirectory.TemporaryDirectory))
+			{
+				FileReference testUproject = new FileReference(testUprojectTestFile.TemporaryFile);
+				RulesAssembly generatedTestRulesAssembly = RulesCompiler.CreateTargetRulesAssembly(testUproject, "TestStubEditor", false, false, false, null, false, new TestLogger());
+				TargetRules testTargetRules = generatedTestRulesAssembly.CreateTargetRules("TestStubEditor", UnrealTargetPlatform.Win64, UnrealTargetConfiguration.Development, null, testUproject, null, new TestLogger());
+				testTargetRules.WindowsPlatform.Compiler = WindowsCompiler.VisualStudio2022;
+				CppCompileWarnings baseWarnings = new CppCompileWarnings(testTargetRules, null);
 
-					deterministicCompileEnvironment.bDeterministic = true;
-					args = baseWarnings.GenerateWarningCommandLineArgs(deterministicCompileEnvironment, typeof(ClangToolChain), newVersionNumber);
+				CppCompileEnvironment newCompileEnvironment = new CppCompileEnvironment(UnrealTargetPlatform.Win64, CppConfiguration.Development, new UnrealArchitectures(UnrealArch.X64), null!);
+				IEnumerable<string> args = null!;
 
+				// VC toolchain check;
+				{
+					baseWarnings.UnsafeTypeCastWarningLevel = WarningLevel.Warning;
+					args = baseWarnings.GenerateWarningCommandLineArgs(newCompileEnvironment, typeof(VCToolChain));
 					Assert.IsTrue(args.Count() == 2);
-					Assert.IsTrue(args.Contains("-Wdate-time"));
-					Assert.IsTrue(args.Contains("-Wno-error=date-time"));
+					Assert.IsTrue(args.Contains("/w44244"));
+					Assert.IsTrue(args.Contains("/w44838"));
 
-					baseWarnings.DeterministicWarningLevel = WarningLevel.Error;
-					args = baseWarnings.GenerateWarningCommandLineArgs(deterministicCompileEnvironment, typeof(ClangToolChain), newVersionNumber);
-					Assert.IsTrue(args.Count() == 1);
-					Assert.IsTrue(args.Contains("-Wdate-time"));
-
-					baseWarnings.DeterministicWarningLevel = WarningLevel.Default;
+					baseWarnings.UnsafeTypeCastWarningLevel = WarningLevel.Default;
 				}
 
-				// PGOOptimizedFilter
+				// MSVC compiler check;
 				{
-					testTargetRules.WindowsPlatform.Compiler = WindowsCompiler.Clang;
-					baseWarnings.ProfileInstructWarningLevel = WarningLevel.Off;
-					VersionNumber newVersionNumber = new VersionNumber(18, 9, 9);
+					baseWarnings.ShadowVariableWarningLevel = WarningLevel.Error;
+					args = baseWarnings.GenerateWarningCommandLineArgs(newCompileEnvironment, typeof(VCToolChain));
 
-					CppCompileEnvironment pgoOptimizedCompileEnvironment = new CppCompileEnvironment(UnrealTargetPlatform.Win64, CppConfiguration.Debug, new UnrealArchitectures(UnrealArch.X64), null!);
-					pgoOptimizedCompileEnvironment.bPGOOptimize = false;
+					Assert.IsTrue(args.Count() == 3);
+					Assert.IsTrue(args.Contains("/we4456"));
+					Assert.IsTrue(args.Contains("/we4458"));
+					Assert.IsTrue(args.Contains("/we4459"));
 
-					args = baseWarnings.GenerateWarningCommandLineArgs(pgoOptimizedCompileEnvironment, typeof(ClangToolChain), newVersionNumber);
-					Assert.IsFalse(args.Any());
+					WindowsCompiler previousCompiler = testTargetRules.WindowsPlatform.Compiler;
+					testTargetRules.WindowsPlatform.Compiler = WindowsCompiler.Intel;
 
-					pgoOptimizedCompileEnvironment.bPGOOptimize = true;
-
-					args = baseWarnings.GenerateWarningCommandLineArgs(pgoOptimizedCompileEnvironment, typeof(ClangToolChain), newVersionNumber);
-
-					Assert.IsTrue(args.Count() == 2);
-					Assert.IsTrue(args.Contains("-Wno-profile-instr-out-of-date"));
-					Assert.IsTrue(args.Contains("-Wno-profile-instr-unprofiled"));
-
-					baseWarnings.ProfileInstructWarningLevel = WarningLevel.Off;
-				}
-
-				// AndroidNDKR28ToolChainVersionExclusion
-				{
-					testTargetRules.WindowsPlatform.Compiler = WindowsCompiler.Clang;
-					baseWarnings.CastFunctionTypeMismatchWarningLevel = WarningLevel.Off;
-
-					VersionNumber newVersionNumber = new VersionNumber(19, 0, 0);
-					CppCompileEnvironment androidCompileEnvironment = new CppCompileEnvironment(UnrealTargetPlatform.Android, CppConfiguration.Development, new UnrealArchitectures(UnrealArch.X64), null!);
-					args = baseWarnings.GenerateWarningCommandLineArgs(androidCompileEnvironment, typeof(ClangToolChain), newVersionNumber);
+					args = baseWarnings.GenerateWarningCommandLineArgs(newCompileEnvironment, typeof(VCToolChain));
 
 					Assert.IsFalse(args.Any());
 
-					args = baseWarnings.GenerateWarningCommandLineArgs(newCompileEnvironment, typeof(ClangToolChain), newVersionNumber);
-
-					Assert.IsTrue(args.Count() == 1);
-					Assert.IsTrue(args.Contains("-Wno-cast-function-type-mismatch"));
-
-					newVersionNumber = new VersionNumber(19, 0, 1);
-					args = baseWarnings.GenerateWarningCommandLineArgs(androidCompileEnvironment, typeof(ClangToolChain), newVersionNumber);
-
-					Assert.IsTrue(args.Count() == 1);
-					Assert.IsTrue(args.Contains("-Wno-cast-function-type-mismatch"));
-
-					baseWarnings.CastFunctionTypeMismatchWarningLevel = WarningLevel.Default;
+					testTargetRules.WindowsPlatform.Compiler = previousCompiler;
+					baseWarnings.ShadowVariableWarningLevel = WarningLevel.Default;
 				}
 
-				// AndroidNDKR26ToolChainVersionExclusion
-				{
-					testTargetRules.WindowsPlatform.Compiler = WindowsCompiler.Clang;
-					baseWarnings.InvalidUnevaluatedStringWarningLevel = WarningLevel.Off;
-					CppCompileEnvironment androidCompileEnvironment = new CppCompileEnvironment(UnrealTargetPlatform.Android, CppConfiguration.Development, new UnrealArchitectures(UnrealArch.X64), null!);
-
-					// Exclusion cases
-					{
-						VersionNumber newVersionNumber = new VersionNumber(17, 0, 2);
-						args = baseWarnings.GenerateWarningCommandLineArgs(androidCompileEnvironment, typeof(ClangToolChain), newVersionNumber);
-
-						Assert.IsTrue(args.Count() == 1 && args.Contains("-Wno-shadow"));
-
-						newVersionNumber = new VersionNumber(16);
-						args = baseWarnings.GenerateWarningCommandLineArgs(androidCompileEnvironment, typeof(ClangToolChain), newVersionNumber);
-
-						Assert.IsFalse(args.Any());
-
-						newVersionNumber = new VersionNumber(17, 1);
-						androidCompileEnvironment.CppStandard = CppStandardVersion.Latest;
-						args = baseWarnings.GenerateWarningCommandLineArgs(androidCompileEnvironment, typeof(ClangToolChain), newVersionNumber);
-
-						Assert.IsTrue(args.Count() == 1 && args.Contains("-Wno-shadow"));
-					}
-
-					// Inclusion cases
-					{
-						VersionNumber newVersionNumber = new VersionNumber(17, 1);
-						androidCompileEnvironment.CppStandard = CppStandardVersion.Cpp20;
-
-						args = baseWarnings.GenerateWarningCommandLineArgs(androidCompileEnvironment, typeof(ClangToolChain), newVersionNumber);
-
-						Assert.IsTrue(args.Contains("-Wno-invalid-unevaluated-string"));
-
-						newVersionNumber = new VersionNumber(17, 0, 2);
-
-						args = baseWarnings.GenerateWarningCommandLineArgs(newCompileEnvironment, typeof(ClangToolChain), newVersionNumber);
-
-						Assert.IsTrue(args.Contains("-Wno-invalid-unevaluated-string"));
-					}
-
-					baseWarnings.InvalidUnevaluatedStringWarningLevel = WarningLevel.Default;
-				}
-
-				// UnusedValueClangToolChainAttribute
-				{
-					testTargetRules.WindowsPlatform.Compiler = WindowsCompiler.Clang;
-					baseWarnings.UnusedValueWarningLevel = WarningLevel.Off;
-					VersionNumber newVersionNumber = new VersionNumber(18, 9, 9);
-
-					CppCompileEnvironment invalidCompileEnvironment = new CppCompileEnvironment(UnrealTargetPlatform.Win64, CppConfiguration.Debug, new UnrealArchitectures(UnrealArch.X64), null!);
-					CppCompileEnvironment validCompileEnvironment = new CppCompileEnvironment(UnrealTargetPlatform.Win64, CppConfiguration.Shipping, new UnrealArchitectures(UnrealArch.X64), null!);
-
-					args = baseWarnings.GenerateWarningCommandLineArgs(invalidCompileEnvironment, typeof(ClangToolChain), newVersionNumber, StaticAnalyzer.None);
-					Assert.IsFalse(args.Any());
-
-					args = baseWarnings.GenerateWarningCommandLineArgs(validCompileEnvironment, typeof(ClangToolChain), newVersionNumber, StaticAnalyzer.None);
-					Assert.IsTrue(args.Count() == 2 && args.Contains("-Wno-unused-value"));
-
-					args = baseWarnings.GenerateWarningCommandLineArgs(validCompileEnvironment, typeof(ClangToolChain), newVersionNumber, StaticAnalyzer.Clang);
-					Assert.IsTrue(args.Count() == 2 && args.Contains("-Wno-unused-value"));
-				}
-
 				// VCClangCompilerFilter
diff --git a/Engine/Source/Programs/Shared/UnrealBuildTool.Tests/WarningLevelDefaultAttributeTests.cs b/Engine/Source/Programs/Shared/UnrealBuildTool.Tests/WarningLevelDefaultAttributeTests.cs
--- a/Engine/Source/Programs/Shared/UnrealBuildTool.Tests/WarningLevelDefaultAttributeTests.cs
+++ b/Engine/Source/Programs/Shared/UnrealBuildTool.Tests/WarningLevelDefaultAttributeTests.cs
@@ -9,94 +9,105 @@ namespace UnrealBuildTool.Tests
 	[TestClass]
 	public class WarningLevelDefaultAttributeTests
 	{
+		internal class TestBuildSettingsProvider : IBuildContextProvider
+		{
+			internal BuildSettingsVersion _testBuildSettings = BuildSettingsVersion.V1;
+
+			public BuildSettingsVersion GetBuildSettings()
+			{
+				return _testBuildSettings;
+			}
+		}
+
+
 		[TestMethod]
 		public void TestWindowsCompletelyDefined()
 		{
-			List<WarningLevelDefaultAttribute> attributes = new List<WarningLevelDefaultAttribute>
+			List<VersionWarningLevelDefaultAttribute> attributes = new List<VersionWarningLevelDefaultAttribute>
 			{
-				new WarningLevelDefaultAttribute( WarningLevel.Off, BuildSettingsVersion.V1, BuildSettingsVersion.V2),
-				new WarningLevelDefaultAttribute(WarningLevel.Warning, BuildSettingsVersion.V3, BuildSettingsVersion.V4),
-				new WarningLevelDefaultAttribute(WarningLevel.Error, BuildSettingsVersion.V5, BuildSettingsVersion.Latest)
+				new VersionWarningLevelDefaultAttribute( WarningLevel.Off, BuildSettingsVersion.V1, BuildSettingsVersion.V2),
+				new VersionWarningLevelDefaultAttribute(WarningLevel.Warning, BuildSettingsVersion.V3, BuildSettingsVersion.V4),
+				new VersionWarningLevelDefaultAttribute(WarningLevel.Error, BuildSettingsVersion.V5, BuildSettingsVersion.Latest)
 			};
 
-			List<WarningLevelDefaultAttribute> merged = WarningLevelDefaultAttribute.EnsureWarningLevelDefaultBounds(attributes);
+			List<VersionWarningLevelDefaultAttribute> merged = VersionWarningLevelDefaultAttribute.EnsureWarningLevelDefaultBounds(attributes);
 
 			Assert.AreEqual(3, merged.Count);
 			Assert.AreEqual(BuildSettingsVersion.V1, merged[0].MinVersion);
 			Assert.AreEqual(BuildSettingsVersion.V2, merged[0].MaxVersion);
-			Assert.AreEqual(WarningLevel.Off, merged[0].DefaultLevel);
+			Assert.AreEqual(WarningLevel.Off, merged[0].GetDefaultLevel());
 
 			Assert.AreEqual(BuildSettingsVersion.V3, merged[1].MinVersion);
 			Assert.AreEqual(BuildSettingsVersion.V4, merged[1].MaxVersion);
-			Assert.AreEqual(WarningLevel.Warning, merged[1].DefaultLevel);
+			Assert.AreEqual(WarningLevel.Warning, merged[1].GetDefaultLevel());
 
 			Assert.AreEqual(BuildSettingsVersion.V5, merged[2].MinVersion);
 			Assert.AreEqual(BuildSettingsVersion.Latest, merged[2].MaxVersion);
-			Assert.AreEqual(WarningLevel.Error, merged[2].DefaultLevel);
+			Assert.AreEqual(WarningLevel.Error, merged[2].GetDefaultLevel());
 		}
 
 		[TestMethod]
 		public void TestGapResolveToLowerBound()
 		{
-			List<WarningLevelDefaultAttribute> attributes = new List<WarningLevelDefaultAttribute>
+			List<VersionWarningLevelDefaultAttribute> attributes = new List<VersionWarningLevelDefaultAttribute>
 			{
-				new WarningLevelDefaultAttribute(WarningLevel.Warning, BuildSettingsVersion.V1, BuildSettingsVersion.V2),
-				new WarningLevelDefaultAttribute(WarningLevel.Error, BuildSettingsVersion.V4, BuildSettingsVersion.V5),
+				new VersionWarningLevelDefaultAttribute(WarningLevel.Warning, BuildSettingsVersion.V1, BuildSettingsVersion.V2),
+				new VersionWarningLevelDefaultAttribute(WarningLevel.Error, BuildSettingsVersion.V4, BuildSettingsVersion.V5),
 			};
 
-			List<WarningLevelDefaultAttribute> merged = WarningLevelDefaultAttribute.EnsureWarningLevelDefaultBounds(attributes);
+			List<VersionWarningLevelDefaultAttribute> merged = VersionWarningLevelDefaultAttribute.EnsureWarningLevelDefaultBounds(attributes);
 
 			Assert.AreEqual(2, merged.Count);
 			Assert.AreEqual(BuildSettingsVersion.V1, merged[0].MinVersion);
 			Assert.AreEqual(BuildSettingsVersion.V3, merged[0].MaxVersion);
-			Assert.AreEqual(WarningLevel.Warning, merged[0].DefaultLevel);
+			Assert.AreEqual(WarningLevel.Warning, merged[0].GetDefaultLevel());
 
 			Assert.AreEqual(BuildSettingsVersion.V4, merged[1].MinVersion);
 			Assert.AreEqual(BuildSettingsVersion.V5, merged[1].MaxVersion);
-			Assert.AreEqual(WarningLevel.Error, merged[1].DefaultLevel);
+			Assert.AreEqual(WarningLevel.Error, merged[1].GetDefaultLevel());
 		}
 
 		[TestMethod]
 		public void TestLowerUpperOverlapResizeBottomUpward()
 		{
-			List<WarningLevelDefaultAttribute> attributes = new List<WarningLevelDefaultAttribute>
+			List<VersionWarningLevelDefaultAttribute> attributes = new List<VersionWarningLevelDefaultAttribute>
 			{
-				new WarningLevelDefaultAttribute(WarningLevel.Warning, BuildSettingsVersion.V1, BuildSettingsVersion.V3),
-				new WarningLevelDefaultAttribute(WarningLevel.Error, BuildSettingsVersion.V2, BuildSettingsVersion.V4)
+				new VersionWarningLevelDefaultAttribute(WarningLevel.Warning, BuildSettingsVersion.V1, BuildSettingsVersion.V3),
+				new VersionWarningLevelDefaultAttribute(WarningLevel.Error, BuildSettingsVersion.V2, BuildSettingsVersion.V4)
 			};
 
-			List<WarningLevelDefaultAttribute> merged = WarningLevelDefaultAttribute.EnsureWarningLevelDefaultBounds(attributes);
+			List<VersionWarningLevelDefaultAttribute> merged = VersionWarningLevelDefaultAttribute.EnsureWarningLevelDefaultBounds(attributes);
 
 			Assert.AreEqual(2, merged.Count);
 			Assert.AreEqual(BuildSettingsVersion.V1, merged[0].MinVersion);
 			Assert.AreEqual(BuildSettingsVersion.V3, merged[0].MaxVersion);
-			Assert.AreEqual(WarningLevel.Warning, merged[0].DefaultLevel);
+			Assert.AreEqual(WarningLevel.Warning, merged[0].GetDefaultLevel());
 
 			Assert.AreEqual(BuildSettingsVersion.V4, merged[1].MinVersion);
 			Assert.AreEqual(BuildSettingsVersion.V4, merged[1].MaxVersion);
-			Assert.AreEqual(WarningLevel.Error, merged[1].DefaultLevel);
+			Assert.AreEqual(WarningLevel.Error, merged[1].GetDefaultLevel());
 		}
 
 		[TestMethod]
 		public void TestAllRangeWide()
 		{
-			List<WarningLevelDefaultAttribute> attributes = new List<WarningLevelDefaultAttribute>
+			List<VersionWarningLevelDefaultAttribute> attributes = new List<VersionWarningLevelDefaultAttribute>
 			{
-				new WarningLevelDefaultAttribute( WarningLevel.Warning, BuildSettingsVersion.V1, BuildSettingsVersion.V2),
-				new WarningLevelDefaultAttribute(WarningLevel.Error, BuildSettingsVersion.V1, BuildSettingsVersion.V4)
+				new VersionWarningLevelDefaultAttribute( WarningLevel.Warning, BuildSettingsVersion.V1, BuildSettingsVersion.V2),
+				new VersionWarningLevelDefaultAttribute(WarningLevel.Error, BuildSettingsVersion.V1, BuildSettingsVersion.V4)
 			};
 
-			List<WarningLevelDefaultAttribute> merged = WarningLevelDefaultAttribute.EnsureWarningLevelDefaultBounds(attributes);
+			List<VersionWarningLevelDefaultAttribute> merged = VersionWarningLevelDefaultAttribute.EnsureWarningLevelDefaultBounds(attributes);
 
 			Assert.AreEqual(2, merged.Count);
 
 			Assert.AreEqual(BuildSettingsVersion.V1, merged[0].MinVersion);
 			Assert.AreEqual(BuildSettingsVersion.V2, merged[0].MaxVersion);
-			Assert.AreEqual(WarningLevel.Warning, merged[0].DefaultLevel);
+			Assert.AreEqual(WarningLevel.Warning, merged[0].GetDefaultLevel());
 
 			Assert.AreEqual(BuildSettingsVersion.V3, merged[1].MinVersion);
 			Assert.AreEqual(BuildSettingsVersion.V4, merged[1].MaxVersion);
-			Assert.AreEqual(WarningLevel.Error, merged[1].DefaultLevel);
+			Assert.AreEqual(WarningLevel.Error, merged[1].GetDefaultLevel());
 		}
 
 		[TestMethod]
@@ -104,17 +115,23 @@ namespace UnrealBuildTool.Tests
 		{
 			List<WarningLevelDefaultAttribute> attributes = new List<WarningLevelDefaultAttribute>
 			{
-				new WarningLevelDefaultAttribute(WarningLevel.Warning, BuildSettingsVersion.V1, BuildSettingsVersion.V2),
-				new WarningLevelDefaultAttribute(WarningLevel.Error, BuildSettingsVersion.V3, BuildSettingsVersion.V4)
+				new VersionWarningLevelDefaultAttribute(WarningLevel.Warning, BuildSettingsVersion.V1, BuildSettingsVersion.V2),
+				new VersionWarningLevelDefaultAttribute(WarningLevel.Error, BuildSettingsVersion.V3, BuildSettingsVersion.V4)
 			};
+			TestBuildSettingsProvider testBuildSettingsProvider = new TestBuildSettingsProvider();
+			BuildSystemContext newBuildSystemContext = new BuildSystemContext(testBuildSettingsProvider);
 
-			Assert.AreEqual(WarningLevel.Warning, WarningLevelDefaultAttribute.ResolveWarningLevelDefault(attributes, BuildSettingsVersion.V1));
-			Assert.AreEqual(WarningLevel.Warning, WarningLevelDefaultAttribute.ResolveWarningLevelDefault(attributes, BuildSettingsVersion.V2));
+			Assert.AreEqual(WarningLevel.Warning, VersionWarningLevelDefaultAttribute.ResolveVersionWarningLevelDefault(attributes, newBuildSystemContext));
+			testBuildSettingsProvider._testBuildSettings = BuildSettingsVersion.V2;
+			Assert.AreEqual(WarningLevel.Warning, VersionWarningLevelDefaultAttribute.ResolveVersionWarningLevelDefault(attributes, newBuildSystemContext));
 
-			Assert.AreEqual(WarningLevel.Error, WarningLevelDefaultAttribute.ResolveWarningLevelDefault(attributes, BuildSettingsVersion.V3));
-			Assert.AreEqual(WarningLevel.Error, WarningLevelDefaultAttribute.ResolveWarningLevelDefault(attributes, BuildSettingsVersion.V4));
+			testBuildSettingsProvider._testBuildSettings = BuildSettingsVersion.V3;
+			Assert.AreEqual(WarningLevel.Error, VersionWarningLevelDefaultAttribute.ResolveVersionWarningLevelDefault(attributes, newBuildSystemContext));
+			testBuildSettingsProvider._testBuildSettings = BuildSettingsVersion.V4;
+			Assert.AreEqual(WarningLevel.Error, VersionWarningLevelDefaultAttribute.ResolveVersionWarningLevelDefault(attributes, newBuildSystemContext));
 
-			Assert.AreEqual(WarningLevel.Default, WarningLevelDefaultAttribute.ResolveWarningLevelDefault(attributes, BuildSettingsVersion.V5));
+			testBuildSettingsProvider._testBuildSettings = BuildSettingsVersion.V5;
+			Assert.AreEqual(WarningLevel.Default, VersionWarningLevelDefaultAttribute.ResolveVersionWarningLevelDefault(attributes, newBuildSystemContext));
 		}
 	}
 }
\ No newline at end of file
diff --git a/Engine/Source/Programs/UnrealBuildTool/Configuration/CompileWarnings/ApplyWarningsAttribute.cs b/Engine/Source/Programs/UnrealBuildTool/Configuration/CompileWarnings/ApplyWarningsAttribute.cs
--- a/Engine/Source/Programs/UnrealBuildTool/Configuration/CompileWarnings/ApplyWarningsAttribute.cs
+++ b/Engine/Source/Programs/UnrealBuildTool/Configuration/CompileWarnings/ApplyWarningsAttribute.cs
@@ -104,15 +104,15 @@ namespace UnrealBuildTool
 		/// <returns>True if applicable, false otherwise.</returns>
 		internal virtual bool CanApplyToContext(CompilerWarningsToolChainContext toolChainContext)
 		{
-			return (ActiveFilters == null || ActiveFilters.TrueForAll(x => x.CanApply(toolChainContext)));
+			return (ActiveFilters == null || (ActiveFilters.TrueForAll(x => x.CanApply(toolChainContext))));
 		}
 
-		protected bool TryAppendNewFilter(string rawFilterName)
+		protected static bool TryAppendNewFilter(string rawFilterName, List<IApplyWarningsFilter> container)
 		{
 			IApplyWarningsFilter? filter = ApplyWarningsFilterRegistry.RequestFilter(rawFilterName);
-			if (filter != null && !ActiveFilters.Contains(filter))
+			if (filter != null && !container.Contains(filter))
 			{
-				ActiveFilters.Add(filter);
+				container.Add(filter);
 				return true;
 			}
 
@@ -125,7 +125,7 @@ namespace UnrealBuildTool
 			{
 				foreach (string rawFilterName in newFilters)
 				{
-					TryAppendNewFilter(rawFilterName);
+					TryAppendNewFilter(rawFilterName, ActiveFilters);
 				}
 			}
 		}
@@ -136,7 +136,7 @@ namespace UnrealBuildTool
 			{
 				foreach (FilterID filterID in newFilters)
 				{
-					TryAppendNewFilter(filterID.ToString());
+					TryAppendNewFilter(filterID.ToString(), ActiveFilters);
 				}
 			}
 		}
@@ -196,6 +196,17 @@ namespace UnrealBuildTool
 
 	#region -- Compiler Specializations -- 
 
+#pragma warning disable CA1813
+	internal class NonVCClangWarningsClangToolChainAttribute : WarningsClangToolChainAttribute
+	{
+#pragma warning restore CA1813
+		/// <inheritdoc/>
+		internal NonVCClangWarningsClangToolChainAttribute(string[]? offArgs = null, string[]? warningArgs = null, string[]? errorArgs = null, string[]? filterNames = null, ArgPosition argPosition = ArgPosition.Middle) : base(offArgs, warningArgs, errorArgs, filterNames, argPosition)
+		{
+			AppendNewFilters([nameof(StandardFilters.NonVCClangCompilerFilter)]);
+		}
+	}
+
 	/// <summary>
 	/// Specialized <see cref="WarningsVCToolChainAttribute"/> for MSVC compiler within the <see cref="VCToolChain"/>.
 	/// </summary>
diff --git a/Engine/Source/Programs/UnrealBuildTool/Configuration/CompileWarnings/ApplyWarningsFilterRegistry.cs b/Engine/Source/Programs/UnrealBuildTool/Configuration/CompileWarnings/ApplyWarningsFilterRegistry.cs
--- a/Engine/Source/Programs/UnrealBuildTool/Configuration/CompileWarnings/ApplyWarningsFilterRegistry.cs
+++ b/Engine/Source/Programs/UnrealBuildTool/Configuration/CompileWarnings/ApplyWarningsFilterRegistry.cs
@@ -74,7 +74,7 @@ namespace UnrealBuildTool
 
 		private ApplyWarningsFilterRegistry()
 		{
-			// Process all ApplyWarningsFilterAttribute added to this class first, and store in the filterIDCacheDict
+			// Check the entire assembly for the ApplyWarningsFilterAttribute
 			Type filterRegistryType = typeof(ApplyWarningsFilterRegistry);
 
 			IEnumerable<ApplyWarningsFilterAttribute> attributes = filterRegistryType
@@ -89,7 +89,7 @@ namespace UnrealBuildTool
 				}
 			}
 
-			// Process all of the static methods in the assembly that are annotated wtih ApplyWarningsFilterDelegateAttribute, and store in _nameCacheDict given they match
+			// Process all of the static methods in the assembly that are annotated wtih ApplyWarningsFilterDelegateAttribute, and evaluate given they match
 			// the execution delegate, which is effectively IApplyWarningsFilter.CanApply
 			MethodInfo? reflectiveInterfaceMethodInfo = typeof(IApplyWarningsFilter).GetMethod(nameof(IApplyWarningsFilter.CanApply));
 
diff --git a/Engine/Source/Programs/UnrealBuildTool/Configuration/CompileWarnings/StandardApplyWarningsFilterDelegates.cs b/Engine/Source/Programs/UnrealBuildTool/Configuration/CompileWarnings/StandardApplyWarningsFilterDelegates.cs
--- a/Engine/Source/Programs/UnrealBuildTool/Configuration/CompileWarnings/StandardApplyWarningsFilterDelegates.cs
+++ b/Engine/Source/Programs/UnrealBuildTool/Configuration/CompileWarnings/StandardApplyWarningsFilterDelegates.cs
@@ -39,7 +39,7 @@ namespace UnrealBuildTool
 		{
 			ReadOnlyTargetRules? readOnlyTargetRules = context._buildSystemContext.GetReadOnlyTargetRules();
 
-			return readOnlyTargetRules != null && readOnlyTargetRules.WindowsPlatform != null && readOnlyTargetRules.WindowsPlatform.Compiler.IsMSVC();
+			return readOnlyTargetRules?.WindowsPlatform?.Compiler.IsMSVC() == true;
 		}
 
 		/// <summary>
@@ -52,7 +52,7 @@ namespace UnrealBuildTool
 		{
 			ReadOnlyTargetRules? readOnlyTargetRules = context._buildSystemContext.GetReadOnlyTargetRules();
 
-			return readOnlyTargetRules != null && readOnlyTargetRules.WindowsPlatform != null && readOnlyTargetRules.WindowsPlatform.Compiler.IsIntel();
+			return readOnlyTargetRules?.WindowsPlatform?.Compiler.IsIntel() == true;
 		}
 
 		/// <summary>
@@ -65,7 +65,18 @@ namespace UnrealBuildTool
 		{
 			ReadOnlyTargetRules? readOnlyTargetRules = context._buildSystemContext.GetReadOnlyTargetRules();
 
-			return readOnlyTargetRules != null && readOnlyTargetRules.WindowsPlatform != null && readOnlyTargetRules.WindowsPlatform.Compiler.IsClang();
+			return readOnlyTargetRules?.WindowsPlatform?.Compiler.IsClang() == true;
+		}
+
+		/// <summary>
+		/// Negation filter for !VCClang.
+		/// </summary>
+		/// <param name="context">The context in which to evaluate the applicability of the filter.</param>
+		/// <returns>True if the ApplyWarnings will apply, false otherwise.</returns>
+		[ApplyWarningsFilterDelegate(nameof(NonVCClangCompilerFilter))]
+		internal static bool NonVCClangCompilerFilter(CompilerWarningsToolChainContext context)
+		{
+			return !VCClangCompilerFilter(context);
 		}
 	}
 }
diff --git a/Engine/Source/Programs/UnrealBuildTool/Configuration/CompileWarnings/WarningLevelDefaultAttribute.cs b/Engine/Source/Programs/UnrealBuildTool/Configuration/CompileWarnings/WarningLevelDefaultAttribute.cs
--- a/Engine/Source/Programs/UnrealBuildTool/Configuration/CompileWarnings/WarningLevelDefaultAttribute.cs
+++ b/Engine/Source/Programs/UnrealBuildTool/Configuration/CompileWarnings/WarningLevelDefaultAttribute.cs
@@ -1,6 +1,7 @@
 // Copyright Epic Games, Inc. All Rights Reserved.
 
 using System;
+using System.Collections;
 using System.Collections.Generic;
 using System.Linq;
 using EpicGames.Core;
@@ -9,86 +10,335 @@ using Microsoft.Extensions.Logging;
 namespace UnrealBuildTool.Configuration.CompileWarnings
 {
 	/// <summary>
-	/// Flag used to communicate which context the <see cref="WarningLevelDefaultAttribute"/> should be applied.
+	/// Flag used to communicate which context the <see cref="VersionWarningLevelDefaultAttribute"/> should be applied.
 	/// </summary>
 	internal enum InitializationContext
 	{
 		/// <summary>
-		/// Apply the <see cref="WarningLevelDefaultAttribute"/> in all contexts.
+		/// Apply the <see cref="VersionWarningLevelDefaultAttribute"/> in all contexts.
 		/// </summary>
 		Any,
 		/// <summary>
-		/// Apply the <see cref="WarningLevelDefaultAttribute"/> exclusively in the <see cref="CppCompileWarnings.ApplyTargetDefaults(CppCompileWarnings, Boolean)"/> context.
+		/// Apply the <see cref="VersionWarningLevelDefaultAttribute"/> exclusively in the <see cref="CppCompileWarnings.ApplyTargetDefaults(CppCompileWarnings, Boolean)"/> context.
 		/// </summary>
 		Target,
 		/// <summary>
-		/// Apply the <see cref="WarningLevelDefaultAttribute"/> exclusively in the <see cref="CppCompileWarnings.ApplyDefaults"/> context.
+		/// Apply the <see cref="VersionWarningLevelDefaultAttribute"/> exclusively in the <see cref="CppCompileWarnings.ApplyDefaults"/> context.
 		/// </summary>
 		Constructor
 	}
 
+	internal static class WarningLevelDefaultHelpers
+	{
+		public static IList<T>? CastWarningLevelDefaultAttirbuteToTarget<T>(IList<WarningLevelDefaultAttribute> uncastedAttributes, ILogger? logger)
+		{
+			List<T>? castedAttributes = null;
+			try
+			{
+				castedAttributes = uncastedAttributes.Cast<T>().ToList();
+			}
+			catch (Exception ex)
+			{
+				logger?.LogError("Exception occurred {ExMessage}. Received a set of attributes that weren't of the correct warning attribute type ({Type})", ex.Message, nameof(T));
+			}
+
+			return castedAttributes;
+		}
+	}
+
 	/// <summary>
-	/// Attribute used to set the default values of a <see cref="WarningLevel"/> under the context of <see cref="CppCompileWarnings.ApplyDefaults"/> context and <see cref="CppCompileWarnings.ApplyTargetDefaults(CppCompileWarnings, Boolean)"/> context.
+	/// Abstract attribute used to specify the default <see cref="WarningLevel"/> to assign a member, and the interface,
 	/// </summary>
 	[AttributeUsage(AttributeTargets.Field | AttributeTargets.Property, AllowMultiple = true)]
-	internal sealed class WarningLevelDefaultAttribute : Attribute
+	internal abstract class WarningLevelDefaultAttribute : Attribute
 	{
-		private static readonly ILogger? s_logger = Log.Logger;
+		protected static readonly ILogger? Logger = Log.Logger;
 
 		/// <summary>
-		/// The lower bound (inclusive) <see cref="BuildSettingsVersion"/> of which to apply this <see cref="DefaultLevel"/> to.
+		/// Constructor a WarningLevelDefaultAttribute.
 		/// </summary>
-		public BuildSettingsVersion MinVersion { get; }
+		/// <param name="context">The context in which to apply this default.</param>
+		public WarningLevelDefaultAttribute(InitializationContext context = InitializationContext.Any)
+		{
+			Context = context;
+		}
 
 		/// <summary>
-		/// The upper bound (inclusive) <see cref="BuildSettingsVersion"/> of which to apply this <see cref="DefaultLevel"/> to.
+		/// Gets the default <see cref="WarningLevel"/> to apply to the member.
 		/// </summary>
-		public BuildSettingsVersion MaxVersion { get; }
+		/// <param name="buildSystemContext">The build system context used to consider defaults and applicability.</param>
+		/// <returns></returns>
+		public abstract WarningLevel GetDefaultLevel(BuildSystemContext? buildSystemContext);
+
+		/// <summary>
+		/// The context in which this attribute should be applied.
+		/// </summary>
+		public InitializationContext Context { get; }
+	}
+
+	/// <summary>
+	/// Basic attribute used to set the default values of a <see cref="WarningLevel"/> under the context of <see cref="CppCompileWarnings.ApplyDefaults"/> context and <see cref="CppCompileWarnings.ApplyTargetDefaults(CppCompileWarnings, Boolean)"/> context.
+	/// </summary>
+	/// <remarks>When combined with any other <see cref="WarningLevelDefaultAttribute"/>, it will act as the fallback in resolution if no other attribute sets a non-<see cref="WarningLevel.Default"/>.</remarks>
+	[AttributeUsage(AttributeTargets.Field | AttributeTargets.Property, AllowMultiple = false)]
+#pragma warning disable CA1813 // Avoid unsealed attributes
+	internal class BasicWarningLevelDefaultAttribute : WarningLevelDefaultAttribute
+#pragma warning restore CA1813 // Avoid unsealed attributes
+	{
+		/// <summary>
+		/// Constructor a BasicWarningLevelDefaultAttribute.
+		/// </summary>
+		/// <param name="defaultLevel">The default warning level to apply for this bound.</param>
+		/// <param name="context">The context in which to apply this default.</param>
+#pragma warning disable CA1019 // Define accessors for attribute arguments
+		public BasicWarningLevelDefaultAttribute(WarningLevel defaultLevel = WarningLevel.Default, InitializationContext context = InitializationContext.Any) : base(context)
+#pragma warning restore CA1019 // Define accessors for attribute arguments
+		{
+			DefaultLevel = defaultLevel;
+		}
 
 		/// <summary>
 		/// The default warning level.
 		/// </summary>
-		public WarningLevel DefaultLevel { get; }
+		/// <param name="buildSystemContext">The build system context used to conditionally augment the <see cref="DefaultLevel"/> by.</param>
+		public override WarningLevel GetDefaultLevel(BuildSystemContext? buildSystemContext = null)
+		{
+			return DefaultLevel;
+		}
 
 		/// <summary>
-		/// The context in which this attribute should be applied.
+		/// Resolver for <see cref="BasicWarningLevelDefaultAttribute"/>.
 		/// </summary>
-		public InitializationContext Context { get; }
+		/// <param name="unsortedDefaultValueAttributes">The set of attributes of which to consider for resolution.</param>
+		/// <param name="buildSystemContext">The build system context used to consider defaults and applicability.</param>
+		/// <returns>The corresponding <see cref="WarningLevel"/> of which to apply to the property with the default value attributes.</returns>
+		[WarningLevelResolverDelegate(nameof(BasicWarningLevelDefaultAttribute))]
+		internal static WarningLevel ResolveWarningLevelDefault(IList<WarningLevelDefaultAttribute> unsortedDefaultValueAttributes, BuildSystemContext? buildSystemContext)
+		{
+			if (unsortedDefaultValueAttributes.Count != 1)
+			{
+				Logger?.LogError("Invalid attribute definition for {Name}.", nameof(WarningLevelDefaultAttribute));
+				return WarningLevel.Default;
+			}
+
+			return unsortedDefaultValueAttributes[0].GetDefaultLevel(buildSystemContext);
+		}
+
+		protected readonly WarningLevel DefaultLevel;
+	}
+
+	#region -- Compiler Constrained Warning Level Defaults --
+	/// <summary>
+	/// Compiler constrained Attribute used to set the default values of a <see cref="WarningLevel"/> under the context of <see cref="CppCompileWarnings.ApplyDefaults"/> context and <see cref="CppCompileWarnings.ApplyTargetDefaults(CppCompileWarnings, Boolean)"/> context.
+	/// </summary>
+	[AttributeUsage(AttributeTargets.Field | AttributeTargets.Property, AllowMultiple = true)]
+	internal abstract class CompilerWarningLevelDefaultAttribute : BasicWarningLevelDefaultAttribute
+	{
+		public WindowsCompiler Compiler { get; }
+
+		/// <summary>
+		/// Constructs a CompilerWarningLevelDefaultAttribute.
+		/// </summary>
+		/// <param name="compiler">The compiler which this attribute should be applied.</param>
+		/// <param name="defaultLevel">The default warning level to apply for this bound.</param>
+		/// <param name="context">The context in which to apply this default.</param>
+		public CompilerWarningLevelDefaultAttribute(WindowsCompiler compiler, WarningLevel defaultLevel = WarningLevel.Default, InitializationContext context = InitializationContext.Any) : base(defaultLevel, context)
+		{
+			Compiler = compiler;
+		}
 
 		/// <summary>
-		/// Constructs a ApplyWarningsDefaultValueAttribute.
+		/// Resolves the set of <see cref="CompilerWarningLevelDefaultAttribute"/> to a <see cref="WarningLevel"/> provided the current <see cref="BuildSystemContext"/>.
+		/// </summary>
+		/// <param name="unsortedDefaultValueAttributes">The set of attributes of which to consider for resolution.</param>
+		/// <param name="buildSystemContext">The build system context used to consider defaults and applicability.</param>
+		/// <returns>The corresponding <see cref="WarningLevel"/> of which to apply to the property with the default value attributes.</returns>
+		[WarningLevelResolverDelegate(nameof(CompilerWarningLevelDefaultAttribute))]
+#pragma warning disable IDE0060 // Remove unused parameter
+		internal static WarningLevel ResolveVersionWarningLevelDefault(IList<WarningLevelDefaultAttribute> unsortedDefaultValueAttributes, BuildSystemContext? buildSystemContext)
+#pragma warning restore IDE0060 // Remove unused parameter
+		{
+			IList<CompilerWarningLevelDefaultAttribute>? castedAttributes = WarningLevelDefaultHelpers.CastWarningLevelDefaultAttirbuteToTarget<CompilerWarningLevelDefaultAttribute>(unsortedDefaultValueAttributes, Logger);
+
+			if (castedAttributes == null)
+			{
+				Logger?.LogWarning("Unable to discern default warning level through resolution. Defaulting to WarningLevel.Default.");
+				return WarningLevel.Default;
+			}
+
+			int nonDefaultCount = castedAttributes.Select(x => x.GetDefaultLevel(buildSystemContext) != WarningLevel.Default).Count();
+
+			if (nonDefaultCount > 1)
+			{
+				Logger?.LogWarning("Unable to disambiguate the warning level through resolution. Too many CompilerWarningLevelDefaultAttribute returned a resolved value. Defaulting to WarningLevel.Default.");
+				return WarningLevel.Default;
+			}
+
+			if (nonDefaultCount == 1)
+			{
+				return castedAttributes.Where(x => x.GetDefaultLevel(buildSystemContext) != WarningLevel.Default).Select(x => x.GetDefaultLevel(buildSystemContext)).First();
+			}
+
+			return WarningLevel.Default;
+		}
+	}
+
+	/// <summary>
+	/// Version constrained attribute used to set the default values of a <see cref="WarningLevel"/> under the context of <see cref="CppCompileWarnings.ApplyDefaults"/> context and <see cref="CppCompileWarnings.ApplyTargetDefaults(CppCompileWarnings, Boolean)"/> context.
+	/// </summary>
+	[AttributeUsage(AttributeTargets.Field | AttributeTargets.Property, AllowMultiple = false)]
+	internal sealed class IntelCompilerWarningLevelDefaultAttribute : CompilerWarningLevelDefaultAttribute
+	{
+		/// <summary>
+		/// Constructs a IntelCompilerWarningLevelDefaultAttribute.
+		/// </summary>
+		/// <param name="defaultLevel">The default warning level to apply for this bound.</param>
+		/// <param name="context">The context in which to apply this default.</param>
+		public IntelCompilerWarningLevelDefaultAttribute(WarningLevel defaultLevel = WarningLevel.Default, InitializationContext context = InitializationContext.Any) : base(WindowsCompiler.Intel, defaultLevel, context) { }
+
+		public override WarningLevel GetDefaultLevel(BuildSystemContext? buildSystemContext = null)
+		{
+			if (buildSystemContext == null)
+			{
+				return WarningLevel.Default;
+			}
+
+			ReadOnlyTargetRules? readOnlyTargetRules = buildSystemContext.GetReadOnlyTargetRules();
+
+			if (readOnlyTargetRules?.WindowsPlatform?.Compiler.IsIntel() == true)
+			{
+				return DefaultLevel;
+			}
+
+			return WarningLevel.Default;
+		}
+	}
+
+	/// <summary>
+	/// Version constrained attribute used to set the default values of a <see cref="WarningLevel"/> under the context of <see cref="CppCompileWarnings.ApplyDefaults"/> context and <see cref="CppCompileWarnings.ApplyTargetDefaults(CppCompileWarnings, Boolean)"/> context.
+	/// </summary>
+	[AttributeUsage(AttributeTargets.Field | AttributeTargets.Property, AllowMultiple = false)]
+	internal sealed class MSVCCompilerWarningLevelDefaultAttribute : CompilerWarningLevelDefaultAttribute
+	{
+		/// <summary>
+		/// Constructs a MSVCCompilerWarningLevelDefaultAttribute.
+		/// </summary>
+		/// <param name="defaultLevel">The default warning level to apply for this bound.</param>
+		/// <param name="context">The context in which to apply this default.</param>
+		public MSVCCompilerWarningLevelDefaultAttribute(WarningLevel defaultLevel = WarningLevel.Default, InitializationContext context = InitializationContext.Any) : base(WindowsCompiler.VisualStudio2022, defaultLevel, context) { }
+
+		public override WarningLevel GetDefaultLevel(BuildSystemContext? buildSystemContext = null)
+		{
+			if (buildSystemContext == null)
+			{
+				return WarningLevel.Default;
+			}
+
+			ReadOnlyTargetRules? readOnlyTargetRules = buildSystemContext.GetReadOnlyTargetRules();
+
+			if (readOnlyTargetRules?.WindowsPlatform?.Compiler.IsMSVC() == true)
+			{
+				return DefaultLevel;
+			}
+
+			return WarningLevel.Default;
+		}
+	}
+
+	/// <summary>
+	/// Version constrained attribute used to set the default values of a <see cref="WarningLevel"/> under the context of <see cref="CppCompileWarnings.ApplyDefaults"/> context and <see cref="CppCompileWarnings.ApplyTargetDefaults(CppCompileWarnings, Boolean)"/> context.
+	/// </summary>
+	/// <remarks>Although <see cref="WindowsCompiler.Intel"/> is considered clang, this specificaly removes that.</remarks>
+	[AttributeUsage(AttributeTargets.Field | AttributeTargets.Property, AllowMultiple = false)]
+	internal sealed class ClangNonIntelSpecializedCompilerWarningLevelDefaultAttribute : CompilerWarningLevelDefaultAttribute
+	{
+		/// <summary>
+		/// Constructs a ClangNonIntelSpecializedCompilerWarningLevelDefaultAttribute.
+		/// </summary>
+		/// <param name="defaultLevel">The default warning level to apply for this bound.</param>
+		/// <param name="context">The context in which to apply this default.</param>
+		public ClangNonIntelSpecializedCompilerWarningLevelDefaultAttribute(WarningLevel defaultLevel = WarningLevel.Default, InitializationContext context = InitializationContext.Any) : base(WindowsCompiler.Clang, defaultLevel, context) { }
+
+		public override WarningLevel GetDefaultLevel(BuildSystemContext? buildSystemContext = null)
+		{
+			if (buildSystemContext == null)
+			{
+				return WarningLevel.Default;
+			}
+
+			ReadOnlyTargetRules? readOnlyTargetRules = buildSystemContext.GetReadOnlyTargetRules();
+
+			if (readOnlyTargetRules?.WindowsPlatform?.Compiler.IsClang() == true && !readOnlyTargetRules.WindowsPlatform.Compiler.IsIntel())
+			{
+				return DefaultLevel;
+			}
+
+			return WarningLevel.Default;
+		}
+	}
+
+	#endregion -- Compiler Constrained Warning Level Defaults --
+
+	/// <summary>
+	/// Version constrained attribute used to set the default values of a <see cref="WarningLevel"/> under the context of <see cref="CppCompileWarnings.ApplyDefaults"/> context and <see cref="CppCompileWarnings.ApplyTargetDefaults(CppCompileWarnings, Boolean)"/> context.
+	/// </summary>
+	[AttributeUsage(AttributeTargets.Field | AttributeTargets.Property, AllowMultiple = true)]
+	internal sealed class VersionWarningLevelDefaultAttribute : BasicWarningLevelDefaultAttribute
+	{
+		/// <summary>
+		/// The lower bound (inclusive) <see cref="BuildSettingsVersion"/> of which to apply this <see cref="BasicWarningLevelDefaultAttribute.DefaultLevel"/> to.
+		/// </summary>
+		public BuildSettingsVersion MinVersion { get; }
+
+		/// <summary>
+		/// The upper bound (inclusive) <see cref="BuildSettingsVersion"/> of which to apply this <see cref="BasicWarningLevelDefaultAttribute.DefaultLevel"/> to.
+		/// </summary>
+		public BuildSettingsVersion MaxVersion { get; }
+
+		/// <summary>
+		/// Constructs a VersionWarningLevelDefaultAttribute.
 		/// </summary>
 		/// <param name="minVersion">Minimum version of the bound to apply this setting to. Inclusive.</param>
 		/// <param name="maxVersion">Maximum version of the bound to apply this setting to. Inclusive.</param>
 		/// <param name="defaultLevel">The default warning level to apply for this bound.</param>
 		/// <param name="context">The context in which to apply this default.</param>
-		public WarningLevelDefaultAttribute(WarningLevel defaultLevel = WarningLevel.Default, BuildSettingsVersion minVersion = BuildSettingsVersion.V1, BuildSettingsVersion maxVersion = BuildSettingsVersion.Latest, InitializationContext context = InitializationContext.Any)
+#pragma warning disable CA1019
+		public VersionWarningLevelDefaultAttribute(WarningLevel defaultLevel = WarningLevel.Default, BuildSettingsVersion minVersion = BuildSettingsVersion.V1, BuildSettingsVersion maxVersion = BuildSettingsVersion.Latest, InitializationContext context = InitializationContext.Any) : base(defaultLevel, context)
+#pragma warning restore CA1019
 		{
 			MinVersion = minVersion;
 			MaxVersion = maxVersion;
-			DefaultLevel = defaultLevel;
-			Context = context;
 		}
 
 		/// <summary>
-		/// Resolves the set of <see cref="WarningLevelDefaultAttribute"/> to a <see cref="WarningLevel"/> provided the current <see cref="BuildSettingsVersion"/>.
+		/// Resolves the set of <see cref="VersionWarningLevelDefaultAttribute"/> to a <see cref="WarningLevel"/> provided the current <see cref="BuildSystemContext"/>.
 		/// </summary>
-		/// <param name="unsortedDefaultValueAttributes">The set of attributes to consider.</param>
-		/// <param name="activeVersion">The current build settings version.</param>
+		/// <param name="unsortedDefaultValueAttributes">The set of attributes of which to consider for resolution.</param>
+		/// <param name="buildSystemContext">The build system context used to consider defaults and applicability.</param>
 		/// <returns>The corresponding <see cref="WarningLevel"/> of which to apply to the property with the default value attributes.</returns>
-		/// <remarks>Will invoke <see cref="EnsureWarningLevelDefaultBounds(IList{WarningLevelDefaultAttribute})"/> on the input <see cref="WarningLevelDefaultAttribute"/>.</remarks>
-		internal static WarningLevel ResolveWarningLevelDefault(IList<WarningLevelDefaultAttribute> unsortedDefaultValueAttributes, BuildSettingsVersion activeVersion)
+		/// <remarks>Will invoke <see cref="EnsureWarningLevelDefaultBounds(IList{VersionWarningLevelDefaultAttribute})"/> on the input <see cref="VersionWarningLevelDefaultAttribute"/>.</remarks>
+		[WarningLevelResolverDelegate(nameof(VersionWarningLevelDefaultAttribute))]
+		internal static WarningLevel ResolveVersionWarningLevelDefault(IList<WarningLevelDefaultAttribute> unsortedDefaultValueAttributes, BuildSystemContext? buildSystemContext)
 		{
-			List<WarningLevelDefaultAttribute> sortedAndMerged = EnsureWarningLevelDefaultBounds(unsortedDefaultValueAttributes);
+			IList<VersionWarningLevelDefaultAttribute>? castedAttributes = WarningLevelDefaultHelpers.CastWarningLevelDefaultAttirbuteToTarget<VersionWarningLevelDefaultAttribute>(unsortedDefaultValueAttributes, Logger);
+
+			if (castedAttributes == null)
+			{
+				Logger?.LogWarning("Unable to discern default warning level through resolution. Defaulting to WarningLevel.Default.");
+				return WarningLevel.Default;
+			}
+
+			BuildSettingsVersion activeVersion = buildSystemContext != null ? buildSystemContext._buildContext.GetBuildSettings() : BuildSettingsVersion.V1;
+			List<VersionWarningLevelDefaultAttribute> sortedAndMerged = EnsureWarningLevelDefaultBounds(castedAttributes);
 
 			WarningLevel returnWarningLevel = WarningLevel.Default;
 
-			foreach (WarningLevelDefaultAttribute attr in sortedAndMerged)
+			foreach (VersionWarningLevelDefaultAttribute attr in sortedAndMerged)
 			{
 				// If we find our appropriate range, we early out.
 				if (attr.MinVersion <= activeVersion && activeVersion <= attr.MaxVersion)
 				{
-					returnWarningLevel = attr.DefaultLevel;
+					returnWarningLevel = attr.GetDefaultLevel(buildSystemContext);
 					break;
 				}
 			}
@@ -101,22 +351,22 @@ namespace UnrealBuildTool.Configuration.CompileWarnings
 		/// </summary>
 		/// <param name="attributes">The list of attributes to verify.</param>
 		/// <returns>A list of attributes that has no overlaps, and is contiguous.</returns>
-		internal static List<WarningLevelDefaultAttribute> EnsureWarningLevelDefaultBounds(IList<WarningLevelDefaultAttribute> attributes)
+		internal static List<VersionWarningLevelDefaultAttribute> EnsureWarningLevelDefaultBounds(IList<VersionWarningLevelDefaultAttribute> attributes)
 		{
-			List<WarningLevelDefaultAttribute> sorted = attributes.OrderBy(a => a.MinVersion).ToList();
-			List<WarningLevelDefaultAttribute> merged = new List<WarningLevelDefaultAttribute>(sorted.Count);
+			List<VersionWarningLevelDefaultAttribute> sorted = attributes.OrderBy(a => a.MinVersion).ToList();
+			List<VersionWarningLevelDefaultAttribute> merged = new List<VersionWarningLevelDefaultAttribute>(sorted.Count);
 
 			for (int i = 0; i < sorted.Count - 1; i++)
 			{
-				WarningLevelDefaultAttribute current = sorted[i];
-				WarningLevelDefaultAttribute next = sorted[i + 1];
+				VersionWarningLevelDefaultAttribute current = sorted[i];
+				VersionWarningLevelDefaultAttribute next = sorted[i + 1];
 
 				if (current.MaxVersion < next.MinVersion - 1)
 				{
-					s_logger?.LogWarning("Malformed ApplyWarningsDefaultValueAttribute collection; taking corrective action to address gap in range (CurrentMax: {CurrentMax} NextMin: {NextMin}-1).", next.MinVersion, current.MaxVersion);
+					Logger?.LogWarning("Malformed VersionWarningLevelDefaultAttribute collection; taking corrective action to address gap in range (CurrentMax: {CurrentMax} NextMin: {NextMin}-1).", next.MinVersion, current.MaxVersion);
 
 					// Extend current range to cover the gap up to the next MinVersion,using the old build settings version standard.
-					current = new WarningLevelDefaultAttribute(
+					current = new VersionWarningLevelDefaultAttribute(
 						current.DefaultLevel,
 						current.MinVersion,
 						next.MinVersion - 1
@@ -124,10 +374,10 @@ namespace UnrealBuildTool.Configuration.CompileWarnings
 				}
 				else if (next.MinVersion <= current.MaxVersion)
 				{
-					s_logger?.LogWarning("Malformed ApplyWarningsDefaultValueAttribute collection; taking corrective action to address overlap in range (NextMin: {NextMin} <= CurrentMax: {CurrentMax}).", next.MinVersion, current.MaxVersion);
+					Logger?.LogWarning("Malformed VersionWarningLevelDefaultAttribute collection; taking corrective action to address overlap in range (NextMin: {NextMin} <= CurrentMax: {CurrentMax}).", next.MinVersion, current.MaxVersion);
 
 					// Reduce next range to be more constrained, leaving the larger current range at the old build settings version standard.
-					next = new WarningLevelDefaultAttribute(
+					next = new VersionWarningLevelDefaultAttribute(
 						next.DefaultLevel,
 						current.MaxVersion + 1,
 						next.MaxVersion
diff --git a/Engine/Source/Programs/UnrealBuildTool/Configuration/CompileWarnings/WarningLevelResolverDelegateAttribute.cs b/Engine/Source/Programs/UnrealBuildTool/Configuration/CompileWarnings/WarningLevelResolverDelegateAttribute.cs
new file mode 100644
--- /dev/null
+++ b/Engine/Source/Programs/UnrealBuildTool/Configuration/CompileWarnings/WarningLevelResolverDelegateAttribute.cs
@@ -0,0 +1,71 @@
+// Copyright Epic Games, Inc. All Rights Reserved.
+
+using System;
+using System.Collections.Generic;
+using System.Reflection;
+using EpicGames.Core;
+using Microsoft.Extensions.Logging;
+using UnrealBuildTool.Configuration.CompileWarnings;
+
+namespace UnrealBuildTool
+{
+	/// <summary>
+	/// A wrapper for the methods annotated by the <see cref="WarningLevelResolverDelegateAttribute"/> to support <see cref="IWarningLevelResolver"/>.
+	/// </summary>
+	/// <remarks>Primarily invoked by <see cref="CppCompileWarnings.CoalesceDefaultWarningFromMemberInfo"/>.</remarks>
+	internal struct WarningLevelResolverWrapper : IWarningLevelResolver
+	{
+		public string FilterName => _applyWarningsFilterDelegateAttribute.ResolverName;
+
+		private readonly ILogger? _logger = Log.Logger;
+		private readonly WarningLevelResolverDelegateAttribute _applyWarningsFilterDelegateAttribute;
+		private readonly MethodInfo _method;
+
+		/// <summary>
+		/// Constructs a ApplyWarningsFilterWrapper.
+		/// </summary>
+		/// <param name="filterAttribute">Instance of the filter attribute that is being wrapped.</param>
+		/// <param name="method">The method that will be used for reflective invocation in <see cref="ResolveWarningLevelDefault"/>.</param>
+		public WarningLevelResolverWrapper(WarningLevelResolverDelegateAttribute filterAttribute, MethodInfo method)
+		{
+			_applyWarningsFilterDelegateAttribute = filterAttribute;
+			_method = method;
+		}
+
+		/// <inheritdoc/>
+		public WarningLevel ResolveWarningLevelDefault(IList<WarningLevelDefaultAttribute> unsortedDefaultValueAttributes, BuildSystemContext? buildSystemContext)
+		{
+			WarningLevel defaultWarningLevel = WarningLevel.Default;
+
+			{
+				try
+				{
+					object? result = _method?.Invoke(null, [unsortedDefaultValueAttributes, buildSystemContext]);
+					defaultWarningLevel = result != null ? (WarningLevel)result : WarningLevel.Default;
+				}
+				catch (Exception e)
+				{
+					_logger?.LogWarning("Have failed to reflectively invoke (name:{MethodName}) in the CanApply context. Exception: {Exception}.", _method?.Name, e.Message);
+				}
+
+				return defaultWarningLevel;
+			}
+		}
+	}
+
+	/// <summary>
+	/// Attribute used to denote methods that should participate in <see cref="IWarningLevelResolver"/> usage.
+	/// </summary>
+	/// <remarks>Primarily invoked by <see cref="CppCompileWarnings.CoalesceDefaultWarningFromMemberInfo"/>.</remarks>
+	[AttributeUsage(AttributeTargets.Method, AllowMultiple = false)]
+	internal sealed class WarningLevelResolverDelegateAttribute : Attribute
+	{
+		private readonly string _resolverName;
+		public string ResolverName => _resolverName;
+
+		public WarningLevelResolverDelegateAttribute(string resolverName)
+		{
+			_resolverName = resolverName;
+		}
+	}
+}
diff --git a/Engine/Source/Programs/UnrealBuildTool/Configuration/CompileWarnings/WarningLevelResolverRegistry.cs b/Engine/Source/Programs/UnrealBuildTool/Configuration/CompileWarnings/WarningLevelResolverRegistry.cs
new file mode 100644
--- /dev/null
+++ b/Engine/Source/Programs/UnrealBuildTool/Configuration/CompileWarnings/WarningLevelResolverRegistry.cs
@@ -0,0 +1,126 @@
+// Copyright Epic Games, Inc. All Rights Reserved.
+
+using System;
+using System.Collections.Generic;
+using System.Linq;
+using System.Reflection;
+using EpicGames.Core;
+using Microsoft.CodeAnalysis;
+using Microsoft.Extensions.Logging;
+using UnrealBuildTool.Configuration.CompileWarnings;
+
+namespace UnrealBuildTool
+{
+	/// <summary>
+	/// Interface that participates in resolving a set of <see cref="WarningLevelDefaultAttribute"/> against a member.
+	/// </summary>
+	internal interface IWarningLevelResolver
+	{
+		/// <summary>
+		/// Resolves a <see cref="WarningLevel"/> given a set of <see cref="WarningLevelDefaultAttribute"/>.
+		/// </summary>
+		/// <param name="unsortedDefaultValueAttributes">The set of attributes to consider in resolution.</param>
+		/// <param name="buildSystemContext">The build system context to evaluate these under.</param>
+		/// <returns>The resolved warning level.</returns>
+		WarningLevel ResolveWarningLevelDefault(IList<WarningLevelDefaultAttribute> unsortedDefaultValueAttributes, BuildSystemContext? buildSystemContext);
+	}
+
+	/// <summary>
+	/// WarningLevelResolverRegistry is the container for all of the resolvers applicable within the <see cref="CppCompileWarnings.ApplyTargetDefaults"/> and <see cref="CppCompileWarnings.ApplyDefaults"/> contexts.
+	/// </summary>
+	/// <remarks>The primary entry point for extensibility is <see cref="WarningLevelResolverDelegateAttribute"/>.</remarks>
+	internal class WarningLevelResolverRegistry
+	{
+		private readonly ILogger? _logger = Log.Logger;
+		private static WarningLevelResolverRegistry? s_resolverRegistry;
+		private readonly Dictionary<string, IWarningLevelResolver> _attributeNameToBehaviour = [];
+
+		/// <summary>
+		/// The resolver name to resolver representation of the registry.
+		/// </summary>
+		internal static IReadOnlyDictionary<string, IWarningLevelResolver> Resolvers
+		{
+			get
+			{
+				if (s_resolverRegistry == null)
+				{
+					s_resolverRegistry = new WarningLevelResolverRegistry();
+				}
+
+				return s_resolverRegistry._attributeNameToBehaviour;
+			}
+		}
+
+		/// <summary>
+		/// Requests a resolver by name.
+		/// </summary>
+		/// <param name="resolverName">The name of the resolver to obtain.</param>
+		/// <returns>The IWarningLevelResolver if it has been registered, null if it was not found.</returns>
+		internal static IWarningLevelResolver? RequestResolver(string resolverName)
+		{
+			IWarningLevelResolver? returnResolver;
+			Resolvers.TryGetValue(resolverName, out returnResolver);
+
+			return returnResolver;
+		}
+
+		private WarningLevelResolverRegistry()
+		{
+			// Process all of the static methods in the assembly that are annotated wtih WarningLevelResolverDelegateAttribute, and evaluate given they match
+			// the execution delegate, which is effectively IWarningLevelResolver.ResolveWarningLevelDefault
+			MethodInfo? reflectiveInterfaceMethodInfo = typeof(IWarningLevelResolver).GetMethod(nameof(IWarningLevelResolver.ResolveWarningLevelDefault));
+
+			if (reflectiveInterfaceMethodInfo != null)
+			{
+				Assembly assembly = Assembly.GetExecutingAssembly();
+				Type[] allTypes = assembly.GetTypes();
+				List<MethodInfo> delegatesForConsideration = new List<MethodInfo>();
+
+				foreach (Type type in allTypes)
+				{
+					IEnumerable<MethodInfo> methodsWithAttribute = type
+						.GetMethods(BindingFlags.Static | BindingFlags.Public | BindingFlags.NonPublic)
+						.Where(m => m.GetCustomAttributes(typeof(WarningLevelResolverDelegateAttribute), false).Length > 0);
+
+					delegatesForConsideration.AddRange(methodsWithAttribute);
+				}
+
+				foreach (MethodInfo mi in delegatesForConsideration)
+				{
+					if (mi.ReturnType != reflectiveInterfaceMethodInfo.ReturnType)
+					{
+						_logger?.LogWarning("Skipping MethodInfo (name:{MethodName}) annotated with WarningLevelResolverDelegateAttribute as it doesn't return {ReturnType}.", mi.Name, reflectiveInterfaceMethodInfo.ReturnType.Name);
+						continue;
+					}
+
+					if (mi.GetParameters().Length != reflectiveInterfaceMethodInfo.GetParameters().Length)
+					{
+						_logger?.LogWarning("Skipping MethodInfo (name:{MethodName}) annotated with WarningLevelResolverDelegateAttribute as it doesn't contain the same number of parameters as the base method interface ({Count})", mi.Name, reflectiveInterfaceMethodInfo.GetParameters().Length);
+						continue;
+					}
+					if (!mi.GetParameters().SequenceEqual(reflectiveInterfaceMethodInfo.GetParameters(), EqualityComparer<ParameterInfo>.Create((a, b) => a?.ParameterType == b?.ParameterType)))
+					{
+						_logger?.LogWarning("Skipping MethodInfo (name:{MethodName}) annotated with WarningLevelResolverDelegateAttribute it's parameters don't match the base method interface.", mi.Name);
+						continue;
+					}
+
+					WarningLevelResolverDelegateAttribute? attributeInstance = mi.GetCustomAttribute<WarningLevelResolverDelegateAttribute>();
+					if (attributeInstance == null)
+					{
+						_logger?.LogWarning("Skipping MethodInfo (name:{MethodName}) annotated with WarningLevelResolverDelegateAttribute as reflection has failed to obtain the attribute instance.", mi.Name);
+						continue;
+					}
+
+					if (!_attributeNameToBehaviour.TryAdd(attributeInstance.ResolverName, new WarningLevelResolverWrapper(attributeInstance, mi)))
+					{
+						_logger?.LogWarning("Duplicate resolver names ({DuplicateName}) encountered when building the WarningLevelResolverRegistry.", attributeInstance.ResolverName);
+					}
+				}
+			}
+			else
+			{
+				_logger?.LogWarning("Base method interface was null. Ignoring all method reflective methods used for resolvers.");
+			}
+		}
+	}
+}
\ No newline at end of file
diff --git a/Engine/Source/Programs/UnrealBuildTool/Configuration/CppCompileWarnings.cs b/Engine/Source/Programs/UnrealBuildTool/Configuration/CppCompileWarnings.cs
--- a/Engine/Source/Programs/UnrealBuildTool/Configuration/CppCompileWarnings.cs
+++ b/Engine/Source/Programs/UnrealBuildTool/Configuration/CppCompileWarnings.cs
@@ -87,9 +87,9 @@ namespace UnrealBuildTool
 		[CommandLine("-ShadowVariableErrors", Value = nameof(WarningLevel.Error))]
 		[WarningsMSCV(["/wd4456", "/wd4458", "/wd4459"], null, ["/we4456", "/we4458", "/we4459"])]
 		[ShadowVariableWarningsClangToolChain(["-Wno-shadow"], ["-Wshadow", "-Wno-error=shadow"], ["-Wshadow"])]
-		[WarningLevelDefault(WarningLevel.Warning, BuildSettingsVersion.V1, BuildSettingsVersion.Latest, InitializationContext.Constructor)]
-		[WarningLevelDefault(WarningLevel.Warning, BuildSettingsVersion.V1, BuildSettingsVersion.V1, InitializationContext.Target)]
-		[WarningLevelDefault(WarningLevel.Error, BuildSettingsVersion.V2, BuildSettingsVersion.Latest, InitializationContext.Target)]
+		[VersionWarningLevelDefault(WarningLevel.Warning, BuildSettingsVersion.V1, BuildSettingsVersion.Latest, InitializationContext.Constructor)]
+		[VersionWarningLevelDefault(WarningLevel.Warning, BuildSettingsVersion.V1, BuildSettingsVersion.V1, InitializationContext.Target)]
+		[VersionWarningLevelDefault(WarningLevel.Error, BuildSettingsVersion.V2, BuildSettingsVersion.Latest, InitializationContext.Target)]
 		public WarningLevel ShadowVariableWarningLevel
 		{
 			get => ResolveWarning(ParentCppCompileWarnings?.ShadowVariableWarningLevel, _shadowVariableWarningLevel);
@@ -124,7 +124,7 @@ namespace UnrealBuildTool
 			["-Wno-float-conversion", "-Wno-implicit-float-conversion", "-Wno-implicit-int-conversion", "-Wno-c++11-narrowing"],
 			["-Wfloat-conversion", "-Wno-error=float-conversion", "-Wimplicit-float-conversion", "-Wno-error=implicit-float-conversion", "-Wimplicit-int-conversion", "-Wno-error=implicit-int-conversion", "-Wc++11-narrowing", "-Wno-error=c++11-narrowing"],
 			["-Wfloat-conversion", "-Wimplicit-float-conversion", "-Wimplicit-int-conversion", "-Wc++11-narrowing"])]
-		[WarningLevelDefault(WarningLevel.Off)]
+		[BasicWarningLevelDefault(WarningLevel.Off)]
 		public WarningLevel UnsafeTypeCastWarningLevel
 		{
 			get => ResolveWarning(ParentCppCompileWarnings?.UnsafeTypeCastWarningLevel, _unsafeTypeCastWarningLevel);
@@ -145,7 +145,7 @@ namespace UnrealBuildTool
 		[CommandLine("-UndefinedIdentifierWarningLevel=")]
 		[UndefinedIdentifierWarningsVCToolChain(null, ["/w44668"], ["/we4668"])]
 		[WarningsClangToolChain(null, ["-Wundef", "-Wno-error=undef"], ["-Wundef"])]
-		[WarningLevelDefault(WarningLevel.Off)]
+		[BasicWarningLevelDefault(WarningLevel.Off)]
 		public WarningLevel UndefinedIdentifierWarningLevel
 		{
 			get => ResolveWarning(ParentCppCompileWarnings?.UndefinedIdentifierWarningLevel, _undefinedIdentifierWarningLevel);
@@ -165,7 +165,7 @@ namespace UnrealBuildTool
 		[CommandLine("-SwitchUnhandledEnumeratorWarningLevel=")]
 		[WarningsVCToolChain(null, ["/w44061"], ["/we4061"])]
 		[WarningsClangToolChain(null, ["-Wswitch-enum", "-Wno-error=switch-enum"], ["-Wswitch-enum"])]
-		[WarningLevelDefault(WarningLevel.Off)]
+		[BasicWarningLevelDefault(WarningLevel.Off)]
 		public WarningLevel SwitchUnhandledEnumeratorWarningLevel
 		{
 			get => ResolveWarning(ParentCppCompileWarnings?.SwitchUnhandledEnumeratorWarningLevel, _switchUnhandledEnumeratorWarningLevel);
@@ -173,7 +173,49 @@ namespace UnrealBuildTool
 		}
 		private WarningLevel _switchUnhandledEnumeratorWarningLevel;
 
-		#region -- Clang Compiler Disables --
+		#region -- Clang Compiler --
+
+		/// <summary>
+		/// Indicates what warning/error level to treat invocations to destructors that aren't virtual.
+		/// Clang -
+		///		https://clang.llvm.org/docs/DiagnosticsReference.html#wdelete-non-virtual-dtor
+		/// </summary>
+		[NonVCClangWarningsClangToolChain(["-Wno-delete-non-virtual-dtor"], ["-Wdelete-non-virtual-dtor", "-Wno-error=delete-non-virtual-dtor"], ["-Wdelete-non-virtual-dtor"])]
+		[BasicWarningLevelDefault(WarningLevel.Error)]
+		public WarningLevel DeleteNonVirtualDtorWarningLevel
+		{
+			get => ResolveWarning(ParentCppCompileWarnings?.DeleteNonVirtualDtorWarningLevel, _deleteNonVirtualDtorWarningLevel);
+			set => _deleteNonVirtualDtorWarningLevel = value;
+		}
+		private WarningLevel _deleteNonVirtualDtorWarningLevel;
+
+		/// <summary>
+		/// Indicates what warning/error level to treat implicit conversions betwen enums and integers.
+		/// Clang -
+		///		https://clang.llvm.org/docs/DiagnosticsReference.html#wenum-conversion
+		/// </summary>
+		[NonVCClangWarningsClangToolChain(["-Wno-enum-conversion"], ["-Wenum-conversion", "-Wno-error=enum-conversion"], ["-Wenum-conversion"], null, ArgPosition.Beginning)]
+		[BasicWarningLevelDefault(WarningLevel.Error)]
+		public WarningLevel EnumConversionWarningLevel
+		{
+			get => ResolveWarning(ParentCppCompileWarnings?.EnumConversionWarningLevel, _enumConversionWarningLevel);
+			set => _enumConversionWarningLevel = value;
+		}
+		private WarningLevel _enumConversionWarningLevel;
+
+		/// <summary>
+		/// Indicates what warning/error level to treat assignments of enum types to a bit=field, which can lead to implicit conversations (causing unexpected results).
+		/// Clang -
+		///		https://clang.llvm.org/docs/DiagnosticsReference.html#wbitfield-enum-conversion
+		/// </summary>
+		[NonVCClangWarningsClangToolChain(["-Wno-bitfield-enum-conversion"], ["-Wbitfield-enum-conversion", "-Wno-error=bitfield-enum-conversion"], ["-Wbitfield-enum-conversion"])]
+		[BasicWarningLevelDefault(WarningLevel.Error)]
+		public WarningLevel ClangBitfieldEnumConversion
+		{
+			get => ResolveWarning(ParentCppCompileWarnings?.ClangBitfieldEnumConversion, _clangBitfieldEnumConversion);
+			set => _clangBitfieldEnumConversion = value;
+		}
+		private WarningLevel _clangBitfieldEnumConversion;
 
 		/// <summary>
 		/// Indicates what warning/error level to treat conversions between enum types.
@@ -181,7 +223,7 @@ namespace UnrealBuildTool
 		///		https://clang.llvm.org/docs/DiagnosticsReference.html#wenum-enum-conversion
 		/// </summary>
 		[WarningsClangToolChain(["-Wno-enum-enum-conversion"], ["-Wenum-enum-conversion", "-Wno-error=enum-enum-conversion"], ["-Wenum-enum-conversion"])]
-		[WarningLevelDefault(WarningLevel.Off)]
+		[BasicWarningLevelDefault(WarningLevel.Off)]
 		public WarningLevel EnumEnumConversionWarningLevel
 		{
 			get => ResolveWarning(ParentCppCompileWarnings?.EnumEnumConversionWarningLevel, _enumEnumConversionWarningLevel);
@@ -195,7 +237,7 @@ namespace UnrealBuildTool
 		///		https://clang.llvm.org/docs/DiagnosticsReference.html#wenum-float-conversion
 		/// </summary>
 		[WarningsClangToolChain(["-Wno-enum-float-conversion"], ["-Wenum-float-conversion", "-Wno-error=enum-float-conversion"], ["-Wenum-float-conversion"])]
-		[WarningLevelDefault(WarningLevel.Off)]
+		[BasicWarningLevelDefault(WarningLevel.Off)]
 		public WarningLevel EnumFloatConversionWarningLevel
 		{
 			get => ResolveWarning(ParentCppCompileWarnings?.EnumFloatConversionWarningLevel, _enumFloatConversionWarningLevel);
@@ -209,7 +251,7 @@ namespace UnrealBuildTool
 		///		https://clang.llvm.org/docs/DiagnosticsReference.html#wambiguous-reversed-operator
 		/// </summary>
 		[WarningsClangToolChain(["-Wno-ambiguous-reversed-operator"], ["-Wambiguous-reversed-operator", "-Wno-error=ambiguous-reversed-operator"], ["-Wambiguous-reversed-operator"], [nameof(FilterID.Cpp20Min)])]
-		[WarningLevelDefault(WarningLevel.Off)]
+		[BasicWarningLevelDefault(WarningLevel.Off)]
 		public WarningLevel AmbiguousReversedOperatorWarningLevel
 		{
 			get => ResolveWarning(ParentCppCompileWarnings?.AmbiguousReversedOperatorWarningLevel, _ambiguousReversedOperatorWarningLevel);
@@ -223,7 +265,7 @@ namespace UnrealBuildTool
 		///		https://clang.llvm.org/docs/DiagnosticsReference.html#wdeprecated-anon-enum-enum-conversion
 		/// </summary>
 		[WarningsClangToolChain(["-Wno-deprecated-anon-enum-enum-conversion"], ["-Wdeprecated-anon-enum-enum-conversion", "-Wno-error=deprecated-anon-enum-enum-conversion"], ["-Wdeprecated-anon-enum-enum-conversion"], [nameof(FilterID.Cpp20Min)])]
-		[WarningLevelDefault(WarningLevel.Off)]
+		[BasicWarningLevelDefault(WarningLevel.Off)]
 		public WarningLevel DeprecatedAnonEnumEnumConversionWarningLevel
 		{
 			get => ResolveWarning(ParentCppCompileWarnings?.DeprecatedAnonEnumEnumConversionWarningLevel, _deprecatedAnonEnumEnumConversionWarningLevel);
@@ -237,7 +279,7 @@ namespace UnrealBuildTool
 		///		https://clang.llvm.org/docs/DiagnosticsReference.html#wdeprecated-volatile
 		/// </summary>
 		[WarningsClangToolChain(["-Wno-deprecated-volatile"], ["-Wdeprecated-volatile", "-Wno-error=deprecated-volatile"], ["-Wdeprecated-volatile"], [nameof(FilterID.Cpp20Min)])]
-		[WarningLevelDefault(WarningLevel.Off)]
+		[BasicWarningLevelDefault(WarningLevel.Off)]
 		public WarningLevel DeprecatedVolatileWarningLevel
 		{
 			get => ResolveWarning(ParentCppCompileWarnings?.DeprecatedVolatileWarningLevel, _deprecatedVolatileWarningLevel);
@@ -251,7 +293,7 @@ namespace UnrealBuildTool
 		///		https://clang.llvm.org/docs/DiagnosticsReference.html#wordered-compare-function-pointers
 		/// </summary>
 		[WarningsClangToolChain(["-Wno-ordered-compare-function-pointers"], ["-Wordered-compare-function-pointers", "-Wno-error=ordered-compare-function-pointers"], ["-Wordered-compare-function-pointers"])]
-		[WarningLevelDefault(WarningLevel.Off)]
+		[BasicWarningLevelDefault(WarningLevel.Off)]
 		public WarningLevel OrderedCompareFunctionPointers
 		{
 			get => ResolveWarning(ParentCppCompileWarnings?.OrderedCompareFunctionPointers, _orderedCompareFunctionPointers);
@@ -265,7 +307,7 @@ namespace UnrealBuildTool
 		///		https://clang.llvm.org/docs/DiagnosticsReference.html#wbitwise-instead-of-logical
 		/// </summary>
 		[WarningsClangToolChain(["-Wno-bitwise-instead-of-logical"], ["-Wbitwise-instead-of-logical", "-Wno-error=bitwise-instead-of-logical"], ["-Wbitwise-instead-of-logical"], [nameof(FilterID.Version14Min)])]
-		[WarningLevelDefault(WarningLevel.Off)]
+		[BasicWarningLevelDefault(WarningLevel.Off)]
 		public WarningLevel BitwiseInsteadOfLogicalWarningLevel
 		{
 			get => ResolveWarning(ParentCppCompileWarnings?.BitwiseInsteadOfLogicalWarningLevel, _bitwiseInsteadOfLogicalWarningLevel);
@@ -279,7 +321,7 @@ namespace UnrealBuildTool
 		///		https://clang.llvm.org/docs/DiagnosticsReference.html#wdeprecated-copy
 		/// </summary>
 		[WarningsClangToolChain(["-Wno-deprecated-copy"], ["-Wdeprecated-copy", "-Wno-error=deprecated-copy"], ["-Wdeprecated-copy"], [nameof(FilterID.Version16Min)])]
-		[WarningLevelDefault(WarningLevel.Off)]
+		[BasicWarningLevelDefault(WarningLevel.Off)]
 		public WarningLevel DeprecatedCopyWarningLevel
 		{
 			get => ResolveWarning(ParentCppCompileWarnings?.DeprecatedCopyWarningLevel, _deprecatedCopyWarningLevel);
@@ -294,7 +336,7 @@ namespace UnrealBuildTool
 		/// </summary>
 		/// <remarks><see cref="ArgPosition.End"/> due to it's relation to <see cref="DeprecatedCopyWarningLevel"/>.</remarks>
 		[WarningsClangToolChain(["-Wno-deprecated-copy-with-user-provided-copy"], ["-Wdeprecated-copy-with-user-provided-copy", "-Wno-error=deprecated-copy-with-user-provided-copy"], ["-Wdeprecated-copy-with-user-provided-copy"], [nameof(FilterID.Version16Min)], ArgPosition.End)]
-		[WarningLevelDefault(WarningLevel.Off)]
+		[BasicWarningLevelDefault(WarningLevel.Off)]
 		public WarningLevel DeprecatedCopyWithUserProvidedCopyWarningLevel
 		{
 			get => ResolveWarning(ParentCppCompileWarnings?.DeprecatedCopyWithUserProvidedCopyWarningLevel, _deprecatedCopyWithUserProvidedCopyWarningLevel);
@@ -308,7 +350,7 @@ namespace UnrealBuildTool
 		///		https://clang.llvm.org/docs/DiagnosticsReference.html#winvalid-unevaluated-string
 		/// </summary>
 		[WarningsClangToolChain(["-Wno-invalid-unevaluated-string"], ["-Winvalid-unevaluated-string", "-Wno-error=invalid-unevaluated-string"], ["-Winvalid-unevaluated-string"], [nameof(SpecializedFilters.AndroidNDKR26ToolChainVersionExclusion)])]
-		[WarningLevelDefault(WarningLevel.Off)]
+		[BasicWarningLevelDefault(WarningLevel.Off)]
 		public WarningLevel InvalidUnevaluatedStringWarningLevel
 		{
 			get => ResolveWarning(ParentCppCompileWarnings?.InvalidUnevaluatedStringWarningLevel, _invalidUnevaluatedStringWarningLevel);
@@ -327,7 +369,7 @@ namespace UnrealBuildTool
 		/// We tend to use NAN just as a bit pattern, and not rely on it in calculations, so it should be reasonably safe to disable.
 		/// </remarks>
 		[WarningsClangToolChain(["-Wno-nan-infinity-disabled"], ["-Wnan-infinity-disabled", "-Wno-error=nan-infinity-disabled"], ["-Wnan-infinity-disabled"], [nameof(FilterID.Version18Min)])]
-		[WarningLevelDefault(WarningLevel.Off)]
+		[BasicWarningLevelDefault(WarningLevel.Off)]
 		public WarningLevel NaNInfinityDisabledWarningLevel
 		{
 			get => ResolveWarning(ParentCppCompileWarnings?.NaNInfinityDisabledWarningLevel, _nanInfinityDisabledWarningLevel);
@@ -341,7 +383,7 @@ namespace UnrealBuildTool
 		///		https://clang.llvm.org/docs/DiagnosticsReference.html#wextra-qualification
 		/// </summary>
 		[WarningsClangToolChain(["-Wno-extra-qualification"], ["-Wextra-qualification", "-Wno-error=extra-qualification"], ["-Wextra-qualification"], [nameof(FilterID.Version19Min)])]
-		[WarningLevelDefault(WarningLevel.Off)]
+		[BasicWarningLevelDefault(WarningLevel.Off)]
 		public WarningLevel LevelExtraQualificationWarningLevel
 		{
 			get => ResolveWarning(ParentCppCompileWarnings?.LevelExtraQualificationWarningLevel, _levelExtraQualificationWarningLevel);
@@ -355,7 +397,7 @@ namespace UnrealBuildTool
 		///		https://clang.llvm.org/docs/DiagnosticsReference.html#wcast-function-type-mismatch
 		/// </summary>
 		[WarningsClangToolChain(["-Wno-cast-function-type-mismatch"], ["-Wcast-function-type-mismatch", "-Wno-error=cast-function-type-mismatch"], ["-Wcast-function-type-mismatch"], [nameof(SpecializedFilters.AndroidNDKR28ToolChainVersionExclusion)])]
-		[WarningLevelDefault(WarningLevel.Off)]
+		[BasicWarningLevelDefault(WarningLevel.Off)]
 		public WarningLevel CastFunctionTypeMismatchWarningLevel
 		{
 			get => ResolveWarning(ParentCppCompileWarnings?.CastFunctionTypeMismatchWarningLevel, _castFunctionTypeMismatchWarningLevel);
@@ -369,7 +411,7 @@ namespace UnrealBuildTool
 		///		https://clang.llvm.org/docs/DiagnosticsReference.html#wmissing-template-arg-list-after-template-kw
 		/// </summary>
 		[WarningsClangToolChain(["-Wno-missing-template-arg-list-after-template-kw"], ["-Wmissing-template-arg-list-after-template-kw", "-Wno-error=missing-template-arg-list-after-template-kw"], ["-Wmissing-template-arg-list-after-template-kw"], [nameof(SpecializedFilters.AndroidNDKR28ToolChainVersionExclusion)])]
-		[WarningLevelDefault(WarningLevel.Off)]
+		[BasicWarningLevelDefault(WarningLevel.Off)]
 		public WarningLevel MissingTemplateArgListAfterTemplateWarningLevel
 		{
 			get => ResolveWarning(ParentCppCompileWarnings?.MissingTemplateArgListAfterTemplateWarningLevel, _missingTemplateArgListAfterTemplateWarningLevel);
@@ -384,7 +426,7 @@ namespace UnrealBuildTool
 		/// </summary>
 		/// <remarks>We use this feature to allow static FNames.</remarks>
 		[WarningsClangToolChain(["-Wno-gnu-string-literal-operator-template"], ["-Wgnu-string-literal-operator-template", "-Wno-error=gnu-string-literal-operator-template"], ["-Wgnu-string-literal-operator-template"])]
-		[WarningLevelDefault(WarningLevel.Off)]
+		[BasicWarningLevelDefault(WarningLevel.Off)]
 		public WarningLevel GNUStringLiteralOperatorTemplateWarningLevel
 		{
 			get => ResolveWarning(ParentCppCompileWarnings?.GNUStringLiteralOperatorTemplateWarningLevel, _gnuStringLiteralOperatorTemplateWarningLevel);
@@ -398,7 +440,7 @@ namespace UnrealBuildTool
 		///		https://clang.llvm.org/docs/DiagnosticsReference.html#winconsistent-missing-override
 		/// </summary>
 		[WarningsClangToolChain(["-Wno-inconsistent-missing-override"], ["-Winconsistent-missing-override", "-Wno-error=inconsistent-missing-override"], ["-Winconsistent-missing-override"])]
-		[WarningLevelDefault(WarningLevel.Off)]
+		[BasicWarningLevelDefault(WarningLevel.Off)]
 		public WarningLevel InconsistentMissingOverrideWarningLevel
 		{
 			get => ResolveWarning(ParentCppCompileWarnings?.InconsistentMissingOverrideWarningLevel, _inconsistentMissingOverrideWarningLevel);
@@ -413,7 +455,7 @@ namespace UnrealBuildTool
 		/// </summary>
 		/// <remarks>Needed to suppress warnings about using offsetof on non-POD types.</remarks>
 		[WarningsClangToolChain(["-Wno-invalid-offsetof"], ["-Winvalid-offsetof", "-Wno-error=invalid-offsetof"], ["-Winvalid-offsetof"])]
-		[WarningLevelDefault(WarningLevel.Off)]
+		[BasicWarningLevelDefault(WarningLevel.Off)]
 		public WarningLevel InvalidOffsetWarningLevel
 		{
 			get => ResolveWarning(ParentCppCompileWarnings?.InvalidOffsetWarningLevel, _invalidOffsetOfWarningLevel);
@@ -428,7 +470,7 @@ namespace UnrealBuildTool
 		/// </summary>
 		/// <remarks>This hides the "enumeration value 'XXXXX' not handled in switch [-Wswitch]" warnings - we should maybe remove this at some point and add UE_LOG(, Fatal, ) to default cases.</remarks>
 		[WarningsClangToolChain(["-Wno-switch"], ["-Wswitch", "-Wno-error=switch"], ["-Wswitch"])]
-		[WarningLevelDefault(WarningLevel.Off)]
+		[BasicWarningLevelDefault(WarningLevel.Off)]
 		public WarningLevel SwitchWarningLevel
 		{
 			get => ResolveWarning(ParentCppCompileWarnings?.SwitchWarningLevel, _switchWarningLevel);
@@ -443,7 +485,7 @@ namespace UnrealBuildTool
 		/// </summary>
 		/// <remarks>This hides the "warning : comparison of unsigned expression &lt; 0 is always false" type warnings due to constant comparisons, which are possible with template arguments/</remarks>
 		[WarningsClangToolChain(["-Wno-tautological-compare"], ["-Wtautological-compare", "Wno-error=tautological-compare"], ["-Wtautological-compare"])]
-		[WarningLevelDefault(WarningLevel.Off)]
+		[BasicWarningLevelDefault(WarningLevel.Off)]
 		public WarningLevel TautologicalCompareWarningLevel
 		{
 			get => ResolveWarning(ParentCppCompileWarnings?.TautologicalCompareWarningLevel, _tautologicalCompareWarningLevel);
@@ -458,7 +500,7 @@ namespace UnrealBuildTool
 		/// </summary>
 		/// <remarks>Slate triggers this (with its optimize on/off pragmas).</remarks>
 		[WarningsClangToolChain(["-Wno-unknown-pragmas"], ["-Wunknown-pragmas", "-Wno-error=unknown-pragmas"], ["-Wunknown-pragmas"])]
-		[WarningLevelDefault(WarningLevel.Off)]
+		[BasicWarningLevelDefault(WarningLevel.Off)]
 		public WarningLevel UnknownPragmasWarningLevel
 		{
 			get => ResolveWarning(ParentCppCompileWarnings?.UnknownPragmasWarningLevel, _unknownPragmasWarningLevel);
@@ -487,7 +529,7 @@ namespace UnrealBuildTool
 		["-Wno-unused-but-set-variable", "-Wno-unused-but-set-parameter", "-Wno-unused-function", "-Wno-unused-lambda-capture", "-Wno-unused-local-typedef", "-Wno-unused-private-field", "-Wno-unused-variable"],
 		["-Wunused-but-set-variable", "-Wno-error=unused-but-set-variable", "-Wunused-but-set-parameter", "-Wno-error=unused-but-set-parameter", "-Wunused-function", "-Wno-error=unused-function", "-Wunused-lambda-capture", "-Wno-error=unused-lambda-capture", "-Wunused-local-typedef", "-Wno-error=unused-local-typedef", "-Wunused-private-field", "-Wno-error=unused-private-field", "-Wunused-variable", "-Wno-error=unused-variable"],
 		["-Wunused-but-set-variable", "-Wunused-but-set-parameter", "-Wunused-function", "-Wunused-lambda-capture", "-Wunused-local-typedef", "-Wunused-private-field", "-Wunused-variable"])]
-		[WarningLevelDefault(WarningLevel.Off)]
+		[BasicWarningLevelDefault(WarningLevel.Off)]
 		public WarningLevel UnusedWarningLevel
 		{
 			get => ResolveWarning(ParentCppCompileWarnings?.UnusedWarningLevel, _unusedWarningLevel);
@@ -502,7 +544,7 @@ namespace UnrealBuildTool
 		/// </summary>
 		/// <remarks>Not a good warning to be disabling.</remarks>
 		[WarningsClangToolChain(["-Wno-undefined-var-template"], ["-Wundefined-var-template", "-Wno-error=undefined-var-template"], ["-Wundefined-var-template"])]
-		[WarningLevelDefault(WarningLevel.Off)]
+		[BasicWarningLevelDefault(WarningLevel.Off)]
 		public WarningLevel UndefinedVarTemplateWarningLevel
 		{
 			get => ResolveWarning(ParentCppCompileWarnings?.UndefinedVarTemplateWarningLevel, _undefinedVarTemplateWarningLevel);
@@ -518,7 +560,7 @@ namespace UnrealBuildTool
 		/// </summary>
 		/// <remarks> Clang emits warnings for each compiled object file that doesn't have a matching entry in the profile data. This can happen when the profile data is older than the binaries we're compiling. Disable these warnings.</remarks>
 		[WarningsClangToolChain(["-Wno-profile-instr-out-of-date", "-Wno-profile-instr-unprofiled"], ["-Wprofile-instr-out-of-date", "-Wno-error=profile-instr-out-of-date", "-Wprofile-instr-unprofiled", "-Wno-error=profile-instr-unprofiled"], ["-Wprofile-instr-out-of-date", "-Wprofile-instr-unprofiled"], [nameof(StandardFilters.PGOOptimizedFilter)])]
-		[WarningLevelDefault(WarningLevel.Off)]
+		[BasicWarningLevelDefault(WarningLevel.Off)]
 		public WarningLevel ProfileInstructWarningLevel
 		{
 			get => ResolveWarning(ParentCppCompileWarnings?.ProfileInstructWarningLevel, _profileInstructWarningLevel);
@@ -532,7 +574,7 @@ namespace UnrealBuildTool
 		/// </summary>
 		/// <remarks>Apparently there can be hashing conflicts with PGO which can result in: 'Function control flow change detected (hash mismatch)' warnings. </remarks>
 		[WarningsClangToolChain(["-Wno-backend-plugin"], ["-Wbackend-plugin", "-Wno-error=backend-plugin"], ["-Wbackend-plugin"], [nameof(StandardFilters.PGOOptimizedFilter)])]
-		[WarningLevelDefault(WarningLevel.Off)]
+		[BasicWarningLevelDefault(WarningLevel.Off)]
 		public WarningLevel BackendPluginWarningLevel
 		{
 			get => ResolveWarning(ParentCppCompileWarnings?.BackendPluginWarningLevel, _backendPluginWarningLevel);
@@ -550,7 +592,7 @@ namespace UnrealBuildTool
 		/// -Wunused-result - we will always re-enable this as it's controlled by the same set. This will be moved out to a separate property. It is on by default, so it only must be re-enabled in the disable context.
 		/// </remarks>
 		[UnusedValueClangToolChain(["-Wno-unused-value", "-Wunused-result"], ["-Wunused-value", "-Wno-error=unused-value", "-Wunused-result"], ["-Wunused-value"])]
-		[WarningLevelDefault(WarningLevel.Off)]
+		[BasicWarningLevelDefault(WarningLevel.Off)]
 		public WarningLevel UnusedValueWarningLevel
 		{
 			get => ResolveWarning(ParentCppCompileWarnings?.UnusedValueWarningLevel, _unusedValueWarningLevel);
@@ -566,7 +608,7 @@ namespace UnrealBuildTool
 		/// <remarks>-Wimplicit-int-conversion also controls -Wshorten-64-to-32 as of clang 19.</remarks>
 		/// <remarks><see cref="ArgPosition.End"/> due to it's relation to <see cref="UnsafeTypeCastWarningLevel"/>.</remarks>
 		[WarningsClangToolChain(["-Wno-shorten-64-to-32"], ["-Wshorten-64-to-32", "-Wno-error=shorten-64-to-32"], ["-Wshorten-64-to-32"], [nameof(FilterID.Version19Min)], ArgPosition.End)]
-		[WarningLevelDefault(WarningLevel.Off)]
+		[BasicWarningLevelDefault(WarningLevel.Off)]
 		public WarningLevel Shorten64To32WarningLevel
 		{
 			get => ResolveWarning(ParentCppCompileWarnings?.Shorten64To32WarningLevel, _shorten64To32WarningLevel);
@@ -581,7 +623,7 @@ namespace UnrealBuildTool
 		/// </summary>
 		/// <remarks>The code base contains lots of places where we do "extern template X_API class ..." and we want to keep doing that to reduce compile times.</remarks>
 		[WarningsClangToolChain(["-Wno-dllexport-explicit-instantiation-decl"], ["-Wdllexport-explicit-instantiation-decl", "-Wno-error=dllexport-explicit-instantiation-decl"], ["-Wdllexport-explicit-instantiation-decl"])]
-		[WarningLevelDefault(WarningLevel.Off)]
+		[BasicWarningLevelDefault(WarningLevel.Off)]
 		public WarningLevel DllExportExplicitInstantiationDeclWarningLevel
 		{
 			get => ResolveWarning(ParentCppCompileWarnings?.DllExportExplicitInstantiationDeclWarningLevel, _dllExportExplicitInstantiationDeclWarningLevel);
@@ -595,7 +637,7 @@ namespace UnrealBuildTool
 		///		https://clang.llvm.org/docs/DiagnosticsReference.html#wnontrivial-memaccess
 		/// </summary>
 		[WarningsClangToolChain(["-Wno-nontrivial-memaccess"], ["-Wnontrivial-memaccess", "-Wno-error=nontrivial-memaccess"], ["-Wnontrivial-memaccess"], [nameof(FilterID.Version20Min)])]
-		[WarningLevelDefault(WarningLevel.Off)]
+		[BasicWarningLevelDefault(WarningLevel.Off)]
 		public WarningLevel NonTrivialMemAccessWarningLevel
 		{
 			get => ResolveWarning(ParentCppCompileWarnings?.NonTrivialMemAccessWarningLevel, _nonTrivialMemAccessWarningLevel);
@@ -610,7 +652,7 @@ namespace UnrealBuildTool
 		/// </summary>
 		/// <remarks>Allow Microsoft-specific syntax to slide, even though it may be non-standard.  Needed for Windows headers.</remarks>
 		[WarningsVCClang(["-Wno-microsoft"], ["-Wmicrosoft", "-Wno-error=microsoft"], ["-Wmicrosoft"])]
-		[WarningLevelDefault(WarningLevel.Off)]
+		[BasicWarningLevelDefault(WarningLevel.Off)]
 		public WarningLevel MicrosoftGroupWarningLevel
 		{
 			get => ResolveWarning(ParentCppCompileWarnings?.MicrosoftGroupWarningLevel, _microsoftGroupWarningLevel);
@@ -625,7 +667,7 @@ namespace UnrealBuildTool
 		/// </summary>
 		/// <remarks> Hack due to how we have our 'DummyPCH' wrappers setup when using unity builds.  This warning should not be disabled!</remarks>
 		[WarningsVCClang(["-Wno-msvc-include"], ["-Wmsvc-include", "-Wno-error=msvc-include"], ["-Wmsvc-include"])]
-		[WarningLevelDefault(WarningLevel.Off)]
+		[BasicWarningLevelDefault(WarningLevel.Off)]
 		public WarningLevel MSVCIncludeWarningLevel
 		{
 			get => ResolveWarning(ParentCppCompileWarnings?.MSVCIncludeWarningLevel, _msvcIncludeWarningLevel);
@@ -643,7 +685,7 @@ namespace UnrealBuildTool
 		/// @TODO clang: push/pop this in  Pre/PostWindowsApi.h headers instead?
 		/// </remarks>
 		[WarningsVCClang(["-Wno-pragma-pack"], ["-Wpragma-pack", "-Wno-error=pragma-pack"], ["-Wpragma-pack"])]
-		[WarningLevelDefault(WarningLevel.Off)]
+		[BasicWarningLevelDefault(WarningLevel.Off)]
 		public WarningLevel PragmaPackWarningLevel
 		{
 			get => ResolveWarning(ParentCppCompileWarnings?.PragmaPackWarningLevel, _pragmaPackWarningLevel);
@@ -658,7 +700,7 @@ namespace UnrealBuildTool
 		/// </summary>
 		/// <remarks>@todo clang: We declare operator new as inline.  Clang doesn't seem to like that.</remarks>
 		[WarningsVCClang(["-Wno-inline-new-delete"], ["-Winline-new-delete", "-Wno-error=inline-new-delete"], ["-Winline-new-delete"])]
-		[WarningLevelDefault(WarningLevel.Off)]
+		[BasicWarningLevelDefault(WarningLevel.Off)]
 		public WarningLevel InlineNewDeleteWarningLevel
 		{
 			get => ResolveWarning(ParentCppCompileWarnings?.InlineNewDeleteWarningLevel, _inlineNewDeleteWarningLevel);
@@ -673,7 +715,7 @@ namespace UnrealBuildTool
 		/// </summary>
 
 		[WarningsVCClang(["-Wno-implicit-exception-spec-mismatch"], ["-Wimplicit-exception-spec-mismatch", "-Wno-error=implicit-exception-spec-mismatch"], ["-Wimplicit-exception-spec-mismatch"])]
-		[WarningLevelDefault(WarningLevel.Off)]
+		[BasicWarningLevelDefault(WarningLevel.Off)]
 		public WarningLevel ImplicitExceptionSpecMismatchWarningLevel
 		{
 			get => ResolveWarning(ParentCppCompileWarnings?.ImplicitExceptionSpecMismatchWarningLevel, _implicitExceptionSpecMismatchWarningLevel);
@@ -688,7 +730,7 @@ namespace UnrealBuildTool
 		/// </summary>
 		/// <remarks>Sometimes we compare 'this' pointers against nullptr, which Clang warns about by default.</remarks>
 		[WarningsVCClang(["-Wno-undefined-bool-conversion"], ["-Wundefined-bool-conversion", "-Wno-error=undefined-bool-conversion"], ["-Wundefined-bool-conversion"])]
-		[WarningLevelDefault(WarningLevel.Off)]
+		[BasicWarningLevelDefault(WarningLevel.Off)]
 		public WarningLevel UndefinedBoolConversionWarningLevel
 		{
 			get => ResolveWarning(ParentCppCompileWarnings?.UndefinedBoolConversionWarningLevel, _undefinedBoolConversionWarningLevel);
@@ -701,9 +743,8 @@ namespace UnrealBuildTool
 		/// Clang -
 		///		https://clang.llvm.org/docs/DiagnosticsReference.html#wdeprecated-writable-strings
 		/// </summary>
-
 		[WarningsVCClang(["-Wno-deprecated-writable-strings"], ["-Wdeprecated-writable-strings", "-Wno-error=deprecated-writable-strings"], ["-Wdeprecated-writable-strings"])]
-		[WarningLevelDefault(WarningLevel.Off)]
+		[BasicWarningLevelDefault(WarningLevel.Off)]
 		public WarningLevel DeprecatedWriteableStringsWarningLevel
 		{
 			get => ResolveWarning(ParentCppCompileWarnings?.DeprecatedWriteableStringsWarningLevel, _deprecatedWriteableStringsWarningLevel);
@@ -717,7 +758,7 @@ namespace UnrealBuildTool
 		///		https://clang.llvm.org/docs/DiagnosticsReference.html#wdeprecated-register
 		/// </summary>
 		[WarningsVCClang(["-Wno-deprecated-register"], ["-Wdeprecated-register", "-Wno-error=deprecated-register"], ["-Wdeprecated-register"])]
-		[WarningLevelDefault(WarningLevel.Off)]
+		[BasicWarningLevelDefault(WarningLevel.Off)]
 		public WarningLevel DeprecatedRegisterWarningLevel
 		{
 			get => ResolveWarning(ParentCppCompileWarnings?.DeprecatedRegisterWarningLevel, _deprecatedRegisterWarningLevel);
@@ -733,7 +774,7 @@ namespace UnrealBuildTool
 		/// </summary>
 		/// <remarks>Disable needed for external headers we shan't change.</remarks>
 		[WarningsVCClang(["-Wno-logical-op-parentheses"], ["-Wlogical-op-parentheses", "-Wno-error=logical-op-parentheses"], ["-Wlogical-op-parentheses"])]
-		[WarningLevelDefault(WarningLevel.Off)]
+		[BasicWarningLevelDefault(WarningLevel.Off)]
 		public WarningLevel LogicalOpParenthesesWarningLevel
 		{
 			get => ResolveWarning(ParentCppCompileWarnings?.LogicalOpParenthesesWarningLevel, _logicalOpParenthesesWarningLevel);
@@ -748,7 +789,7 @@ namespace UnrealBuildTool
 		/// </summary>
 		/// <remarks>Disable needed for external headers we shan't change.</remarks>
 		[WarningsVCClang(["-Wno-null-arithmetic"], ["-Wnull-arithmetic", "-Wno-error=null-arithmetic"], ["-Wnull-arithmetic"])]
-		[WarningLevelDefault(WarningLevel.Off)]
+		[BasicWarningLevelDefault(WarningLevel.Off)]
 		public WarningLevel NullArithmeticWarningLevel
 		{
 			get => ResolveWarning(ParentCppCompileWarnings?.NullArithmeticWarningLevel, _nullArithmeticWarningLevel);
@@ -763,7 +804,7 @@ namespace UnrealBuildTool
 		/// </summary>
 		/// <remarks>Disable needed for PhysX.</remarks>
 		[WarningsVCClang(["-Wno-return-type-c-linkage"], ["-Wreturn-type-c-linkage", "-Wno-error=return-type-c-linkage"], ["-Wreturn-type-c-linkage"])]
-		[WarningLevelDefault(WarningLevel.Off)]
+		[BasicWarningLevelDefault(WarningLevel.Off)]
 		public WarningLevel ReturnTypeCLinkageWarningLevel
 		{
 			get => ResolveWarning(ParentCppCompileWarnings?.ReturnTypeCLinkageWarningLevel, _returnTypeCLinkageWarningLevel);
@@ -778,7 +819,7 @@ namespace UnrealBuildTool
 		/// </summary>
 		/// <remarks>Disable needed for nvtesslib.</remarks>
 		[WarningsVCClang(["-Wno-ignored-attributes"], ["-Wignored-attributes", "-Wno-error=ignored-attributes"], ["-Wignored-attributes"])]
-		[WarningLevelDefault(WarningLevel.Off)]
+		[BasicWarningLevelDefault(WarningLevel.Off)]
 		public WarningLevel IgnoredAttributesWarningLevel
 		{
 			get => ResolveWarning(ParentCppCompileWarnings?.IgnoredAttributesWarningLevel, _ignoredAttributesWarningLevel);
@@ -792,7 +833,7 @@ namespace UnrealBuildTool
 		///		https://clang.llvm.org/docs/DiagnosticsReference.html#wuninitialized
 		/// </summary>
 		[WarningsVCClang(["-Wno-uninitialized"], ["-Wuninitialized", "-Wno-error=uninitialized"], ["-Wuninitialized"])]
-		[WarningLevelDefault(WarningLevel.Off)]
+		[BasicWarningLevelDefault(WarningLevel.Off)]
 		public WarningLevel UninitializedWarningLevel
 		{
 			get => ResolveWarning(ParentCppCompileWarnings?.UninitializedWarningLevel, _uninitializedWarningLevel);
@@ -807,7 +848,7 @@ namespace UnrealBuildTool
 		/// </summary>
 		/// <remarks>Disable needed for external headers we shan't change.</remarks>
 		[WarningsVCClang(["-Wno-return-type"], ["-Wreturn-type", "-Wno-error=return-type"], ["-Wreturn-type"])]
-		[WarningLevelDefault(WarningLevel.Off)]
+		[BasicWarningLevelDefault(WarningLevel.Off)]
 		public WarningLevel ReturnTypeWarningLevel
 		{
 			get => ResolveWarning(ParentCppCompileWarnings?.ReturnTypeWarningLevel, _returnTypeWarningLevel);
@@ -822,7 +863,7 @@ namespace UnrealBuildTool
 		/// </summary>
 		/// <remarks>Unused function parameter. A lot are named 'bUnused'...</remarks>
 		[WarningsVCClang(["-Wno-unused-parameter"], ["-Wunused-parameter", "-Wno-error=unused-parameter"], ["-Wunused-parameter"])]
-		[WarningLevelDefault(WarningLevel.Off)]
+		[BasicWarningLevelDefault(WarningLevel.Off)]
 		public WarningLevel UnusedParameterWarningLevel
 		{
 			get => ResolveWarning(ParentCppCompileWarnings?.UnusedParameterWarningLevel, _unusedParameterWarningLevel);
@@ -837,7 +878,7 @@ namespace UnrealBuildTool
 		/// </summary>
 		/// <remarks>const ignored when returning by value e.g. 'const int foo() { return 4; }'.</remarks>
 		[WarningsVCClang(["-Wno-ignored-qualifiers"], ["-Wignored-qualifiers", "-Wno-error=ignored-qualifiers"], ["-Wignored-qualifiers"])]
-		[WarningLevelDefault(WarningLevel.Off)]
+		[BasicWarningLevelDefault(WarningLevel.Off)]
 		public WarningLevel IgnoredQualifiersWarningLevel
 		{
 			get => ResolveWarning(ParentCppCompileWarnings?.IgnoredQualifiersWarningLevel, _ignoredQualifiersWarningLevel);
@@ -852,7 +893,7 @@ namespace UnrealBuildTool
 		/// </summary>
 		/// <remarks>Usage of 'defined(X)' in a macro definition. Gives different results under MSVC.</remarks>
 		[WarningsVCClang(["-Wno-expansion-to-defined"], ["-Wexpansion-to-defined", "-Wno-error=expansion-to-defined"], ["-Wexpansion-to-defined"])]
-		[WarningLevelDefault(WarningLevel.Off)]
+		[BasicWarningLevelDefault(WarningLevel.Off)]
 		public WarningLevel ExpansionToDefined
 		{
 			get => ResolveWarning(ParentCppCompileWarnings?.ExpansionToDefined, _expansionToDefined);
@@ -867,7 +908,7 @@ namespace UnrealBuildTool
 		/// </summary>
 		/// <remarks>Signed/unsigned comparison - millions of these.</remarks>
 		[WarningsVCClang(["-Wno-sign-compare"], ["-Wsign-compare", "-Wno-error=sign-compare"], ["-Wsign-compare"])]
-		[WarningLevelDefault(WarningLevel.Off)]
+		[BasicWarningLevelDefault(WarningLevel.Off)]
 		public WarningLevel SignCompareWarningLevel
 		{
 			get => ResolveWarning(ParentCppCompileWarnings?.SignCompareWarningLevel, _signCompareWarningLevel);
@@ -882,7 +923,7 @@ namespace UnrealBuildTool
 		/// </summary>
 		/// <remarks>Excessive warning, generated when you initialize with MyStruct A = {0};</remarks>
 		[WarningsVCClang(["-Wno-missing-field-initializers"], ["-Wmissing-field-initializers", "-Wno-error=missing-field-initializers"], ["-Wmissing-field-initializers"])]
-		[WarningLevelDefault(WarningLevel.Off)]
+		[BasicWarningLevelDefault(WarningLevel.Off)]
 		public WarningLevel MissingFieldInitializersWarningLevel
 		{
 			get => ResolveWarning(ParentCppCompileWarnings?.MissingFieldInitializersWarningLevel, _missingFieldInitializersWarningLevel);
@@ -896,7 +937,7 @@ namespace UnrealBuildTool
 		///		https://clang.llvm.org/docs/DiagnosticsReference.html#wnonportable-include-path
 		/// </summary>
 		[WarningsVCClang(["-Wno-nonportable-include-path"], ["-Wnonportable-include-path", "-Wno-error=nonportable-include-path"], ["-Wnonportable-include-path"])]
-		[WarningLevelDefault(WarningLevel.Off)]
+		[BasicWarningLevelDefault(WarningLevel.Off)]
 		public WarningLevel NonPortableIncludePathWarningLevel
 		{
 			get => ResolveWarning(ParentCppCompileWarnings?.NonPortableIncludePathWarningLevel, _nonPortableIncludePathWarningLevel);
@@ -912,7 +953,7 @@ namespace UnrealBuildTool
 		///		https://learn.microsoft.com/en-us/cpp/preprocessor/token-pasting-operator-hash-hash?view=msvc-170
 		/// </summary>
 		[WarningsVCClang(["-Wno-invalid-token-paste"], ["-Winvalid-token-paste", "-Wno-error=invalid-token-paste"], ["-Winvalid-token-paste"])]
-		[WarningLevelDefault(WarningLevel.Off)]
+		[BasicWarningLevelDefault(WarningLevel.Off)]
 		public WarningLevel InvalidTokenPasteWarningLevel
 		{
 			get => ResolveWarning(ParentCppCompileWarnings?.InvalidTokenPasteWarningLevel, _invalidTokenPasteWarningLevel);
@@ -926,7 +967,7 @@ namespace UnrealBuildTool
 		///		https://clang.llvm.org/docs/DiagnosticsReference.html#wnull-pointer-arithmetic
 		/// </summary>
 		[WarningsVCClang(["-Wno-null-pointer-arithmetic"], ["-Wnull-pointer-arithmetic", "-Wno-error=null-pointer-arithmetic"], ["-Wnull-pointer-arithmetic"])]
-		[WarningLevelDefault(WarningLevel.Off)]
+		[BasicWarningLevelDefault(WarningLevel.Off)]
 		public WarningLevel NullPointerArithmeticWarningLevel
 		{
 			get => ResolveWarning(ParentCppCompileWarnings?.NullPointerArithmeticWarningLevel, _nullPointerArithmeticWarningLevel);
@@ -941,7 +982,7 @@ namespace UnrealBuildTool
 		/// </summary>
 		/// <remarks>Triggered by || of two template-derived values inside a static_assert.</remarks>
 		[WarningsVCClang(["-Wno-constant-logical-operand"], ["-Wconstant-logical-operand", "-Wno-error=constant-logical-operand"], ["-Wconstant-logical-operand"])]
-		[WarningLevelDefault(WarningLevel.Off)]
+		[BasicWarningLevelDefault(WarningLevel.Off)]
 		public WarningLevel ConstantLogicalOperandWarningLevel
 		{
 			get => ResolveWarning(ParentCppCompileWarnings?.ConstantLogicalOperandWarningLevel, _constantLogicalOperandWarningLevel);
@@ -955,7 +996,7 @@ namespace UnrealBuildTool
 		///		https://clang.llvm.org/docs/DiagnosticsReference.html#wbitfield-enum-conversion
 		/// </summary>
 		[WarningsVCClang(["-Wno-bitfield-enum-conversion"], ["-Wbitfield-enum-conversion", "-Wno-error=bitfield-enum-conversion"], ["-Wbitfield-enum-conversion"])]
-		[WarningLevelDefault(WarningLevel.Off)]
+		[BasicWarningLevelDefault(WarningLevel.Off)]
 		public WarningLevel BitfieldEnumConversion
 		{
 			get => ResolveWarning(ParentCppCompileWarnings?.BitfieldEnumConversion, _bitfieldEnumConversion);
@@ -969,7 +1010,7 @@ namespace UnrealBuildTool
 		///		https://clang.llvm.org/docs/DiagnosticsReference.html#wnull-pointer-subtraction
 		/// </summary>
 		[WarningsVCClang(["-Wno-null-pointer-subtraction"], ["-Wnull-pointer-subtraction", "-Wno-error=null-pointer-subtraction"], ["-Wnull-pointer-subtraction"])]
-		[WarningLevelDefault(WarningLevel.Off)]
+		[BasicWarningLevelDefault(WarningLevel.Off)]
 		public WarningLevel NullPointerSubtractionWarningLevel
 		{
 			get => ResolveWarning(ParentCppCompileWarnings?.NullPointerSubtractionWarningLevel, _nullPointerSubtractionWarningLevel);
@@ -983,7 +1024,7 @@ namespace UnrealBuildTool
 		///		https://clang.llvm.org/docs/DiagnosticsReference.html#wdangling
 		/// </summary>	
 		[WarningsVCClang(["-Wno-dangling"], ["-Wdangling", "-Wno-error=dangling"], ["-Wdangling"])]
-		[WarningLevelDefault(WarningLevel.Off)]
+		[BasicWarningLevelDefault(WarningLevel.Off)]
 		public WarningLevel DanglingWarningLevel
 		{
 			get => ResolveWarning(ParentCppCompileWarnings?.DanglingWarningLevel, _danglingWarningLevel);
@@ -991,7 +1032,7 @@ namespace UnrealBuildTool
 		}
 		private WarningLevel _danglingWarningLevel;
 
-		#region -- Duplicated Clang VC Compiler Disables --
+		#region -- Duplicated Clang VC Compiler --
 
 		/// <summary>
 		/// Indicates what warning/error level to treat unhandled enumerators in switches on enumeration-typed values.
@@ -1007,7 +1048,7 @@ namespace UnrealBuildTool
 		/// In order to ensure no change of behavior for MSVC toolchain within the clang compilation context, we will keep this in, and deprecate with assignments to the clang setting.
 		/// </remarks>
 		[WarningsVCClang(["-Wno-switch-enum"])]
-		[WarningLevelDefault(WarningLevel.Off)]
+		[BasicWarningLevelDefault(WarningLevel.Off)]
 		public WarningLevel MSVCSwitchEnumWarningLevel
 		{
 			get => ResolveWarning(ParentCppCompileWarnings?.MSVCSwitchEnumWarningLevel, _msvcSwitchEnumWarningLevel);
@@ -1027,7 +1068,7 @@ namespace UnrealBuildTool
 		/// -Wunused-result - we will always re-enable this as it's controlled by the same set. This will be moved out to a separate property. It is on by default, so it only must be re-enabled in the disable context.
 		/// </remarks>
 		[WarningsVCClang(["-Wno-unused-value", "-Wunused-result"], ["-Wunused-value", "-Wno-error=unused-value", "-Wunused-result"], ["-Wunused-value"])]
-		[WarningLevelDefault(WarningLevel.Off)]
+		[BasicWarningLevelDefault(WarningLevel.Off)]
 		public WarningLevel MSVCUnusedValueWarningLevel
 		{
 			get => ResolveWarning(ParentCppCompileWarnings?.MSVCUnusedValueWarningLevel, _msvcUnusedValueWarningLevel);
@@ -1046,7 +1087,7 @@ namespace UnrealBuildTool
 		/// In order to ensure no change of behavior for MSVC toolchain within the clang compilation context, we will keep this in, and deprecate with assignments to the clang setting.
 		/// </remarks>
 		[WarningsVCClang(["-Wno-deprecated-declarations"], ["-Wdeprecated-declarations", "-Wno-error=deprecated-declarations"], ["-Wdeprecated-declarations"])]
-		[WarningLevelDefault(WarningLevel.Off)]
+		[BasicWarningLevelDefault(WarningLevel.Off)]
 		public WarningLevel MSVCDeprecationWarningLevel
 		{
 			get => ResolveWarning(ParentCppCompileWarnings?.MSVCDeprecationWarningLevel, _msvcDeprecationWarningLevel);
@@ -1054,17 +1095,17 @@ namespace UnrealBuildTool
 		}
 		private WarningLevel _msvcDeprecationWarningLevel;
 
-		#endregion -- Duplicated Clang VC Compiler Disables --
+		#endregion -- Duplicated Clang VC Compiler --
 
-		#endregion -- Clang VC Compiler Disables --
+		#endregion -- Clang Compiler --
 
-		#region -- Intel Compiler Disables --
+		#region -- Intel Compiler --
 
 		/// <summary>
 		/// Indicates what warning/error level to treat instances of mismatched printf style format specifiers.
 		/// </summary>
-		[WarningsIntelCompiler(["-Wno-format"])]
-		[WarningLevelDefault(WarningLevel.Off)]
+		[WarningsIntelCompiler(["-Wno-format"], ["-Wformat", "-Wno-error=format"], ["-Wformat"])]
+		[BasicWarningLevelDefault(WarningLevel.Off)]
 		public WarningLevel FormatWarningLevel
 		{
 			get => ResolveWarning(ParentCppCompileWarnings?.FormatWarningLevel, _formatWarningLevel);
@@ -1075,8 +1116,8 @@ namespace UnrealBuildTool
 		/// <summary>
 		/// Indicates what warning/error level to treat instances of implicit conversion between floating-point types of different precision.
 		/// </summary>
-		[WarningsIntelCompiler(["-Wno-implicit-float-size-conversion"])]
-		[WarningLevelDefault(WarningLevel.Off)]
+		[WarningsIntelCompiler(["-Wno-implicit-float-size-conversion"], ["-Wimplicit-float-size-conversion", "-Wno-error=implicit-float-size-conversion"], ["-Wimplicit-float-size-conversion"])]
+		[BasicWarningLevelDefault(WarningLevel.Off)]
 		public WarningLevel ImplicitFloatSizeConversionWarningLevel
 		{
 			get => ResolveWarning(ParentCppCompileWarnings?.ImplicitFloatSizeConversionWarningLevel, _implicitFloatSizeConversionWarningLevel);
@@ -1087,8 +1128,8 @@ namespace UnrealBuildTool
 		/// <summary>
 		/// Indicates what warning/error level to treat instances where a command line argument is provided, but not used for anything.
 		/// </summary>
-		[WarningsIntelCompiler(["-Wno-single-bit-bitfield-constant-conversion"])]
-		[WarningLevelDefault(WarningLevel.Off)]
+		[WarningsIntelCompiler(["-Wno-single-bit-bitfield-constant-conversion"], ["-Wsingle-bit-bitfield-constant-conversion", "-Wno-error=single-bit-bitfield-constant-conversion"], ["-Wsingle-bit-bitfield-constant-conversion"])]
+		[BasicWarningLevelDefault(WarningLevel.Off)]
 		public WarningLevel SingleBitfieldConstantConversionWarningLevel
 		{
 			get => ResolveWarning(ParentCppCompileWarnings?.SingleBitfieldConstantConversionWarningLevel, _singleBitfieldConstantConversionWarningLevel);
@@ -1100,8 +1141,8 @@ namespace UnrealBuildTool
 		/// 
 		/// Clang - https://clang.llvm.org/docs/DiagnosticsReference.html#wunused-command-line-argument
 		/// </summary>
-		[WarningsIntelCompiler(["-Wno-unused-command-line-argument"])]
-		[WarningLevelDefault(WarningLevel.Off)]
+		[WarningsIntelCompiler(["-Wno-unused-command-line-argument"], ["-Wunused-command-line-argument", "-Wno-error=unused-command-line-argument"], ["-Wunused-command-line-argument"])]
+		[BasicWarningLevelDefault(WarningLevel.Off)]
 		public WarningLevel UnusedCommandLineArgumentWarningLevel
 		{
 			get => ResolveWarning(ParentCppCompileWarnings?.UnusedCommandLineArgumentWarningLevel, _unusedCommandLineArgumentWarningLevel);
@@ -1113,8 +1154,8 @@ namespace UnrealBuildTool
 		/// Indicates what warning/error level to treat instances of malformed comments.
 		///  Clang - https://clang.llvm.org/docs/DiagnosticsReference.html#wcomment
 		/// </summary>
-		[WarningsIntelCompiler(["-Wno-comment"])]
-		[WarningLevelDefault(WarningLevel.Off)]
+		[WarningsIntelCompiler(["-Wno-comment"], ["-Wcomment", "-Wno-error=comment"], ["-Wcomment"])]
+		[BasicWarningLevelDefault(WarningLevel.Off)]
 		public WarningLevel CommentWarningLevel
 		{
 			get => ResolveWarning(ParentCppCompileWarnings?.CommentWarningLevel, _commentWarningLevel);
@@ -1127,8 +1168,8 @@ namespace UnrealBuildTool
 		/// Indicates what warning/error level to treat instances of copy construct in range loop iteration.
 		///  Clang - https://clang.llvm.org/docs/DiagnosticsReference.html#wrange-loop-construct
 		/// </summary>
-		[WarningsIntelCompiler(["-Wno-range-loop-construct"])]
-		[WarningLevelDefault(WarningLevel.Off)]
+		[WarningsIntelCompiler(["-Wno-range-loop-construct"], ["-Wrange-loop-construct", "-Wno-error=Wrange-loop-construct"], ["-Wrange-loop-construct"])]
+		[BasicWarningLevelDefault(WarningLevel.Off)]
 		public WarningLevel RangeLoopConstructWarningLevel
 		{
 			get => ResolveWarning(ParentCppCompileWarnings?.RangeLoopConstructWarningLevel, _rangeLoopConstructWarningLevel);
@@ -1140,8 +1181,8 @@ namespace UnrealBuildTool
 		/// Indicates what warning/error level to treat instances where a #pragma once is used in a source file instead of a header.
 		/// Clang - https://clang.llvm.org/docs/DiagnosticsReference.html#wpragma-once-outside-header
 		/// </summary>
-		[WarningsIntelCompiler(["-Wno-pragma-once-outside-header"])]
-		[WarningLevelDefault(WarningLevel.Off)]
+		[WarningsIntelCompiler(["-Wno-pragma-once-outside-header"], ["-Wpragma-once-outside-header", "-Wno-error=Wpragma-once-outside-header"], ["-Wpragma-once-outside-header"])]
+		[BasicWarningLevelDefault(WarningLevel.Off)]
 		public WarningLevel PragmaOnceOutsideHeaderWarningLevel
 		{
 			get => ResolveWarning(ParentCppCompileWarnings?.PragmaOnceOutsideHeaderWarningLevel, _pragmaOnceOutsideHeaderWarningLevel);
@@ -1153,8 +1194,8 @@ namespace UnrealBuildTool
 		/// Indicates what warning/error level to treat instances where the logical NOT operator (!) is used in a way that's ambiguous with respect to parentheses.
 		/// Clang - https://clang.llvm.org/docs/DiagnosticsReference.html#wlogical-not-parentheses
 		/// </summary>
-		[WarningsIntelCompiler(["-Wno-logical-not-parentheses"])]
-		[WarningLevelDefault(WarningLevel.Off)]
+		[WarningsIntelCompiler(["-Wno-logical-not-parentheses"], ["-Wlogical-not-parentheses", "-Wno-error=Wlogical-not-parentheses"], ["-Wlogical-not-parentheses"])]
+		[BasicWarningLevelDefault(WarningLevel.Off)]
 		public WarningLevel LogicalNotParenthesesWarningLevel
 		{
 			get => ResolveWarning(ParentCppCompileWarnings?.LogicalNotParenthesesWarningLevel, _logicalNotParenthesesWarningLevel);
@@ -1166,8 +1207,8 @@ namespace UnrealBuildTool
 		///  Indicates what warning/error level to treat instances where a C++20 feature is being used, but compiling with an older standard.
 		/// Clang - https://clang.llvm.org/docs/DiagnosticsReference.html#wc-20-extensions
 		/// </summary>
-		[WarningsIntelCompiler(["-Wno-c++20-extensions"])]
-		[WarningLevelDefault(WarningLevel.Off)]
+		[WarningsIntelCompiler(["-Wno-c++20-extensions"], ["-Wc++20-extensions", "-Wno-error=c++20-extensions"], ["-Wc++20-extensions"])]
+		[BasicWarningLevelDefault(WarningLevel.Off)]
 		public WarningLevel Cpp20ExtensionsWarningLevel
 		{
 			get => ResolveWarning(ParentCppCompileWarnings?.Cpp20ExtensionsWarningLevel, _cpp20ExtensionsWarningLevel);
@@ -1175,7 +1216,7 @@ namespace UnrealBuildTool
 		}
 		private WarningLevel _cpp20ExtensionsWarningLevel;
 
-		#endregion -- Intel Compiler Disables --
+		#endregion -- Intel Compiler --
 
 		#region -- CppCompileEnvironment Settings
 
@@ -1188,7 +1229,7 @@ namespace UnrealBuildTool
 		[XmlConfigFile(Category = "BuildConfiguration", Name = nameof(DeprecationWarningLevel))]
 		[WarningsClangToolChain(["-Wno-deprecated-declarations"], ["-Wdeprecated-declarations", "-Wno-error=deprecated-declarations"], ["-Wdeprecated-declarations"])]
 		[WarningsVCToolChain(["/wd4996"], null, ["/we4996"])]
-		[WarningLevelDefault(WarningLevel.Warning)]
+		[BasicWarningLevelDefault(WarningLevel.Warning)]
 		public WarningLevel DeprecationWarningLevel
 		{
 			get => ResolveWarning(ParentCppCompileWarnings?.DeprecationWarningLevel, _deprecationWarningLevel);
@@ -1203,7 +1244,7 @@ namespace UnrealBuildTool
 		/// </summary>
 		[WarningsClangToolChain(["-Wno-date-time"], ["-Wdate-time", "-Wno-error=date-time"], ["-Wdate-time"], [nameof(StandardFilters.DeterministcFlagSetFilter)])]
 		[WarningsMSCV(["/wd5048"], null, ["/we5048"], [nameof(StandardFilters.DeterministcFlagSetFilter)])]
-		[WarningLevelDefault(WarningLevel.Off)]
+		[BasicWarningLevelDefault(WarningLevel.Off)]
 		public WarningLevel DeterministicWarningLevel
 		{
 			get => ResolveWarning(ParentCppCompileWarnings?.DeterministicWarningLevel, _deterministicWarningLevel);
@@ -1222,7 +1263,7 @@ namespace UnrealBuildTool
 		/// </summary>
 		[XmlConfigFile(Category = "BuildConfiguration")]
 		[CommandLine("-PCHPerformanceIssueWarningLevel=")]
-		[WarningLevelDefault(WarningLevel.Off)]
+		[BasicWarningLevelDefault(WarningLevel.Off)]
 		public WarningLevel PCHPerformanceIssueWarningLevel
 		{
 			get => ResolveWarning(ParentCppCompileWarnings?.PCHPerformanceIssueWarningLevel, _pchPerformanceIssueWarningLevel);
@@ -1239,7 +1280,7 @@ namespace UnrealBuildTool
 		/// </summary>
 		[XmlConfigFile(Category = "BuildConfiguration")]
 		[CommandLine("-ModuleUnsupportedWarningLevel=")]
-		[WarningLevelDefault(WarningLevel.Warning)]
+		[BasicWarningLevelDefault(WarningLevel.Warning)]
 		public WarningLevel ModuleUnsupportedWarningLevel
 		{
 			get => ResolveWarning(ParentCppCompileWarnings?.ModuleUnsupportedWarningLevel, _moduleUnsupportedWarningLevel);
@@ -1252,7 +1293,7 @@ namespace UnrealBuildTool
 		/// </summary>
 		[XmlConfigFile(Category = "BuildConfiguration")]
 		[CommandLine("-PluginModuleUnsupportedWarningLevel=")]
-		[WarningLevelDefault(WarningLevel.Off)]
+		[BasicWarningLevelDefault(WarningLevel.Off)]
 		public WarningLevel PluginModuleUnsupportedWarningLevel
 		{
 			get => ResolveWarning(ParentCppCompileWarnings?.PluginModuleUnsupportedWarningLevel, _pluginModuleUnsupportedWarningLevel);
@@ -1265,7 +1306,7 @@ namespace UnrealBuildTool
 		/// </summary>
 		[XmlConfigFile(Category = "BuildConfiguration")]
 		[CommandLine("-ModuleIncludePathWarningLevel=")]
-		[WarningLevelDefault(WarningLevel.Off)]
+		[BasicWarningLevelDefault(WarningLevel.Off)]
 		public WarningLevel ModuleIncludePathWarningLevel
 		{
 			get => ResolveWarning(ParentCppCompileWarnings?.ModuleIncludePathWarningLevel, _moduleIncludePathWarningLevel);
@@ -1278,7 +1319,7 @@ namespace UnrealBuildTool
 		/// </summary>
 		[XmlConfigFile(Category = "BuildConfiguration")]
 		[CommandLine("-ModuleIncludePrivateWarningLevel=")]
-		[WarningLevelDefault(WarningLevel.Off)]
+		[BasicWarningLevelDefault(WarningLevel.Off)]
 		public WarningLevel ModuleIncludePrivateWarningLevel
 		{
 			get => ResolveWarning(ParentCppCompileWarnings?.ModuleIncludePrivateWarningLevel, _moduleIncludePrivateWarningLevel);
@@ -1291,7 +1332,7 @@ namespace UnrealBuildTool
 		/// </summary>
 		[XmlConfigFile(Category = "BuildConfiguration")]
 		[CommandLine("-ModuleIncludeSubdirectoryWarningLevel=")]
-		[WarningLevelDefault(WarningLevel.Off)]
+		[BasicWarningLevelDefault(WarningLevel.Off)]
 		public WarningLevel ModuleIncludeSubdirectoryWarningLevel
 		{
 			get => ResolveWarning(ParentCppCompileWarnings?.ModuleIncludeSubdirectoryWarningLevel, _moduleIncludeSubdirectoryWarningLevel);
@@ -1302,7 +1343,7 @@ namespace UnrealBuildTool
 		/// <summary>
 		/// How to treat conflicts when a disabled plugin is being enabled by another plugin referencing it
 		/// </summary>
-		[WarningLevelDefault(WarningLevel.Default)]
+		[BasicWarningLevelDefault(WarningLevel.Default)]
 		public WarningLevel DisablePluginsConflictWarningLevel
 		{
 			get => ResolveWarning(ParentCppCompileWarnings?.DisablePluginsConflictWarningLevel, _disablePluginsConflictWarningLevel);
@@ -1313,7 +1354,7 @@ namespace UnrealBuildTool
 		/// <summary>
 		/// Enable warnings for when there are .gen.cpp files that could be inlined in a matching handwritten cpp file
 		/// </summary>
-		[WarningLevelDefault(WarningLevel.Off)]
+		[BasicWarningLevelDefault(WarningLevel.Off)]
 		public WarningLevel NonInlinedGenCppWarningLevel
 		{
 			get => ResolveWarning(ParentCppCompileWarnings?.NonInlinedGenCppWarningLevel, _nonInlinedGenCppWarningLevel);
@@ -1486,19 +1527,18 @@ namespace UnrealBuildTool
 		/// <param name="cppCompileWarnings">The CppCompileWarnings object to apply target defaults to.</param>
 		/// <param name="overwriteSourceValue">Whether to overwrite warning setting values, or set if default.</param>
 		/// <remarks>
-		/// To have items automatically participate in default setting, use <see cref="WarningLevelDefaultAttribute"/>, with <see cref="InitializationContext.Any"/> or <see cref="InitializationContext.Target"/>.
+		/// To have items automatically participate in default setting, use <see cref="BasicWarningLevelDefaultAttribute"/>, with <see cref="InitializationContext.Any"/> or <see cref="InitializationContext.Target"/>.
 		/// </remarks>
 		internal static void ApplyTargetDefaults(CppCompileWarnings cppCompileWarnings, bool overwriteSourceValue = false)
 		{
 			cppCompileWarnings._logger?.LogTrace("Applying default C++ Compiler warnings for the Target({TargetName}) context.", cppCompileWarnings.BuildContextProvider?._targetRulesPrivate?.Name);
 
-			BuildSettingsVersion buildSettingsVersion = cppCompileWarnings.BuildContextProvider != null ? cppCompileWarnings.BuildContextProvider._buildContext.GetBuildSettings() : BuildSettingsVersion.V1;
 			IEnumerable<MemberInfo> memberInfos = GetMembersWithAttributes(typeof(CppCompileWarnings), typeof(WarningLevelDefaultAttribute), typeof(WarningLevel));
 
 			foreach (MemberInfo mi in memberInfos)
 			{
-				IList<WarningLevelDefaultAttribute> attributes = [.. mi.GetCustomAttributes<WarningLevelDefaultAttribute>().Where(prop => prop.Context == InitializationContext.Any || prop.Context == InitializationContext.Target)];
-				WarningLevel defaultWarningLevel = WarningLevelDefaultAttribute.ResolveWarningLevelDefault(attributes, buildSettingsVersion);
+				WarningLevel defaultWarningLevel = WarningLevel.Default;
+				defaultWarningLevel = CoalesceDefaultWarningFromMemberInfo(mi, InitializationContext.Target, defaultWarningLevel, cppCompileWarnings.BuildContextProvider, cppCompileWarnings._logger);
 
 				// Obtain instance value, and see if we need to preserve the TargetRules' constructor setting, or apply the default.
 				WarningLevel currentWarningLevel = GetCurrentWarningLevel(mi, cppCompileWarnings);
@@ -1509,6 +1549,46 @@ namespace UnrealBuildTool
 			}
 		}
 
+		internal static WarningLevel CoalesceDefaultWarningFromMemberInfo(MemberInfo mi, InitializationContext defaultWarningContext, WarningLevel defaultWarningLevel, BuildSystemContext? buildSystemContext, ILogger? logger)
+		{
+			IList<WarningLevelDefaultAttribute> attributes = [.. mi.GetCustomAttributes<WarningLevelDefaultAttribute>().Where(prop => prop.Context == InitializationContext.Any || prop.Context == defaultWarningContext)];
+			Type fallbackType = typeof(BasicWarningLevelDefaultAttribute);
+			Dictionary<Type, List<WarningLevelDefaultAttribute>> groupedAttributes = attributes
+				.GroupBy(attr => attr.GetType())
+				.OrderBy(g => g.Key == fallbackType ? 1 : 0)
+				.ToDictionary(g => g.Key, g => g.ToList());
+
+			foreach (KeyValuePair<Type, List<WarningLevelDefaultAttribute>> item in groupedAttributes)
+			{
+				if (WarningLevelResolverRegistry.Resolvers.ContainsKey(item.Key.Name))
+				{
+					WarningLevel resolvedWarning = WarningLevelResolverRegistry.Resolvers[item.Key.Name].ResolveWarningLevelDefault(item.Value, buildSystemContext);
+
+					defaultWarningLevel = resolvedWarning;
+				}
+				else if (item.Value.Count == 1)
+				{
+					logger?.LogInformation("Was unable to find a {Resolver} for the attribute type: {AttributeTypeName}. Falling back to basic retrieval as there's only one attribute.", nameof(IWarningLevelResolver), item.Key.Name);
+
+					WarningLevel resolvedWarning = item.Value[0].GetDefaultLevel(buildSystemContext);
+
+					defaultWarningLevel = resolvedWarning;
+				}
+				else
+				{
+					logger?.LogError("Was unable to find a {Resolver} for the attribute type: {AttributeTypeName}, and there was multiple attributes of the same type.", nameof(IWarningLevelResolver), item.Key.Name);
+				}
+
+				if (defaultWarningLevel != WarningLevel.Default)
+				{
+					logger?.LogTrace("Obtained default value for ({Property}) via the attribute ({AttributeType}).", mi.Name, item.Key.Name);
+					break;
+				}
+			}
+
+			return defaultWarningLevel;
+		}
+
 		#endregion
 
 		#region -- Private API --
@@ -1519,7 +1599,7 @@ namespace UnrealBuildTool
 		/// <param name="parentWarningLevel">The parent <see cref="CppCompileWarnings"/> <see cref="WarningLevel"/>.</param>
 		/// <param name="instanceWarningLevel">The instance <see cref="WarningLevel"/>.</param>
 		/// <returns>The resolved warning level.</returns>
-		private static WarningLevel ResolveWarning(WarningLevel? parentWarningLevel, WarningLevel instanceWarningLevel)
+		private WarningLevel ResolveWarning(WarningLevel? parentWarningLevel, WarningLevel instanceWarningLevel)
 		{
 			WarningLevel resolvedWarning = WarningLevel.Default;
 
@@ -1533,7 +1613,7 @@ namespace UnrealBuildTool
 		/// Default settings for the no-arg constructor <see cref="CppCompileWarnings()"/>.
 		/// </summary>
 		/// <remarks>
-		/// To have items automatically participate in default setting, use <see cref="WarningLevelDefaultAttribute"/>, with <see cref="InitializationContext.Any"/> or <see cref="InitializationContext.Constructor"/>.
+		/// To have items automatically participate in default setting, use <see cref="BasicWarningLevelDefaultAttribute"/>, with <see cref="InitializationContext.Any"/> or <see cref="InitializationContext.Constructor"/>.
 		/// </remarks>
 		private void ApplyDefaults()
 		{
@@ -1545,8 +1625,8 @@ namespace UnrealBuildTool
 
 			foreach (MemberInfo mi in memberInfos)
 			{
-				IList<WarningLevelDefaultAttribute> attributes = [.. mi.GetCustomAttributes<WarningLevelDefaultAttribute>().Where(prop => prop.Context == InitializationContext.Any || prop.Context == InitializationContext.Constructor)];
-				WarningLevel defaultWarningLevel = WarningLevelDefaultAttribute.ResolveWarningLevelDefault(attributes, buildSettingsVersion);
+				WarningLevel defaultWarningLevel = WarningLevel.Default;
+				defaultWarningLevel = CoalesceDefaultWarningFromMemberInfo(mi, InitializationContext.Constructor, defaultWarningLevel, BuildContextProvider, _logger);
 
 				SetWarningLevel(mi, this, defaultWarningLevel);
 			}
diff --git a/Engine/Source/Programs/UnrealBuildTool/Configuration/ReadOnlyCppCompileWarnings.cs b/Engine/Source/Programs/UnrealBuildTool/Configuration/ReadOnlyCppCompileWarnings.cs
--- a/Engine/Source/Programs/UnrealBuildTool/Configuration/ReadOnlyCppCompileWarnings.cs
+++ b/Engine/Source/Programs/UnrealBuildTool/Configuration/ReadOnlyCppCompileWarnings.cs
@@ -46,6 +46,9 @@ namespace UnrealBuildTool
 
 		#region -- Clang Compiler Disables --
 
+		public WarningLevel DeleteNonVirtualDtorWarningLevel => _inner.DeleteNonVirtualDtorWarningLevel;
+		public WarningLevel EnumConversionWarningLevel => _inner.EnumConversionWarningLevel;
+		public WarningLevel ClangBitfieldEnumConversion => _inner.ClangBitfieldEnumConversion;
 		public WarningLevel EnumEnumConversionWarningLevel => _inner.EnumEnumConversionWarningLevel;
 		public WarningLevel EnumFloatConversionWarningLevel => _inner.EnumFloatConversionWarningLevel;
 		public WarningLevel AmbiguousReversedOperatorWarningLevel => _inner.AmbiguousReversedOperatorWarningLevel;
diff --git a/Engine/Source/Programs/UnrealBuildTool/Platform/Windows/VCToolChain.cs b/Engine/Source/Programs/UnrealBuildTool/Platform/Windows/VCToolChain.cs
--- a/Engine/Source/Programs/UnrealBuildTool/Platform/Windows/VCToolChain.cs
+++ b/Engine/Source/Programs/UnrealBuildTool/Platform/Windows/VCToolChain.cs
@@ -1464,8 +1464,7 @@ namespace UnrealBuildTool
 
 				VersionNumber ClangVersion = Target.WindowsPlatform.Compiler == WindowsCompiler.Intel ? MicrosoftPlatformSDK.GetClangVersionForIntelCompiler(EnvVars.CompilerPath) : EnvVars.CompilerVersion;
 
-				ClangWarnings.GetEnabledWarnings(Arguments);
-				ClangWarnings.GetDisabledWarnings(CompileEnvironment, Target.StaticAnalyzer, ClangVersion, Arguments);
+				ClangWarnings.GetWarnings(CompileEnvironment, ClangVersion, Arguments);
 			}
 
 			Arguments.AddRange(CompileEnvironment.CppCompileWarnings.GenerateWarningCommandLineArgs(CompileEnvironment, typeof(VCToolChain)));
diff --git a/Engine/Source/Programs/UnrealBuildTool/ToolChain/ClangToolChain.cs b/Engine/Source/Programs/UnrealBuildTool/ToolChain/ClangToolChain.cs
--- a/Engine/Source/Programs/UnrealBuildTool/ToolChain/ClangToolChain.cs
+++ b/Engine/Source/Programs/UnrealBuildTool/ToolChain/ClangToolChain.cs
@@ -639,13 +639,9 @@ namespace UnrealBuildTool
 		{
 			Arguments.Add("-Wall");                                     // https://clang.llvm.org/docs/DiagnosticsReference.html#wall
 			Arguments.Add("-Werror");                                   // https://clang.llvm.org/docs/UsersManual.html#cmdoption-werror
-
-			ClangWarnings.GetEnabledWarnings(Arguments);
-
-			ClangWarnings.GetDisabledWarnings(CompileEnvironment,
-				StaticAnalyzer,
-				new VersionNumber(Info.ClangVersion.Major, Info.ClangVersion.Minor, Info.ClangVersion.Build),
-				Arguments);
+			
+			VersionNumber clangVersionNumber = new VersionNumber(Info.ClangVersion.Major, Info.ClangVersion.Minor, Info.ClangVersion.Build);
+			ClangWarnings.GetWarnings(CompileEnvironment, clangVersionNumber, Arguments);
 
 			// always use absolute paths for errors, this can help IDEs go to the error properly
 			Arguments.Add("-fdiagnostics-absolute-paths");  // https://clang.llvm.org/docs/ClangCommandLineReference.html#cmdoption-clang-fdiagnostics-absolute-paths
diff --git a/Engine/Source/Programs/UnrealBuildTool/ToolChain/ClangWarnings.cs b/Engine/Source/Programs/UnrealBuildTool/ToolChain/ClangWarnings.cs
--- a/Engine/Source/Programs/UnrealBuildTool/ToolChain/ClangWarnings.cs
+++ b/Engine/Source/Programs/UnrealBuildTool/ToolChain/ClangWarnings.cs
@@ -8,14 +8,7 @@ namespace UnrealBuildTool
 {
 	internal static class ClangWarnings
 	{
-		internal static void GetEnabledWarnings(List<string> Arguments)
-		{
-			Arguments.Add("-Wdelete-non-virtual-dtor");                 // https://clang.llvm.org/docs/DiagnosticsReference.html#wdelete-non-virtual-dtor
-			Arguments.Add("-Wenum-conversion");                         // https://clang.llvm.org/docs/DiagnosticsReference.html#wenum-conversion
-			Arguments.Add("-Wbitfield-enum-conversion");                // https://clang.llvm.org/docs/DiagnosticsReference.html#wbitfield-enum-conversion
-		}
-
-		internal static void GetDisabledWarnings(CppCompileEnvironment CompileEnvironment, StaticAnalyzer Analyzer, VersionNumber ClangVersion, List<string> Arguments)
+		internal static void GetWarnings(CppCompileEnvironment CompileEnvironment, VersionNumber ClangVersion, List<string> Arguments)
 		{
 			Arguments.AddRange(CompileEnvironment.CppCompileWarnings.GenerateWarningCommandLineArgs(CompileEnvironment, typeof(ClangToolChain), ClangVersion));
 		}
