[UBA]
* Moved cancelling of build to outside callback when bStopCompilationAfterErrors is set to true

diff --git a/Engine/Source/Programs/UnrealBuildTool/Executors/UnrealBuildAccelerator/UBAExecutor.cs b/Engine/Source/Programs/UnrealBuildTool/Executors/UnrealBuildAccelerator/UBAExecutor.cs
--- a/Engine/Source/Programs/UnrealBuildTool/Executors/UnrealBuildAccelerator/UBAExecutor.cs
+++ b/Engine/Source/Programs/UnrealBuildTool/Executors/UnrealBuildAccelerator/UBAExecutor.cs
@@ -626,6 +626,7 @@ namespace UnrealBuildTool
 			int remoteProcessedActions = 0;
 			int localRetryActions = 0;
 			int forcedRetryActions = 0;
+			bool shouldCancel = false;
 
 			scheduler.SetProcessFinishedCallback((info) =>
 			{
@@ -790,7 +791,7 @@ namespace UnrealBuildTool
 
 					if (bStopCompilationAfterErrors)
 					{
-						cancelBuild();
+						shouldCancel = true;
 					}
 				}
 
@@ -850,7 +851,13 @@ namespace UnrealBuildTool
 				}
 
 				while (!scheduler.IsEmpty)
+				{
+					if (shouldCancel)
+					{
+						cancelBuild();
+					}
 					Thread.Sleep(100);
+				}
 
 				bool res = failedActions == 0;
 
