UnrealBuildTool: Added logic for enabling distributed thin lto linking. (only supported by some console platform but will likely be supported in vanilla clang toolchain at some point)

#rnx

diff --git a/Engine/Source/Programs/UnrealBuildTool/Configuration/ReadOnlyTargetRules.cs b/Engine/Source/Programs/UnrealBuildTool/Configuration/ReadOnlyTargetRules.cs
--- a/Engine/Source/Programs/UnrealBuildTool/Configuration/ReadOnlyTargetRules.cs
+++ b/Engine/Source/Programs/UnrealBuildTool/Configuration/ReadOnlyTargetRules.cs
@@ -403,6 +403,8 @@ namespace UnrealBuildTool
 
 		public bool bPreferThinLTO => Inner.bPreferThinLTO;
 
+		public bool bThinLTODistributed => Inner.bThinLTODistributed;
+
 		public bool bCompressDebugFile => Inner.bCompressDebugFile;
 
 		public bool bPGOProfile => Inner.bPGOProfile;
diff --git a/Engine/Source/Programs/UnrealBuildTool/Configuration/TargetRules.cs b/Engine/Source/Programs/UnrealBuildTool/Configuration/TargetRules.cs
--- a/Engine/Source/Programs/UnrealBuildTool/Configuration/TargetRules.cs
+++ b/Engine/Source/Programs/UnrealBuildTool/Configuration/TargetRules.cs
@@ -2051,6 +2051,14 @@ namespace UnrealBuildTool
 		[XmlConfigFile(Category = "BuildConfiguration")]
 		public string? ThinLTOCachePruningArguments { get; set; }
 
+		/// <summary>
+		/// When ThinLTO is enabled we use distributed linking if supported
+		/// This can save a lot of time for large binaries
+		/// </summary>
+		[CommandLine("-ThinLTODistributed")]
+		[XmlConfigFile(Category = "BuildConfiguration")]
+		public bool bThinLTODistributed { get; set; }
+
 		/// <summary>
 		/// Whether to enable Profile Guided Optimization (PGO) instrumentation in this build.
 		/// </summary>
diff --git a/Engine/Source/Programs/UnrealBuildTool/ToolChain/ClangToolChain.cs b/Engine/Source/Programs/UnrealBuildTool/ToolChain/ClangToolChain.cs
--- a/Engine/Source/Programs/UnrealBuildTool/ToolChain/ClangToolChain.cs
+++ b/Engine/Source/Programs/UnrealBuildTool/ToolChain/ClangToolChain.cs
@@ -68,6 +68,11 @@ namespace UnrealBuildTool
 		/// </summary>
 		EnableThinLTO = 1 << 8,
 
+		/// <summary>
+		/// Enable thin LTO
+		/// </summary>
+		ThinLTODistributed = 1 << 9,
+
 		/// <summary>
 		/// Enable tuning of debug info for LLDB
 		/// </summary>
