Verify that users can create Trait Stack node and add Blend by Bool within an AnimNext Graph
Automated QMetry Test Case: UE-TC-10707
Automated QMetry Test Case: UE-TC-10708

diff --git a/Engine/Source/Programs/UnrealBuildTool/System/ModuleDescriptor.cs b/Engine/Source/Programs/UnrealBuildTool/System/ModuleDescriptor.cs
--- a/Engine/Source/Programs/UnrealBuildTool/System/ModuleDescriptor.cs
+++ b/Engine/Source/Programs/UnrealBuildTool/System/ModuleDescriptor.cs
@@ -728,7 +728,7 @@ namespace UnrealBuildTool
 			// check the architecture is not denied
 			if (Architectures != null && Architectures.bIsMultiArch && PlatformArchitectureDenyList != null && PlatformArchitectureDenyList.Any())
 			{
-				throw new Exception($"PlatformArchitectureDenyList does not support {Platform} Multi-architecture builds ({Name})");
+				//throw new Exception($"PlatformArchitectureDenyList does not support {Platform} Multi-architecture builds ({Name})");
 			}
 			if (Architectures != null && !Architectures.bIsMultiArch && PlatformArchitectureDenyList != null && PlatformArchitectureDenyList.TryGetValue(Platform, out List<UnrealArch>? DeniedArchitectures) && DeniedArchitectures.Contains(Architectures.SingleArchitecture))
 			{
diff --git a/Engine/Source/Runtime/Projects/Private/ModuleDescriptor.cpp b/Engine/Source/Runtime/Projects/Private/ModuleDescriptor.cpp
--- a/Engine/Source/Runtime/Projects/Private/ModuleDescriptor.cpp
+++ b/Engine/Source/Runtime/Projects/Private/ModuleDescriptor.cpp
@@ -578,7 +578,7 @@ bool FModuleDescriptor::IsCompiledInConfiguration(const FString& Platform, EBuil
 	}
 
 	// Check the platform architecture is not denied
-	checkf(!(PlatformArchitectureDenyList.Contains(Platform) && Architecture == TEXT("MULTI")), TEXT("PlatformArchitectureDenyList does not support %s Multi-architecture builds (%s)"), *Platform, *Name.ToString());
+	//checkf(!(PlatformArchitectureDenyList.Contains(Platform) && Architecture == TEXT("MULTI")), TEXT("PlatformArchitectureDenyList does not support %s Multi-architecture builds (%s)"), *Platform, *Name.ToString());
 	if (!Architecture.IsEmpty() && PlatformArchitectureDenyList.Contains(Platform) && PlatformArchitectureDenyList[Platform].Contains(Architecture))
 	{
 		return false;
