UnrealBuildTool: Adjust cache bucket hash to use compile and link environment instead of only target
Simplify ISPC cache bucket key as it can be shared
Add optional target cache salt to seed cache bucket key hash

#rnx

diff --git a/Engine/Source/Programs/UnrealBuildTool/Configuration/ReadOnlyTargetRules.cs b/Engine/Source/Programs/UnrealBuildTool/Configuration/ReadOnlyTargetRules.cs
--- a/Engine/Source/Programs/UnrealBuildTool/Configuration/ReadOnlyTargetRules.cs
+++ b/Engine/Source/Programs/UnrealBuildTool/Configuration/ReadOnlyTargetRules.cs
@@ -369,6 +369,8 @@ namespace UnrealBuildTool
 
 		public bool bUseVFS => Inner.bUseVFS;
 
+		public string? CacheSalt => Inner.CacheSalt;
+
 		public bool bChainPCHs => Inner.bChainPCHs;
 
 		public bool bForceIncludePCHHeadersForGenCppFilesWhenPCHIsDisabled => Inner.bForceIncludePCHHeadersForGenCppFilesWhenPCHIsDisabled;
diff --git a/Engine/Source/Programs/UnrealBuildTool/Configuration/TargetRules.cs b/Engine/Source/Programs/UnrealBuildTool/Configuration/TargetRules.cs
--- a/Engine/Source/Programs/UnrealBuildTool/Configuration/TargetRules.cs
+++ b/Engine/Source/Programs/UnrealBuildTool/Configuration/TargetRules.cs
@@ -1912,6 +1912,13 @@ namespace UnrealBuildTool
 		[RequiresUniqueBuildEnvironment]
 		public bool bUseVFS { get; set; }
 
+		/// <summary>
+		/// Salt value to modify cache bucket calculation, can be used to invalidate cached actions outputs
+		/// </summary>
+		[CommandLine("-CacheSalt")]
+		[XmlConfigFile(Category = "BuildConfiguration")]
+		public string? CacheSalt { get; set; }
+
 		/// <summary>
 		/// Whether PCHs should be chained when compiling with clang.
 		/// </summary>
diff --git a/Engine/Source/Programs/UnrealBuildTool/Platform/Linux/LinuxToolChain.cs b/Engine/Source/Programs/UnrealBuildTool/Platform/Linux/LinuxToolChain.cs
--- a/Engine/Source/Programs/UnrealBuildTool/Platform/Linux/LinuxToolChain.cs
+++ b/Engine/Source/Programs/UnrealBuildTool/Platform/Linux/LinuxToolChain.cs
@@ -1132,7 +1132,7 @@ namespace UnrealBuildTool
 			LinkAction.CommandDescription += (LinkCommandString.Contains("-fuse-ld=lld")) ? " (lld)" : " (ld)";
 			LinkAction.CommandVersion = Info.ClangVersionString;
 			LinkAction.StatusDescription = Path.GetFileName(OutputFile.AbsolutePath);
-			LinkAction.CacheBucket = GetCacheBucket(TargetRules, null);
+			LinkAction.CacheBucket = GetCacheBucket(TargetRules, LinkEnvironment);
 			LinkAction.ArtifactMode = ArtifactMode.Enabled;
 
 			// Add the output file to the command-line.
@@ -1483,7 +1483,7 @@ namespace UnrealBuildTool
 				RelinkAction.StatusDescription = LinkAction.StatusDescription;
 				RelinkAction.CommandDescription = "Relink";
 				RelinkAction.bCanExecuteRemotely = false;
-				RelinkAction.CacheBucket = GetCacheBucket(TargetRules, null);
+				RelinkAction.CacheBucket = GetCacheBucket(TargetRules, LinkEnvironment);
 				RelinkAction.ArtifactMode = ArtifactMode.Enabled;
 				RelinkAction.ProducedItems.Add(RealOutputFile);
 
diff --git a/Engine/Source/Programs/UnrealBuildTool/Platform/Mac/MacToolChain.cs b/Engine/Source/Programs/UnrealBuildTool/Platform/Mac/MacToolChain.cs
--- a/Engine/Source/Programs/UnrealBuildTool/Platform/Mac/MacToolChain.cs
+++ b/Engine/Source/Programs/UnrealBuildTool/Platform/Mac/MacToolChain.cs
@@ -690,7 +690,7 @@ namespace UnrealBuildTool
 			LinkAction.CommandVersion = Info.ClangVersionString;
 			LinkAction.bProducesImportLibrary = bBuildImportLibraryOnly || LinkEnvironment.bIsBuildingDLL;
 			LinkAction.bCanExecuteInUBA = true;
-			LinkAction.CacheBucket = GetCacheBucket(Target, null);
+			LinkAction.CacheBucket = GetCacheBucket(Target, LinkEnvironment);
 			LinkAction.ArtifactMode = ArtifactMode.Enabled;
 
 			List<string> LinkArguments = new();
diff --git a/Engine/Source/Programs/UnrealBuildTool/Platform/Windows/PVSToolChain.cs b/Engine/Source/Programs/UnrealBuildTool/Platform/Windows/PVSToolChain.cs
--- a/Engine/Source/Programs/UnrealBuildTool/Platform/Windows/PVSToolChain.cs
+++ b/Engine/Source/Programs/UnrealBuildTool/Platform/Windows/PVSToolChain.cs
@@ -882,7 +882,7 @@ namespace UnrealBuildTool
 			analyzeAction.bCanExecuteRemotelyWithXGE = false;
 
 			analyzeAction.RootPaths = compileEnvironment.RootPaths;
-			analyzeAction.CacheBucket = GetCacheBucket(Target, null);
+			analyzeAction.CacheBucket = GetCacheBucket(Target, compileEnvironment);
 			analyzeAction.ArtifactMode = ArtifactMode.Enabled;
 
 			result.ObjectFiles.AddRange(analyzeAction.ProducedItems);
diff --git a/Engine/Source/Programs/UnrealBuildTool/Platform/Windows/VCToolChain.cs b/Engine/Source/Programs/UnrealBuildTool/Platform/Windows/VCToolChain.cs
--- a/Engine/Source/Programs/UnrealBuildTool/Platform/Windows/VCToolChain.cs
+++ b/Engine/Source/Programs/UnrealBuildTool/Platform/Windows/VCToolChain.cs
@@ -1946,7 +1946,7 @@ namespace UnrealBuildTool
 
 			// Calculate cache bucket. We want to use as few buckets as possible per build but we also
 			// want to make sure that the buckets are not getting too big.
-			BaseCompileAction.CacheBucket = GetCacheBucket(Target, (!Target.bUseVFS && CompileEnvironment.PrecompiledHeaderAction != PrecompiledHeaderAction.None) ? CompileEnvironment.RootPaths : null);
+			BaseCompileAction.CacheBucket = GetCacheBucket(Target, CompileEnvironment);
 
 			return BaseCompileAction;
 		}
@@ -2591,7 +2591,7 @@ namespace UnrealBuildTool
 				CompileAction.PrerequisiteItems.UnionWith(CompileEnvironment.ForceIncludeFiles);
 				CompileAction.PrerequisiteItems.UnionWith(CompileEnvironment.AdditionalPrerequisites);
 				CompileAction.ArtifactMode = ArtifactMode.Enabled;
-				CompileAction.CacheBucket = GetCacheBucket(Target, null);
+				CompileAction.CacheBucket = GetCacheBucket(Target, CompileEnvironment);
 
 				// Resource tool can run remotely if possible
 				CompileAction.bCanExecuteRemotely = true;
@@ -2800,7 +2800,7 @@ namespace UnrealBuildTool
 			CompileAction.bCanExecuteRemotely = false; // Incompatible with remote distribution
 			CompileAction.RootPaths = CompileEnvironment.RootPaths;
 			CompileAction.ArtifactMode = ArtifactMode.Enabled;
-			CompileAction.CacheBucket = GetCacheBucket(Target, null);
+			CompileAction.CacheBucket = GetCacheBucket(Target, CompileEnvironment);
 		}
 
 		public override IEnumerable<string> GetGlobalCommandLineArgs(CppCompileEnvironment CompileEnvironment)
@@ -3282,7 +3282,7 @@ namespace UnrealBuildTool
 				LinkAction.CommandVersion = $"{LinkAction.CommandVersion} Compressed";
 			}
 
-			LinkAction.CacheBucket = GetCacheBucket(Target, null);
+			LinkAction.CacheBucket = GetCacheBucket(Target, LinkEnvironment);
 
 			Logger.LogDebug("     Linking: {StatusDescription}", LinkAction.StatusDescription);
 			Logger.LogDebug("     Command: {CommandArguments}", LinkAction.CommandArguments);
diff --git a/Engine/Source/Programs/UnrealBuildTool/ToolChain/ClangToolChain.cs b/Engine/Source/Programs/UnrealBuildTool/ToolChain/ClangToolChain.cs
--- a/Engine/Source/Programs/UnrealBuildTool/ToolChain/ClangToolChain.cs
+++ b/Engine/Source/Programs/UnrealBuildTool/ToolChain/ClangToolChain.cs
@@ -1365,7 +1365,7 @@ namespace UnrealBuildTool
 			}
 
 			CompileAction.ArtifactMode = ArtifactMode.Enabled;
-			CompileAction.CacheBucket = GetCacheBucket(TargetRules, (!CompileEnvironment.RootPaths.bUseVfs && CompileEnvironment.PrecompiledHeaderAction != PrecompiledHeaderAction.None) ? CompileEnvironment.RootPaths : null);
+			CompileAction.CacheBucket = GetCacheBucket(TargetRules, CompileEnvironment);
 
 			if (CompileEnvironment.PrecompiledHeaderAction != PrecompiledHeaderAction.None)
 			{
diff --git a/Engine/Source/Programs/UnrealBuildTool/ToolChain/ISPCToolChain.cs b/Engine/Source/Programs/UnrealBuildTool/ToolChain/ISPCToolChain.cs
--- a/Engine/Source/Programs/UnrealBuildTool/ToolChain/ISPCToolChain.cs
+++ b/Engine/Source/Programs/UnrealBuildTool/ToolChain/ISPCToolChain.cs
@@ -362,7 +362,7 @@ namespace UnrealBuildTool
 		/// <param name="str">The string to hash</param>
 		/// <param name="hash">The hash input</param>
 		/// <returns>hash of input hash + string</returns>
-		protected static uint GetStringHash(string str, uint hash = 5381)
+		static uint GetStringHash(string str, uint hash = 5381)
 		{
 			for (int idx = 0; idx < str.Length; idx++)
 			{
@@ -372,59 +372,109 @@ namespace UnrealBuildTool
 		}
 
 		/// <summary>
-		/// Calculates hash bucket based on target and root paths
+		/// Calculates hash bucket based on target, do not call directly
 		/// </summary>
 		/// <param name="target">Target rules</param>
-		/// <param name="rootPaths">Optional root paths. Should be null except in special scenario (using pch and vfs is disabled)</param>
 		/// <returns>bucket hash</returns>
-		public static uint GetCacheBucket(ReadOnlyTargetRules? target, CppRootPaths? rootPaths)
+		static uint GetCacheBucket(ReadOnlyTargetRules? target)
 		{
 			if (target == null)
 			{
 				return 0;
 			}
 
-			uint cacheBucket = GetStringHash(target.Platform.ToString()); // Platform
+			IEnumerable<string> strings = [
+				target.CacheSalt ?? String.Empty,			// Optional salt
+				target.BuildEnvironment.ToString(),			// Unique or Share environment
+				target.bMergeModules ? "MM" : String.Empty	// Merged modules will modify all object files
+			];
+			return strings.Aggregate(GetStringHash(String.Empty), (acc, str) => GetStringHash(str, acc));
+		}
 
-			cacheBucket = GetStringHash(target.Configuration.ToString(), cacheBucket); // Configuration
+		/// <summary>
+		/// Calculates hash bucket based on target for ISPC
+		/// </summary>
+		/// <param name="target">Target rules</param>
+		/// <param name="compileEnvironment">Compile environment</param>
+		/// <returns>bucket hash</returns>
+		static uint GetISPCCacheBucket(ReadOnlyTargetRules? target, CppCompileEnvironment compileEnvironment)
+		{
+			IEnumerable<string> strings = [
+				"ISCP",										// Create a separate bucket for ispc
+				compileEnvironment.Platform.ToString(),		// Platform
+				compileEnvironment.Architecture.ToString(),	// ispc only ever has a single arch
+			];
+			return strings.Aggregate(GetCacheBucket(target), (acc, str) => GetStringHash(str, acc));
+		}
 
-			cacheBucket = GetStringHash(target.LinkType.ToString(), cacheBucket); // Modular or monolithic
+		/// <summary>
+		/// Calculates hash bucket based on target and compile environment
+		/// </summary>
+		/// <param name="target">Target rules</param>
+		/// <param name="compileEnvironment">Compile environment</param>
+		/// <returns>bucket hash</returns>
+		public static uint GetCacheBucket(ReadOnlyTargetRules? target, CppCompileEnvironment compileEnvironment)
+		{
+			IEnumerable<string> strings = [
+				compileEnvironment.Platform.ToString(), // Platform
+				compileEnvironment.Configuration.ToString(), // Configuration
+				!compileEnvironment.Architectures.bIsMultiArch ? compileEnvironment.Architecture.ToString() : compileEnvironment.Architectures.ToString(), // arm64 or x64
+				target != null && target.StaticAnalyzer != StaticAnalyzer.None && !compileEnvironment.bDisableStaticAnalysis ? target.StaticAnalyzer.ToString() : String.Empty, // If static analyzer
+				!compileEnvironment.bUseUnity ? "NU" : String.Empty, // Non-unity
+			];
 
-			if (!target.Architectures.bIsMultiArch)
+			// Optimize settings
+			if (compileEnvironment.bPGOProfile)
 			{
-				cacheBucket = GetStringHash(target.Architecture.ToString(), cacheBucket); // arm64 or x64
+				strings = [.. strings, "PGOProfile"];
 			}
-
-			if (target.bPGOProfile)
+			else if (compileEnvironment.bPGOOptimize)
 			{
-				cacheBucket = GetStringHash("PGOProfile", cacheBucket); // pgo profile
+				strings = [.. strings, "PGOOptimize"];
 			}
-			if (target.bPGOOptimize)
+			else if (compileEnvironment.bAllowLTCG)
 			{
-				cacheBucket = GetStringHash("PGOOptimize", cacheBucket); // pgo profile
+				strings = [.. strings, "LTCG"];
 			}
 
-			if (target.StaticAnalyzer != StaticAnalyzer.None)
+			// If vfs is disabled and pch is enabled
+			if (!compileEnvironment.RootPaths.bUseVfs && compileEnvironment.PrecompiledHeaderAction != PrecompiledHeaderAction.None)
 			{
-				cacheBucket = GetStringHash(target.StaticAnalyzer.ToString(), cacheBucket); // If static analyzer
+				strings = [.. strings, .. compileEnvironment.RootPaths.Select(item => item.local.FullName)];
 			}
-			if (!target.bUseUnityBuild)
+			return strings.Aggregate(GetCacheBucket(target), (acc, str) => GetStringHash(str, acc));
+		}
+
+		/// <summary>
+		/// Calculates hash bucket based on target and link environment
+		/// </summary>
+		/// <param name="target">Target rules</param>
+		/// <param name="linkEnvironment">Link environment</param>
+		/// <returns>bucket hash</returns>
+		public static uint GetCacheBucket(ReadOnlyTargetRules? target, LinkEnvironment linkEnvironment)
+		{
+			IEnumerable<string> strings = [
+				linkEnvironment.Platform.ToString(), // Platform
+				linkEnvironment.Configuration.ToString(), // Configuration
+				linkEnvironment.LinkType.ToString(), // Modular or monolithic
+				!linkEnvironment.Architectures.bIsMultiArch ? linkEnvironment.Architecture.ToString() : linkEnvironment.Architectures.ToString(), // arm64 or x64
+				target?.bUseUnityBuild == false ? "NU" : String.Empty // Non-unity
+			];
+
+			// Optimize settings
+			if (linkEnvironment.bPGOProfile)
 			{
-				cacheBucket = GetStringHash("NU", cacheBucket); // Non-unity. There might be matching cache entries between unity and non unity but this will keep bucket size down
+				strings = [.. strings, "PGOProfile"];
 			}
-			if (target.bMergeModules)
+			else if (linkEnvironment.bPGOOptimize)
 			{
-				cacheBucket = GetStringHash("MM", cacheBucket); // Merged modules.. will modify all .obj files
+				strings = [.. strings, "PGOOptimize"];
 			}
-
-			if (rootPaths != null) // Usually if vfs is disabled and pch is enabled
+			else if (linkEnvironment.bAllowLTCG)
 			{
-				foreach (var (id, vfs, local) in rootPaths)
-				{
-					cacheBucket = GetStringHash(local.FullName, cacheBucket);
-				}
+				strings = [.. strings, "LTCG"];
 			}
-			return cacheBucket;
+			return strings.Aggregate(GetCacheBucket(target), (acc, str) => GetStringHash(str, acc));
 		}
 
 		protected override CPPOutput GenerateISPCHeaders(CppCompileEnvironment CompileEnvironment, IEnumerable<FileItem> InputFiles, DirectoryReference OutputDir, IActionGraphBuilder Graph)
@@ -500,9 +550,9 @@ namespace UnrealBuildTool
 				CompileAction.WorkingDirectory = Unreal.EngineSourceDirectory;
 				CompileAction.CommandPath = new FileReference(GetISPCHostCompilerPath(BuildHostPlatform.Current.Platform));
 				CompileAction.StatusDescription = Path.GetFileName(ISPCFile.AbsolutePath);
-				CompileAction.CommandVersion = GetISPCHostCompilerVersion(BuildHostPlatform.Current.Platform).ToString();
+				CompileAction.CommandVersion = GetISPCHostCompilerVersion(BuildHostPlatform.Current.Platform);
 				CompileAction.ArtifactMode = ArtifactMode.Enabled;
-				CompileAction.CacheBucket = GetCacheBucket(TargetRules, null);
+				CompileAction.CacheBucket = GetISPCCacheBucket(TargetRules, CompileEnvironment);
 
 				CompileAction.bCanExecuteRemotely = true;
 				CompileAction.bCanExecuteInUBACrossArchitecture = false;
@@ -728,7 +778,7 @@ namespace UnrealBuildTool
 				CompileAction.bCanExecuteRemotelyWithSNDBS = false;
 
 				CompileAction.ArtifactMode = ArtifactMode.Enabled;
-				CompileAction.CacheBucket = GetCacheBucket(TargetRules, null);
+				CompileAction.CacheBucket = GetISPCCacheBucket(TargetRules, CompileEnvironment);
 
 				List<string> Arguments = new List<string>();
 
@@ -880,7 +930,7 @@ namespace UnrealBuildTool
 
 							PostCompileAction.RootPaths = RootPaths;
 							PostCompileAction.ArtifactMode = ArtifactMode.Enabled;
-							PostCompileAction.CacheBucket = GetCacheBucket(TargetRules, null);
+							PostCompileAction.CacheBucket = GetISPCCacheBucket(TargetRules, CompileEnvironment);
 
 							// Disable remote execution to workaround mismatched case on XGE
 							PostCompileAction.bCanExecuteRemotelyWithXGE = false;
