UnrealBuildTool: Add in toolchain version to action version when not compiling with msvc

#rnx

diff --git a/Engine/Source/Programs/UnrealBuildTool/Platform/Windows/VCCompileAction.cs b/Engine/Source/Programs/UnrealBuildTool/Platform/Windows/VCCompileAction.cs
--- a/Engine/Source/Programs/UnrealBuildTool/Platform/Windows/VCCompileAction.cs
+++ b/Engine/Source/Programs/UnrealBuildTool/Platform/Windows/VCCompileAction.cs
@@ -48,6 +48,11 @@ namespace UnrealBuildTool
 		/// <summary>
 		/// The version of the compiler being used
 		/// </summary>
+		public string CompilerVersion { get; set; }
+
+		/// <summary>
+		/// The version of the toolchain being used (if the compiler is not MSVC)
+		/// </summary>
 		public string ToolChainVersion { get; set; }
 
 		/// <summary>
@@ -369,7 +374,7 @@ namespace UnrealBuildTool
 		}
 
 		/// <inheritdoc/>
-		string IExternalAction.CommandVersion => ToolChainVersion;
+		string IExternalAction.CommandVersion => CompilerType.IsMSVC() ? CompilerVersion : $"{CompilerType} {CompilerVersion} MSVC {ToolChainVersion}";
 
 		#endregion
 
@@ -381,6 +386,7 @@ namespace UnrealBuildTool
 		{
 			CompilerExe = FileItem.GetItemByFileReference(Environment.CompilerPath);
 			CompilerType = Environment.Compiler;
+			CompilerVersion = Environment.CompilerVersion.ToString();
 			ToolChainVersion = Environment.ToolChainVersion.ToString();
 
 			RootPaths = new();
@@ -396,6 +402,7 @@ namespace UnrealBuildTool
 			ArtifactMode = InAction.ArtifactMode;
 			CompilerExe = InAction.CompilerExe;
 			CompilerType = InAction.CompilerType;
+			CompilerVersion = InAction.CompilerVersion;
 			ToolChainVersion = InAction.ToolChainVersion;
 			SourceFile = InAction.SourceFile;
 			ObjectFile = InAction.ObjectFile;
@@ -443,6 +450,7 @@ namespace UnrealBuildTool
 			ArtifactMode = (ArtifactMode)Reader.ReadByte();
 			CompilerExe = Reader.ReadFileItem()!;
 			CompilerType = (WindowsCompiler)Reader.ReadInt();
+			CompilerVersion = Reader.ReadString()!;
 			ToolChainVersion = Reader.ReadString()!;
 			SourceFile = Reader.ReadFileItem();
 			ObjectFile = Reader.ReadFileItem();
diff --git a/Engine/Source/Programs/UnrealBuildTool/Platform/Windows/VCToolChain.cs b/Engine/Source/Programs/UnrealBuildTool/Platform/Windows/VCToolChain.cs
--- a/Engine/Source/Programs/UnrealBuildTool/Platform/Windows/VCToolChain.cs
+++ b/Engine/Source/Programs/UnrealBuildTool/Platform/Windows/VCToolChain.cs
@@ -1941,7 +1941,7 @@ namespace UnrealBuildTool
 
 			if (Target.bAllowUbaCompression)
 			{
-				BaseCompileAction.ToolChainVersion = $"{BaseCompileAction.ToolChainVersion} Compressed";
+				BaseCompileAction.CompilerVersion = $"{BaseCompileAction.CompilerVersion} Compressed";
 			}
 
 			// Calculate cache bucket. We want to use as few buckets as possible per build but we also
@@ -2796,6 +2796,10 @@ namespace UnrealBuildTool
 			CompileAction.CommandPath = BuildHostPlatform.Current.Shell;
 			CompileAction.CommandArguments = $"/C \"{BatchFile}\"";
 			CompileAction.CommandVersion = EnvVars.ToolChainVersion.ToString();
+			if (Target.bAllowUbaCompression)
+			{
+				CompileAction.CommandVersion = $"{CompileAction.CommandVersion} Compressed";
+			}
 			CompileAction.bShouldOutputStatusDescription = false;
 			CompileAction.bCanExecuteRemotely = false; // Incompatible with remote distribution
 			CompileAction.RootPaths = CompileEnvironment.RootPaths;
@@ -3218,7 +3222,8 @@ namespace UnrealBuildTool
 				LinkAction.CommandPath = EnvVars.LinkerPath;
 				LinkAction.CommandArguments = $"@\"{NormalizeCommandLinePath(ResponseFileName, RootPaths)}\"";
 			}
-			LinkAction.CommandVersion = EnvVars.ToolChainVersion.ToString();
+
+			LinkAction.CommandVersion = EnvVars.Compiler.IsMSVC() ? EnvVars.CompilerVersion.ToString() : $"{EnvVars.Compiler} {EnvVars.CompilerVersion} MSVC {EnvVars.ToolChainVersion}";
 			LinkAction.ProducedItems.UnionWith(ProducedItems);
 			LinkAction.PrerequisiteItems.UnionWith(PrerequisiteItems);
 			LinkAction.StatusDescription = Path.GetFileName(OutputFile.AbsolutePath);
