UnrealBuildTool: Remove some unneeded spammy log output

#rnx

diff --git a/Engine/Source/Programs/UnrealBuildTool/Configuration/UEBuildBinary.cs b/Engine/Source/Programs/UnrealBuildTool/Configuration/UEBuildBinary.cs
--- a/Engine/Source/Programs/UnrealBuildTool/Configuration/UEBuildBinary.cs
+++ b/Engine/Source/Programs/UnrealBuildTool/Configuration/UEBuildBinary.cs
@@ -777,7 +777,7 @@ namespace UnrealBuildTool
 				if (Module.Binary == null || Module.Binary == this)
 				{
 					// Compile each module.
-					Logger.LogDebug("Compile module: {ModuleName}", Module.Name);
+					Logger.LogTrace("Compile module: {ModuleName}", Module.Name);
 					List<FileItem> LinkInputFiles = Module.Compile(Target, ToolChain, BinaryCompileEnvironment, WorkingSet, Graph, Logger);
 					InputObjects.UnionWith(LinkInputFiles.Where(x => x.HasExtension(".obj") || x.HasExtension(".o")));
 
diff --git a/Engine/Source/Programs/UnrealBuildTool/Platform/Windows/VCToolChain.cs b/Engine/Source/Programs/UnrealBuildTool/Platform/Windows/VCToolChain.cs
--- a/Engine/Source/Programs/UnrealBuildTool/Platform/Windows/VCToolChain.cs
+++ b/Engine/Source/Programs/UnrealBuildTool/Platform/Windows/VCToolChain.cs
@@ -3289,9 +3289,6 @@ namespace UnrealBuildTool
 
 			LinkAction.CacheBucket = GetCacheBucket(Target, LinkEnvironment);
 
-			Logger.LogDebug("     Linking: {StatusDescription}", LinkAction.StatusDescription);
-			Logger.LogDebug("     Command: {CommandArguments}", LinkAction.CommandArguments);
-
 			return OutputFile;
 		}
 
diff --git a/Engine/Source/Programs/UnrealBuildTool/ToolChain/ISPCToolChain.cs b/Engine/Source/Programs/UnrealBuildTool/ToolChain/ISPCToolChain.cs
--- a/Engine/Source/Programs/UnrealBuildTool/ToolChain/ISPCToolChain.cs
+++ b/Engine/Source/Programs/UnrealBuildTool/ToolChain/ISPCToolChain.cs
@@ -642,8 +642,6 @@ namespace UnrealBuildTool
 				CopyAction.bShouldOutputStatusDescription = false;
 
 				Result.GeneratedHeaderFiles.Add(ISPCFinalHeaderFile);
-
-				Logger.LogDebug("   ISPC Generating Header {StatusDescription}: \"{CommandPath}\" {CommandArguments}", CompileAction.StatusDescription, CompileAction.CommandPath, CompileAction.CommandArguments);
 			}
 
 			return Result;
@@ -876,8 +874,6 @@ namespace UnrealBuildTool
 				// Add the source file and its included files to the prerequisite item list.
 				CompileAction.PrerequisiteItems.Add(ISPCFile);
 
-				Logger.LogDebug("   ISPC Compiling {StatusDescription}: \"{CommandPath}\" {CommandArguments}", CompileAction.StatusDescription, CompileAction.CommandPath, CompileAction.CommandArguments);
-
 				if (bByteCodeOutput)
 				{
 					// If the platform toolchain supports bytecode compilation for ISPC, compile the bytecode object files to actual native object files 
@@ -936,7 +932,6 @@ namespace UnrealBuildTool
 							PostCompileAction.bCanExecuteRemotelyWithXGE = false;
 
 							FinalObjectFiles.Add(FinalCompiledISPCObjFile);
-							Logger.LogDebug("   ISPC Compiling bytecode {StatusDescription}: \"{CommandPath}\" {CommandArguments} {ProducedItems}", PostCompileAction.StatusDescription, PostCompileAction.CommandPath, PostCompileAction.CommandArguments, PostCompileAction.ProducedItems);
 						}
 						// Override the output object files
 						CompiledISPCObjFiles = FinalObjectFiles;
