[UBT]
* Added ApplicationCore to RenderCore and to RHI to circular dependencies allow list

diff --git a/Engine/Source/Programs/UnrealBuildTool/Configuration/UEBuildModuleCPP.cs b/Engine/Source/Programs/UnrealBuildTool/Configuration/UEBuildModuleCPP.cs
--- a/Engine/Source/Programs/UnrealBuildTool/Configuration/UEBuildModuleCPP.cs
+++ b/Engine/Source/Programs/UnrealBuildTool/Configuration/UEBuildModuleCPP.cs
@@ -139,6 +139,8 @@ namespace UnrealBuildTool
 			new("AnimGraph", "KismetCompiler"),
 			new("AnimGraph", "PropertyEditor"),
 			new("AnimGraph", "UnrealEd"),
+			new("ApplicationCore", "RenderCore"),
+			new("ApplicationCore", "RHI"),
 			new("AssetTools", "Kismet"),
 			new("AssetTools", "Landscape"),
 			new("AssetTools", "MaterialEditor"),
