UnrealBuildTool: Fix VCCompileAction serialization

#rnx

diff --git a/Engine/Source/Programs/UnrealBuildTool/Platform/Windows/VCCompileAction.cs b/Engine/Source/Programs/UnrealBuildTool/Platform/Windows/VCCompileAction.cs
--- a/Engine/Source/Programs/UnrealBuildTool/Platform/Windows/VCCompileAction.cs
+++ b/Engine/Source/Programs/UnrealBuildTool/Platform/Windows/VCCompileAction.cs
@@ -497,6 +497,7 @@ namespace UnrealBuildTool
 			Writer.WriteByte((byte)ArtifactMode);
 			Writer.WriteFileItem(CompilerExe);
 			Writer.WriteInt((int)CompilerType);
+			Writer.WriteString(CompilerVersion);
 			Writer.WriteString(ToolChainVersion);
 			Writer.WriteFileItem(SourceFile);
 			Writer.WriteFileItem(ObjectFile);
