UnrealBuildTool: Clean up VCCompileAction.cs

#rnx

diff --git a/Engine/Source/Programs/UnrealBuildTool/Platform/Windows/VCCompileAction.cs b/Engine/Source/Programs/UnrealBuildTool/Platform/Windows/VCCompileAction.cs
--- a/Engine/Source/Programs/UnrealBuildTool/Platform/Windows/VCCompileAction.cs
+++ b/Engine/Source/Programs/UnrealBuildTool/Platform/Windows/VCCompileAction.cs
@@ -2,7 +2,6 @@
 
 using System;
 using System.Collections.Generic;
-using System.Linq;
 using EpicGames.Core;
 using Microsoft.Extensions.Logging;
 using UnrealBuildBase;
@@ -128,42 +127,42 @@ namespace UnrealBuildTool
 		/// <summary>
 		/// List of additional response paths
 		/// </summary>
-		public List<FileItem> AdditionalResponseFiles { get; } = new List<FileItem>();
+		public List<FileItem> AdditionalResponseFiles { get; } = [];
 
 		/// <summary>
 		/// List of include paths
 		/// </summary>
-		public List<DirectoryReference> IncludePaths { get; } = new List<DirectoryReference>();
+		public List<DirectoryReference> IncludePaths { get; } = [];
 
 		/// <summary>
 		/// List of system include paths
 		/// </summary>
-		public List<DirectoryReference> SystemIncludePaths { get; } = new List<DirectoryReference>();
+		public List<DirectoryReference> SystemIncludePaths { get; } = [];
 
 		/// <summary>
 		/// List of macro definitions
 		/// </summary>
-		public List<string> Definitions { get; } = new List<string>();
+		public List<string> Definitions { get; } = [];
 
 		/// <summary>
 		/// List of force included files
 		/// </summary>
-		public List<FileItem> ForceIncludeFiles = new List<FileItem>();
+		public List<FileItem> ForceIncludeFiles = [];
 
 		/// <summary>
 		/// Every file this action depends on.  These files need to exist and be up to date in order for this action to even be considered
 		/// </summary>
-		public List<FileItem> AdditionalPrerequisiteItems { get; } = new List<FileItem>();
+		public List<FileItem> AdditionalPrerequisiteItems { get; } = [];
 
 		/// <summary>
 		/// The files that this action produces after completing
 		/// </summary>
-		public List<FileItem> AdditionalProducedItems { get; } = new List<FileItem>();
+		public List<FileItem> AdditionalProducedItems { get; } = [];
 
 		/// <summary>
 		/// Arguments to pass to the compiler
 		/// </summary>
-		public List<string> Arguments { get; } = new List<string>();
+		public List<string> Arguments { get; } = [];
 
 		/// <summary>
 		/// Whether to show included files to the console
@@ -192,7 +191,7 @@ namespace UnrealBuildTool
 		/// <summary>
 		/// Items that should be deleted before running this action
 		/// </summary>
-		public List<FileItem> DeleteItems { get; } = new List<FileItem>();
+		public List<FileItem> DeleteItems { get; } = [];
 
 		/// <summary>
 		/// Root paths for this action (generally engine root project root, toolchain root, sdk root)
@@ -277,17 +276,17 @@ namespace UnrealBuildTool
 				{
 					yield return VfsOverlayFile;
 				}
-				foreach (FileItem AdditionalResponseFile in AdditionalResponseFiles)
+				foreach (FileItem additionalResponseFile in AdditionalResponseFiles)
 				{
-					yield return AdditionalResponseFile;
+					yield return additionalResponseFile;
 				}
-				foreach (FileItem ForceIncludeFile in ForceIncludeFiles)
+				foreach (FileItem forceIncludeFile in ForceIncludeFiles)
 				{
-					yield return ForceIncludeFile;
+					yield return forceIncludeFile;
 				}
-				foreach (FileItem AdditionalPrerequisiteItem in AdditionalPrerequisiteItems)
+				foreach (FileItem additionalPrerequisiteItem in AdditionalPrerequisiteItems)
 				{
-					yield return AdditionalPrerequisiteItem;
+					yield return additionalPrerequisiteItem;
 				}
 			}
 		}
@@ -329,9 +328,9 @@ namespace UnrealBuildTool
 				{
 					yield return CreatePchFile;
 				}
-				foreach (FileItem AdditionalProducedItem in AdditionalProducedItems)
+				foreach (FileItem additionalProducedItem in AdditionalProducedItems)
 				{
-					yield return AdditionalProducedItem;
+					yield return additionalProducedItem;
 				}
 			}
 		}
@@ -342,52 +341,32 @@ namespace UnrealBuildTool
 		bool UsingClFilter => ForceClFilter || (DependencyListFile != null && !DependencyListFile.HasExtension(".json") && !DependencyListFile.HasExtension(".d"));
 
 		/// <inheritdoc/>
-		FileReference IExternalAction.CommandPath
-		{
-			get
-			{
-				if (UsingClFilter)
-				{
-					return FileReference.Combine(Unreal.EngineDirectory, "Build", "Windows", "cl-filter", "cl-filter.exe");
-				}
-				else
-				{
-					return CompilerExe.Location;
-				}
-			}
-		}
+		FileReference IExternalAction.CommandPath =>
+			UsingClFilter
+				? FileReference.Combine(Unreal.EngineDirectory, "Build", "Windows", "cl-filter", "cl-filter.exe")
+				: CompilerExe.Location;
 
 		/// <inheritdoc/>
-		string IExternalAction.CommandArguments
-		{
-			get
-			{
-				if (UsingClFilter)
-				{
-					return GetClFilterArguments();
-				}
-				else
-				{
-					return GetClArguments();
-				}
-			}
-		}
+		string IExternalAction.CommandArguments => UsingClFilter ? GetClFilterArguments() : GetClArguments();
 
 		/// <inheritdoc/>
-		string IExternalAction.CommandVersion => CompilerType.IsMSVC() ? CompilerVersion : $"{CompilerType} {CompilerVersion} MSVC {ToolChainVersion}";
+		string IExternalAction.CommandVersion =>
+			CompilerType.IsMSVC()
+				? CompilerVersion
+				: $"{CompilerType} {CompilerVersion} MSVC {ToolChainVersion}";
 
 		#endregion
 
 		/// <summary>
 		/// Constructor
 		/// </summary>
-		/// <param name="Environment">Compiler executable</param>
-		public VCCompileAction(VCEnvironment Environment)
+		/// <param name="environment">Compiler executable</param>
+		public VCCompileAction(VCEnvironment environment)
 		{
-			CompilerExe = FileItem.GetItemByFileReference(Environment.CompilerPath);
-			CompilerType = Environment.Compiler;
-			CompilerVersion = Environment.CompilerVersion.ToString();
-			ToolChainVersion = Environment.ToolChainVersion.ToString();
+			CompilerExe = FileItem.GetItemByFileReference(environment.CompilerPath);
+			CompilerType = environment.Compiler;
+			CompilerVersion = environment.CompilerVersion.ToString();
+			ToolChainVersion = environment.ToolChainVersion.ToString();
 
 			RootPaths = new();
 		}
@@ -395,246 +374,246 @@ namespace UnrealBuildTool
 		/// <summary>
 		/// Copy constructor
 		/// </summary>
-		/// <param name="InAction">Action to copy from</param>
-		public VCCompileAction(VCCompileAction InAction)
+		/// <param name="inAction">Action to copy from</param>
+		public VCCompileAction(VCCompileAction inAction)
 		{
-			ActionType = InAction.ActionType;
-			ArtifactMode = InAction.ArtifactMode;
-			CompilerExe = InAction.CompilerExe;
-			CompilerType = InAction.CompilerType;
-			CompilerVersion = InAction.CompilerVersion;
-			ToolChainVersion = InAction.ToolChainVersion;
-			SourceFile = InAction.SourceFile;
-			ObjectFile = InAction.ObjectFile;
-			AssemblyFile = InAction.AssemblyFile;
-			PreprocessedFile = InAction.PreprocessedFile;
-			AnalyzeLogFile = InAction.AnalyzeLogFile;
-			ExperimentalLogFile = InAction.ExperimentalLogFile;
-			DependencyListFile = InAction.DependencyListFile;
-			CompiledModuleInterfaceFile = InAction.CompiledModuleInterfaceFile;
-			TimingFile = InAction.TimingFile;
-			ResponseFile = InAction.ResponseFile;
-			CreatePchFile = InAction.CreatePchFile;
-			UsingPchFile = InAction.UsingPchFile;
-			PchThroughHeaderFile = InAction.PchThroughHeaderFile;
-			VfsOverlayFile = InAction.VfsOverlayFile;
-			AdditionalResponseFiles = new List<FileItem>(InAction.AdditionalResponseFiles);
-			IncludePaths = new List<DirectoryReference>(InAction.IncludePaths);
-			SystemIncludePaths = new List<DirectoryReference>(InAction.SystemIncludePaths);
-			Definitions = new List<string>(InAction.Definitions);
-			ForceIncludeFiles = new List<FileItem>(InAction.ForceIncludeFiles);
-			Arguments = new List<string>(InAction.Arguments);
-			bShowIncludes = InAction.bShowIncludes;
-			bCanExecuteRemotely = InAction.bCanExecuteRemotely;
-			bCanExecuteRemotelyWithSNDBS = InAction.bCanExecuteRemotelyWithSNDBS;
-			bCanExecuteRemotelyWithXGE = InAction.bCanExecuteRemotelyWithXGE;
-			Architecture = InAction.Architecture;
-			bIsAnalyzing = InAction.bIsAnalyzing;
-			bWarningsAsErrors = InAction.bWarningsAsErrors;
-			Weight = InAction.Weight;
-			CacheBucket = InAction.CacheBucket;
-
-			AdditionalPrerequisiteItems = new List<FileItem>(InAction.AdditionalPrerequisiteItems);
-			AdditionalProducedItems = new List<FileItem>(InAction.AdditionalProducedItems);
-			DeleteItems = new List<FileItem>(InAction.DeleteItems);
-			RootPaths = new(InAction.RootPaths);
+			ActionType = inAction.ActionType;
+			ArtifactMode = inAction.ArtifactMode;
+			CompilerExe = inAction.CompilerExe;
+			CompilerType = inAction.CompilerType;
+			CompilerVersion = inAction.CompilerVersion;
+			ToolChainVersion = inAction.ToolChainVersion;
+			SourceFile = inAction.SourceFile;
+			ObjectFile = inAction.ObjectFile;
+			AssemblyFile = inAction.AssemblyFile;
+			PreprocessedFile = inAction.PreprocessedFile;
+			AnalyzeLogFile = inAction.AnalyzeLogFile;
+			ExperimentalLogFile = inAction.ExperimentalLogFile;
+			DependencyListFile = inAction.DependencyListFile;
+			CompiledModuleInterfaceFile = inAction.CompiledModuleInterfaceFile;
+			TimingFile = inAction.TimingFile;
+			ResponseFile = inAction.ResponseFile;
+			CreatePchFile = inAction.CreatePchFile;
+			UsingPchFile = inAction.UsingPchFile;
+			PchThroughHeaderFile = inAction.PchThroughHeaderFile;
+			VfsOverlayFile = inAction.VfsOverlayFile;
+			AdditionalResponseFiles = [.. inAction.AdditionalResponseFiles];
+			IncludePaths = [.. inAction.IncludePaths];
+			SystemIncludePaths = [.. inAction.SystemIncludePaths];
+			Definitions = [.. inAction.Definitions];
+			ForceIncludeFiles = [.. inAction.ForceIncludeFiles];
+			Arguments = [.. inAction.Arguments];
+			bShowIncludes = inAction.bShowIncludes;
+			bCanExecuteRemotely = inAction.bCanExecuteRemotely;
+			bCanExecuteRemotelyWithSNDBS = inAction.bCanExecuteRemotelyWithSNDBS;
+			bCanExecuteRemotelyWithXGE = inAction.bCanExecuteRemotelyWithXGE;
+			Architecture = inAction.Architecture;
+			bIsAnalyzing = inAction.bIsAnalyzing;
+			bWarningsAsErrors = inAction.bWarningsAsErrors;
+			Weight = inAction.Weight;
+			CacheBucket = inAction.CacheBucket;
+
+			AdditionalPrerequisiteItems = [.. inAction.AdditionalPrerequisiteItems];
+			AdditionalProducedItems = [.. inAction.AdditionalProducedItems];
+			DeleteItems = [.. inAction.DeleteItems];
+			RootPaths = new(inAction.RootPaths);
 		}
 
 		/// <summary>
 		/// Serialize a cache handler from an archive
 		/// </summary>
-		/// <param name="Reader">Reader to serialize from</param>
-		public VCCompileAction(BinaryArchiveReader Reader)
+		/// <param name="reader">Reader to serialize from</param>
+		public VCCompileAction(BinaryArchiveReader reader)
 		{
-			ActionType = (ActionType)Reader.ReadInt();
-			ArtifactMode = (ArtifactMode)Reader.ReadByte();
-			CompilerExe = Reader.ReadFileItem()!;
-			CompilerType = (WindowsCompiler)Reader.ReadInt();
-			CompilerVersion = Reader.ReadString()!;
-			ToolChainVersion = Reader.ReadString()!;
-			SourceFile = Reader.ReadFileItem();
-			ObjectFile = Reader.ReadFileItem();
-			AssemblyFile = Reader.ReadFileItem();
-			PreprocessedFile = Reader.ReadFileItem();
-			AnalyzeLogFile = Reader.ReadFileItem();
-			ExperimentalLogFile = Reader.ReadFileItem();
-			DependencyListFile = Reader.ReadFileItem();
-			CompiledModuleInterfaceFile = Reader.ReadFileItem();
-			TimingFile = Reader.ReadFileItem();
-			ResponseFile = Reader.ReadFileItem();
-			CreatePchFile = Reader.ReadFileItem();
-			UsingPchFile = Reader.ReadFileItem();
-			PchThroughHeaderFile = Reader.ReadFileItem();
-			VfsOverlayFile = Reader.ReadFileItem();
-			AdditionalResponseFiles = Reader.ReadList(() => Reader.ReadFileItem())!;
-			IncludePaths = Reader.ReadList(() => Reader.ReadCompactDirectoryReference())!;
-			SystemIncludePaths = Reader.ReadList(() => Reader.ReadCompactDirectoryReference())!;
-			Definitions = Reader.ReadList(() => Reader.ReadString())!;
-			ForceIncludeFiles = Reader.ReadList(() => Reader.ReadFileItem())!;
-			Arguments = Reader.ReadList(() => Reader.ReadString())!;
-			bShowIncludes = Reader.ReadBool();
-			bCanExecuteRemotely = Reader.ReadBool();
-			bCanExecuteRemotelyWithSNDBS = Reader.ReadBool();
-			bCanExecuteRemotelyWithXGE = Reader.ReadBool();
-			bCanExecuteInUBA = Reader.ReadBool();
-			bCanExecuteInUBACrossArchitecture = Reader.ReadBool();
-			Architecture = UnrealArch.Parse(Reader.ReadString()!);
-			bIsAnalyzing = Reader.ReadBool();
-			bWarningsAsErrors = Reader.ReadBool();
-			Weight = Reader.ReadDouble();
-			CacheBucket = Reader.ReadUnsignedInt();
-
-			AdditionalPrerequisiteItems = Reader.ReadList(() => Reader.ReadFileItem())!;
-			AdditionalProducedItems = Reader.ReadList(() => Reader.ReadFileItem())!;
-			DeleteItems = Reader.ReadList(() => Reader.ReadFileItem())!;
-			RootPaths = new(Reader);
+			ActionType = (ActionType)reader.ReadInt();
+			ArtifactMode = (ArtifactMode)reader.ReadByte();
+			CompilerExe = reader.ReadFileItem()!;
+			CompilerType = (WindowsCompiler)reader.ReadInt();
+			CompilerVersion = reader.ReadString()!;
+			ToolChainVersion = reader.ReadString()!;
+			SourceFile = reader.ReadFileItem();
+			ObjectFile = reader.ReadFileItem();
+			AssemblyFile = reader.ReadFileItem();
+			PreprocessedFile = reader.ReadFileItem();
+			AnalyzeLogFile = reader.ReadFileItem();
+			ExperimentalLogFile = reader.ReadFileItem();
+			DependencyListFile = reader.ReadFileItem();
+			CompiledModuleInterfaceFile = reader.ReadFileItem();
+			TimingFile = reader.ReadFileItem();
+			ResponseFile = reader.ReadFileItem();
+			CreatePchFile = reader.ReadFileItem();
+			UsingPchFile = reader.ReadFileItem();
+			PchThroughHeaderFile = reader.ReadFileItem();
+			VfsOverlayFile = reader.ReadFileItem();
+			AdditionalResponseFiles = reader.ReadList(reader.ReadFileItem)!;
+			IncludePaths = reader.ReadList(reader.ReadCompactDirectoryReference)!;
+			SystemIncludePaths = reader.ReadList(reader.ReadCompactDirectoryReference)!;
+			Definitions = reader.ReadList(reader.ReadString)!;
+			ForceIncludeFiles = reader.ReadList(reader.ReadFileItem)!;
+			Arguments = reader.ReadList(reader.ReadString)!;
+			bShowIncludes = reader.ReadBool();
+			bCanExecuteRemotely = reader.ReadBool();
+			bCanExecuteRemotelyWithSNDBS = reader.ReadBool();
+			bCanExecuteRemotelyWithXGE = reader.ReadBool();
+			bCanExecuteInUBA = reader.ReadBool();
+			bCanExecuteInUBACrossArchitecture = reader.ReadBool();
+			Architecture = UnrealArch.Parse(reader.ReadString()!);
+			bIsAnalyzing = reader.ReadBool();
+			bWarningsAsErrors = reader.ReadBool();
+			Weight = reader.ReadDouble();
+			CacheBucket = reader.ReadUnsignedInt();
+
+			AdditionalPrerequisiteItems = reader.ReadList(reader.ReadFileItem)!;
+			AdditionalProducedItems = reader.ReadList(reader.ReadFileItem)!;
+			DeleteItems = reader.ReadList(reader.ReadFileItem)!;
+			RootPaths = new(reader);
 		}
 
 		/// <inheritdoc/>
-		public void Write(BinaryArchiveWriter Writer)
+		public void Write(BinaryArchiveWriter writer)
 		{
-			Writer.WriteInt((int)ActionType);
-			Writer.WriteByte((byte)ArtifactMode);
-			Writer.WriteFileItem(CompilerExe);
-			Writer.WriteInt((int)CompilerType);
-			Writer.WriteString(CompilerVersion);
-			Writer.WriteString(ToolChainVersion);
-			Writer.WriteFileItem(SourceFile);
-			Writer.WriteFileItem(ObjectFile);
-			Writer.WriteFileItem(AssemblyFile);
-			Writer.WriteFileItem(PreprocessedFile);
-			Writer.WriteFileItem(AnalyzeLogFile);
-			Writer.WriteFileItem(ExperimentalLogFile);
-			Writer.WriteFileItem(DependencyListFile);
-			Writer.WriteFileItem(CompiledModuleInterfaceFile);
-			Writer.WriteFileItem(TimingFile);
-			Writer.WriteFileItem(ResponseFile);
-			Writer.WriteFileItem(CreatePchFile);
-			Writer.WriteFileItem(UsingPchFile);
-			Writer.WriteFileItem(PchThroughHeaderFile);
-			Writer.WriteFileItem(VfsOverlayFile);
-			Writer.WriteList(AdditionalResponseFiles, Item => Writer.WriteFileItem(Item));
-			Writer.WriteList(IncludePaths, Item => Writer.WriteCompactDirectoryReference(Item));
-			Writer.WriteList(SystemIncludePaths, Item => Writer.WriteCompactDirectoryReference(Item));
-			Writer.WriteList(Definitions, Item => Writer.WriteString(Item));
-			Writer.WriteList(ForceIncludeFiles, Item => Writer.WriteFileItem(Item));
-			Writer.WriteList(Arguments, Item => Writer.WriteString(Item));
-			Writer.WriteBool(bShowIncludes);
-			Writer.WriteBool(bCanExecuteRemotely);
-			Writer.WriteBool(bCanExecuteRemotelyWithSNDBS);
-			Writer.WriteBool(bCanExecuteRemotelyWithXGE);
-			Writer.WriteBool(bCanExecuteInUBA);
-			Writer.WriteBool(bCanExecuteInUBACrossArchitecture);
-			Writer.WriteString(Architecture.ToString());
-			Writer.WriteBool(bIsAnalyzing);
-			Writer.WriteBool(bWarningsAsErrors);
-			Writer.WriteDouble(Weight);
-			Writer.WriteUnsignedInt(CacheBucket);
-
-			Writer.WriteList(AdditionalPrerequisiteItems, Item => Writer.WriteFileItem(Item));
-			Writer.WriteList(AdditionalProducedItems, Item => Writer.WriteFileItem(Item));
-			Writer.WriteList(DeleteItems, Item => Writer.WriteFileItem(Item));
-			RootPaths.Write(Writer);
+			writer.WriteInt((int)ActionType);
+			writer.WriteByte((byte)ArtifactMode);
+			writer.WriteFileItem(CompilerExe);
+			writer.WriteInt((int)CompilerType);
+			writer.WriteString(CompilerVersion);
+			writer.WriteString(ToolChainVersion);
+			writer.WriteFileItem(SourceFile);
+			writer.WriteFileItem(ObjectFile);
+			writer.WriteFileItem(AssemblyFile);
+			writer.WriteFileItem(PreprocessedFile);
+			writer.WriteFileItem(AnalyzeLogFile);
+			writer.WriteFileItem(ExperimentalLogFile);
+			writer.WriteFileItem(DependencyListFile);
+			writer.WriteFileItem(CompiledModuleInterfaceFile);
+			writer.WriteFileItem(TimingFile);
+			writer.WriteFileItem(ResponseFile);
+			writer.WriteFileItem(CreatePchFile);
+			writer.WriteFileItem(UsingPchFile);
+			writer.WriteFileItem(PchThroughHeaderFile);
+			writer.WriteFileItem(VfsOverlayFile);
+			writer.WriteList(AdditionalResponseFiles, writer.WriteFileItem);
+			writer.WriteList(IncludePaths, writer.WriteCompactDirectoryReference);
+			writer.WriteList(SystemIncludePaths, writer.WriteCompactDirectoryReference);
+			writer.WriteList(Definitions, writer.WriteString);
+			writer.WriteList(ForceIncludeFiles, writer.WriteFileItem);
+			writer.WriteList(Arguments, writer.WriteString);
+			writer.WriteBool(bShowIncludes);
+			writer.WriteBool(bCanExecuteRemotely);
+			writer.WriteBool(bCanExecuteRemotelyWithSNDBS);
+			writer.WriteBool(bCanExecuteRemotelyWithXGE);
+			writer.WriteBool(bCanExecuteInUBA);
+			writer.WriteBool(bCanExecuteInUBACrossArchitecture);
+			writer.WriteString(Architecture.ToString());
+			writer.WriteBool(bIsAnalyzing);
+			writer.WriteBool(bWarningsAsErrors);
+			writer.WriteDouble(Weight);
+			writer.WriteUnsignedInt(CacheBucket);
+
+			writer.WriteList(AdditionalPrerequisiteItems, writer.WriteFileItem);
+			writer.WriteList(AdditionalProducedItems, writer.WriteFileItem);
+			writer.WriteList(DeleteItems, writer.WriteFileItem);
+			RootPaths.Write(writer);
 		}
 
 		/// <summary>
 		/// Writes the response file with the action's arguments
 		/// </summary>
-		/// <param name="Graph">The graph builder</param>
-		/// <param name="Logger">Logger for output</param>
-		public void WriteResponseFile(IActionGraphBuilder Graph, ILogger Logger)
+		/// <param name="graph">The graph builder</param>
+		/// <param name="logger">Logger for output</param>
+		public void WriteResponseFile(IActionGraphBuilder graph, ILogger logger)
 		{
 			if (ResponseFile != null)
 			{
-				Graph.CreateIntermediateTextFile(ResponseFile, GetCompilerArguments(Logger));
+				graph.CreateIntermediateTextFile(ResponseFile, GetCompilerArguments(logger));
 				Arguments.Clear();
 			}
 		}
 
-		public List<string> GetCompilerArguments(ILogger Logger)
+		public List<string> GetCompilerArguments(ILogger logger)
 		{
-			List<string> Arguments = new List<string>();
+			List<string> arguments = [];
 
 			if (SourceFile != null)
 			{
-				VCToolChain.AddSourceFile(Arguments, SourceFile, RootPaths);
+				VCToolChain.AddSourceFile(arguments, SourceFile, RootPaths);
 			}
 
-			foreach (FileItem AdditionalRsp in AdditionalResponseFiles)
+			foreach (FileItem additionalResponseFile in AdditionalResponseFiles)
 			{
-				VCToolChain.AddResponseFile(Arguments, AdditionalRsp, RootPaths);
+				VCToolChain.AddResponseFile(arguments, additionalResponseFile, RootPaths);
 			}
 
-			foreach (DirectoryReference IncludePath in IncludePaths)
+			foreach (DirectoryReference includePath in IncludePaths)
 			{
-				VCToolChain.AddIncludePath(Arguments, IncludePath, CompilerType, RootPaths);
+				VCToolChain.AddIncludePath(arguments, includePath, CompilerType, RootPaths);
 			}
 
-			foreach (DirectoryReference SystemIncludePath in SystemIncludePaths)
+			foreach (DirectoryReference systemIncludePath in SystemIncludePaths)
 			{
-				VCToolChain.AddSystemIncludePath(Arguments, SystemIncludePath, CompilerType, RootPaths);
+				VCToolChain.AddSystemIncludePath(arguments, systemIncludePath, CompilerType, RootPaths);
 			}
 
-			foreach (string Definition in Definitions)
+			foreach (string definition in Definitions)
 			{
 				// Escape all quotation marks so that they get properly passed with the command line.
-				string DefinitionArgument = Definition.Contains('"') ? Definition.Replace("\"", "\\\"") : Definition;
-				VCToolChain.AddDefinition(Arguments, DefinitionArgument);
+				string definitionArgument = definition.Contains('"') ? definition.Replace("\"", "\\\"") : definition;
+				VCToolChain.AddDefinition(arguments, definitionArgument);
 			}
 
-			foreach (FileItem ForceIncludeFile in ForceIncludeFiles)
+			foreach (FileItem forceIncludeFile in ForceIncludeFiles)
 			{
-				VCToolChain.AddForceIncludeFile(Arguments, ForceIncludeFile, RootPaths);
+				VCToolChain.AddForceIncludeFile(arguments, forceIncludeFile, RootPaths);
 			}
 
 			if (CreatePchFile != null)
 			{
-				VCToolChain.AddCreatePchFile(Arguments, PchThroughHeaderFile!, CreatePchFile, RootPaths);
+				VCToolChain.AddCreatePchFile(arguments, PchThroughHeaderFile!, CreatePchFile, RootPaths);
 			}
 
 			if (UsingPchFile != null && CompilerType.IsMSVC())
 			{
-				VCToolChain.AddUsingPchFile(Arguments, PchThroughHeaderFile!, UsingPchFile, RootPaths);
+				VCToolChain.AddUsingPchFile(arguments, PchThroughHeaderFile!, UsingPchFile, RootPaths);
 			}
 
 			if (PreprocessedFile != null)
 			{
-				VCToolChain.AddPreprocessedFile(Arguments, PreprocessedFile, RootPaths);
+				VCToolChain.AddPreprocessedFile(arguments, PreprocessedFile, RootPaths);
 			}
 
 			if (ObjectFile != null)
 			{
-				VCToolChain.AddObjectFile(Arguments, ObjectFile, RootPaths);
+				VCToolChain.AddObjectFile(arguments, ObjectFile, RootPaths);
 			}
 
 			if (AssemblyFile != null)
 			{
-				VCToolChain.AddAssemblyFile(Arguments, AssemblyFile, RootPaths);
+				VCToolChain.AddAssemblyFile(arguments, AssemblyFile, RootPaths);
 			}
 
 			if (AnalyzeLogFile != null)
 			{
-				VCToolChain.AddAnalyzeLogFile(Arguments, AnalyzeLogFile, RootPaths);
+				VCToolChain.AddAnalyzeLogFile(arguments, AnalyzeLogFile, RootPaths);
 			}
 
 			if (ExperimentalLogFile != null)
 			{
-				VCToolChain.AddExperimentalLogFile(Arguments, ExperimentalLogFile, RootPaths);
+				VCToolChain.AddExperimentalLogFile(arguments, ExperimentalLogFile, RootPaths);
 			}
 
 			// A better way to express this? .json is used as output for /sourceDependencies), but .md.json is used as output for /sourceDependencies:directives)
 			if (DependencyListFile != null && DependencyListFile.HasExtension(".json") && !DependencyListFile.HasExtension(".md.json"))
 			{
-				VCToolChain.AddSourceDependenciesFile(Arguments, DependencyListFile, RootPaths);
+				VCToolChain.AddSourceDependenciesFile(arguments, DependencyListFile, RootPaths);
 			}
 
 			if (DependencyListFile != null && DependencyListFile.HasExtension(".d"))
 			{
-				VCToolChain.AddSourceDependsFile(Arguments, DependencyListFile, RootPaths);
+				VCToolChain.AddSourceDependsFile(arguments, DependencyListFile, RootPaths);
 			}
 
-			Arguments.AddRange(this.Arguments);
-			return Arguments;
+			arguments.AddRange(Arguments);
+			return arguments;
 		}
 
 		string GetClArguments()
@@ -644,39 +623,39 @@ namespace UnrealBuildTool
 				return String.Join(' ', Arguments);
 			}
 
-			string ResponseFileString = VCToolChain.NormalizeCommandLinePath(ResponseFile, RootPaths);
+			string responseFileString = VCToolChain.NormalizeCommandLinePath(ResponseFile, RootPaths);
 
 			// cl.exe can't handle response files with a path longer than 260 characters, and relative paths can push it over the limit
-			if (!System.IO.Path.IsPathRooted(ResponseFileString) && System.IO.Path.Combine(WorkingDirectory.FullName, ResponseFileString).Length > 260)
+			if (!System.IO.Path.IsPathRooted(responseFileString) && System.IO.Path.Combine(WorkingDirectory.FullName, responseFileString).Length > 260)
 			{
-				ResponseFileString = ResponseFile.FullName;
+				responseFileString = ResponseFile.FullName;
 			}
-			return String.Format("@{0} {1}", Utils.MakePathSafeToUseWithCommandLine(ResponseFileString), String.Join(' ', Arguments));
+			return $"@{Utils.MakePathSafeToUseWithCommandLine(responseFileString)} {String.Join(' ', Arguments)}";
 		}
 
 		string GetClFilterArguments()
 		{
-			List<string> Arguments = new List<string>();
-			string DependencyListFileString = VCToolChain.NormalizeCommandLinePath(DependencyListFile!, RootPaths);
-			Arguments.Add(String.Format("-dependencies={0}", Utils.MakePathSafeToUseWithCommandLine(DependencyListFileString)));
+			List<string> arguments = [];
+			string dependencyListFileString = VCToolChain.NormalizeCommandLinePath(DependencyListFile!, RootPaths);
+			arguments.Add($"-dependencies={Utils.MakePathSafeToUseWithCommandLine(dependencyListFileString)}");
 
 			if (TimingFile != null)
 			{
-				string TimingFileString = VCToolChain.NormalizeCommandLinePath(TimingFile, RootPaths);
-				Arguments.Add(String.Format("-timing={0}", Utils.MakePathSafeToUseWithCommandLine(TimingFileString)));
+				string timingFileString = VCToolChain.NormalizeCommandLinePath(TimingFile, RootPaths);
+				arguments.Add($"-timing={Utils.MakePathSafeToUseWithCommandLine(timingFileString)}");
 			}
 			if (bShowIncludes)
 			{
-				Arguments.Add("-showincludes");
+				arguments.Add("-showincludes");
 			}
 
-			Arguments.Add(String.Format("-compiler={0}", Utils.MakePathSafeToUseWithCommandLine(CompilerExe.AbsolutePath)));
-			Arguments.Add("--");
-			Arguments.Add(Utils.MakePathSafeToUseWithCommandLine(CompilerExe.AbsolutePath));
-			Arguments.Add(GetClArguments());
-			Arguments.Add("/showIncludes");
+			arguments.Add($"-compiler={Utils.MakePathSafeToUseWithCommandLine(CompilerExe.AbsolutePath)}");
+			arguments.Add("--");
+			arguments.Add(Utils.MakePathSafeToUseWithCommandLine(CompilerExe.AbsolutePath));
+			arguments.Add(GetClArguments());
+			arguments.Add("/showIncludes");
 
-			return String.Join(' ', Arguments);
+			return String.Join(' ', arguments);
 		}
 	}
 
@@ -686,15 +665,15 @@ namespace UnrealBuildTool
 	class VCCompileActionSerializer : ActionSerializerBase<VCCompileAction>
 	{
 		/// <inheritdoc/>
-		public override VCCompileAction Read(BinaryArchiveReader Reader)
+		public override VCCompileAction Read(BinaryArchiveReader reader)
 		{
-			return new VCCompileAction(Reader);
+			return new(reader);
 		}
 
 		/// <inheritdoc/>
-		public override void Write(BinaryArchiveWriter Writer, VCCompileAction Action)
+		public override void Write(BinaryArchiveWriter writer, VCCompileAction action)
 		{
-			Action.Write(Writer);
+			action.Write(writer);
 		}
 	}
 }
