SNDBS: Removed incorrect entries for c1ui.dll.

The removed lines were ignored by SNDBS because they used 'file01', which was already used by other lines. Also, c1ui.dll doesn't seem to be necessary and in fact is not even present in the default installation of MSVC.

diff --git a/Engine/Source/Programs/UnrealBuildTool/Executors/SNDBS.cs b/Engine/Source/Programs/UnrealBuildTool/Executors/SNDBS.cs
--- a/Engine/Source/Programs/UnrealBuildTool/Executors/SNDBS.cs
+++ b/Engine/Source/Programs/UnrealBuildTool/Executors/SNDBS.cs
@@ -467,7 +467,6 @@ server={VC_TOOLCHAIN_DIR}\mspdbsrv.exe
 [files]
 main=cl-filter.exe
 file01={VC_TOOLCHAIN_DIR}\c1.dll
-file01={VC_TOOLCHAIN_DIR}\c1ui.dll
 file02={VC_TOOLCHAIN_DIR}\c1xx.dll
 file03={VC_TOOLCHAIN_DIR}\c2.dll
 file04={VC_TOOLCHAIN_DIR}\mspdb140.dll
@@ -519,7 +518,6 @@ server={VC_TOOLCHAIN_DIR}\mspdbsrv.exe
 [files]
 main={VC_TOOLCHAIN_DIR}\cl.exe
 file01={VC_TOOLCHAIN_DIR}\c1.dll
-file01={VC_TOOLCHAIN_DIR}\c1ui.dll
 file02={VC_TOOLCHAIN_DIR}\c1xx.dll
 file03={VC_TOOLCHAIN_DIR}\c2.dll
 file04={VC_TOOLCHAIN_DIR}\mspdb140.dll
