Update Mac Toolchain to strip symbols & consolidate shared Apple platform logic.
- Consolidate shared logic between iOS & Mac to UEBuildApple.cs and AppleToolChain.cs, including symbol stripping.
- Mac inherits symbol stripping support from new shared Apple platform logic.
- Add bCreateStripFlagFile to create a ".stripped" file representing a stripped library.
- Disable chaining PCHs in shared Apple platform logic; disabling it both iOS & Mac, whereas it was previously only disabled on Mac.
- Consolidate stripped symbol command logic into CreateStripSymbolsProcessInfo function.
- Update post-build symbol stripping logic to handle new bCreateStripFlagFile flag and shared CreateStripSymbolsProcessInfo function.
- New types: AppleTargetRules, ReadOnlyAppleTargetRules, and AppleArchitectureConfig.

#domain Build System
#domain Apple Platform

diff --git a/Engine/Source/Programs/UnrealBuildTool/Platform/Apple/UEBuildApple.cs b/Engine/Source/Programs/UnrealBuildTool/Platform/Apple/UEBuildApple.cs
new file mode 100644
--- /dev/null
+++ b/Engine/Source/Programs/UnrealBuildTool/Platform/Apple/UEBuildApple.cs
@@ -0,0 +1,247 @@
+// Copyright Epic Games, Inc. All Rights Reserved.
+
+using System;
+using System.Collections.Generic;
+using System.Text;
+using EpicGames.Core;
+using Microsoft.Extensions.Logging;
+
+// ReSharper disable once CheckNamespace
+namespace UnrealBuildTool
+{
+	/// <summary>
+	/// Apple-specific target settings
+	/// </summary>
+	public abstract class AppleTargetRules
+	{ 
+		/// <summary>
+		/// Whether to strip symbols or not (implied by Shipping config).
+		/// </summary>
+		[XmlConfigFile(Category = "BuildConfiguration")]
+		[CommandLine("-stripsymbols", Value = "true")]
+		public bool bStripSymbols = false;
+		
+		/// <summary>
+		/// Whether to create a ".stripped" file when stripping symbols or not.
+		/// </summary>
+		[XmlConfigFile(Category = "BuildConfiguration")]
+		[CommandLine("-createstripflagfile", Value = "true")]
+		public bool bCreateStripFlagFile = false;
+
+		/// <summary>
+		/// Disables clang build verification checks on static libraries
+		/// </summary>
+		[CommandLine("-skipclangvalidation", Value = "true")]
+		[XmlConfigFile(Category = "BuildConfiguration")]
+		public bool bSkipClangValidation = false;
+		
+		/// <summary>
+		/// Enables address sanitizer (ASan).
+		/// </summary>
+		[CommandLine("-EnableASan")]
+		[XmlConfigFile(Category = "BuildConfiguration")]
+		public bool bEnableAddressSanitizer = false;
+
+		/// <summary>
+		/// Enables thread sanitizer (TSan).
+		/// </summary>
+		[CommandLine("-EnableTSan")]
+		[XmlConfigFile(Category = "BuildConfiguration")]
+		public bool bEnableThreadSanitizer = false;
+
+		/// <summary>
+		/// Enables undefined behavior sanitizer (UBSan).
+		/// </summary>
+		[CommandLine("-EnableUBSan")]
+		[XmlConfigFile(Category = "BuildConfiguration")]
+		public bool bEnableUndefinedBehaviorSanitizer = false;
+	}
+
+	/// <summary>
+	/// Read-only wrapper for Apple-specific target settings
+	/// </summary>
+	public class ReadOnlyAppleTargetRules
+	{
+		/// <summary>
+		/// The private mutable settings object
+		/// </summary>
+		private readonly AppleTargetRules Inner;
+
+		/// <summary>
+		/// Constructor
+		/// </summary>
+		/// <param name="Inner">The settings object to wrap</param>
+		public ReadOnlyAppleTargetRules(AppleTargetRules Inner)
+		{
+			this.Inner = Inner;
+		}
+
+		/// <summary>
+		/// Accessors for fields on the inner TargetRules instance
+		/// </summary>
+
+		#region Read-only accessor properties
+#pragma warning disable CS1591
+		public bool bStripSymbols => Inner.bStripSymbols;
+		public bool bCreateStripFlagFile => Inner.bCreateStripFlagFile;
+		public bool bSkipClangValidation => Inner.bSkipClangValidation;
+		public bool bEnableAddressSanitizer => Inner.bEnableAddressSanitizer;
+		public bool bEnableThreadSanitizer => Inner.bEnableThreadSanitizer;
+		public bool bEnableUndefinedBehaviorSanitizer => Inner.bEnableUndefinedBehaviorSanitizer;
+#pragma warning restore CS1591
+		#endregion
+	}
+	
+	/// <summary>
+	/// Architecture config for Apple platforms
+	/// </summary>
+	class AppleArchitectureConfig : UnrealArchitectureConfig
+	{
+		public AppleArchitectureConfig(IEnumerable<UnrealArch> SupportedArchitectures)
+			: base(UnrealArchitectureMode.SingleTargetCompileSeparately, SupportedArchitectures)
+		{
+		}
+	}
+	
+	abstract class AppleBuildPlatform : UEBuildPlatform
+	{
+		public AppleBuildPlatform(UnrealTargetPlatform Platform, UEBuildPlatformSDK SDK, UnrealArchitectureConfig ArchitectureConfig, ILogger Logger)
+			: base(Platform, SDK, ArchitectureConfig, Logger)
+		{
+		}
+
+		public abstract AppleTargetRules GetAppleTargetRules(TargetRules Target);
+
+		public abstract ReadOnlyAppleTargetRules GetAppleTargetRules(ReadOnlyTargetRules Target);
+
+		public override void GetExternalBuildMetadata(FileReference? ProjectFile, StringBuilder Metadata)
+		{
+			base.GetExternalBuildMetadata(ProjectFile, Metadata);
+			Metadata.AppendLine("xcode-select: {0}", ApplePlatformSDK.DeveloperDir);
+		}
+
+		public override bool IsBuildProduct(string FileName, string[] NamePrefixes, string[] NameSuffixes)
+		{
+			return IsBuildProductName(FileName, NamePrefixes, NameSuffixes, "")
+			       || IsBuildProductName(FileName, NamePrefixes, NameSuffixes, ".dylib")
+			       || IsBuildProductName(FileName, NamePrefixes, NameSuffixes, ".dSYM");
+		}
+
+		public override bool CanUseFASTBuild()
+		{
+			return true;
+		}
+
+		public override void ResetTarget(TargetRules Target)
+		{
+			base.ResetTarget(Target);
+
+			// always strip in shipping configuration (commandline could have set it also)
+			if (Target.Configuration == UnrealTargetConfiguration.Shipping)
+			{
+				GetAppleTargetRules(Target).bStripSymbols = true;
+			}
+		}
+
+		public override void ValidateTarget(TargetRules Target)
+		{
+			base.ValidateTarget(Target);
+			
+			if (!string.IsNullOrWhiteSpace(Environment.GetEnvironmentVariable("CLANG_STATIC_ANALYZER_MODE")))
+			{
+				Target.StaticAnalyzer = StaticAnalyzer.Default;
+				Target.StaticAnalyzerOutputType = (Environment.GetEnvironmentVariable("CLANG_ANALYZER_OUTPUT")?.Contains("html", StringComparison.OrdinalIgnoreCase) == true) ? StaticAnalyzerOutputType.Html : StaticAnalyzerOutputType.Text;
+				Target.StaticAnalyzerMode = string.Equals(Environment.GetEnvironmentVariable("CLANG_STATIC_ANALYZER_MODE"), "shallow", StringComparison.OrdinalIgnoreCase) ? StaticAnalyzerMode.Shallow : StaticAnalyzerMode.Deep;
+			}
+			else if (Target.StaticAnalyzer == StaticAnalyzer.Clang)
+			{
+				Target.StaticAnalyzer = StaticAnalyzer.Default;
+			}
+
+			// Disable linking and ignore build outputs if we're using a static analyzer
+			if (Target.StaticAnalyzer == StaticAnalyzer.Default)
+			{
+				Target.bDisableLinking = true;
+				Target.bIgnoreBuildOutputs = true;
+
+				// Clang static analysis requires non unity builds
+				Target.bUseUnityBuild = false;
+				if (Target.bStaticAnalyzerIncludeGenerated)
+				{
+					Target.bAlwaysUseUnityForGeneratedFiles = false;
+				}
+
+				// Disable chaining PCHs for the moment because it is crashing clang
+				Target.bChainPCHs = false;
+			}
+
+			if (Target.bCompileAgainstEngine)
+			{
+				Target.GlobalDefinitions.Add("HAS_METAL=1");
+				Target.ExtraModuleNames.Add("MetalRHI");
+			}
+			else
+			{
+				Target.GlobalDefinitions.Add("HAS_METAL=0");
+			}
+			
+			// Force using the ANSI allocator if ASan is enabled
+			if (GetAppleTargetRules(Target).bEnableAddressSanitizer)
+			{
+				Target.StaticAllocator = StaticAllocatorType.Ansi;
+			}
+		}
+
+		/// <summary>
+		/// Get the extension to use for the given binary type
+		/// </summary>
+		/// <param name="InBinaryType"> The binary type being built</param>
+		/// <returns>string    The binary extenstion (ie 'exe' or 'dll')</returns>
+		public override string GetBinaryExtension(UEBuildBinaryType InBinaryType)
+		{
+			switch (InBinaryType)
+			{
+				case UEBuildBinaryType.DynamicLinkLibrary:
+					return ".dylib";
+				case UEBuildBinaryType.Executable:
+					return "";
+				case UEBuildBinaryType.StaticLibrary:
+					return ".a";
+			}
+
+			return base.GetBinaryExtension(InBinaryType);
+		}
+
+		public override void SetUpEnvironment(ReadOnlyTargetRules Target, CppCompileEnvironment CompileEnvironment, LinkEnvironment LinkEnvironment)
+		{
+			// shared stuff for all apple platforms
+			
+			CompileEnvironment.Definitions.Add("PLATFORM_APPLE=1");
+
+			AppleExports.GetSwiftIntegrationSettings(Target.ProjectFile, Target.Type, Platform, out var bUseSwiftUIMain, out _);
+			CompileEnvironment.Definitions.Add("UE_USE_SWIFT_UI_MAIN=" + (bUseSwiftUIMain ? "1" : "0"));
+
+			CompileEnvironment.Definitions.Add("WITH_TTS=0");
+			CompileEnvironment.Definitions.Add("WITH_SPEECH_RECOGNITION=0");
+		}
+		
+		protected ClangToolChainOptions GetToolChainOptionsForSanitizers(ReadOnlyTargetRules Target)
+		{
+			ReadOnlyAppleTargetRules AppleTarget = GetAppleTargetRules(Target);
+			ClangToolChainOptions Options = ClangToolChainOptions.None;
+			if (AppleTarget.bEnableAddressSanitizer)
+			{
+				Options |= ClangToolChainOptions.EnableAddressSanitizer;
+			}
+			if (AppleTarget.bEnableThreadSanitizer)
+			{
+				Options |= ClangToolChainOptions.EnableThreadSanitizer;
+			}
+			if (AppleTarget.bEnableUndefinedBehaviorSanitizer)
+			{
+				Options |= ClangToolChainOptions.EnableUndefinedBehaviorSanitizer;
+			}
+			return Options;
+		}
+	}
+}
diff --git a/Engine/Source/Programs/UnrealBuildTool/Platform/IOS/IOSToolChain.cs b/Engine/Source/Programs/UnrealBuildTool/Platform/IOS/IOSToolChain.cs
--- a/Engine/Source/Programs/UnrealBuildTool/Platform/IOS/IOSToolChain.cs
+++ b/Engine/Source/Programs/UnrealBuildTool/Platform/IOS/IOSToolChain.cs
@@ -135,6 +135,11 @@ namespace UnrealBuildTool
 			return OutputFile.ChangeExtension(DebugExtension);
 		}
 
+		public override ReadOnlyAppleTargetRules GetAppleTargetRules(ReadOnlyTargetRules InTarget)
+		{
+			return InTarget.IOSPlatform;
+		}
+
 		/// <inheritdoc/>
 		protected override void GetCompileArguments_WarningsAndErrors(CppCompileEnvironment CompileEnvironment, List<string> Arguments)
 		{
@@ -1076,27 +1081,6 @@ namespace UnrealBuildTool
 				return OutputFiles;
 			}
 
-			// strip the debug info from the executable if needed. creates a dummy output file for the action graph to track that it's out of date.
-			if (Target.IOSPlatform.bStripSymbols || (Target.Configuration == UnrealTargetConfiguration.Shipping))
-			{
-				FileItem StripCompleteFile = FileItem.GetItemByFileReference(FileReference.Combine(BinaryLinkEnvironment.IntermediateDirectory!, Executable.Location.GetFileName() + ".stripped"));
-
-				// If building a framework we can only strip local symbols, need to leave global in place
-				string StripArguments = BinaryLinkEnvironment.bIsBuildingDLL ? "-x" : "";
-
-				Action StripAction = Graph.CreateAction(ActionType.CreateAppBundle);
-				StripAction.WorkingDirectory = GetMacDevSrcRoot();
-				StripAction.CommandPath = BuildHostPlatform.Current.Shell;
-				StripAction.CommandArguments = String.Format("-c \"\\\"{0}/strip\\\" {1} \\\"{2}\\\" && touch \\\"{3}\\\"\"", Settings.ToolchainDir, StripArguments, Executable.Location, StripCompleteFile);
-				StripAction.PrerequisiteItems.Add(Executable);
-				StripAction.PrerequisiteItems.UnionWith(OutputFiles);
-				StripAction.ProducedItems.Add(StripCompleteFile);
-				StripAction.StatusDescription = String.Format("Stripping symbols from {0}", Executable.AbsolutePath);
-				StripAction.bCanExecuteRemotely = false;
-
-				OutputFiles.Add(StripCompleteFile);
-			}
-
 			if (!AppleExports.UseModernXcode(Target.ProjectFile) && !BinaryLinkEnvironment.bIsBuildingDLL)
 			{
 				// generate the asset catalog
diff --git a/Engine/Source/Programs/UnrealBuildTool/Platform/IOS/UEBuildIOS.cs b/Engine/Source/Programs/UnrealBuildTool/Platform/IOS/UEBuildIOS.cs
--- a/Engine/Source/Programs/UnrealBuildTool/Platform/IOS/UEBuildIOS.cs
+++ b/Engine/Source/Programs/UnrealBuildTool/Platform/IOS/UEBuildIOS.cs
@@ -53,15 +53,8 @@ namespace UnrealBuildTool
 	/// <summary>
 	/// IOS-specific target settings
 	/// </summary>
-	public partial class IOSTargetRules
+	public class IOSTargetRules : AppleTargetRules
 	{
-		/// <summary>
-		/// Whether to strip iOS symbols or not (implied by Shipping config).
-		/// </summary>
-		[XmlConfigFile(Category = "BuildConfiguration")]
-		[CommandLine("-stripsymbols", Value = "true")]
-		public bool bStripSymbols = false;
-
 		/// <summary>
 		///
 		/// </summary>
@@ -92,13 +85,6 @@ namespace UnrealBuildTool
 		[CommandLine("-skipcrashlytics")]
 		public bool bSkipCrashlytics = false;
 
-		/// <summary>
-		/// Disables clang build verification checks on static libraries
-		/// </summary>
-		[CommandLine("-skipclangvalidation", Value = "true")]
-		[XmlConfigFile(Category = "BuildConfiguration", Name = "bSkipClangValidation")]
-		public bool bSkipClangValidation = false;
-
 		/// <summary>
 		/// Mark the build for distribution
 		/// </summary>
@@ -128,24 +114,6 @@ namespace UnrealBuildTool
 		/// </summary>
 		public IOSProjectSettings? ProjectSettings = null;
 
-		/// <summary>
-		/// Enables address sanitizer (ASan)
-		/// </summary>
-		[CommandLine("-EnableASan")]
-		public bool bEnableAddressSanitizer = false;
-
-		/// <summary>
-		/// Enables thread sanitizer (TSan)
-		/// </summary>
-		[CommandLine("-EnableTSan")]
-		public bool bEnableThreadSanitizer = false;
-
-		/// <summary>
-		/// Enables undefined behavior sanitizer (UBSan)
-		/// </summary>
-		[CommandLine("-EnableUBSan")]
-		public bool bEnableUndefinedBehaviorSanitizer = false;
-
 		/// <summary>
 		/// Constructor
 		/// </summary>
@@ -174,18 +142,19 @@ namespace UnrealBuildTool
 	/// <summary>
 	/// Read-only wrapper for IOS-specific target settings
 	/// </summary>
-	public partial class ReadOnlyIOSTargetRules
+	public class ReadOnlyIOSTargetRules : ReadOnlyAppleTargetRules
 	{
 		/// <summary>
 		/// The private mutable settings object
 		/// </summary>
-		private IOSTargetRules Inner;
+		private readonly IOSTargetRules Inner;
 
 		/// <summary>
 		/// Constructor
 		/// </summary>
 		/// <param name="Inner">The settings object to wrap</param>
 		public ReadOnlyIOSTargetRules(IOSTargetRules Inner)
+			: base(Inner)
 		{
 			this.Inner = Inner;
 		}
@@ -195,8 +164,6 @@ namespace UnrealBuildTool
 		/// </summary>
 		#region Read-only accessor properties 
 #pragma warning disable CS1591
-		public bool bStripSymbols => Inner.bStripSymbols;
-
 		public bool bGenerateFrameworkWrapperProject => Inner.bGenerateFrameworkWrapperProject;
 
 		public bool bGeneratedSYM => Inner.bGeneratedSYM;
@@ -205,8 +172,6 @@ namespace UnrealBuildTool
 
 		public bool bSkipCrashlytics => Inner.bSkipCrashlytics;
 
-		public bool bSkipClangValidation => Inner.bSkipClangValidation;
-
 		public bool bForDistribution => Inner.bForDistribution;
 
 		public string? ImportProvision => Inner.ImportProvision;
@@ -217,11 +182,6 @@ namespace UnrealBuildTool
 
 		public float RuntimeVersion => Single.Parse(Inner.ProjectSettings!.RuntimeVersion, System.Globalization.CultureInfo.InvariantCulture);
 
-		public bool bEnableAddressSanitizer => Inner.bEnableAddressSanitizer;
-
-		public bool bEnableThreadSanitizer => Inner.bEnableThreadSanitizer;
-
-		public bool bEnableUndefinedBehaviorSanitizer => Inner.bEnableUndefinedBehaviorSanitizer;
 #pragma warning restore CS1591
 		#endregion
 	}
@@ -686,10 +646,10 @@ namespace UnrealBuildTool
 		}
 	}
 
-	class IOSArchitectureConfig : UnrealArchitectureConfig
+	class IOSArchitectureConfig : AppleArchitectureConfig
 	{
 		public IOSArchitectureConfig()
-			: base(UnrealArchitectureMode.SingleTargetCompileSeparately, new[] { UnrealArch.Arm64, UnrealArch.IOSSimulator })
+			: base(new[] { UnrealArch.Arm64, UnrealArch.IOSSimulator })
 		{
 
 		}
@@ -733,18 +693,14 @@ namespace UnrealBuildTool
 
 		public override void ResetTarget(TargetRules Target)
 		{
+			base.ResetTarget(Target);
+			
 			Target.bDeployAfterCompile = true;
 
 			Target.IOSPlatform.ProjectSettings = ((IOSPlatform)GetBuildPlatform(Target.Platform)).ReadProjectSettings(Target.ProjectFile);
 
 			if (!AppleExports.UseModernXcode(Target.ProjectFile))
 			{
-				// always strip in shipping configuration (commandline could have set it also)
-				if (Target.Configuration == UnrealTargetConfiguration.Shipping)
-				{
-					Target.IOSPlatform.bStripSymbols = true;
-				}
-
 				// if we are stripping the executable, or if the project requested it, or if it's a buildmachine, generate the dsym
 				if (Target.IOSPlatform.bStripSymbols || Target.IOSPlatform.ProjectSettings.bGeneratedSYMFile || Unreal.IsBuildMachine())
 				{
@@ -760,41 +716,6 @@ namespace UnrealBuildTool
 		{
 			base.ValidateTarget(Target);
 
-			if (!String.IsNullOrWhiteSpace(Environment.GetEnvironmentVariable("CLANG_STATIC_ANALYZER_MODE")))
-			{
-				Target.StaticAnalyzer = StaticAnalyzer.Default;
-				Target.StaticAnalyzerOutputType = (Environment.GetEnvironmentVariable("CLANG_ANALYZER_OUTPUT")?.Contains("html", StringComparison.OrdinalIgnoreCase) == true) ? StaticAnalyzerOutputType.Html : StaticAnalyzerOutputType.Text;
-				Target.StaticAnalyzerMode = String.Equals(Environment.GetEnvironmentVariable("CLANG_STATIC_ANALYZER_MODE"), "shallow", StringComparison.OrdinalIgnoreCase) ? StaticAnalyzerMode.Shallow : StaticAnalyzerMode.Deep;
-			}
-			else if (Target.StaticAnalyzer == StaticAnalyzer.Clang)
-			{
-				Target.StaticAnalyzer = StaticAnalyzer.Default;
-			}
-
-			// Disable linking and ignore build outputs if we're using a static analyzer
-			if (Target.StaticAnalyzer == StaticAnalyzer.Default)
-			{
-				Target.bDisableLinking = true;
-				Target.bIgnoreBuildOutputs = true;
-
-				// Clang static analysis requires non unity builds
-				Target.bUseUnityBuild = false;
-				if (Target.bStaticAnalyzerIncludeGenerated)
-				{
-					Target.bAlwaysUseUnityForGeneratedFiles = false;
-				}
-			}
-
-			// we assume now we are building with IOS8 or later
-			if (Target.bCompileAgainstEngine)
-			{
-				Target.GlobalDefinitions.Add("HAS_METAL=1");
-				Target.ExtraModuleNames.Add("MetalRHI");
-			}
-			else
-			{
-				Target.GlobalDefinitions.Add("HAS_METAL=0");
-			}
 
 			if (Target.bShouldCompileAsDLL)
 			{
@@ -819,10 +740,6 @@ namespace UnrealBuildTool
 				// Prefer MB2 as MB3 requires extended virtual address space entitlement
 				Target.StaticAllocator = StaticAllocatorType.Binned2;
 			}
-			if (Target.IOSPlatform.bEnableAddressSanitizer)
-			{
-				Target.StaticAllocator = StaticAllocatorType.Ansi;
-			}
 		}
 
 		public override void ValidateModule(UEBuildModule Module, ReadOnlyTargetRules Target, ILogger Logger)
@@ -875,6 +792,16 @@ namespace UnrealBuildTool
 #endif
 		}
 
+		public override AppleTargetRules GetAppleTargetRules(TargetRules Target)
+		{
+			return Target.IOSPlatform;
+		}
+
+		public override ReadOnlyAppleTargetRules GetAppleTargetRules(ReadOnlyTargetRules Target)
+		{
+			return Target.IOSPlatform;
+		}
+
 		/// <summary>
 		/// Determines if the given name is a build product for a target.
 		/// </summary>
@@ -884,33 +811,12 @@ namespace UnrealBuildTool
 		/// <returns>True if the string matches the name of a build product, false otherwise</returns>
 		public override bool IsBuildProduct(string FileName, string[] NamePrefixes, string[] NameSuffixes)
 		{
-			return IsBuildProductName(FileName, NamePrefixes, NameSuffixes, "")
+			return base.IsBuildProduct(FileName, NamePrefixes, NameSuffixes)
 				|| IsBuildProductName(FileName, NamePrefixes, NameSuffixes, ".stub")
-				|| IsBuildProductName(FileName, NamePrefixes, NameSuffixes, ".dylib")
-				|| IsBuildProductName(FileName, NamePrefixes, NameSuffixes, ".dSYM")
 				|| IsBuildProductName(FileName, NamePrefixes, NameSuffixes, ".dSYM.zip")
 				|| IsBuildProductName(FileName, NamePrefixes, NameSuffixes, ".o");
 		}
 
-		/// <summary>
-		/// Get the extension to use for the given binary type
-		/// </summary>
-		/// <param name="InBinaryType"> The binary type being built</param>
-		/// <returns>string    The binary extenstion (ie 'exe' or 'dll')</returns>
-		public override string GetBinaryExtension(UEBuildBinaryType InBinaryType)
-		{
-			switch (InBinaryType)
-			{
-				case UEBuildBinaryType.DynamicLinkLibrary:
-					return ".dylib";
-				case UEBuildBinaryType.Executable:
-					return "";
-				case UEBuildBinaryType.StaticLibrary:
-					return ".a";
-			}
-			return base.GetBinaryExtension(InBinaryType);
-		}
-
 		public IOSProjectSettings ReadProjectSettings(FileReference? ProjectFile, string? Bundle = "")
 		{
 			IOSProjectSettings? ProjectSettings = null;
@@ -1003,11 +909,6 @@ namespace UnrealBuildTool
 			return BuildHostPlatform.Current.Platform == UnrealTargetPlatform.Mac;
 		}
 
-		public override bool CanUseFASTBuild()
-		{
-			return true;
-		}
-
 		public bool HasCustomIcons(DirectoryReference ProjectDirectoryName, ILogger Logger)
 		{
 			string IconDir = Path.Combine(ProjectDirectoryName.FullName, "Build", "IOS", "Resources", "Graphics");
@@ -1227,8 +1128,6 @@ namespace UnrealBuildTool
 				CompileEnvironment.Definitions.Add("UE_PLATFORM_IOS_ONLY=1");
 			}
 
-			CompileEnvironment.Definitions.Add("WITH_TTS=0");
-			CompileEnvironment.Definitions.Add("WITH_SPEECH_RECOGNITION=0");
 			CompileEnvironment.Definitions.Add("WITH_EDITOR=0");
 			CompileEnvironment.Definitions.Add("USE_NULL_RHI=0");
 
@@ -1326,20 +1225,7 @@ namespace UnrealBuildTool
 		/// <returns>New toolchain instance.</returns>
 		public override UEToolChain CreateToolChain(ReadOnlyTargetRules Target)
 		{
-			ClangToolChainOptions Options = ClangToolChainOptions.None;
-			if (Target.IOSPlatform.bEnableAddressSanitizer)
-			{
-				Options |= ClangToolChainOptions.EnableAddressSanitizer;
-			}
-			if (Target.IOSPlatform.bEnableThreadSanitizer)
-			{
-				Options |= ClangToolChainOptions.EnableThreadSanitizer;
-			}
-			if (Target.IOSPlatform.bEnableUndefinedBehaviorSanitizer)
-			{
-				Options |= ClangToolChainOptions.EnableUndefinedBehaviorSanitizer;
-			}
-
+			ClangToolChainOptions Options = GetToolChainOptionsForSanitizers(Target);
 			IOSProjectSettings ProjectSettings = ReadProjectSettings(Target.ProjectFile);
 			return new IOSToolChain(Target, ProjectSettings, Options, Logger);
 		}
diff --git a/Engine/Source/Programs/UnrealBuildTool/Platform/Mac/MacToolChain.cs b/Engine/Source/Programs/UnrealBuildTool/Platform/Mac/MacToolChain.cs
--- a/Engine/Source/Programs/UnrealBuildTool/Platform/Mac/MacToolChain.cs
+++ b/Engine/Source/Programs/UnrealBuildTool/Platform/Mac/MacToolChain.cs
@@ -125,6 +125,11 @@ namespace UnrealBuildTool
 			return null;
 		}
 
+		public override ReadOnlyAppleTargetRules GetAppleTargetRules(ReadOnlyTargetRules InTarget)
+		{
+			return InTarget.MacPlatform;
+		}
+
 		/// <inheritdoc/>
 		protected override void GetCompileArguments_WarningsAndErrors(CppCompileEnvironment CompileEnvironment, List<string> Arguments)
 		{
diff --git a/Engine/Source/Programs/UnrealBuildTool/Platform/Mac/UEBuildMac.cs b/Engine/Source/Programs/UnrealBuildTool/Platform/Mac/UEBuildMac.cs
--- a/Engine/Source/Programs/UnrealBuildTool/Platform/Mac/UEBuildMac.cs
+++ b/Engine/Source/Programs/UnrealBuildTool/Platform/Mac/UEBuildMac.cs
@@ -12,15 +12,8 @@ namespace UnrealBuildTool
 	/// <summary>
 	/// Mac-specific target settings
 	/// </summary>
-	public class MacTargetRules
+	public class MacTargetRules : AppleTargetRules
 	{
-		/// <summary>
-		/// Enables address sanitizer (ASan).
-		/// </summary>
-		[CommandLine("-EnableASan")]
-		[XmlConfigFile(Category = "BuildConfiguration", Name = "bEnableAddressSanitizer")]
-		public bool bEnableAddressSanitizer = false;
-
 		/// <summary>
 		/// Enables LibFuzzer.
 		/// </summary>
@@ -28,20 +21,6 @@ namespace UnrealBuildTool
 		[XmlConfigFile(Category = "BuildConfiguration", Name = "bEnableLibFuzzer")]
 		public bool bEnableLibFuzzer = false;
 
-		/// <summary>
-		/// Enables thread sanitizer (TSan).
-		/// </summary>
-		[CommandLine("-EnableTSan")]
-		[XmlConfigFile(Category = "BuildConfiguration", Name = "bEnableThreadSanitizer")]
-		public bool bEnableThreadSanitizer = false;
-
-		/// <summary>
-		/// Enables undefined behavior sanitizer (UBSan).
-		/// </summary>
-		[CommandLine("-EnableUBSan")]
-		[XmlConfigFile(Category = "BuildConfiguration", Name = "bEnableUndefinedBehaviorSanitizer")]
-		public bool bEnableUndefinedBehaviorSanitizer = false;
-
 		/// <summary>
 		/// Enables the generation of .dsym files. This can be disabled to enable faster iteration times during development.
 		/// </summary>
@@ -49,13 +28,6 @@ namespace UnrealBuildTool
 		[XmlConfigFile(Category = "BuildConfiguration", Name = "bUseDSYMFiles")]
 		public bool bUseDSYMFiles = false;
 
-		/// <summary>
-		/// Disables clang build verification checks on static libraries
-		/// </summary>
-		[CommandLine("-skipclangvalidation", Value = "true")]
-		[XmlConfigFile(Category = "BuildConfiguration", Name = "bSkipClangValidation")]
-		public bool bSkipClangValidation = false;
-
 		/// <summary>
 		/// Enables runtime ray tracing support.
 		/// </summary>
@@ -90,18 +62,19 @@ namespace UnrealBuildTool
 	/// <summary>
 	/// Read-only wrapper for Mac-specific target settings
 	/// </summary>
-	public class ReadOnlyMacTargetRules
+	public class ReadOnlyMacTargetRules : ReadOnlyAppleTargetRules
 	{
 		/// <summary>
 		/// The private mutable settings object
 		/// </summary>
-		private MacTargetRules Inner;
+		private readonly MacTargetRules Inner;
 
 		/// <summary>
 		/// Constructor
 		/// </summary>
 		/// <param name="Inner">The settings object to wrap</param>
 		public ReadOnlyMacTargetRules(MacTargetRules Inner)
+			: base(Inner)
 		{
 			this.Inner = Inner;
 		}
@@ -112,26 +85,18 @@ namespace UnrealBuildTool
 		#region Read-only accessor properties 
 #pragma warning disable CS1591
 
-		public bool bEnableAddressSanitizer => Inner.bEnableAddressSanitizer;
-
 		public bool bEnableLibFuzzer => Inner.bEnableLibFuzzer;
 
-		public bool bEnableThreadSanitizer => Inner.bEnableThreadSanitizer;
-
-		public bool bEnableUndefinedBehaviorSanitizer => Inner.bEnableUndefinedBehaviorSanitizer;
-
-		public bool bSkipClangValidation => Inner.bSkipClangValidation;
-
 		public bool bEnableRayTracing => Inner.bEnableRayTracing;
 
 #pragma warning restore CS1591
 		#endregion
 	}
 
-	class MacArchitectureConfig : UnrealArchitectureConfig
+	class MacArchitectureConfig : AppleArchitectureConfig
 	{
 		public MacArchitectureConfig()
-			: base(UnrealArchitectureMode.SingleTargetCompileSeparately, new[] { UnrealArch.X64, UnrealArch.Arm64 })
+			: base(new[] { UnrealArch.X64, UnrealArch.Arm64 })
 		{
 		}
 
@@ -293,33 +258,6 @@ namespace UnrealBuildTool
 		}
 	}
 
-	abstract class AppleBuildPlatform : UEBuildPlatform
-	{
-		public AppleBuildPlatform(UnrealTargetPlatform Platform, UEBuildPlatformSDK SDK, UnrealArchitectureConfig ArchitectureConfig, ILogger Logger)
-			: base(Platform, SDK, ArchitectureConfig, Logger)
-		{
-
-		}
-
-		public override void GetExternalBuildMetadata(FileReference? ProjectFile, StringBuilder Metadata)
-		{
-			base.GetExternalBuildMetadata(ProjectFile, Metadata);
-
-			Metadata.AppendLine("xcode-select: {0}", ApplePlatformSDK.DeveloperDir);
-		}
-
-		public override void SetUpEnvironment(ReadOnlyTargetRules Target, CppCompileEnvironment CompileEnvironment, LinkEnvironment LinkEnvironment)
-		{
-			// shared stuff for all apple platforms
-			
-			CompileEnvironment.Definitions.Add("PLATFORM_APPLE=1");
-
-			bool bUseSwiftUIMain;
-			AppleExports.GetSwiftIntegrationSettings(Target.ProjectFile, Target.Type, Platform, out bUseSwiftUIMain, out _);
-			CompileEnvironment.Definitions.Add("UE_USE_SWIFT_UI_MAIN=" + (bUseSwiftUIMain ? "1" : "0"));
-		}
-	}
-
 	class MacPlatform : AppleBuildPlatform
 	{
 		public MacPlatform(UEBuildPlatformSDK InSDK, ILogger InLogger)
@@ -327,47 +265,15 @@ namespace UnrealBuildTool
 		{
 		}
 
-		public override bool CanUseFASTBuild()
-		{
-			return true;
-		}
-
-		public override void ResetTarget(TargetRules Target)
-		{
-		}
-
 		public override void ValidateTarget(TargetRules Target)
 		{
 			base.ValidateTarget(Target);
 
-			if (!String.IsNullOrWhiteSpace(Environment.GetEnvironmentVariable("CLANG_STATIC_ANALYZER_MODE")))
-			{
-				Target.StaticAnalyzer = StaticAnalyzer.Default;
-				Target.StaticAnalyzerOutputType = (Environment.GetEnvironmentVariable("CLANG_ANALYZER_OUTPUT")?.Contains("html", StringComparison.OrdinalIgnoreCase) == true) ? StaticAnalyzerOutputType.Html : StaticAnalyzerOutputType.Text;
-				Target.StaticAnalyzerMode = String.Equals(Environment.GetEnvironmentVariable("CLANG_STATIC_ANALYZER_MODE"), "shallow", StringComparison.OrdinalIgnoreCase) ? StaticAnalyzerMode.Shallow : StaticAnalyzerMode.Deep;
-			}
-			else if (Target.StaticAnalyzer == StaticAnalyzer.Clang)
-			{
-				Target.StaticAnalyzer = StaticAnalyzer.Default;
-			}
-
-			// Disable linking and ignore build outputs if we're using a static analyzer
-			if (Target.StaticAnalyzer == StaticAnalyzer.Default)
-			{
-				Target.bDisableLinking = true;
-				Target.bIgnoreBuildOutputs = true;
-
-				// Clang static analysis requires non unity builds
-				Target.bUseUnityBuild = false;
 				if (Target.bStaticAnalyzerIncludeGenerated)
 				{
 					Target.bAlwaysUseUnityForGeneratedFiles = false;
 				}
 
-				// Disable chaining PCHs for the moment because it is crashing clang
-				Target.bChainPCHs = false;
-			}
-
 			if (BuildHostPlatform.Current.Platform != UnrealTargetPlatform.Mac)
 			{
 				// @todo: Temporarily disable precompiled header files when building remotely due to errors
@@ -381,25 +287,10 @@ namespace UnrealBuildTool
 				Target.DisablePlugins.AddRange(Array.Empty<string>());
 			}
 
-			if (Target.bCompileAgainstEngine)
-			{
-				Target.GlobalDefinitions.Add("HAS_METAL=1");
-				Target.ExtraModuleNames.Add("MetalRHI");
-			}
-			else
-			{
-				Target.GlobalDefinitions.Add("HAS_METAL=0");
-			}
-
 			if (Target.StaticAllocator == StaticAllocatorType.None)
 			{
 				Target.StaticAllocator = StaticAllocatorType.Binned3;
 			}
-			// Force using the ANSI allocator if ASan is enabled
-			if (Target.MacPlatform.bEnableAddressSanitizer)
-			{
-				Target.StaticAllocator = StaticAllocatorType.Ansi;
-			}
 
 			Target.GlobalDefinitions.Add("GL_SILENCE_DEPRECATION=1");
 
@@ -412,6 +303,16 @@ namespace UnrealBuildTool
 			Target.bCheckSystemHeadersForModification = BuildHostPlatform.Current.Platform != UnrealTargetPlatform.Mac;
 		}
 
+		public override AppleTargetRules GetAppleTargetRules(TargetRules Target)
+		{
+			return Target.MacPlatform;
+		}
+
+		public override ReadOnlyAppleTargetRules GetAppleTargetRules(ReadOnlyTargetRules Target)
+		{
+			return Target.MacPlatform;
+		}
+
 		/// <summary>
 		/// Determines if the given name is a build product for a target.
 		/// </summary>
@@ -421,32 +322,11 @@ namespace UnrealBuildTool
 		/// <returns>True if the string matches the name of a build product, false otherwise</returns>
 		public override bool IsBuildProduct(string FileName, string[] NamePrefixes, string[] NameSuffixes)
 		{
-			return IsBuildProductName(FileName, NamePrefixes, NameSuffixes, "")
-				|| IsBuildProductName(FileName, NamePrefixes, NameSuffixes, ".dsym")
-				|| IsBuildProductName(FileName, NamePrefixes, NameSuffixes, ".dylib")
+			return base.IsBuildProduct(FileName, NamePrefixes, NameSuffixes)
 				|| IsBuildProductName(FileName, NamePrefixes, NameSuffixes, ".a")
 				|| IsBuildProductName(FileName, NamePrefixes, NameSuffixes, ".app");
 		}
 
-		/// <summary>
-		/// Get the extension to use for the given binary type
-		/// </summary>
-		/// <param name="InBinaryType"> The binrary type being built</param>
-		/// <returns>string    The binary extenstion (ie 'exe' or 'dll')</returns>
-		public override string GetBinaryExtension(UEBuildBinaryType InBinaryType)
-		{
-			switch (InBinaryType)
-			{
-				case UEBuildBinaryType.DynamicLinkLibrary:
-					return ".dylib";
-				case UEBuildBinaryType.Executable:
-					return "";
-				case UEBuildBinaryType.StaticLibrary:
-					return ".a";
-			}
-			return base.GetBinaryExtension(InBinaryType);
-		}
-
 		/// <summary>
 		/// Get the extensions to use for debug info for the given binary type
 		/// </summary>
@@ -588,9 +468,6 @@ namespace UnrealBuildTool
 			{
 				CompileEnvironment.Definitions.Add("RHI_RAYTRACING=1");
 			}
-
-			CompileEnvironment.Definitions.Add("WITH_TTS=0");
-			CompileEnvironment.Definitions.Add("WITH_SPEECH_RECOGNITION=0");
 		}
 
 		/// <summary>
@@ -623,20 +500,7 @@ namespace UnrealBuildTool
 		/// <returns>New toolchain instance.</returns>
 		public override UEToolChain CreateToolChain(ReadOnlyTargetRules Target)
 		{
-			ClangToolChainOptions Options = ClangToolChainOptions.None;
-
-			if (Target.MacPlatform.bEnableAddressSanitizer)
-			{
-				Options |= ClangToolChainOptions.EnableAddressSanitizer;
-			}
-			if (Target.MacPlatform.bEnableThreadSanitizer)
-			{
-				Options |= ClangToolChainOptions.EnableThreadSanitizer;
-			}
-			if (Target.MacPlatform.bEnableUndefinedBehaviorSanitizer)
-			{
-				Options |= ClangToolChainOptions.EnableUndefinedBehaviorSanitizer;
-			}
+			ClangToolChainOptions Options = GetToolChainOptionsForSanitizers(Target);
 			if (Target.bShouldCompileAsDLL)
 			{
 				Options |= ClangToolChainOptions.OutputDylib;
diff --git a/Engine/Source/Programs/UnrealBuildTool/ToolChain/AppleToolChain.cs b/Engine/Source/Programs/UnrealBuildTool/ToolChain/AppleToolChain.cs
--- a/Engine/Source/Programs/UnrealBuildTool/ToolChain/AppleToolChain.cs
+++ b/Engine/Source/Programs/UnrealBuildTool/ToolChain/AppleToolChain.cs
@@ -242,6 +242,7 @@ namespace UnrealBuildTool
 		
 		protected AppleToolChainSettings Settings => ToolChainSettings.Value;
 
+		public abstract ReadOnlyAppleTargetRules GetAppleTargetRules(ReadOnlyTargetRules Target);
 
 		/// <summary>
 		/// The SDK version being used to compile with the toolchain
@@ -345,14 +346,18 @@ namespace UnrealBuildTool
 				File.Copy(SourceFile.FullName, TargetFile.FullName, true);
 			}
 
-			ProcessStartInfo StartInfo = new ProcessStartInfo();
-			StartInfo.FileName = Path.Combine(ToolchainDir.FullName, "strip");
-			StartInfo.Arguments = String.Format("\"{0}\" -S", TargetFile.FullName);
-			StartInfo.UseShellExecute = false;
-			StartInfo.CreateNoWindow = true;
+			ProcessStartInfo StartInfo = CreateStripSymbolsProcessInfo(TargetFile, ToolchainDir);
 			Utils.RunLocalProcessAndLogOutput(StartInfo, Logger);
 		}
 
+		protected static ProcessStartInfo CreateStripSymbolsProcessInfo(FileReference File, DirectoryReference ToolchainDir) => new()
+		{
+			FileName = Path.Combine(ToolchainDir.FullName, "strip"),
+			Arguments = $"\"{File.FullName}\" -S",
+			UseShellExecute = false,
+			CreateNoWindow = true
+		};
+
 		/// <summary>
 		/// Writes a versions.xcconfig file for xcode to pull in when making an app plist
 		/// </summary>
@@ -785,6 +790,56 @@ namespace UnrealBuildTool
 
 				OutputFiles.Add(UpdateVersionFile(BinaryLinkEnvironment, FileItem.GetItemByFileReference(BinaryLinkEnvironment.OutputFilePath), Graph));
 			}
+			
+			ReadOnlyAppleTargetRules AppleTargetRules = GetAppleTargetRules(Target);
+			if (!BinaryLinkEnvironment.bIsBuildingLibrary && (AppleTargetRules.bStripSymbols || Target.Configuration == UnrealTargetConfiguration.Shipping))
+			{
+				// Generate the standard strip command. Should our default strip command ever change, we have a centralized place to update it.
+				ProcessStartInfo StripProcessInfo = CreateStripSymbolsProcessInfo(Executable.Location, Settings.ToolchainDir);
+
+				// If building a framework, we can only strip local symbols and we need to leave global in place (-x argument)
+				string AddedArguments = BinaryLinkEnvironment.bIsBuildingDLL ? " -x" : "";
+				// Determine all of the arguments we're passing to the strip command (the process info plus framework argument above)
+				string StripCommandArguments = $"{StripProcessInfo.Arguments}{AddedArguments}";
+
+				// Convert the process info to an action graph entry
+				Action StripAction = Graph.CreateAction(ActionType.CreateAppBundle);
+				StripAction.WorkingDirectory = GetMacDevSrcRoot();
+				StripAction.CommandPath = BuildHostPlatform.Current.Shell;
+				StripAction.PrerequisiteItems.Add(Executable);
+				StripAction.PrerequisiteItems.UnionWith(OutputFiles);
+				StripAction.StatusDescription = $"Stripping symbols from {Executable.AbsolutePath}";
+				StripAction.bCanExecuteRemotely = false;
+				
+				// Construct the strip command we want to execute
+				string Command = $"\"{StripProcessInfo.FileName}\" {StripCommandArguments}";
+				
+				// Determine if we are creating a strip flag file.
+				// This will essentially be the same location & file name, but with a ".stripped" extension.
+				if (AppleTargetRules.bCreateStripFlagFile)
+				{
+					// Determine the name & location of the strip flag file
+					FileItem StripCompleteFile = FileItem.GetItemByFileReference(FileReference.Combine(BinaryLinkEnvironment.IntermediateDirectory!, Executable.Location.GetFileName() + ".stripped"));
+					
+					// Create this strip flag file as part of the command we are executing
+					Command += $" && touch \"{StripCompleteFile.FullName}\"";
+					
+					// Ensure our action identifies this as a file it is producing and it's flagged as an output file from this process
+					StripAction.ProducedItems.Add(StripCompleteFile);
+					OutputFiles.Add(StripCompleteFile);
+				}
+				else
+				{
+					// If we're not creating a stripped flag file, we're stripping the actual
+					// target file, so we need to add that as an output from this process
+					OutputFiles.Add(Executable);
+				}
+				
+				// Additional escaping for backslashes and quotes since we are executing a shell and passing in our command as a quoted string
+				string ShellCommand = $"-c \"{Command.Replace(@"\", @"\\").Replace("\"", "\\\"")}\"";
+				// Finally, we assign the action our finished shell command before it executes
+				StripAction.CommandArguments = ShellCommand;
+			}
 
 			return OutputFiles;
 		}
