Extract PROVISIONING_PROFILE_SPECIFIER directly from mobileprovision instead of using security command, which caused a code sign error on a specific build machine

#rnx

diff --git a/Engine/Source/Programs/UnrealBuildTool/ProjectFiles/Xcode/XcodeProject.cs b/Engine/Source/Programs/UnrealBuildTool/ProjectFiles/Xcode/XcodeProject.cs
--- a/Engine/Source/Programs/UnrealBuildTool/ProjectFiles/Xcode/XcodeProject.cs
+++ b/Engine/Source/Programs/UnrealBuildTool/ProjectFiles/Xcode/XcodeProject.cs
@@ -1310,8 +1310,11 @@ namespace UnrealBuildTool.XcodeProjectXcconfig
 					// get mobile provision UUID, either directly or from the provision
 					if (ProfileFile.GetExtension() == ".mobileprovision")
 					{
-						// security will read the provision and dump out the text (plist) bits of the profile, and plutil will extract the UUID to stdout
-						string UUID = Utils.RunLocalProcessAndReturnStdOut("/bin/sh", $"-c 'security cms -D -i {ProfileFile.FullName.Replace(" ", "\\ ")} | plutil -extract UUID raw -'");
+						// extract UUID from .mobileprovision
+						string command = $"-c \"egrep -a -A 2 UUID {ProfileFile.FullName.Replace(" ", "\\ ")} | grep string | xargs\"";
+						string UUID = Utils.RunLocalProcessAndReturnStdOut("/bin/sh", command);
+						// extract from <string>XXX-XXX-XXX</string>
+						UUID = UUID.Substring("<string>".Length, UUID.Length - "<string></string>".Length);
 
 						// make sure it's living in user's library
 						// (note: i couldn't find a way to install it without opening Xcode, so copying it works well enough)
