UnrealBuildTool: FIx default code folding setting to also consider

#rnx

diff --git a/Engine/Source/Programs/UnrealBuildTool/Configuration/TargetRules.cs b/Engine/Source/Programs/UnrealBuildTool/Configuration/TargetRules.cs
--- a/Engine/Source/Programs/UnrealBuildTool/Configuration/TargetRules.cs
+++ b/Engine/Source/Programs/UnrealBuildTool/Configuration/TargetRules.cs
@@ -2285,7 +2285,7 @@ namespace UnrealBuildTool
 		[RequiresUniqueBuildEnvironment]
 		public bool bIdenticalCodeFolding
 		{
-			get => bIdenticalCodeFoldingPrivate ?? Configuration == UnrealTargetConfiguration.Shipping;
+			get => bIdenticalCodeFoldingPrivate ?? Configuration == UnrealTargetConfiguration.Test || Configuration == UnrealTargetConfiguration.Shipping;
 			set => bIdenticalCodeFoldingPrivate = value;
 		}
 		private bool? bIdenticalCodeFoldingPrivate;
