IOS operators new/delete dynamic linking fix

Apple's dynamic linker will correctly use operators new/delete from binary if they are public global symbols.

This is required as we dynamically link libc++, meaning that some if its code can be inlined into our binary calling our operators and some of its code is in libc++.dylib. If that dynamic lib code is not using our operators, we crash in runtime due to allocators mismatch. The fix is to make sure we expose our operators symbols and do not strip those.

statsig is now enabled on IOS. It was previously disabled due to crashes in std:: with a mismatching allocator:
basic_string() is inlined and called our new, ~basic_string() is dynamically linked, hence called libc++ delete. With this fix, dynamic linker makes sure libc++ calls our delete.

#rnx

diff --git a/Engine/Source/Programs/UnrealBuildTool/Platform/IOS/IOSToolChain.cs b/Engine/Source/Programs/UnrealBuildTool/Platform/IOS/IOSToolChain.cs
--- a/Engine/Source/Programs/UnrealBuildTool/Platform/IOS/IOSToolChain.cs
+++ b/Engine/Source/Programs/UnrealBuildTool/Platform/IOS/IOSToolChain.cs
@@ -316,6 +316,29 @@ namespace UnrealBuildTool
 			string PlatformName = LinkEnvironment.Architecture == UnrealArch.IOSSimulator ? Settings.SimulatorPlatformDirName : Settings.PlatformDirName;
 			Arguments.Add($"-L{ApplePlatformSDK.GetToolchainDirectory()}/../lib/swift/{PlatformName}");
 
+			if (!LinkEnvironment.bIsBuildingLibrary && !LinkEnvironment.bIsBuildingDLL)
+			{
+				// We explicitly export our operators new/delete in binaries, so
+				// dynamic linker can correctly redirect those from libc++ and
+				// any other dynamic library we load.
+				
+				// _operator new(unsigned long)
+				Arguments.Add("-Wl,-exported_symbol,__Znwm");
+				// _operator new(unsigned long, std::nothrow_t const&)
+				Arguments.Add("-Wl,-exported_symbol,__ZnwmRKSt9nothrow_t");
+				// _operator delete(void*)
+				Arguments.Add("-Wl,-exported_symbol,__ZdlPv");
+				// _operator delete(void*, std::nothrow_t const&)
+				Arguments.Add("-Wl,-exported_symbol,__ZdlPvRKSt9nothrow_t");
+				// _operator new[](unsigned long)
+				Arguments.Add("-Wl,-exported_symbol,__Znam");
+				// _operator new[](unsigned long, std::nothrow_t const&)
+				Arguments.Add("-Wl,-exported_symbol,__ZnamRKSt9nothrow_t");
+				// _operator delete[](void*)
+				Arguments.Add("-Wl,-exported_symbol,__ZdaPv");
+				// _operator delete[](void*, std::nothrow_t const&)		
+				Arguments.Add("-Wl,-exported_symbol,__ZdaPvRKSt9nothrow_t");
+			}
 
 			Arguments.Add("-dead_strip");
 			Arguments.Add($"-target {Settings.GetTargetTuple(LinkEnvironment.Architecture, Target!.Platform, Target!.Type)}");
@@ -555,6 +578,7 @@ namespace UnrealBuildTool
 				// This is not a shipping build so no need to delete the output file since symbols will not have been stripped from it.
 				LinkAction.CommandArguments = String.Format("-c \"\\\"{0}\\\" {1}\"", LinkerPath, LinkCommandArguments);
 			}
+
 			return OutputFile;
 		}
 
diff --git a/Engine/Source/Programs/UnrealBuildTool/ProjectFiles/Xcode/XcodeProject.cs b/Engine/Source/Programs/UnrealBuildTool/ProjectFiles/Xcode/XcodeProject.cs
--- a/Engine/Source/Programs/UnrealBuildTool/ProjectFiles/Xcode/XcodeProject.cs
+++ b/Engine/Source/Programs/UnrealBuildTool/ProjectFiles/Xcode/XcodeProject.cs
@@ -1656,9 +1656,16 @@ namespace UnrealBuildTool.XcodeProjectXcconfig
 					"\\\"${UE_ENGINE_DIR}/Build/BatchFiles/Mac/GenerateUniversalDSYM.sh\\\" \\\"${CONFIGURATION_BUILD_DIR}/${EXECUTABLE_PATH}\\\" \\\"${DWARF_DSYM_FOLDER_PATH}/${DWARF_DSYM_FILE_NAME}\\\"",
 				});
 			}
+			string StripCommand = "strip -S -no_code_signature_warning -D \\\"${CONFIGURATION_BUILD_DIR}/${EXECUTABLE_PATH}\\\"";
+			if (Platform == UnrealTargetPlatform.IOS || Platform == UnrealTargetPlatform.TVOS)
+			{
+				// We only strip local/debug symbols. Globals should be preserved even in binaries,
+				// as we expose operators new/delete.
+				StripCommand = "strip -x -S -no_code_signature_warning -D \\\"${CONFIGURATION_BUILD_DIR}/${EXECUTABLE_PATH}\\\"";
+			}
 			DsymScript.AddRange(new string[]
 			{
-				"strip -no_code_signature_warning -D \\\"${CONFIGURATION_BUILD_DIR}/${EXECUTABLE_PATH}\\\"",
+				StripCommand,
 				"",
 				"# Remove any unused architectures from dylibs in the .app (param1) that don't match the executable (param2). Also error if a dylib is missing arches",
 				"\\\"${UE_ENGINE_DIR}/Build/BatchFiles/Mac/ThinApp.sh\\\" \\\"${CONFIGURATION_BUILD_DIR}/${CONTENTS_FOLDER_PATH}\\\" \\\"${CONFIGURATION_BUILD_DIR}/${EXECUTABLE_PATH}\\\"",
diff --git a/Engine/Source/Programs/UnrealBuildTool/ToolChain/AppleToolChain.cs b/Engine/Source/Programs/UnrealBuildTool/ToolChain/AppleToolChain.cs
--- a/Engine/Source/Programs/UnrealBuildTool/ToolChain/AppleToolChain.cs
+++ b/Engine/Source/Programs/UnrealBuildTool/ToolChain/AppleToolChain.cs
@@ -798,8 +798,10 @@ namespace UnrealBuildTool
 				ProcessStartInfo StripProcessInfo = CreateStripSymbolsProcessInfo(Executable.Location, Settings.ToolchainDir);
 
 				// If building a framework, we can only strip local symbols and we need to leave global in place (-x argument)
-				string AddedArguments = BinaryLinkEnvironment.bIsBuildingDLL ? " -x" : "";
-				// Determine all of the arguments we're passing to the strip command (the process info plus framework argument above)
+				// We also preserve global in binaries on IOS/TVOS as we expose operators new/delete.
+				bool bPreserveGlobalSymbols = BinaryLinkEnvironment.bIsBuildingDLL || (Target.Platform == UnrealTargetPlatform.IOS || Target.Platform == UnrealTargetPlatform.TVOS);
+				string AddedArguments = bPreserveGlobalSymbols ? " -x" : "";
+				// Determine all of the arguments we're passing to the strip command (the process info plus arguments above)
 				string StripCommandArguments = $"{StripProcessInfo.Arguments}{AddedArguments}";
 
 				// Convert the process info to an action graph entry
diff --git a/Engine/Source/Runtime/Core/Private/IOS/IOSPlatformCrashContext.cpp b/Engine/Source/Runtime/Core/Private/IOS/IOSPlatformCrashContext.cpp
--- a/Engine/Source/Runtime/Core/Private/IOS/IOSPlatformCrashContext.cpp
+++ b/Engine/Source/Runtime/Core/Private/IOS/IOSPlatformCrashContext.cpp
@@ -332,6 +332,17 @@ void FIOSApplicationInfo::Init()
 {
 	SCOPED_AUTORELEASE_POOL;
 	
+	// Prevent clang/ld from dead-code-eliminating the nothrow_t variants of global new.
+	// This ensures that all OS calls to operator new go through our allocators and delete cleanly.
+	{
+		std::nothrow_t t;
+		char* d = (char*)(operator new(8, t));
+		delete d;
+		
+		d = (char*)operator new[]( 8, t );
+		delete [] d;
+	}
+	
 	AppName = FApp::GetProjectName();
 	FCStringAnsi::Strncpy(AppNameUTF8, TCHAR_TO_UTF8(*AppName), PATH_MAX + 1);
 	
