Enable UE types in namespaces by default for Verse projects.
Add VMODULE as something UBT will detect when running UHT.

#rnx

diff --git a/Engine/Source/Programs/UnrealBuildTool/Configuration/ModuleRules.cs b/Engine/Source/Programs/UnrealBuildTool/Configuration/ModuleRules.cs
--- a/Engine/Source/Programs/UnrealBuildTool/Configuration/ModuleRules.cs
+++ b/Engine/Source/Programs/UnrealBuildTool/Configuration/ModuleRules.cs
@@ -1526,6 +1526,7 @@ namespace UnrealBuildTool
 		/// </summary>
 		public void SetupVerse()
 		{
+			bAllowUETypesInNamespaces = true;
 			if (Target.bUseVerseVM)
 			{
 				PublicDependencyModuleNames.AddRange(
diff --git a/Engine/Source/Programs/UnrealBuildTool/System/SourceFileMetadataCache.cs b/Engine/Source/Programs/UnrealBuildTool/System/SourceFileMetadataCache.cs
--- a/Engine/Source/Programs/UnrealBuildTool/System/SourceFileMetadataCache.cs
+++ b/Engine/Source/Programs/UnrealBuildTool/System/SourceFileMetadataCache.cs
@@ -134,7 +134,7 @@ namespace UnrealBuildTool
 		/// <summary>
 		/// Regex that matches C++ code with UObject declarations which we will need to generated code for.
 		/// </summary>
-		static readonly Regex ReflectionMarkupRegex = new Regex("^\\s*U(CLASS|STRUCT|ENUM|INTERFACE|DELEGATE)\\b", RegexOptions.Compiled);
+		static readonly Regex ReflectionMarkupRegex = new Regex("^\\s*(UCLASS|USTRUCT|UENUM|UINTERFACE|UDELEGATE|VMODULE)\\b", RegexOptions.Compiled);
 
 		/// <summary>
 		/// Regex that matches #include UE_INLINE_GENERATED_CPP_BY_NAME(****) statements.
