UnrealBuildTool: Track ThinLTOCache path with UBA

#rnx

diff --git a/Engine/Source/Programs/UnrealBuildTool/ToolChain/UEToolChain.cs b/Engine/Source/Programs/UnrealBuildTool/ToolChain/UEToolChain.cs
--- a/Engine/Source/Programs/UnrealBuildTool/ToolChain/UEToolChain.cs
+++ b/Engine/Source/Programs/UnrealBuildTool/ToolChain/UEToolChain.cs
@@ -282,17 +282,16 @@ namespace UnrealBuildTool
 		{
 		}
 
+		// Construct a relative path for the intermediate response file
+		private FileReference GetResponseFileName(FileReference Location) => Location.ChangeExtension($"{Location.GetExtension()}{ResponseExt}");
+
 		/// <summary>
 		/// Get the name of the response file for the current compile environment and output file
 		/// </summary>
 		/// <param name="CompileEnvironment"></param>
 		/// <param name="OutputFile"></param>
 		/// <returns></returns>
-		public virtual FileReference GetResponseFileName(CppCompileEnvironment CompileEnvironment, FileItem OutputFile)
-		{
-			// Construct a relative path for the intermediate response file
-			return OutputFile.Location.ChangeExtension(OutputFile.Location.GetExtension() + ResponseExt);
-		}
+		public FileReference GetResponseFileName(CppCompileEnvironment CompileEnvironment, FileItem OutputFile) => GetResponseFileName(OutputFile.Location);
 
 		/// <summary>
 		/// Get the name of the response file for the current linker environment and output file
@@ -300,10 +299,10 @@ namespace UnrealBuildTool
 		/// <param name="LinkEnvironment"></param>
 		/// <param name="OutputFile"></param>
 		/// <returns></returns>
-		public virtual FileReference GetResponseFileName(LinkEnvironment LinkEnvironment, FileItem OutputFile)
+		public FileReference GetResponseFileName(LinkEnvironment LinkEnvironment, FileItem OutputFile)
 		{
 			// Construct a relative path for the intermediate response file
-			return FileReference.Combine(LinkEnvironment.IntermediateDirectory!, OutputFile.Location.GetFileName() + ResponseExt);
+			return GetResponseFileName(FileReference.Combine(LinkEnvironment.IntermediateDirectory ?? OutputFile.Directory.Location, OutputFile.Location.GetFileName()));
 		}
 
 		public virtual ICollection<FileItem> PostBuild(ReadOnlyTargetRules Target, FileItem Executable, LinkEnvironment ExecutableLinkEnvironment, IActionGraphBuilder Graph)
@@ -314,7 +313,7 @@ namespace UnrealBuildTool
 		public virtual ICollection<FileItem> PostBuild(ReadOnlyTargetRules Target, IEnumerable<FileItem> Executables, LinkEnvironment ExecutableLinkEnvironment, IActionGraphBuilder Graph)
 		{
 			// by default, run PostBuild for exe Exe and merge results
-			return Executables.SelectMany(x => PostBuild(Target, x, ExecutableLinkEnvironment, Graph)).ToList();
+			return [.. Executables.SelectMany(x => PostBuild(Target, x, ExecutableLinkEnvironment, Graph))];
 		}
 
 		public virtual void SetUpGlobalEnvironment(ReadOnlyTargetRules Target, CppCompileEnvironment GlobalCompileEnvironment, LinkEnvironment GlobalLinkEnvironment)
@@ -324,6 +323,13 @@ namespace UnrealBuildTool
 
 			GlobalCompileEnvironment.RootPaths[CppRootPathFolder.Root] = Unreal.RootDirectory;
 			GlobalLinkEnvironment.RootPaths[CppRootPathFolder.Root] = Unreal.RootDirectory;
+
+			DirectoryReference? thinLtoCache = DirectoryReference.FromString(GlobalLinkEnvironment.ThinLTOCacheDirectory);
+			if (thinLtoCache != null)
+			{
+				GlobalLinkEnvironment.RootPaths.AddExtraPath(("ThinLTOCache", thinLtoCache.FullName));
+				EpicGames.UBA.Utils.RegisterPathHash(thinLtoCache.FullName, $"{Target.Platform}-ThinLTOCache");
+			}
 		}
 
 		public virtual void ModifyBuildProducts(ReadOnlyTargetRules Target, UEBuildBinary Binary, IEnumerable<string> Libraries, IEnumerable<UEBuildBundleResource> BundleResources, Dictionary<FileReference, BuildProductType> BuildProducts)
@@ -352,6 +358,12 @@ namespace UnrealBuildTool
 			return true;
 		}
 
+		/// <summary>
+		/// Get the file path to a debug file
+		/// </summary>
+		/// <param name="OutputFile">Binary output file</param>
+		/// <param name="DebugExtension">debug file extension</param>
+		/// <returns>Path to the debug file</returns>
 		public virtual FileReference GetDebugFile(FileReference OutputFile, string DebugExtension)
 		{
 			//  by default, just change the extension to the debug extension
