Bump buildtools to 35.0.1 and change zipalign to page align uncompressed .so to 16kb

diff --git a/Engine/Config/Android/Android_SDK.json b/Engine/Config/Android/Android_SDK.json
--- a/Engine/Config/Android/Android_SDK.json
+++ b/Engine/Config/Android/Android_SDK.json
@@ -6,7 +6,7 @@
 	"AutoSDKDirectory": "-28",
 
 	"platforms": "android-34",
-	"build-tools": "34.0.0",
+	"build-tools": "35.0.1",
 	"cmake": "3.22.1",
 	"ndk": "27.2.12479018"
 }
diff --git a/Engine/Extras/Android/SetupAndroid.bat b/Engine/Extras/Android/SetupAndroid.bat
--- a/Engine/Extras/Android/SetupAndroid.bat
+++ b/Engine/Extras/Android/SetupAndroid.bat
@@ -18,7 +18,7 @@ SET NDK_VERSION=%4
 
 rem hardcoded versions for compatibility with non-Turnkey manual running
 if "%PLATFORMS_VERSION%" == "" SET PLATFORMS_VERSION=android-34
-if "%BUILDTOOLS_VERSION%" == "" SET BUILDTOOLS_VERSION=34.0.0
+if "%BUILDTOOLS_VERSION%" == "" SET BUILDTOOLS_VERSION=35.0.1
 if "%CMAKE_VERSION%" == "" SET CMAKE_VERSION=3.22.1
 if "%NDK_VERSION%" == "" SET NDK_VERSION=27.2.12479018
 
diff --git a/Engine/Extras/Android/SetupAndroid.command b/Engine/Extras/Android/SetupAndroid.command
--- a/Engine/Extras/Android/SetupAndroid.command
+++ b/Engine/Extras/Android/SetupAndroid.command
@@ -62,7 +62,7 @@ if [[ -z "${PLATFORMS_VERSION}" ]]; then
     PLATFORMS_VERSION="android-34"
 fi
 if [[ -z "${BUILDTOOLS_VERSION}" ]]; then
-    BUILDTOOLS_VERSION="34.0.0"
+    BUILDTOOLS_VERSION="35.0.1"
 fi
 if [[ -z "${CMAKE_VERSION}" ]]; then
     CMAKE_VERSION="3.22.1"
diff --git a/Engine/Extras/Android/SetupAndroid.sh b/Engine/Extras/Android/SetupAndroid.sh
--- a/Engine/Extras/Android/SetupAndroid.sh
+++ b/Engine/Extras/Android/SetupAndroid.sh
@@ -30,7 +30,7 @@ if [[ -z "${PLATFORMS_VERSION}" ]]; then
     PLATFORMS_VERSION="android-34"
 fi
 if [[ -z "${BUILDTOOLS_VERSION}" ]]; then
-    BUILDTOOLS_VERSION="34.0.0"
+    BUILDTOOLS_VERSION="35.0.1"
 fi
 if [[ -z "${CMAKE_VERSION}" ]]; then
     CMAKE_VERSION="3.22.1"
diff --git a/Engine/Source/Programs/UnrealBuildTool/Platform/Android/UEDeployAndroid.cs b/Engine/Source/Programs/UnrealBuildTool/Platform/Android/UEDeployAndroid.cs
--- a/Engine/Source/Programs/UnrealBuildTool/Platform/Android/UEDeployAndroid.cs
+++ b/Engine/Source/Programs/UnrealBuildTool/Platform/Android/UEDeployAndroid.cs
@@ -387,7 +387,7 @@ namespace UnrealBuildTool
 
 			if (BestVersionString == null)
 			{
-				BestVersionString = "34.0.0";
+				BestVersionString = "35.0.1";
 				Logger.LogWarning("Failed to find %ANDROID_HOME%/build-tools subdirectory. Will attempt to use {BestVersionString}.", BestVersionString);
 			}
 
