UnrealBuildTool: Use log event parser for PVS-Studio warnings

#rnx

diff --git a/Engine/Source/Programs/UnrealBuildTool/Platform/Windows/PVSToolChain.cs b/Engine/Source/Programs/UnrealBuildTool/Platform/Windows/PVSToolChain.cs
--- a/Engine/Source/Programs/UnrealBuildTool/Platform/Windows/PVSToolChain.cs
+++ b/Engine/Source/Programs/UnrealBuildTool/Platform/Windows/PVSToolChain.cs
@@ -8,6 +8,7 @@ using System.Collections.Generic;
 using System.Diagnostics;
 using System.IO;
 using System.Linq;
+using System.Reflection;
 using System.Text;
 using System.Text.RegularExpressions;
 using System.Threading.Tasks;
@@ -446,13 +447,16 @@ namespace UnrealBuildTool
 				.OrderBy(x => x.Positions.FirstOrDefault()?.File)
 				.ThenBy(x => x.Positions.FirstOrDefault()?.Lines.FirstOrDefault());
 
-			// Create the combined output file, and print the diagnostics to the log
-			foreach (PVSErrorInfo errorInfo in filteredErrors)
+			using (LogEventParser parser = new(logger))
 			{
-				string file = errorInfo.Positions.FirstOrDefault()?.File ?? Unreal.EngineDirectory.FullName;
-				int lineNumber = errorInfo.Positions.FirstOrDefault()?.Lines.FirstOrDefault() ?? 1;
-				LogValue logValue = new(LogValueType.SourceFile, file, new Dictionary<Utf8String, object?> { [LogEventPropertyName.File] = file });
-				logger.LogWarning(KnownLogEvents.Compiler, "{Path}({LineNumber}): warning {WarningCode}: {WarningMessage}", logValue, lineNumber, errorInfo.Code, errorInfo.Message);
+				parser.AddMatchersFromAssembly(Assembly.GetExecutingAssembly());
+				// Create the combined output file, and print the diagnostics to the log
+				foreach (PVSErrorInfo errorInfo in filteredErrors)
+				{
+					string file = errorInfo.Positions.FirstOrDefault()?.File ?? Unreal.EngineDirectory.FullName;
+					int lineNumber = errorInfo.Positions.FirstOrDefault()?.Lines.FirstOrDefault() ?? 1;
+					parser.WriteLine($"{file}({lineNumber}): warning {errorInfo.Code}: {errorInfo.Message}");
+				}
 			}
 
 			PVSErrorInfo? renewError = allErrors.FirstOrDefault(x => x.Code.Equals("Renew", StringComparison.Ordinal));
@@ -952,7 +956,7 @@ namespace UnrealBuildTool
 
 			// Collect the sourcefile items off of the Compile action added in CompileCPPFiles so that in SingleFileCompile mode the PVSGather step is also not filtered out
 			IEnumerable<FileItem> compileSourceFiles = [.. makefile.Actions.OfType<VCCompileAction>().Select(x => x.SourceFile!)];
-			CppRootPaths rootPaths = makefile.Actions.OfType<VCCompileAction>().FirstOrDefault(x => x.RootPaths.Any())?.RootPaths ?? new();
+			CppRootPaths rootPaths = makefile.Actions.OfType<VCCompileAction>().LastOrDefault(x => x.RootPaths.Any())?.RootPaths ?? new();
 
 			// Store list of system paths that should be excluded
 			IEnumerable<DirectoryReference> systemIncludePaths = [.. makefile.Actions.OfType<VCCompileAction>().SelectMany(x => x.SystemIncludePaths)];
