[UBT]
* Fixed bug where action weights ended up wrong when compiling with unity files disabled

diff --git a/Engine/Source/Programs/UnrealBuildTool/Configuration/UEBuildModuleCPP.cs b/Engine/Source/Programs/UnrealBuildTool/Configuration/UEBuildModuleCPP.cs
--- a/Engine/Source/Programs/UnrealBuildTool/Configuration/UEBuildModuleCPP.cs
+++ b/Engine/Source/Programs/UnrealBuildTool/Configuration/UEBuildModuleCPP.cs
@@ -727,6 +727,11 @@ namespace UnrealBuildTool
 				}
 			}
 
+			if (GeneratedCPPCompileEnvironment == CppCompileEnvironment)
+			{
+				GeneratedCPPCompileEnvironment = new(CppCompileEnvironment);
+			}
+
 			if (Target.WindowsPlatform.bEnableInstrumentation)
 			{
 				CppCompileEnvironment.UserIncludePaths.Add(DirectoryReference.Combine(Unreal.EngineSourceDirectory, "Runtime", "Instrumentation", "Public"));
@@ -816,13 +821,14 @@ namespace UnrealBuildTool
 					Logger.LogTrace("Module '{ModuleName}' using unity build mode", Name);
 					bModuleUsesUnityBuild = true;
 				}
+				bModuleUsesUnityBuild = bModuleUsesUnityBuild && Rules.UnityFileTypes.HasFlag(ModuleRules.UnityFileType.Cpp);
 			}
 			else
 			{
 				Logger.LogTrace("Module '{ModuleName}' not using unity build mode", Name);
 			}
 
-			CppCompileEnvironment.bUseUnity = bModuleUsesUnityBuild && Rules.UnityFileTypes.HasFlag(ModuleRules.UnityFileType.Cpp);
+			CppCompileEnvironment.bUseUnity = bModuleUsesUnityBuild;
 			GeneratedCPPCompileEnvironment.bUseUnity = CppCompileEnvironment.bUseUnity;
 
 			// Create and register a special action that can be used to compile single files (even when unity is enabled) for generated files
@@ -850,7 +856,7 @@ namespace UnrealBuildTool
 			}
 
 			// Compile Generated CPP Files, always use unity for .gen.cpp even if disabled for the current module
-			if (bModuleUsesUnityBuild && bMergeUnityFiles && Rules.UnityFileTypes.HasFlag(ModuleRules.UnityFileType.Cpp))
+			if (bModuleUsesUnityBuild && bMergeUnityFiles)
 			{
 				CPPFiles.AddRange(GeneratedFileItems);
 			}
@@ -906,7 +912,7 @@ namespace UnrealBuildTool
 			// Compile CPP files
 			if (CPPFiles.Count > 0)
 			{
-				if (bModuleUsesUnityBuild && Rules.UnityFileTypes.HasFlag(ModuleRules.UnityFileType.Cpp))
+				if (bModuleUsesUnityBuild)
 				{
 					Unity.GenerateUnitySource(Target, CPPFiles, InputFiles.HeaderFiles, CppCompileEnvironment, WorkingSet, Rules.ShortName ?? Name, IntermediateDirectory, Graph, SourceFileToUnityFile,
 						out List<FileItem> NormalFiles, out List<FileItem> AdaptiveFiles, NumIncludedBytesPerUnityCPP);
@@ -1090,6 +1096,7 @@ namespace UnrealBuildTool
 				if (HeaderFileItems.Count > 0)
 				{
 					CppCompileEnvironment HeaderCompileEnvironment = new CppCompileEnvironment(ModuleCompileEnvironment);
+					HeaderCompileEnvironment.bUseUnity = false;
 					DirectoryReference HeaderIntermediateDirectory = DirectoryReference.Combine(IntermediateDirectory, "H");
 
 					CreateHeaderForDefinitions(HeaderCompileEnvironment, HeaderIntermediateDirectory, "h", Graph);
diff --git a/Engine/Source/Programs/UnrealBuildTool/Platform/Windows/VCToolChain.cs b/Engine/Source/Programs/UnrealBuildTool/Platform/Windows/VCToolChain.cs
--- a/Engine/Source/Programs/UnrealBuildTool/Platform/Windows/VCToolChain.cs
+++ b/Engine/Source/Programs/UnrealBuildTool/Platform/Windows/VCToolChain.cs
@@ -1921,7 +1921,15 @@ namespace UnrealBuildTool
 					// If we are building without unity files the multithread part balances out with the start/exit of all the actions
 					if (!CompileEnvironment.bUseUnity)
 					{
-						BaseCompileAction.Weight = 1.0f + (BaseCompileAction.Weight - 1.0f) * 0.5f;
+						// This is a very lightweight build. no pch or debug info. I/O is also low. Set weight to 1.0f
+						if (CompileEnvironment.PrecompiledHeaderAction == PrecompiledHeaderAction.None && !CompileEnvironment.bCreateDebugInfo)
+						{
+							BaseCompileAction.Weight = 1.0f;
+						}
+						else
+						{
+							BaseCompileAction.Weight = 1.0f + (BaseCompileAction.Weight - 1.0f) * 0.5f;
+						}
 					}
 				}
 			}
