Add program size reporting on large unmergeable GFPs with -MergeModulesReport

- When the option is passed, the game target is given the results of the merging strategy for inspection
- Also fix non-code plugins being used to cull dependencies, resulting in GFS plugins preventing core code plugins from getting merged
- Remove recursive option, we only ever do recursive

#rnx

diff --git a/Engine/Source/Programs/Shared/EpicGames.Core/Telemetry/SimpleJsonTelemetry.cs b/Engine/Source/Programs/Shared/EpicGames.Core/Telemetry/SimpleJsonTelemetry.cs
new file mode 100644
--- /dev/null
+++ b/Engine/Source/Programs/Shared/EpicGames.Core/Telemetry/SimpleJsonTelemetry.cs
@@ -0,0 +1,52 @@
+// Copyright Epic Games, Inc. All Rights Reserved.
+
+using System;
+using System.Net.Http;
+using System.Text.Json;
+using System.Threading.Tasks;
+using Microsoft.Extensions.Logging;
+
+namespace EpicGames.Core.Telemetry
+{
+	/// <summary>
+	/// Simple telemetry wrapper for accessing JSON data
+	/// </summary>
+	public class SimpleJsonTelemetry
+	{
+		private readonly HttpClient _httpClient;
+
+		/// <summary>
+		/// Constructor
+		/// </summary>
+		public SimpleJsonTelemetry(string server)
+		{
+			_httpClient = new()
+			{
+				BaseAddress = new Uri(server),
+			};
+		}
+
+		/// <summary>
+		/// Fetch the data at a given endpoint and return it as string
+		/// </summary>
+		public async Task<T?> GetAsync<T>(string endpoint)
+			where T: new()
+		{
+			try
+			{
+				using HttpResponseMessage response = await _httpClient.GetAsync(endpoint);
+
+				response.EnsureSuccessStatusCode();
+
+				System.IO.Stream responseStream = await response.Content.ReadAsStreamAsync();
+
+				return await JsonSerializer.DeserializeAsync<T>(responseStream)!;
+			}
+			catch (HttpRequestException e)
+			{
+				Log.Logger.LogWarning("GetAsync: failed to retrieve {endpoint}: {1}", endpoint, e.Message);
+				return new T();
+			}
+		}
+	};
+}
diff --git a/Engine/Source/Programs/UnrealBuildTool/Configuration/ReadOnlyTargetRules.cs b/Engine/Source/Programs/UnrealBuildTool/Configuration/ReadOnlyTargetRules.cs
--- a/Engine/Source/Programs/UnrealBuildTool/Configuration/ReadOnlyTargetRules.cs
+++ b/Engine/Source/Programs/UnrealBuildTool/Configuration/ReadOnlyTargetRules.cs
@@ -521,7 +521,9 @@ namespace UnrealBuildTool
 
 		public IReadOnlyList<string> MergePluginsLaunch => Inner.MergePluginsLaunch;
 
-		public bool MergeOnlyGameFeatures => Inner.MergeOnlyGameFeatures;
+		public bool MergeOnlyGameFeatures => Inner.bMergeOnlyGameFeatures;
+
+		public bool bReportModuleMergingData => Inner.bReportModuleMergingData;
 
 		public IReadOnlyList<string> GlobalDefinitions => Inner.GlobalDefinitions.AsReadOnly();
 
@@ -653,6 +655,12 @@ namespace UnrealBuildTool
 		/// <returns>true if the plugin is allowed</returns>
 		public bool ShouldIgnorePluginDependency(PluginInfo parentInfo, PluginReferenceDescriptor childDescriptor) => Inner.ShouldIgnorePluginDependency(parentInfo, childDescriptor);
 
+		/// <summary>
+		/// Report merged modules structure to the target so that the game may use this information.
+		/// </summary>
+		public virtual void ReportMergedModules(Dictionary<string, IEnumerable<ModuleRules>> ModuleGroups, Dictionary<string, IEnumerable<string>> PluginUsage, Dictionary<string, IEnumerable<string>> RootPluginUsage)
+			=> Inner.ReportMergedModules(ModuleGroups, PluginUsage, RootPluginUsage);
+
 		/// <summary>
 		/// Determines if the automation tests should be compiled based on the current configuration and optional forced settings.
 		/// If either the development tests or performance tests are compiled (unless explicitly disabled), this property returns true.
diff --git a/Engine/Source/Programs/UnrealBuildTool/Configuration/TargetRules.cs b/Engine/Source/Programs/UnrealBuildTool/Configuration/TargetRules.cs
--- a/Engine/Source/Programs/UnrealBuildTool/Configuration/TargetRules.cs
+++ b/Engine/Source/Programs/UnrealBuildTool/Configuration/TargetRules.cs
@@ -2477,7 +2477,13 @@ namespace UnrealBuildTool
 		/// <summary>
 		/// Experimental: Whether to merge only game feature plugins, ignoring regular plugins.
 		/// </summary>
-		public bool MergeOnlyGameFeatures { get; set; } = true;
+		public bool bMergeOnlyGameFeatures { get; set; } = true;
+
+		/// <summary>
+		/// Experimental: Whether to report module merging data to the game target.
+		/// </summary>
+		[CommandLine("-MergeModulesReport")]
+		public bool bReportModuleMergingData { get; set; } = false;
 
 		/// <summary>
 		/// Backing storage for the LinkType property.
@@ -2761,6 +2767,11 @@ namespace UnrealBuildTool
 		/// <returns>true if the plugin is allowed</returns>
 		public virtual bool ShouldIgnorePluginDependency(PluginInfo parentInfo, PluginReferenceDescriptor descriptor) => false;
 
+		/// <summary>
+		/// Validate that the merged modules setup is appropriate
+		/// </summary>
+		public virtual void ReportMergedModules(Dictionary<string, IEnumerable<ModuleRules>> ModuleGroups, Dictionary<string, IEnumerable<string>> PluginUsage, Dictionary<string, IEnumerable<string>> RootPluginUsage) { }
+
 		/// <summary>
 		/// Checks if a property has been set with the RequiresUniqueBuildEnvironmentAttribute, and is different from it's base
 		/// </summary>
diff --git a/Engine/Source/Programs/UnrealBuildTool/Configuration/UEBuildTarget.cs b/Engine/Source/Programs/UnrealBuildTool/Configuration/UEBuildTarget.cs
--- a/Engine/Source/Programs/UnrealBuildTool/Configuration/UEBuildTarget.cs
+++ b/Engine/Source/Programs/UnrealBuildTool/Configuration/UEBuildTarget.cs
@@ -3613,8 +3613,15 @@ namespace UnrealBuildTool
 				}
 			}
 
-			void GetPluginDependencies(UEBuildPlugin plugin, ref HashSet<UEBuildPlugin> dependencies, bool recursively = false, bool checkAllPlugins = false)
+			Dictionary<UEBuildPlugin, HashSet<UEBuildPlugin>> globalPluginUserMap = new();
+
+			void GetPluginDependencies(UEBuildPlugin plugin, ref HashSet<UEBuildPlugin> dependencies, bool strictFollowMode = false, bool hasCodeInHierarchy = false)
 			{
+				if (plugin.Modules.Any())
+				{
+					hasCodeInHierarchy = true;
+				}
+
 				foreach (PluginReferenceDescriptor dependency in plugin.Info.Descriptor.Plugins ?? new())
 				{
 					if (!Rules.ShouldIgnorePluginDependency(plugin.Info, dependency))
@@ -3623,13 +3630,20 @@ namespace UnrealBuildTool
 						if (dependentPluginCandidates.Any())
 						{
 							UEBuildPlugin dependentPlugin = dependentPluginCandidates.First();
-							if (!dependencies.Contains(dependentPlugin) && (CanPluginBeMerged(dependentPlugin) || checkAllPlugins))
+
+							bool shouldFollowPlugin = (CanPluginBeMerged(dependentPlugin) && !strictFollowMode) || hasCodeInHierarchy;
+
+							if (!dependencies.Contains(dependentPlugin) && shouldFollowPlugin)
 							{
 								dependencies.Add(dependentPlugin);
-								if (recursively)
+
+								if (!globalPluginUserMap.ContainsKey(dependentPlugin))
 								{
-									GetPluginDependencies(dependentPlugin, ref dependencies, recursively, checkAllPlugins);
+									globalPluginUserMap.Add(dependentPlugin, new());
 								}
+								globalPluginUserMap[dependentPlugin].Add(plugin);
+
+								GetPluginDependencies(dependentPlugin, ref dependencies, strictFollowMode, hasCodeInHierarchy);
 							}
 						}
 					}
@@ -3700,7 +3714,7 @@ namespace UnrealBuildTool
 					{
 						HashSet<UEBuildPlugin> self = new HashSet<UEBuildPlugin> { plugin };
 						HashSet<UEBuildPlugin> dependencies = new();
-						GetPluginDependencies(plugin, ref dependencies, true, false);
+						GetPluginDependencies(plugin, ref dependencies, false);
 
 						pluginGroups.Add(name, self);
 						pluginGroupDependencies.Add(name, dependencies);
@@ -3788,7 +3802,7 @@ namespace UnrealBuildTool
 				foreach (UEBuildPlugin plugin in unmergablePlugins)
 				{
 					HashSet<UEBuildPlugin> dependencies = new();
-					GetPluginDependencies(plugin, ref dependencies, true, true);
+					GetPluginDependencies(plugin, ref dependencies, true);
 
 					foreach (UEBuildPlugin dependency in dependencies)
 					{
@@ -3839,6 +3853,48 @@ namespace UnrealBuildTool
 				HashSet<UEBuildModuleCPP> allLeftoverDependencies = new HashSet<UEBuildModuleCPP>();
 				GetModulesDependenciesRecursively(allLeftoverModules.OfType<UEBuildModuleCPP>(), ref allLeftoverDependencies);
 
+				Dictionary<string, IEnumerable<ModuleRules>> moduleGroups = new();
+				Dictionary<string, IEnumerable<string>> pluginUsageMap = new();
+				Dictionary<string, IEnumerable<string>> rootPluginUsageMap = new();
+				if (Rules.bReportModuleMergingData)
+				{ 
+					// prepare the merged modules group data for validation by the game
+					moduleGroups.Add("", allLeftoverDependencies.Select(x => x.Rules));
+					foreach (UEBuildModuleCPP module in allLeftoverDependencies)
+					{
+						IEnumerable<UEBuildPlugin> pluginCandidates = BuildPlugins.Where(x => x.Modules.Contains(module));
+						if (pluginCandidates.Any())
+						{
+							List<string> userNames = new();
+							UEBuildPlugin plugin = pluginCandidates.First();
+
+							if (globalPluginUserMap.ContainsKey(plugin))
+							{
+								pluginUsageMap.Add(module.Name, globalPluginUserMap[plugin].Select(x => x.Name));
+							}
+
+							void FindRootUsers(UEBuildPlugin plugin, ref HashSet<UEBuildPlugin> roots)
+							{
+								if (globalPluginUserMap.ContainsKey(plugin))
+								{
+									foreach (UEBuildPlugin user in globalPluginUserMap[plugin])
+									{
+										FindRootUsers(user, ref roots);
+									}
+								}
+								else if (!roots.Contains(plugin))
+								{
+									roots.Add(plugin);
+								}
+							}
+
+							HashSet<UEBuildPlugin> roots = new();
+							FindRootUsers(plugin, ref roots);
+							rootPluginUsageMap.Add(module.Name, roots.Select(x => x.Name));
+						}
+					}
+				}
+
 				// Remove all leftover dependencies from plugin modules
 				// then remove remaining plugin modules from all modules set
 				// then merge the module into a binary and update plugin modules in all modules to that binary
@@ -3849,9 +3905,15 @@ namespace UnrealBuildTool
 
 					pluginModules.ExceptWith(allLeftoverDependencies);
 					allModules.ExceptWith(pluginModules);
+					moduleGroups.Add(item.Key, pluginModules.Select(x => x.Rules));
 
 					ConsumeModules(item.Key, pluginModules);
 				}
+
+				if (Rules.bReportModuleMergingData)
+				{
+					Rules.ReportMergedModules(moduleGroups, pluginUsageMap, rootPluginUsageMap);
+				}
 			}
 
 			List<UEBuildModuleCPP> engineModules = [];
