[Linux] Add flag to disable symbol stripping

[RN] [Linux] Add -NoStripSymbols to UBT to support build systems where stripping needs to happen later in the process

diff --git a/Engine/Source/Programs/UnrealBuildTool/Platform/Linux/LinuxToolChain.cs b/Engine/Source/Programs/UnrealBuildTool/Platform/Linux/LinuxToolChain.cs
--- a/Engine/Source/Programs/UnrealBuildTool/Platform/Linux/LinuxToolChain.cs
+++ b/Engine/Source/Programs/UnrealBuildTool/Platform/Linux/LinuxToolChain.cs
@@ -217,7 +217,7 @@ namespace UnrealBuildTool
 
 			action.ProducedItems.Add(EncodedBinarySymbolsFile);
 
-			if (LinkEnvironment.bCreateDebugInfo)
+			if (LinkEnvironment.bCreateDebugInfo && !Options.HasFlag(ClangToolChainOptions.DisableStripSymbols))
 			{
 				if (MaxBinarySizeOverrideForObjcopy > 0 && bUseCmdExe)
 				{
diff --git a/Engine/Source/Programs/UnrealBuildTool/Platform/Linux/UEBuildLinux.cs b/Engine/Source/Programs/UnrealBuildTool/Platform/Linux/UEBuildLinux.cs
--- a/Engine/Source/Programs/UnrealBuildTool/Platform/Linux/UEBuildLinux.cs
+++ b/Engine/Source/Programs/UnrealBuildTool/Platform/Linux/UEBuildLinux.cs
@@ -102,6 +102,13 @@ namespace UnrealBuildTool
 		[XmlConfigFile(Category = "BuildConfiguration", Name = "bDisableDumpSyms")]
 		public bool bDisableDumpSyms = false;
 
+		/// <summary>
+		/// Whether to globally skip symbol stripping
+		/// </summary>
+		[CommandLine("-NoStripSymbols")]
+		[XmlConfigFile(Category = "BuildConfiguration", Name = "bDisableStripSymbols")]
+		public bool bDisableStripSymbols = false;
+
 		/// <summary>
 		/// Enables runtime ray tracing support.
 		/// </summary>
@@ -149,6 +156,7 @@ namespace UnrealBuildTool
 
 		public bool bEnableRayTracing => Inner.bEnableRayTracing;
 
+		public bool bDisableStripSymbols => Inner.bDisableStripSymbols;
 #pragma warning restore CS1591
 		#endregion
 	}
@@ -653,6 +661,10 @@ namespace UnrealBuildTool
 			{
 				Options |= ClangToolChainOptions.DisableDumpSyms;
 			}
+			if (Target.LinuxPlatform.bDisableStripSymbols)
+			{
+				Options |= ClangToolChainOptions.DisableStripSymbols;
+			}
 			if (Target.LinuxPlatform.bEnableLibFuzzer)
 			{
 				Options |= ClangToolChainOptions.EnableLibFuzzer;
diff --git a/Engine/Source/Programs/UnrealBuildTool/ToolChain/ClangToolChain.cs b/Engine/Source/Programs/UnrealBuildTool/ToolChain/ClangToolChain.cs
--- a/Engine/Source/Programs/UnrealBuildTool/ToolChain/ClangToolChain.cs
+++ b/Engine/Source/Programs/UnrealBuildTool/ToolChain/ClangToolChain.cs
@@ -129,6 +129,12 @@ namespace UnrealBuildTool
 		/// Enables compressing the debug sections if the platform supports this
 		/// </summary>
 		CompressDebugFile = 1 << 20,
+
+		/// <summary>
+		/// Disable Symbol stripping, for use when an external part of the build pipeline requires symbols to be present
+		/// and will strip them in a later step
+		/// </summary>
+		DisableStripSymbols = 1 << 21,
 	}
 
 	abstract class ClangToolChain : ISPCToolChain
