[UBT]
* Fixed so export generation action is assigned to a cache bucket (so it will be cached)

diff --git a/Engine/Source/Programs/UnrealBuildTool/Configuration/UEBuildTarget.cs b/Engine/Source/Programs/UnrealBuildTool/Configuration/UEBuildTarget.cs
--- a/Engine/Source/Programs/UnrealBuildTool/Configuration/UEBuildTarget.cs
+++ b/Engine/Source/Programs/UnrealBuildTool/Configuration/UEBuildTarget.cs
@@ -3077,18 +3077,23 @@ namespace UnrealBuildTool
 			}
 
 			HashSet<FileItem> allObjects = new(Binaries.SelectMany(x => x.InputObjects));
+			SortedSet<FileItem> allExis = new(allObjects.Select(x => FileItem.GetItemByFileReference(x.Location.ChangeExtension(".exi"))));
 
 			void CreateStripAction(string Name, UEBuildBinaryType binaryType, IEnumerable< FileItem> toStrip, DirectoryReference IntermediateDirectory)
 			{
-				HashSet<FileItem> otherObjects = new(allObjects.Except(toStrip));
+				SortedSet<FileItem> toStripExis = new(toStrip.Select(x => FileItem.GetItemByFileReference(x.Location.ChangeExtension(".exi"))));
+				SortedSet<FileItem> otherExis = new(allExis);
+				otherExis.ExceptWith(toStripExis);
+
 				FileItem exportsRsp = FileItem.GetItemByFileReference(FileReference.Combine(IntermediateDirectory, $"{Name}.exports.rsp"));
 				string Ext = TargetToolChain.GetExtraLinkFileExtension();
-				FileReference extraObj = FileReference.Combine(IntermediateDirectory, $"{Name}.exports.{Ext}");
+				FileReference outputFile = FileReference.Combine(IntermediateDirectory, $"{Name}.exports.{Ext}");
+
 				IEnumerable<string> arguments =
 				[
-					.. toStrip.Select(path => $"/S:{path.Location.ChangeExtension(".exi")}").OrderBy(x => x),
-					.. otherObjects.Select(path => $"/D:{path.Location.ChangeExtension(".exi")}").OrderBy(x => x),
-					$"/O:{extraObj}",
+					.. toStripExis.Select(path => $"/S:{path}"),
+					.. otherExis.Select(path => $"/D:{path}"),
+					$"/O:{outputFile}",
 					$"/T:{Platform}",
 					$"/M:{Name}",
 					..TargetToolChain.GetExtraLinkFileAdditionalSymbols(binaryType).Select(symbol=> $"/E:{symbol}"),
@@ -3098,9 +3103,9 @@ namespace UnrealBuildTool
 
 				Action exportsAction = MakefileBuilder.CreateAction(ActionType.Link);
 
+				exportsAction.PrerequisiteItems.UnionWith(allExis);
 				exportsAction.PrerequisiteItems.Add(exportsRsp);
-				exportsAction.PrerequisiteItems.UnionWith(allObjects);
-				exportsAction.ProducedItems.Add(FileItem.GetItemByFileReference(extraObj));
+				exportsAction.ProducedItems.Add(FileItem.GetItemByFileReference(outputFile));
 				exportsAction.CommandPath = commandPath;
 				exportsAction.CommandArguments = $"@\"{TargetToolChain.NormalizeCommandLinePath(exportsRsp.Location, CompileEnvironment.RootPaths)}\"";
 				exportsAction.CommandDescription = "GenerateExports";
@@ -3109,6 +3114,7 @@ namespace UnrealBuildTool
 				exportsAction.bCanExecuteRemotely = false; // Can run remotely but is so fast so it would be a waste
 				exportsAction.ArtifactMode = ArtifactMode.Enabled;
 				exportsAction.RootPaths = CompileEnvironment.RootPaths;
+				exportsAction.CacheBucket = ISPCToolChain.GetCacheBucket(Rules, CompileEnvironment);
 
 				Makefile.OutputItems.AddRange(exportsAction.ProducedItems);
 			}
