[AutoRTFM] Add UBT support for AutoRTFM external mapping files.

AutoRTFM External Mapping files (.aem) are a new DSL file format for controling the AutoRTFM open-to-closed mappings on functions that cannot be modified (e.g. standard library, third-party code)

This will be used to support AutoRTFM with modular linux builds.

diff --git a/Engine/Source/Programs/UnrealBuildTool/Configuration/ModuleRules.cs b/Engine/Source/Programs/UnrealBuildTool/Configuration/ModuleRules.cs
--- a/Engine/Source/Programs/UnrealBuildTool/Configuration/ModuleRules.cs
+++ b/Engine/Source/Programs/UnrealBuildTool/Configuration/ModuleRules.cs
@@ -1301,6 +1301,11 @@ namespace UnrealBuildTool
 		/// </summary>
 		public List<string> PublicDefinitions = new();
 
+		/// <summary>
+		/// AutoRTFM external open -> closed mapping files
+		/// </summary>
+		public List<string> AutoRTFMExternalMappingFiles = new();
+
 		/// <summary>
 		/// Append (or create)
 		/// </summary>
diff --git a/Engine/Source/Programs/UnrealBuildTool/Configuration/UEBuildBinary.cs b/Engine/Source/Programs/UnrealBuildTool/Configuration/UEBuildBinary.cs
--- a/Engine/Source/Programs/UnrealBuildTool/Configuration/UEBuildBinary.cs
+++ b/Engine/Source/Programs/UnrealBuildTool/Configuration/UEBuildBinary.cs
@@ -776,6 +776,9 @@ namespace UnrealBuildTool
 			{
 				if (Module.Binary == null || Module.Binary == this)
 				{
+					// Add the module's AutoRTFM external mapping files to the compile environment
+					BinaryCompileEnvironment.AutoRTFMExternalMappingFiles.UnionWith(Module.PublicAutoRTFMExternalMappingFiles);
+
 					// Compile each module.
 					Logger.LogTrace("Compile module: {ModuleName}", Module.Name);
 					List<FileItem> LinkInputFiles = Module.Compile(Target, ToolChain, BinaryCompileEnvironment, WorkingSet, Graph, Logger);
diff --git a/Engine/Source/Programs/UnrealBuildTool/Configuration/UEBuildModule.cs b/Engine/Source/Programs/UnrealBuildTool/Configuration/UEBuildModule.cs
--- a/Engine/Source/Programs/UnrealBuildTool/Configuration/UEBuildModule.cs
+++ b/Engine/Source/Programs/UnrealBuildTool/Configuration/UEBuildModule.cs
@@ -205,6 +205,11 @@ namespace UnrealBuildTool
 		/// </summary>
 		public HashSet<FileItem> NatvisFiles = new();
 
+		/// <summary>
+		/// AutoRTFM function-map definition files
+		/// </summary>
+		public HashSet<FileItem> PublicAutoRTFMExternalMappingFiles = new();
+
 		/// <summary>
 		/// Constructor
 		/// </summary>
@@ -230,6 +235,7 @@ namespace UnrealBuildTool
 			PublicSystemLibraries = HashSetFromOptionalEnumerableStringParameter(Rules.PublicSystemLibraries);
 			PublicFrameworks = HashSetFromOptionalEnumerableStringParameter(Rules.PublicFrameworks);
 			PublicWeakFrameworks = HashSetFromOptionalEnumerableStringParameter(Rules.PublicWeakFrameworks);
+			PublicAutoRTFMExternalMappingFiles = new HashSet<FileItem>(Rules.AutoRTFMExternalMappingFiles.Select(x => FileItem.GetItemByPath(x)));
 
 			foreach (string LibraryName in PublicAdditionalLibraries)
 			{
@@ -638,6 +644,7 @@ namespace UnrealBuildTool
 			HashSet<DirectoryReference> ModuleInterfacePaths,
 			List<string> Definitions,
 			List<UEBuildFramework> AdditionalFrameworks,
+			HashSet<FileItem> AutoRTFMExternalMappingFiles,
 			List<FileItem> AdditionalPrerequisites,
 			bool bLegacyPublicIncludePaths,
 			bool bLegacyParentIncludePaths
@@ -715,6 +722,8 @@ namespace UnrealBuildTool
 					}
 				}
 			}
+
+			AutoRTFMExternalMappingFiles.UnionWith(PublicAutoRTFMExternalMappingFiles);
 		}
 
 		/// <summary>
@@ -726,6 +735,7 @@ namespace UnrealBuildTool
 			HashSet<DirectoryReference> ModuleInterfacePaths,
 			List<string> Definitions,
 			List<UEBuildFramework> AdditionalFrameworks,
+			HashSet<FileItem> AutoRTFMExternalMappingFiles,
 			List<FileItem> AdditionalPrerequisites,
 			bool bWithLegacyPublicIncludePaths,
 			bool bWithLegacyParentIncludePaths
@@ -741,7 +751,7 @@ namespace UnrealBuildTool
 			// Now set up the compile environment for the modules in the original order that we encountered them
 			foreach (UEBuildModule Module in ModuleToIncludePathsOnlyFlag.Keys)
 			{
-				Module.AddModuleToCompileEnvironment(this, Binary, IncludePaths, SystemIncludePaths, ModuleInterfacePaths, Definitions, AdditionalFrameworks, AdditionalPrerequisites, bWithLegacyPublicIncludePaths, bWithLegacyParentIncludePaths);
+				Module.AddModuleToCompileEnvironment(this, Binary, IncludePaths, SystemIncludePaths, ModuleInterfacePaths, Definitions, AdditionalFrameworks, AutoRTFMExternalMappingFiles, AdditionalPrerequisites, bWithLegacyPublicIncludePaths, bWithLegacyParentIncludePaths);
 			}
 		}
 
@@ -1496,6 +1506,8 @@ namespace UnrealBuildTool
 			ExportJsonStringArray(Writer, "PublicDelayLoadDLLs", PublicDelayLoadDLLs);
 			ExportJsonStringArray(Writer, "PublicDefinitions", PublicDefinitions);
 
+			ExportJsonStringArray(Writer, "AutoRTFMExternalMappingFiles", PublicAutoRTFMExternalMappingFiles.Select(x => x.FullName));
+
 			Writer.WriteArrayStart("CircularlyReferencedModules");
 			foreach (string ModuleName in Rules.CircularlyReferencedDependentModules)
 			{
diff --git a/Engine/Source/Programs/UnrealBuildTool/Configuration/UEBuildModuleCPP.cs b/Engine/Source/Programs/UnrealBuildTool/Configuration/UEBuildModuleCPP.cs
--- a/Engine/Source/Programs/UnrealBuildTool/Configuration/UEBuildModuleCPP.cs
+++ b/Engine/Source/Programs/UnrealBuildTool/Configuration/UEBuildModuleCPP.cs
@@ -417,6 +417,7 @@ namespace UnrealBuildTool
 			HashSet<DirectoryReference> ModuleInterfacePaths,
 			List<string> Definitions,
 			List<UEBuildFramework> AdditionalFrameworks,
+			HashSet<FileItem> AutoRTFMExternalMappingFiles,
 			List<FileItem> AdditionalPrerequisites,
 			bool bLegacyPublicIncludePaths,
 			bool bLegacyParentIncludePaths
@@ -446,7 +447,7 @@ namespace UnrealBuildTool
 
 			ModuleInterfacePaths.Add(UEToolChain.GetModuleInterfaceDir(IntermediateDirectory));
 
-			base.AddModuleToCompileEnvironment(SourceModule, SourceBinary, IncludePaths, SystemIncludePaths, ModuleInterfacePaths, Definitions, AdditionalFrameworks, AdditionalPrerequisites, bLegacyPublicIncludePaths, bLegacyParentIncludePaths);
+			base.AddModuleToCompileEnvironment(SourceModule, SourceBinary, IncludePaths, SystemIncludePaths, ModuleInterfacePaths, Definitions, AdditionalFrameworks, AutoRTFMExternalMappingFiles, AdditionalPrerequisites, bLegacyPublicIncludePaths, bLegacyParentIncludePaths);
 		}
 
 		// UEBuildModule interface.
@@ -2497,7 +2498,7 @@ namespace UnrealBuildTool
 			}
 
 			// Setup the compile environment for the module.
-			SetupPrivateCompileEnvironment(Result.UserIncludePaths, Result.SystemIncludePaths, Result.ModuleInterfacePaths, Result.Definitions, Result.AdditionalFrameworks, Result.AdditionalPrerequisites, Rules.bLegacyPublicIncludePaths, Rules.bLegacyParentIncludePaths);
+			SetupPrivateCompileEnvironment(Result.UserIncludePaths, Result.SystemIncludePaths, Result.ModuleInterfacePaths, Result.Definitions, Result.AdditionalFrameworks, Result.AutoRTFMExternalMappingFiles, Result.AdditionalPrerequisites, Rules.bLegacyPublicIncludePaths, Rules.bLegacyParentIncludePaths);
 
 			foreach (string ForceIncludeFile in Rules.ForceIncludeFiles)
 			{
@@ -2592,7 +2593,7 @@ namespace UnrealBuildTool
 			// Now set up the compile environment for the modules in the original order that we encountered them
 			foreach (UEBuildModule Module in ModuleToIncludePathsOnlyFlag.Keys)
 			{
-				Module.AddModuleToCompileEnvironment(this, null, CompileEnvironment.UserIncludePaths, CompileEnvironment.SystemIncludePaths, CompileEnvironment.ModuleInterfacePaths, CompileEnvironment.Definitions, CompileEnvironment.AdditionalFrameworks, CompileEnvironment.AdditionalPrerequisites, Rules.bLegacyPublicIncludePaths, Rules.bLegacyParentIncludePaths);
+				Module.AddModuleToCompileEnvironment(this, null, CompileEnvironment.UserIncludePaths, CompileEnvironment.SystemIncludePaths, CompileEnvironment.ModuleInterfacePaths, CompileEnvironment.Definitions, CompileEnvironment.AdditionalFrameworks, CompileEnvironment.AutoRTFMExternalMappingFiles, CompileEnvironment.AdditionalPrerequisites, Rules.bLegacyPublicIncludePaths, Rules.bLegacyParentIncludePaths);
 			}
 			return CompileEnvironment;
 		}
diff --git a/Engine/Source/Programs/UnrealBuildTool/Configuration/UEBuildTarget.cs b/Engine/Source/Programs/UnrealBuildTool/Configuration/UEBuildTarget.cs
--- a/Engine/Source/Programs/UnrealBuildTool/Configuration/UEBuildTarget.cs
+++ b/Engine/Source/Programs/UnrealBuildTool/Configuration/UEBuildTarget.cs
@@ -3411,7 +3411,7 @@ namespace UnrealBuildTool
 			List<string> Definitions = new List<string>(GlobalCompileEnvironment.Definitions);
 			foreach (UEBuildModule Module in Binary.Modules)
 			{
-				Module.AddModuleToCompileEnvironment(null, null, new HashSet<DirectoryReference>(), new HashSet<DirectoryReference>(), new HashSet<DirectoryReference>(), Definitions, new List<UEBuildFramework>(), new List<FileItem>(), false, false);
+				Module.AddModuleToCompileEnvironment(null, null, new HashSet<DirectoryReference>(), new HashSet<DirectoryReference>(), new HashSet<DirectoryReference>(), Definitions, new List<UEBuildFramework>(), new HashSet<FileItem>(), new List<FileItem>(), false, false);
 			}
 
 			// Write the header
diff --git a/Engine/Source/Programs/UnrealBuildTool/Platform/Windows/VCToolChain.cs b/Engine/Source/Programs/UnrealBuildTool/Platform/Windows/VCToolChain.cs
--- a/Engine/Source/Programs/UnrealBuildTool/Platform/Windows/VCToolChain.cs
+++ b/Engine/Source/Programs/UnrealBuildTool/Platform/Windows/VCToolChain.cs
@@ -2171,12 +2171,18 @@ namespace UnrealBuildTool
 					// If we are using the AutoRTFM compiler, we make the compile action depend on the version of the compiler itself.
 					// This lets us update the compiler (which might not cause a version update of the compiler, which instead tracks
 					// the LLVM versioning scheme that Clang uses), but ensure that we rebuild the source if the compiler has changed.
+					// AutoRTFM also depends on external mapping files, which requires a full rebuild when changed.
 					if (CompileEnvironment.bUseAutoRTFMCompiler)
 					{
 						FileReference? CompilerPath = GetCppCompilerPath();
 						if (null != CompilerPath)
 						{
 							CompileAction.AdditionalPrerequisiteItems.Add(FileItem.GetItemByFileReference(CompilerPath));
+							foreach (FileItem File in CompileEnvironment.AutoRTFMExternalMappingFiles)
+							{
+								CompileAction.AdditionalPrerequisiteItems.Add(File);
+								CompileAction.Arguments.Add($"-Xclang -autortfm-mappings -Xclang \"{NormalizeCommandLinePath(File.Location, CompileEnvironment.RootPaths)}\"");
+							}
 						}
 					}
 				}
diff --git a/Engine/Source/Programs/UnrealBuildTool/System/CppCompileEnvironment.cs b/Engine/Source/Programs/UnrealBuildTool/System/CppCompileEnvironment.cs
--- a/Engine/Source/Programs/UnrealBuildTool/System/CppCompileEnvironment.cs
+++ b/Engine/Source/Programs/UnrealBuildTool/System/CppCompileEnvironment.cs
@@ -337,6 +337,11 @@ namespace UnrealBuildTool
 		/// </summary>
 		public bool bEnableAutoRTFMVerification = false;
 
+		/// <summary>
+		/// AutoRTFM external open -> closed mapping files
+		/// </summary>
+		public HashSet<FileItem> AutoRTFMExternalMappingFiles = new();
+
 		/// <summary>
 		/// Whether we should run an extra LLVM verification pass after the AutoRTFM compiler pass.
 		/// </summary>
@@ -844,6 +849,7 @@ namespace UnrealBuildTool
 			bUseAutoRTFMCompiler = Other.bUseAutoRTFMCompiler;
 			bDisableAutoRTFMInstrumentation = Other.bDisableAutoRTFMInstrumentation;
 			bEnableAutoRTFMVerification = Other.bEnableAutoRTFMVerification;
+			AutoRTFMExternalMappingFiles = new HashSet<FileItem>(Other.AutoRTFMExternalMappingFiles);
 			bAutoRTFMVerify = Other.bAutoRTFMVerify;
 			bAutoRTFMClosedStaticLinkage = Other.bAutoRTFMClosedStaticLinkage;
 			bCodeCoverage = Other.bCodeCoverage;
diff --git a/Engine/Source/Programs/UnrealBuildTool/ToolChain/ClangToolChain.cs b/Engine/Source/Programs/UnrealBuildTool/ToolChain/ClangToolChain.cs
--- a/Engine/Source/Programs/UnrealBuildTool/ToolChain/ClangToolChain.cs
+++ b/Engine/Source/Programs/UnrealBuildTool/ToolChain/ClangToolChain.cs
@@ -1285,19 +1285,25 @@ namespace UnrealBuildTool
 			Action CompileAction = Graph.CreateAction(ActionType.Compile);
 			CompileAction.RootPaths = CompileEnvironment.RootPaths;
 
+			CompileAction.Weight = CompileActionWeight;
+
+			// copy the global arguments into the file arguments, so GetCompileArguments_FileType can remove entries if needed (special case but can be important)
+			List<string> FileArguments = new(GlobalArguments);
+
 			// If we are using the AutoRTFM compiler, we make the compile action depend on the version of the compiler itself.
 			// This lets us update the compiler (which might not cause a version update of the compiler, which instead tracks
 			// the LLVM versioning scheme that Clang uses), but ensure that we rebuild the source if the compiler has changed.
+			// AutoRTFM also depends on external mapping files, which requires a full rebuild when changed.
 			if (CompileEnvironment.bUseAutoRTFMCompiler)
 			{
 				CompileAction.PrerequisiteItems.Add(FileItem.GetItemByFileReference(Info.Clang));
+				foreach (FileItem File in CompileEnvironment.AutoRTFMExternalMappingFiles)
+				{
+					CompileAction.PrerequisiteItems.Add(File);
+					FileArguments.Add($"-Xclang -autortfm-mappings -Xclang \"{NormalizeCommandLinePath(File.Location, CompileEnvironment.RootPaths)}\"");
+				}
 			}
 
-			CompileAction.Weight = CompileActionWeight;
-
-			// copy the global arguments into the file arguments, so GetCompileArguments_FileType can remove entries if needed (special case but can be important)
-			List<string> FileArguments = new(GlobalArguments);
-
 			// Add C or C++ specific compiler arguments and get the target file so we can get the correct output path.
 			FileItem TargetFile = GetCompileArguments_FileType(CompileEnvironment, SourceFile, OutputDir, FileArguments, CompileAction, Result);
 
