GenerateProjectFiles: Add -VisualStudio option

#rnx

diff --git a/Engine/Source/Programs/UnrealBuildTool/Modes/GenerateProjectFilesMode.cs b/Engine/Source/Programs/UnrealBuildTool/Modes/GenerateProjectFilesMode.cs
--- a/Engine/Source/Programs/UnrealBuildTool/Modes/GenerateProjectFilesMode.cs
+++ b/Engine/Source/Programs/UnrealBuildTool/Modes/GenerateProjectFilesMode.cs
@@ -23,6 +23,7 @@ namespace UnrealBuildTool
 		/// Types of project files to generate
 		/// </summary>
 		[CommandLine("-ProjectFileFormat")]
+		[CommandLine("-VisualStudio", Value = nameof(ProjectFileFormat.VisualStudio))]
 		[CommandLine("-2022", Value = nameof(ProjectFileFormat.VisualStudio2022))] // + override compiler
 		[CommandLine("-Makefile", Value = nameof(ProjectFileFormat.Make))]
 		[CommandLine("-CMakefile", Value = nameof(ProjectFileFormat.CMake))]
