[UBT]
* Fixed bug where bHasSharedResponseFile was not copied in copy ctor of CppCompileEnvironment

diff --git a/Engine/Source/Programs/UnrealBuildTool/System/CppCompileEnvironment.cs b/Engine/Source/Programs/UnrealBuildTool/System/CppCompileEnvironment.cs
--- a/Engine/Source/Programs/UnrealBuildTool/System/CppCompileEnvironment.cs
+++ b/Engine/Source/Programs/UnrealBuildTool/System/CppCompileEnvironment.cs
@@ -875,6 +875,7 @@ namespace UnrealBuildTool
 			PGODirectory = Other.PGODirectory;
 			bPrintTimingInfo = Other.bPrintTimingInfo;
 			bAllowRemotelyCompiledPCHs = Other.bAllowRemotelyCompiledPCHs;
+			bHasSharedResponseFile = Other.bHasSharedResponseFile;
 			bUseHeaderUnitsForPch = Other.bUseHeaderUnitsForPch;
 			UserIncludePaths = new HashSet<DirectoryReference>(Other.UserIncludePaths);
 			SystemIncludePaths = new HashSet<DirectoryReference>(Other.SystemIncludePaths);
