[UBT][Linux] Do not apply `-fno-PIE` to executables when building merged modules

Compiling with -fno-PIE will cause functions implemented in shared libraries to be declared as dso_local, which results in function addresses in the main executable resolving to PLT stubs instead of the address of the true function. This breaks function pointer equality, which is important for AutoRTFM.

diff --git a/Engine/Source/Programs/UnrealBuildTool/Platform/Linux/LinuxToolChain.cs b/Engine/Source/Programs/UnrealBuildTool/Platform/Linux/LinuxToolChain.cs
--- a/Engine/Source/Programs/UnrealBuildTool/Platform/Linux/LinuxToolChain.cs
+++ b/Engine/Source/Programs/UnrealBuildTool/Platform/Linux/LinuxToolChain.cs
@@ -576,7 +576,15 @@ namespace UnrealBuildTool
 				}
 				else
 				{
-					Arguments.Add("-fno-PIE");
+					// Compiling with -fno-PIE will cause functions implemented
+					// in shared libraries to be declared as dso_local, which
+					// results in function addresses in the main executable
+					// resolving to PLT stubs instead of the address of the true
+					// function. This breaks function pointer equality.
+					if (!bMergeModules)
+					{
+						Arguments.Add("-fno-PIE");
+					}
 				}
 			}
 
