[AutoRTFM] Disable magic constant emission when building linux shared modules.

Enabling the AutoRTFM "magic prefix" for a shared object will result in a R_X86_64_64 relocation linker error. Disable this optimization until we find a fix.

diff --git a/Engine/Source/Programs/UnrealBuildTool/Platform/Linux/LinuxToolChain.cs b/Engine/Source/Programs/UnrealBuildTool/Platform/Linux/LinuxToolChain.cs
--- a/Engine/Source/Programs/UnrealBuildTool/Platform/Linux/LinuxToolChain.cs
+++ b/Engine/Source/Programs/UnrealBuildTool/Platform/Linux/LinuxToolChain.cs
@@ -560,6 +560,14 @@ namespace UnrealBuildTool
 				// Use local-dynamic TLS model. This generates less efficient runtime code for __thread variables, but avoids problems of running into
 				// glibc/ld.so limit (DTV_SURPLUS) for number of dlopen()'ed DSOs with static TLS (see e.g. https://www.cygwin.com/ml/libc-help/2013-11/msg00033.html)
 				Arguments.Add("-ftls-model=local-dynamic");
+
+				if (CompileEnvironment.bEnableAutoRTFMInstrumentation)
+				{
+					// Enabling the AutoRTFM "magic prefix" for a shared object will result in a
+					// R_X86_64_64 relocation linker error, due to the closed function address
+					// prefixed onto the open functions.
+					Arguments.Add("-Xclang -mllvm -Xclang -autortfm-no-magic-prefix");
+				}
 			}
 			else
 			{
