Add target settings to reduce debug info size on clang.
Remove line mappings for inlined functions.
Remove template arguments from type names (info still remains on template type parameters).

diff --git a/Engine/Source/Programs/UnrealBuildTool/Configuration/ReadOnlyTargetRules.cs b/Engine/Source/Programs/UnrealBuildTool/Configuration/ReadOnlyTargetRules.cs
--- a/Engine/Source/Programs/UnrealBuildTool/Configuration/ReadOnlyTargetRules.cs
+++ b/Engine/Source/Programs/UnrealBuildTool/Configuration/ReadOnlyTargetRules.cs
@@ -357,6 +357,14 @@ namespace UnrealBuildTool
 
 		public IReadOnlySet<string> DebugInfoLineTablesOnlyPlugins => Inner.DebugInfoLineTablesOnlyPlugins;
 
+		public DebugInfoMode DebugInfoNoInlineLineTables => Inner.DebugInfoNoInlineLineTables;
+		public IReadOnlySet<string> DebugInfoNoInlineLineTablesModules => Inner.DebugInfoNoInlineLineTablesModules;
+		public IReadOnlySet<string> DebugInfoNoInlineLineTablesPlugins => Inner.DebugInfoNoInlineLineTablesPlugins;
+		
+		public DebugInfoMode DebugInfoSimpleTemplateNames => Inner.DebugInfoSimpleTemplateNames;
+		public IReadOnlySet<string> DebugInfoSimpleTemplateNamesModules => Inner.DebugInfoSimpleTemplateNamesModules;
+		public IReadOnlySet<string> DebugInfoSimpleTemplateNamesPlugins => Inner.DebugInfoSimpleTemplateNamesPlugins;
+
 		public bool bDisableDebugInfoForGeneratedCode => Inner.bDisableDebugInfoForGeneratedCode;
 
 		public bool bOmitPCDebugInfoInDevelopment => Inner.bOmitPCDebugInfoInDevelopment;
diff --git a/Engine/Source/Programs/UnrealBuildTool/Configuration/TargetRules.cs b/Engine/Source/Programs/UnrealBuildTool/Configuration/TargetRules.cs
--- a/Engine/Source/Programs/UnrealBuildTool/Configuration/TargetRules.cs
+++ b/Engine/Source/Programs/UnrealBuildTool/Configuration/TargetRules.cs
@@ -1861,6 +1861,44 @@ namespace UnrealBuildTool
 		[CommandLine("-DebugInfoLineTablesOnlyPlugins=", ListSeparator = '+')]
 		public HashSet<string> DebugInfoLineTablesOnlyPlugins { get; } = [];
 
+		/// <summary>
+		/// True if debug line number tables should be omitted in debug information for inline functions, compilers that support doing so.
+		/// </summary>
+		[XmlConfigFile(Category = "BuildConfiguration")]
+		[CommandLine("-DebugInfoNoInlineLineTables=")]
+		public DebugInfoMode DebugInfoNoInlineLineTables { get; set; } = DebugInfoMode.None;
+
+		/// <summary>
+		/// Modules that should omit line number tables for inline functions for compilers that support doing so. 
+		/// </summary>
+		[CommandLine("-DebugInfoNoInlineLineTablesModules=", ListSeparator = '+')]
+		public HashSet<string> DebugInfoNoInlineLineTablesModules { get; } = [];
+
+		/// <summary>
+		/// Plugins that should omit line number tables for inline functions for compilers that support doing so. 
+		/// </summary>
+		[CommandLine("-DebugInfoNoInlineLineTablesPlugins=", ListSeparator = '+')]
+		public HashSet<string> DebugInfoNoInlineLineTablesPlugins { get; } = [];
+
+		/// <summary>
+		/// True if template type names should be simplified to reduce the size of debug info for compilers that support doing so.
+		/// </summary>
+		[XmlConfigFile(Category = "BuildConfiguration")]
+		[CommandLine("-DebugInfoSimpleTemplateNames=")]
+		public DebugInfoMode DebugInfoSimpleTemplateNames { get; set; } = DebugInfoMode.None;
+
+		/// <summary>
+		/// Modules for which template type names should be simplified to reduce the size of debug info for compilers that support doing so.
+		/// </summary>
+		[CommandLine("-DebugInfoSimpleTemplateNamesModules=", ListSeparator = '+')]
+		public HashSet<string> DebugInfoSimpleTemplateNamesModules { get; } = [];
+
+		/// <summary>
+		/// Plugins for which template type names should be simplified to reduce the size of debug info for compilers that support doing so.
+		/// </summary>
+		[CommandLine("-DebugInfoSimpleTemplateNamesPlugins=", ListSeparator = '+')]
+		public HashSet<string> DebugInfoSimpleTemplateNamesPlugins { get; } = [];
+
 		/// <summary>
 		/// Whether to disable debug info generation for generated files. This improves link times and reduces pdb size for modules that have a lot of generated glue code.
 		/// </summary>
diff --git a/Engine/Source/Programs/UnrealBuildTool/Configuration/UEBuildModuleCPP.cs b/Engine/Source/Programs/UnrealBuildTool/Configuration/UEBuildModuleCPP.cs
--- a/Engine/Source/Programs/UnrealBuildTool/Configuration/UEBuildModuleCPP.cs
+++ b/Engine/Source/Programs/UnrealBuildTool/Configuration/UEBuildModuleCPP.cs
@@ -2207,6 +2207,24 @@ namespace UnrealBuildTool
 				{
 					Result.bDebugLineTablesOnly = true;
 				}
+
+				if (Target.DebugInfoNoInlineLineTables.HasFlag(DebugInfoMode.ProjectPlugins) && Rules.Plugin != null)
+				{
+					Result.bDebugNoInlineLineTables = true;
+				}
+				else if (Target.DebugInfoNoInlineLineTables.HasFlag(DebugInfoMode.Project) && Rules.Plugin == null)
+				{
+					Result.bDebugNoInlineLineTables = true;
+				}
+
+				if (Target.DebugInfoSimpleTemplateNames.HasFlag(DebugInfoMode.ProjectPlugins) && Rules.Plugin != null)
+				{
+					Result.bDebugSimpleTemplateNames = true;
+				}
+				else if (Target.DebugInfoSimpleTemplateNames.HasFlag(DebugInfoMode.Project) && Rules.Plugin == null)
+				{
+					Result.bDebugSimpleTemplateNames = true;
+				}
 			}
 			else
 			{
@@ -2218,6 +2236,24 @@ namespace UnrealBuildTool
 				{
 					Result.bDebugLineTablesOnly = true;
 				}
+
+				if (Target.DebugInfoNoInlineLineTables.HasFlag(DebugInfoMode.EnginePlugins) && Rules.Plugin != null)
+				{
+					Result.bDebugNoInlineLineTables = true;
+				}
+				else if (Target.DebugInfoNoInlineLineTables.HasFlag(DebugInfoMode.Engine) && Rules.Plugin == null)
+				{
+					Result.bDebugNoInlineLineTables = true;
+				}
+
+				if (Target.DebugInfoSimpleTemplateNames.HasFlag(DebugInfoMode.EnginePlugins) && Rules.Plugin != null)
+				{
+					Result.bDebugSimpleTemplateNames = true;
+				}
+				else if (Target.DebugInfoSimpleTemplateNames.HasFlag(DebugInfoMode.Engine) && Rules.Plugin == null)
+				{
+					Result.bDebugSimpleTemplateNames = true;
+				}
 			}
 
 			if (Rules.Plugin != null && Rules.Target.DebugInfoLineTablesOnlyPlugins.Contains(Rules.Plugin.Name))
@@ -2229,6 +2265,24 @@ namespace UnrealBuildTool
 				Result.bDebugLineTablesOnly = true;
 			}
 
+			if (Rules.Plugin != null && Rules.Target.DebugInfoNoInlineLineTablesPlugins.Contains(Rules.Plugin.Name))
+			{
+				Result.bDebugNoInlineLineTables = true;
+			}
+			else if (Rules.Target.DebugInfoNoInlineLineTablesModules.Contains(Name))
+			{
+				Result.bDebugNoInlineLineTables = true;
+			}
+
+			if (Rules.Plugin != null && Rules.Target.DebugInfoSimpleTemplateNamesPlugins.Contains(Rules.Plugin.Name))
+			{
+				Result.bDebugSimpleTemplateNames = true;
+			}
+			else if (Rules.Target.DebugInfoSimpleTemplateNamesModules.Contains(Name))
+			{
+				Result.bDebugSimpleTemplateNames = true;
+			}
+
 			if (Result.bDebugLineTablesOnly)
 			{
 				// Don't disable debug info if only line tables are requested
diff --git a/Engine/Source/Programs/UnrealBuildTool/Platform/Windows/VCToolChain.cs b/Engine/Source/Programs/UnrealBuildTool/Platform/Windows/VCToolChain.cs
--- a/Engine/Source/Programs/UnrealBuildTool/Platform/Windows/VCToolChain.cs
+++ b/Engine/Source/Programs/UnrealBuildTool/Platform/Windows/VCToolChain.cs
@@ -1126,6 +1126,16 @@ namespace UnrealBuildTool
 						Arguments.Add("-fstandalone-debug");
 					}
 
+					if (CompileEnvironment.bDebugNoInlineLineTables)
+					{
+						Arguments.Add("-Xclang -gno-inline-line-tables");
+					}
+
+					if (CompileEnvironment.bDebugSimpleTemplateNames)
+					{
+						Arguments.Add("-Xclang -gsimple-template-names");
+					}
+
 					// https://clang.llvm.org/docs/UsersManual.html#cmdoption-feliminate-unused-debug-types
 					// Intel ICX 2024.2 does not have this option, but reports Clang 19, so turn it off here
 					if (ClangVersion >= new VersionNumber(19) && !Target.WindowsPlatform.Compiler.IsIntel())
diff --git a/Engine/Source/Programs/UnrealBuildTool/System/CppCompileEnvironment.cs b/Engine/Source/Programs/UnrealBuildTool/System/CppCompileEnvironment.cs
--- a/Engine/Source/Programs/UnrealBuildTool/System/CppCompileEnvironment.cs
+++ b/Engine/Source/Programs/UnrealBuildTool/System/CppCompileEnvironment.cs
@@ -479,6 +479,18 @@ namespace UnrealBuildTool
 		/// </summary>
 		public bool bDebugLineTablesOnly = false;
 
+		/// <summary>
+		/// If true, omits line tables for inlined functions at their callsites, and only includes line information 
+		/// for the caller.
+		/// </summary>
+		public bool bDebugNoInlineLineTables = false;
+
+		/// <summary>
+		/// If true, use simple names for templates in debug type info. i.e. just TArray without template parameters.
+		/// Template argument information is still present in the type information. 
+		/// </summary>
+		public bool bDebugSimpleTemplateNames = false;
+
 		/// <summary>
 		/// True if we're compiling .cpp files that will go into a library (.lib file)
 		/// </summary>
@@ -857,6 +869,8 @@ namespace UnrealBuildTool
 			FPSemantics = Other.FPSemantics;
 			bCreateDebugInfo = Other.bCreateDebugInfo;
 			bDebugLineTablesOnly = Other.bDebugLineTablesOnly;
+			bDebugNoInlineLineTables = Other.bDebugNoInlineLineTables;
+			bDebugSimpleTemplateNames = Other.bDebugSimpleTemplateNames;
 			bIsBuildingLibrary = Other.bIsBuildingLibrary;
 			bIsBuildingDLL = Other.bIsBuildingDLL;
 			bUseStaticCRT = Other.bUseStaticCRT;
diff --git a/Engine/Source/Programs/UnrealBuildTool/ToolChain/ClangToolChain.cs b/Engine/Source/Programs/UnrealBuildTool/ToolChain/ClangToolChain.cs
--- a/Engine/Source/Programs/UnrealBuildTool/ToolChain/ClangToolChain.cs
+++ b/Engine/Source/Programs/UnrealBuildTool/ToolChain/ClangToolChain.cs
@@ -836,6 +836,19 @@ namespace UnrealBuildTool
 			{
 				Arguments.Add("-fextend-lifetimes");
 			}
+
+			if (CompileEnvironment.bCreateDebugInfo)
+			{
+				if (CompileEnvironment.bDebugNoInlineLineTables)
+				{
+					Arguments.Add("-gno-inline-line-tables");
+				}
+
+				if (CompileEnvironment.bDebugSimpleTemplateNames)
+				{
+					Arguments.Add("-gsimple-template-names");
+				}
+			}
 		}
 
 		/// <summary>
