UnrealBuildTool: Fix vfs resolution for pvs-studio analyzer

#rnx

diff --git a/Engine/Source/Programs/UnrealBuildTool/Platform/Windows/PVSToolChain.cs b/Engine/Source/Programs/UnrealBuildTool/Platform/Windows/PVSToolChain.cs
--- a/Engine/Source/Programs/UnrealBuildTool/Platform/Windows/PVSToolChain.cs
+++ b/Engine/Source/Programs/UnrealBuildTool/Platform/Windows/PVSToolChain.cs
@@ -956,7 +956,12 @@ namespace UnrealBuildTool
 
 			// Collect the sourcefile items off of the Compile action added in CompileCPPFiles so that in SingleFileCompile mode the PVSGather step is also not filtered out
 			IEnumerable<FileItem> compileSourceFiles = [.. makefile.Actions.OfType<VCCompileAction>().Select(x => x.SourceFile!)];
-			CppRootPaths rootPaths = makefile.Actions.OfType<VCCompileAction>().LastOrDefault(x => x.RootPaths.Any())?.RootPaths ?? new();
+			IEnumerable<CppRootPaths> allRootPaths = makefile.Actions.OfType<VCCompileAction>().Where(x => x.RootPaths.Any()).Select(x => x.RootPaths);
+			CppRootPaths rootPaths = allRootPaths.FirstOrDefault() ?? new();
+			foreach (CppRootPaths item in allRootPaths.Except([rootPaths]))
+			{
+				rootPaths.Merge(item);
+			}
 
 			// Store list of system paths that should be excluded
 			IEnumerable<DirectoryReference> systemIncludePaths = [.. makefile.Actions.OfType<VCCompileAction>().SelectMany(x => x.SystemIncludePaths)];
diff --git a/Engine/Source/Programs/UnrealBuildTool/System/CppRootPaths.cs b/Engine/Source/Programs/UnrealBuildTool/System/CppRootPaths.cs
--- a/Engine/Source/Programs/UnrealBuildTool/System/CppRootPaths.cs
+++ b/Engine/Source/Programs/UnrealBuildTool/System/CppRootPaths.cs
@@ -206,6 +206,26 @@ namespace UnrealBuildTool
 			}
 		}
 
+		/// <summary>
+		/// Merge in entries from another root paths object, discarding duplicates
+		/// </summary>
+		/// <param name="other">The paths to merge from</param>
+		public void Merge(CppRootPaths other)
+		{
+			foreach (KeyValuePair<CppRootPathFolder, (DirectoryReference vfs, DirectoryReference local)> lookup in other._lookup.Where(x => !_lookup.ContainsKey(x.Key)))
+			{
+				this[lookup.Key] = lookup.Value.local;
+			}
+			foreach (KeyValuePair<string, (DirectoryReference vfs, DirectoryReference local)> extra in other._extras.Where(x => !_extras.ContainsKey(x.Key)))
+			{
+				AddExtraPath((extra.Key, extra.Value.local.FullName));
+			}
+			foreach (KeyValuePair<CppRootPathFolder, string> name in other._folderNames.Where(x => !_folderNames.ContainsKey(x.Key)))
+			{
+				AddFolderName(name.Key, name.Value);
+			}
+		}
+
 		/// <summary>
 		/// Get a virtual filesystem overlay path if in use for a local path
 		/// </summary>
