UnrealBuildTool: Ensure directory exists before writing pvs temp files

#rnx

diff --git a/Engine/Source/Programs/UnrealBuildTool/Platform/Windows/PVSToolChain.cs b/Engine/Source/Programs/UnrealBuildTool/Platform/Windows/PVSToolChain.cs
--- a/Engine/Source/Programs/UnrealBuildTool/Platform/Windows/PVSToolChain.cs
+++ b/Engine/Source/Programs/UnrealBuildTool/Platform/Windows/PVSToolChain.cs
@@ -543,9 +543,9 @@ namespace UnrealBuildTool
 		readonly Version AnalyzerVersion;
 
 		static readonly Version _minAnalyzerVersion = new Version("7.30");
-        static readonly Version _analysisPathsSkipVersion = new Version("7.34");
+		static readonly Version _analysisPathsSkipVersion = new Version("7.34");
 
-        const string OutputFileExtension = ".PVS-Studio.log";
+		const string OutputFileExtension = ".PVS-Studio.log";
 
 		public PVSToolChain(ReadOnlyTargetRules Target, VCToolChain InInnerToolchain, ILogger Logger)
 			: base(Logger)
@@ -956,7 +956,7 @@ namespace UnrealBuildTool
 
 			// Collect the sourcefile items off of the Compile action added in CompileCPPFiles so that in SingleFileCompile mode the PVSGather step is also not filtered out
 			IEnumerable<FileItem> compileSourceFiles = [.. makefile.Actions.OfType<VCCompileAction>().Select(x => x.SourceFile!)];
-			IEnumerable<CppRootPaths> allRootPaths = makefile.Actions.OfType<VCCompileAction>().Where(x => x.RootPaths.Any()).Select(x => x.RootPaths);
+			IEnumerable<CppRootPaths> allRootPaths = makefile.Actions.OfType<VCCompileAction>().Where(x => x.RootPaths.bUseVfs && x.RootPaths.Any()).Select(x => x.RootPaths);
 			CppRootPaths rootPaths = allRootPaths.FirstOrDefault() ?? new();
 			foreach (CppRootPaths item in allRootPaths.Except([rootPaths]))
 			{
@@ -974,6 +974,7 @@ namespace UnrealBuildTool
 				using BinaryArchiveWriter binaryArchiveWriter = new(stream);
 				rootPaths.Write(binaryArchiveWriter);
 				binaryArchiveWriter.Flush();
+				DirectoryReference.CreateDirectory(rootPathsItem.Location.Directory);
 				FileReference.WriteAllBytesIfDifferent(rootPathsItem.Location, stream.ToArray());
 			}
 
