Fix ISPC compilation when using PVS for platforms that override ISPC accessors

diff --git a/Engine/Source/Programs/UnrealBuildTool/Platform/Windows/PVSToolChain.cs b/Engine/Source/Programs/UnrealBuildTool/Platform/Windows/PVSToolChain.cs
--- a/Engine/Source/Programs/UnrealBuildTool/Platform/Windows/PVSToolChain.cs
+++ b/Engine/Source/Programs/UnrealBuildTool/Platform/Windows/PVSToolChain.cs
@@ -992,5 +992,18 @@ namespace UnrealBuildTool
 
 			makefile.OutputItems.AddRange(finalizeAction.ProducedItems);
 		}
+
+		public override List<string> GetISPCCompileTargets(UnrealTargetPlatform Platform, UnrealArch Arch) => InnerToolChain.GetISPCCompileTargets(Platform, Arch);
+
+		public override string GetISPCOSTarget(UnrealTargetPlatform Platform) => InnerToolChain.GetISPCOSTarget(Platform);
+
+		public override string GetISPCArchTarget(UnrealTargetPlatform Platform, UnrealArch Arch) => InnerToolChain.GetISPCArchTarget(Platform, Arch);
+
+		public override string? GetISPCCpuTarget(UnrealTargetPlatform Platform) => InnerToolChain.GetISPCCpuTarget(Platform);
+
+		public override string GetISPCObjectFileFormat(UnrealTargetPlatform Platform) => InnerToolChain.GetISPCObjectFileFormat(Platform);
+
+		public override string GetISPCObjectFileSuffix(UnrealTargetPlatform Platform) => InnerToolChain.GetISPCObjectFileSuffix(Platform);
+
 	}
 }
