Bump NDK to r27c, buildtools to 35.0.1 and bundletool to 1.18.1 to support 16kb page size requirements

diff --git a/Engine/Source/Programs/UnrealBuildTool/Platform/Android/UEDeployAndroid.cs b/Engine/Source/Programs/UnrealBuildTool/Platform/Android/UEDeployAndroid.cs
--- a/Engine/Source/Programs/UnrealBuildTool/Platform/Android/UEDeployAndroid.cs
+++ b/Engine/Source/Programs/UnrealBuildTool/Platform/Android/UEDeployAndroid.cs
@@ -26,7 +26,7 @@ namespace UnrealBuildTool
 		private const string ANDROID_TOOLS_BUILD_GRADLE_VERSION = "com.android.tools.build:gradle:8.5.2";
 
 		// default NDK version if not set
-		private const string DEFAULT_NDK_VERSION = "25.1.8937393";
+		private const string DEFAULT_NDK_VERSION = "27.2.12479018";
 
 		// name of the only vulkan validation layer we're interested in 
 		private const string ANDROID_VULKAN_VALIDATION_LAYER = "libVkLayer_khronos_validation.so";
