[UBT/IWYU]
* Fixed so include-what-you-use works again! It is now using the win64 target instead so build win64 with -iwyu to enable it
* New iwyu binaries (iwyu 0.19 upgraded to use clang 18)
* Minor tweak to iwyu mappings rules
* Changed so IWYUToolChain now references "real" toolchain and calls it to generate response files
* Added in target rules so options required by iwyu are explicitly set
* Removed check in ModuleRules not needed since we disable pch in target rules
* Changed so iwyu use shared response files and use header as main file when compiling headers (UEBuildModuleCPP.cs)
* VCToolChain - Added AppendCLArguments_GlobalIWYU to keep vctoolchain command line code in one place
* VCToolChain - Deduplicated code for setting up include path arguments
* VCToolChain - Deduplicated code for exception arguments
* VCToolChain - Made file extension specific methods internal instead of protected so iwyu toolchain can call them

diff --git a/Engine/Source/Programs/UnrealBuildTool/Configuration/ModuleRules.cs b/Engine/Source/Programs/UnrealBuildTool/Configuration/ModuleRules.cs
--- a/Engine/Source/Programs/UnrealBuildTool/Configuration/ModuleRules.cs
+++ b/Engine/Source/Programs/UnrealBuildTool/Configuration/ModuleRules.cs
@@ -819,9 +819,9 @@ namespace UnrealBuildTool
 				{
 					return PCHUsageMode.NoPCHs;
 				}
-				else if (Target.bIWYU || DefaultBuildSettings >= BuildSettingsVersion.V2)
+				else if (DefaultBuildSettings >= BuildSettingsVersion.V2)
 				{
-					// Use shared or explicit PCHs, and enable IWYU
+					// Use shared or explicit PCHs
 					return PCHUsageMode.UseExplicitOrSharedPCHs;
 				}
 				else if (Plugin != null)
diff --git a/Engine/Source/Programs/UnrealBuildTool/Configuration/TargetRules.cs b/Engine/Source/Programs/UnrealBuildTool/Configuration/TargetRules.cs
--- a/Engine/Source/Programs/UnrealBuildTool/Configuration/TargetRules.cs
+++ b/Engine/Source/Programs/UnrealBuildTool/Configuration/TargetRules.cs
@@ -3113,6 +3113,17 @@ namespace UnrealBuildTool
 				}
 			}
 
+			if (bIWYU)
+			{
+				WindowsPlatform.Compiler = WindowsCompiler.Clang;
+				bDisableLinking = true; // No linking in iwyu builds
+				CppCompileWarningSettings.UnusedWarningLevel = WarningLevel.Off; // Disable warnings to avoid extra noise
+				bUsePCHFiles = false; // No pch for iwyu
+				bUseVFS = false; // Do not use vfs because IWYUMode can't handle virtual paths (yet)
+				bUseUnityBuild = false; // No unity for iwyu
+				// bAlwaysUseUnityForGeneratedFiles = false; // This should be used but adds too many extra files atm.
+			}
+
 			// Set the default build version
 			if (String.IsNullOrEmpty(BuildVersion))
 			{
diff --git a/Engine/Source/Programs/UnrealBuildTool/Configuration/UEBuildModuleCPP.cs b/Engine/Source/Programs/UnrealBuildTool/Configuration/UEBuildModuleCPP.cs
--- a/Engine/Source/Programs/UnrealBuildTool/Configuration/UEBuildModuleCPP.cs
+++ b/Engine/Source/Programs/UnrealBuildTool/Configuration/UEBuildModuleCPP.cs
@@ -743,7 +743,7 @@ namespace UnrealBuildTool
 
 			// The operator new/delete overrides need to exist in all modules.
 			// We put these overrides in a separate file in order to be able to ensure only one file is added per binary for the primary module
-			if (Target.LinkType == TargetLinkType.Modular && bCreatePerModuleFile)
+			if (Target.LinkType == TargetLinkType.Modular && bCreatePerModuleFile && !Target.bDisableLinking)
 			{
 				FileReference NewDeleteOverrides = FileReference.Combine(IntermediateDirectory, "PerModuleInline.gen.cpp");
 				List<string> Content =
@@ -1096,10 +1096,20 @@ namespace UnrealBuildTool
 				List<FileItem> HeaderFileItems = GetCompilableHeaders(InputFiles, CppCompileEnvironment);
 				if (HeaderFileItems.Count > 0)
 				{
-					CppCompileEnvironment HeaderCompileEnvironment = new CppCompileEnvironment(ModuleCompileEnvironment);
-					HeaderCompileEnvironment.bUseUnity = false;
-					DirectoryReference HeaderIntermediateDirectory = DirectoryReference.Combine(IntermediateDirectory, "H");
+					CppCompileEnvironment HeaderCompileEnvironment;
+
+					if (Target.bIWYU)
+					{
+						// We want to use CppCompileEnvironment with shared rsp etc in iwyu case
+						HeaderCompileEnvironment = new CppCompileEnvironment(CppCompileEnvironment);
+					}
+					else
+					{
+						HeaderCompileEnvironment = new CppCompileEnvironment(ModuleCompileEnvironment);
+						HeaderCompileEnvironment.bUseUnity = false;
+					}
 
+					DirectoryReference HeaderIntermediateDirectory = DirectoryReference.Combine(IntermediateDirectory, "H");
 					CreateHeaderForDefinitions(HeaderCompileEnvironment, HeaderIntermediateDirectory, "h", Graph);
 
 					// Duplicate named headers are allowed, so adjust the IntermediateDirectory to compensate for this by using an index
@@ -1110,20 +1120,28 @@ namespace UnrealBuildTool
 						List<FileItem> HeadersToCompile = new();
 						foreach (FileItem Header in HeaderFileItems.Where(x => x.Directory == ParentDir))
 						{
-							string IncludeFileString = Header.AbsolutePath;
-							if (HeaderCompileEnvironment.RootPaths.GetVfsOverlayPath(Header.Location, out string? VfsPath))
+							if (Target.bIWYU)
 							{
-								IncludeFileString = VfsPath;
+								// IWYU needs the actual headers as main file input
+								HeadersToCompile.Add(Header);
 							}
-							else if (Header.Location.IsUnderDirectory(Unreal.RootDirectory))
+							else
 							{
-								IncludeFileString = Header.Location.MakeRelativeTo(Unreal.EngineSourceDirectory);
-							}
+								string IncludeFileString = Header.AbsolutePath;
+								if (HeaderCompileEnvironment.RootPaths.GetVfsOverlayPath(Header.Location, out string? VfsPath))
+								{
+									IncludeFileString = VfsPath;
+								}
+								else if (Header.Location.IsUnderDirectory(Unreal.RootDirectory))
+								{
+									IncludeFileString = Header.Location.MakeRelativeTo(Unreal.EngineSourceDirectory);
+								}
 
-							List<string> GeneratedHeaderCppContents = UEBuildModuleCPP.GenerateHeaderCpp(Header.Name, IncludeFileString);
-							FileItem GeneratedHeaderCpp = FileItem.GetItemByFileReference(FileReference.Combine(HeaderSubIntermediateDirectory, $"{Header.Name}.cpp"));
-							Graph.CreateIntermediateTextFile(GeneratedHeaderCpp, GeneratedHeaderCppContents);
-							HeadersToCompile.Add(GeneratedHeaderCpp);
+								List<string> GeneratedHeaderCppContents = UEBuildModuleCPP.GenerateHeaderCpp(Header.Name, IncludeFileString);
+								FileItem GeneratedHeaderCpp = FileItem.GetItemByFileReference(FileReference.Combine(HeaderSubIntermediateDirectory, $"{Header.Name}.cpp"));
+								Graph.CreateIntermediateTextFile(GeneratedHeaderCpp, GeneratedHeaderCppContents);
+								HeadersToCompile.Add(GeneratedHeaderCpp);
+							}
 						}
 
 						// Add the compile actions
diff --git a/Engine/Source/Programs/UnrealBuildTool/Modes/IWYUMode.cs b/Engine/Source/Programs/UnrealBuildTool/Modes/IWYUMode.cs
--- a/Engine/Source/Programs/UnrealBuildTool/Modes/IWYUMode.cs
+++ b/Engine/Source/Programs/UnrealBuildTool/Modes/IWYUMode.cs
@@ -905,6 +905,10 @@ namespace UnrealBuildTool
 					Info.ForwardDeclarations.RemoveAll(Entry =>
 					{
 						int LastSpace = Entry.Printable.LastIndexOf(' ');
+						if (LastSpace == -1)
+						{
+							throw new BuildException($"{Info.Source?.Name} contains ForwardDeclaration that is empty");
+						}
 						string TypeName = Entry.Printable.Substring(LastSpace + 1, Entry.Printable.Length - LastSpace - 2);
 						IWYUIncludeEntry? IncludeEntry;
 						ForwardingHeaders.TryGetValue(TypeName, out IncludeEntry);
@@ -936,32 +940,7 @@ namespace UnrealBuildTool
 						}
 					}
 				}
-				/*
-				if (Info.IncludesSeenInFile.Count == 0)
-				{
-					if (Info.File.EndsWith(".inl"))
-					{
-						Info.Includes.Clear();
-						Info.ForwardDeclarations.Clear();
-					}
-				}
-				*/
-				// Definitions.h is automatically added in the reponse file and iwyu sees it as not included
-				// If there are no includes but we depend on Definitions.h (which is missing) we add Platform.h because it is most likely a <module>_API entry in the file
-				/*
-				if (Info.Includes.Count == 0)
-				{
-					bool IsUsingDefinitionsH = false;
-					foreach (var Missing in Info.MissingIncludes)
-					{
-						IsUsingDefinitionsH = IsUsingDefinitionsH || Missing.Full.Contains("Definitions.h");
-					}
-					if (IsUsingDefinitionsH)
-					{
-						Info.Includes.Add(new IWYUIncludeEntry() { Full = PlatformInfo!.File, Printable = "HAL/Platform.h", Resolved = PlatformInfo });
-					}
-				}
-				*/
+
 				foreach (IWYUIncludeEntry Include in Info.Includes)
 				{
 					if (Include.Resolved == null)
@@ -984,6 +963,29 @@ namespace UnrealBuildTool
 						}
 					}
 				}
+
+				foreach (IWYUMissingInclude MissingInclude in Info.MissingIncludes)
+				{
+					if (Info.TransitiveIncludes.ContainsKey(MissingInclude.Full))
+					{
+						continue;
+					}
+					if (Infos.TryGetValue(MissingInclude.Full, out IWYUInfo? value) && value.Module != null)
+					{
+						string fixedPath = value.File.Replace('/', '\\');
+						foreach (DirectoryReference? Dir in value.Module.PublicIncludePaths.Union(value.Module.PrivateIncludePaths).Union(value.Module.PublicSystemIncludePaths))
+						{
+							if (fixedPath.StartsWith(Dir.FullName))
+							{
+								string printable = value.File.Substring(Dir.FullName.Length + 1);
+								var entry = new IWYUIncludeEntry() { Full = MissingInclude.Full, Printable = printable, System = false, Resolved = value, };
+								Info.Includes.Add(entry);
+								break;
+							}
+						}
+					}
+				}
+
 			});
 
 			int InfosCount = Infos.Count;
@@ -1081,6 +1083,10 @@ namespace UnrealBuildTool
 					{
 						if (!Include.Full.Contains("/ThirdParty/") && !Include.Full.StartsWith(SystemPath))
 						{
+							if (Include.Resolved == null)
+							{
+								Infos.TryGetValue(Include.Full, out Include.Resolved);
+							}
 							IncluderInfo.Includes.Add(Include);
 						}
 					}
@@ -2181,10 +2187,22 @@ namespace UnrealBuildTool
 				return Size + "b";
 			};
 
+			int PrintCount = 20;
+
 			Logger.LogInformation("");
-			Logger.LogInformation("Top 20 most reduction in bytes parsed by compiler frontend");
+			Logger.LogInformation($"Top {PrintCount} most reduction in bytes parsed by compiler frontend");
 			Logger.LogInformation("----------------------------------------------------------");
-			for (int I = 0; I != Math.Min(20, ByteSizeDiffList.Count); ++I)
+			int End = Math.Min(20, ByteSizeDiffList.Count);
+
+
+			int LongestName = 0;
+			for (int I = 0; I != End; ++I)
+			{
+				IWYUInfo File = ByteSizeDiffList[I].Item2;
+				LongestName = Math.Max(LongestName, Path.GetFileName(File.File).Length);
+			}
+
+			for (int I = 0; I != End; ++I)
 			{
 				long Saved = ByteSizeDiffList[I].Item1;
 				IWYUInfo File = ByteSizeDiffList[I].Item2;
@@ -2195,21 +2213,16 @@ namespace UnrealBuildTool
 				}
 
 				float Percent = 100.0f - (((float)OptimizedSize) / (OptimizedSize + Saved) * 100.0f);
-				Logger.LogInformation($"{Path.GetFileName(File.File)}   {PrettySize(OptimizedSize + Saved)} -> {PrettySize(OptimizedSize)}  (Saved {PrettySize(Saved)} or {Percent:0.0}%)");
+				string Name = Path.GetFileName(File.File);
+				Logger.LogInformation($"{Name.PadRight(LongestName)}   {PrettySize(OptimizedSize + Saved).PadRight(8)} -> {PrettySize(OptimizedSize).PadRight(8)}  (Saved {PrettySize(Saved).PadRight(7)} or {Percent:0.0}%)");
 			}
 			Logger.LogInformation("");
 
-			/*
-			RatioList.Sort((a, b) =>
-			{
-				if (a.Item1 == b.Item1)
-					return 0;
-				if (a.Item1 < b.Item1)
-					return -1;
-				return 1;
-			});
-			*/
-
+			DirectoryReference iwyuDir = DirectoryReference.Combine(Unreal.EngineDirectory, "Intermediate", "IWYU");
+			Directory.CreateDirectory(iwyuDir.FullName);
+			string resultFile = DirectoryReference.Combine(iwyuDir, "CompareToplist.txt").FullName;
+			File.WriteAllLines(resultFile, ByteSizeDiffList.Take(200).Select(Item => Item.Item2.File));
+			Logger.LogInformation($"Sorted top list written to {resultFile} (Can be used as input to -PathToUpdateFile=<file>)");
 			return 0;
 		}
 	}
diff --git a/Engine/Source/Programs/UnrealBuildTool/Platform/Linux/UEBuildLinux.cs b/Engine/Source/Programs/UnrealBuildTool/Platform/Linux/UEBuildLinux.cs
--- a/Engine/Source/Programs/UnrealBuildTool/Platform/Linux/UEBuildLinux.cs
+++ b/Engine/Source/Programs/UnrealBuildTool/Platform/Linux/UEBuildLinux.cs
@@ -614,11 +614,6 @@ namespace UnrealBuildTool
 		/// <returns>New toolchain instance.</returns>
 		public override UEToolChain CreateToolChain(ReadOnlyTargetRules Target)
 		{
-			if (Target.bIWYU)
-			{
-				return new IWYUToolChain(Target, Logger);
-			}
-
 			ClangToolChainOptions Options = ClangToolChainOptions.None;
 
 			if (Target.LinuxPlatform.bEnableAddressSanitizer)
@@ -712,7 +707,14 @@ namespace UnrealBuildTool
 				Log.ColorConsoleOutput = false;
 			}
 
-			return new LinuxToolChain(Target.Architecture, SDK, Options, Logger);
+			LinuxToolChain toolChain = new LinuxToolChain(Target.Architecture, SDK, Options, Logger);
+
+			if (Target.bIWYU)
+			{
+				return new IWYUToolChain(Target, toolChain, Logger);
+			}
+
+			return toolChain;
 		}
 
 		/// <inheritdoc/>
diff --git a/Engine/Source/Programs/UnrealBuildTool/Platform/Windows/UEBuildWindows.cs b/Engine/Source/Programs/UnrealBuildTool/Platform/Windows/UEBuildWindows.cs
--- a/Engine/Source/Programs/UnrealBuildTool/Platform/Windows/UEBuildWindows.cs
+++ b/Engine/Source/Programs/UnrealBuildTool/Platform/Windows/UEBuildWindows.cs
@@ -2089,6 +2089,11 @@ namespace UnrealBuildTool
 			{
 				return new PVSToolChain(Target, toolchain, Logger);
 			}
+			else if (Target.bIWYU)
+			{
+				return new IWYUToolChain(Target, toolchain, Logger);
+			}
+				
 			return toolchain;
 		}
 
diff --git a/Engine/Source/Programs/UnrealBuildTool/Platform/Windows/VCToolChain.cs b/Engine/Source/Programs/UnrealBuildTool/Platform/Windows/VCToolChain.cs
--- a/Engine/Source/Programs/UnrealBuildTool/Platform/Windows/VCToolChain.cs
+++ b/Engine/Source/Programs/UnrealBuildTool/Platform/Windows/VCToolChain.cs
@@ -405,6 +405,30 @@ namespace UnrealBuildTool
 			return new FileReference(ProfDataFiles.First());
 		}
 
+		private void AddExceptionArguments(CppCompileEnvironment CompileEnvironment, List<string> Arguments)
+		{
+			// Enable C++ exceptions when building with the editor or when building UHT.
+			if (CompileEnvironment.bEnableExceptions)
+			{
+				// Enable C++ exception handling, but not C exceptions.
+				Arguments.Add("/EHsc");
+				Arguments.Add("/DPLATFORM_EXCEPTIONS_DISABLED=0");
+			}
+			else
+			{
+				// This is required to disable exception handling in VC platform headers.
+				AddDefinition(Arguments, "_HAS_EXCEPTIONS=0");
+				Arguments.Add("/DPLATFORM_EXCEPTIONS_DISABLED=1");
+			}
+		}
+
+		private void AddClangResourceDirArgument(CppCompileEnvironment CompileEnvironment, List<string> Arguments)
+		{
+			VersionNumber ClangVersion = Target.WindowsPlatform.Compiler == WindowsCompiler.Intel ? MicrosoftPlatformSDK.GetClangVersionForIntelCompiler(EnvVars.CompilerPath) : EnvVars.CompilerVersion;
+			DirectoryReference ResourceDir = DirectoryReference.Combine(EnvVars.CompilerDir!, "lib", "clang", ClangVersion.GetComponent(0).ToString());
+			Arguments.Add($"-resource-dir=\"{NormalizeCommandLinePath(ResourceDir, CompileEnvironment.RootPaths)}\"");
+		}
+
 		protected virtual void AppendCLArguments_Global(CppCompileEnvironment CompileEnvironment, List<string> Arguments)
 		{
 			// Workaround for MSVC 14.31 compiler crash
@@ -518,9 +542,7 @@ namespace UnrealBuildTool
 				}
 
 				// Resource dir is for built-in includes. We need it for vfs but also for uba include crawler to find includes to transfer
-				VersionNumber ClangVersion = Target.WindowsPlatform.Compiler == WindowsCompiler.Intel ? MicrosoftPlatformSDK.GetClangVersionForIntelCompiler(EnvVars.CompilerPath) : EnvVars.CompilerVersion;
-				DirectoryReference ResourceDir = DirectoryReference.Combine(EnvVars.CompilerDir!, "lib", "clang", ClangVersion.GetComponent(0).ToString());
-				Arguments.Add($"-resource-dir=\"{NormalizeCommandLinePath(ResourceDir, CompileEnvironment.RootPaths)}\"");
+				AddClangResourceDirArgument(CompileEnvironment, Arguments);
 
 			}
 			else if (Target.StaticAnalyzer == StaticAnalyzer.Default && !CompileEnvironment.bDisableStaticAnalysis)
@@ -1058,19 +1080,7 @@ namespace UnrealBuildTool
 			// Prompt the user before reporting internal errors to Microsoft.
 			Arguments.Add("/errorReport:prompt");
 
-			// Enable C++ exceptions when building with the editor or when building UHT.
-			if (CompileEnvironment.bEnableExceptions)
-			{
-				// Enable C++ exception handling, but not C exceptions.
-				Arguments.Add("/EHsc");
-				Arguments.Add("/DPLATFORM_EXCEPTIONS_DISABLED=0");
-			}
-			else
-			{
-				// This is required to disable exception handling in VC platform headers.
-				AddDefinition(Arguments, "_HAS_EXCEPTIONS=0");
-				Arguments.Add("/DPLATFORM_EXCEPTIONS_DISABLED=1");
-			}
+			AddExceptionArguments(CompileEnvironment, Arguments);
 
 			// If enabled, create debug information.
 			if (CompileEnvironment.bCreateDebugInfo)
@@ -1272,7 +1282,44 @@ namespace UnrealBuildTool
 			}
 		}
 
-		protected virtual void AppendCLArguments_H(CppCompileEnvironment CompileEnvironment, List<string> Arguments)
+		private void AddIncludePathArguments(CppCompileEnvironment CompileEnvironment, WindowsCompiler Compiler, List<string> Arguments)
+		{
+			foreach (DirectoryReference IncludePath in CompileEnvironment.UserIncludePaths)
+			{
+				AddIncludePath(Arguments, IncludePath, Compiler, CompileEnvironment.RootPaths);
+			}
+
+			foreach (DirectoryReference IncludePath in CompileEnvironment.SystemIncludePaths)
+			{
+				AddSystemIncludePath(Arguments, IncludePath, Compiler, CompileEnvironment.RootPaths);
+			}
+
+			foreach (DirectoryReference IncludePath in EnvVars.IncludePaths)
+			{
+				AddSystemIncludePath(Arguments, IncludePath, Compiler, CompileEnvironment.RootPaths);
+			}
+		}
+
+		internal void AppendCLArguments_GlobalIWYU(CppCompileEnvironment CompileEnvironment, List<string> Arguments)
+		{
+			AddIncludePathArguments(CompileEnvironment, Target.WindowsPlatform.Compiler, Arguments);
+
+			Arguments.Add($"-fms-compatibility-version=19.{EnvVars.ToolChainVersion.GetComponent(1)}");
+			Arguments.Add("-Wno-invalid-constexpr");
+
+			AddDefinition(Arguments, "_CRT_USE_BUILTIN_OFFSETOF");
+			AddDefinition(Arguments, "_CRT_SECURE_NO_WARNINGS");
+
+			AddExceptionArguments(CompileEnvironment, Arguments);
+
+			// This is a bit messed up. Since include-what-you-use exe and this path does not really match.
+			// But some code includes cpuid.h behind #if PLATFORM_COMPILER_CLANG so we need this folder to find that file
+			AddClangResourceDirArgument(CompileEnvironment, Arguments);
+
+			AppendCLArguments_CompileWarnings(CompileEnvironment, Arguments);
+		}
+
+		internal virtual void AppendCLArguments_H(CppCompileEnvironment CompileEnvironment, List<string> Arguments)
 		{
 			AppendCLArguments_CPP(CompileEnvironment, Arguments);
 			
@@ -1282,7 +1329,7 @@ namespace UnrealBuildTool
 			}
 		}
 
-		protected virtual void AppendCLArguments_CPP(CppCompileEnvironment CompileEnvironment, List<string> Arguments)
+		internal virtual void AppendCLArguments_CPP(CppCompileEnvironment CompileEnvironment, List<string> Arguments)
 		{
 			// Explicitly compile the file as C++.
 			Arguments.Add("/TP");
@@ -1395,7 +1442,7 @@ namespace UnrealBuildTool
 			}
 		}
 
-		void AppendCLArguments_C(CppCompileEnvironment CompileEnvironment, List<string> Arguments)
+		internal void AppendCLArguments_C(CppCompileEnvironment CompileEnvironment, List<string> Arguments)
 		{
 			// Explicitly compile the file as C.
 			Arguments.Add("/TC");
@@ -2772,20 +2819,7 @@ namespace UnrealBuildTool
 				$"/Fo\"{NormalizeCommandLinePath(ObjectFile, RootPaths)}\""
 			};
 
-			foreach (DirectoryReference IncludePath in CompileEnvironment.UserIncludePaths)
-			{
-				AddIncludePath(Arguments, IncludePath, EnvVars.ToolChain, RootPaths);
-			}
-
-			foreach (DirectoryReference IncludePath in CompileEnvironment.SystemIncludePaths)
-			{
-				AddSystemIncludePath(Arguments, IncludePath, EnvVars.ToolChain, RootPaths);
-			}
-
-			foreach (DirectoryReference IncludePath in EnvVars.IncludePaths)
-			{
-				AddSystemIncludePath(Arguments, IncludePath, EnvVars.ToolChain, RootPaths);
-			}
+			AddIncludePathArguments(CompileEnvironment, EnvVars.ToolChain, Arguments);
 
 			FileItem ResponseFile = Graph.CreateIntermediateTextFile(OutputFile.ChangeExtension(ResponseExt), Arguments);
 
@@ -2864,20 +2898,7 @@ namespace UnrealBuildTool
 			CppCompileEnvironment NewCompileEnvironment = new CppCompileEnvironment(CompileEnvironment);
 			List<string> Arguments = new List<string>();
 
-			foreach (DirectoryReference IncludePath in NewCompileEnvironment.UserIncludePaths)
-			{
-				AddIncludePath(Arguments, IncludePath, Target.WindowsPlatform.Compiler, NewCompileEnvironment.RootPaths);
-			}
-
-			foreach (DirectoryReference IncludePath in NewCompileEnvironment.SystemIncludePaths)
-			{
-				AddSystemIncludePath(Arguments, IncludePath, Target.WindowsPlatform.Compiler, NewCompileEnvironment.RootPaths);
-			}
-
-			foreach (DirectoryReference IncludePath in EnvVars.IncludePaths)
-			{
-				AddSystemIncludePath(Arguments, IncludePath, Target.WindowsPlatform.Compiler, NewCompileEnvironment.RootPaths);
-			}
+			AddIncludePathArguments(NewCompileEnvironment, Target.WindowsPlatform.Compiler, Arguments);
 
 			// Stash shared include paths for validation purposes
 			NewCompileEnvironment.SharedUserIncludePaths = new(NewCompileEnvironment.UserIncludePaths);
diff --git a/Engine/Source/Programs/UnrealBuildTool/ToolChain/IWYUToolChain.cs b/Engine/Source/Programs/UnrealBuildTool/ToolChain/IWYUToolChain.cs
--- a/Engine/Source/Programs/UnrealBuildTool/ToolChain/IWYUToolChain.cs
+++ b/Engine/Source/Programs/UnrealBuildTool/ToolChain/IWYUToolChain.cs
@@ -13,8 +13,10 @@ namespace UnrealBuildTool
 	class IWYUToolChain : ClangToolChain
 	{
 		private List<string> CrossCompilingArguments = new();
-		private FileReference IWYUMappingFile;
+		private string IWYUMappingFile;
 		private string RelativePathToIWYUDirectory = @"Binaries\ThirdParty\IWYU";
+		private UnrealTargetPlatform TargetPlatform;
+		private ISPCToolChain InnerToolChain;
 
 		public static void ValidateTarget(TargetRules Target)
 		{
@@ -26,10 +28,13 @@ namespace UnrealBuildTool
 			Target.bCompileISPC = false;
 		}
 
-		public IWYUToolChain(ReadOnlyTargetRules Target, ILogger InLogger) : base(ClangToolChainOptions.None, InLogger)
+		public IWYUToolChain(ReadOnlyTargetRules target, ISPCToolChain internalToolChain, ILogger InLogger) : base(ClangToolChainOptions.None, InLogger)
 		{
+			TargetPlatform = target.Platform;
+			InnerToolChain = internalToolChain;
+
 			// set up the path to our toolchains
-			IWYUMappingFile = FileReference.Combine(Unreal.EngineDirectory, RelativePathToIWYUDirectory, "ue_mapping.imp");
+			IWYUMappingFile = FileReference.Combine(Unreal.EngineDirectory, RelativePathToIWYUDirectory, "ue_mapping.imp").ToNormalizedPath();
 			if (IWYUMappingFile == null)
 			{
 				throw new BuildException("It seems you don't have access to NotForLicensees folder.\n" +
@@ -40,25 +45,27 @@ namespace UnrealBuildTool
 
 		protected override ClangToolChainInfo GetToolChainInfo()
 		{
-			LinuxPlatformSDK PlatformSDK = new LinuxPlatformSDK(Logger);
-			DirectoryReference? BaseLinuxPath = PlatformSDK.GetBaseLinuxPathForArchitecture(LinuxPlatform.DefaultHostArchitecture);
-
-			SourceFileMetadataCache MetadataCache = SourceFileMetadataCache.CreateHierarchy(null, Logger);
-			CppCompileEnvironment CrossCompileEnvironment = new CppCompileEnvironment(UnrealTargetPlatform.Linux, CppConfiguration.Development, new(LinuxPlatform.DefaultHostArchitecture), MetadataCache);
+			if (TargetPlatform == UnrealTargetPlatform.Linux)
+			{
+				LinuxPlatformSDK PlatformSDK = new LinuxPlatformSDK(Logger);
+				DirectoryReference? BaseLinuxPath = PlatformSDK.GetBaseLinuxPathForArchitecture(LinuxPlatform.DefaultHostArchitecture);
 
-			CrossCompilingArguments.Add($"--sysroot=\"{NormalizeCommandLinePath(BaseLinuxPath!)}\"");
+				SourceFileMetadataCache MetadataCache = SourceFileMetadataCache.CreateHierarchy(null, Logger);
+				CppCompileEnvironment CrossCompileEnvironment = new CppCompileEnvironment(UnrealTargetPlatform.Linux, CppConfiguration.Development, new(LinuxPlatform.DefaultHostArchitecture), MetadataCache);
 
-			FileReference ClangPath = FileReference.Combine(BaseLinuxPath!, "bin", $"clang++{BuildHostPlatform.Current.BinarySuffix}");
-			ClangToolChainInfo CompilerToolChainInfo = new ClangToolChainInfo(BaseLinuxPath, ClangPath, null!, Logger);
+				CrossCompilingArguments.Add($"--sysroot=\"{NormalizeCommandLinePath(BaseLinuxPath!)}\"");
 
-			// starting with clang 16.x the directory naming changed to include major version only
-			string ClangVersionString = (CompilerToolChainInfo.ClangVersion.Major >= 16) ? CompilerToolChainInfo.ClangVersion.Major.ToString() : CompilerToolChainInfo.ClangVersion.ToString();
-			DirectoryReference SystemPath = DirectoryReference.Combine(BaseLinuxPath!, "lib", "clang", ClangVersionString, "include");
-			CrossCompilingArguments.Add(GetSystemIncludePathArgument(SystemPath, CrossCompileEnvironment));
+				FileReference ClangPath = FileReference.Combine(BaseLinuxPath!, "bin", $"clang++{BuildHostPlatform.Current.BinarySuffix}");
+				ClangToolChainInfo CompilerToolChainInfo = new ClangToolChainInfo(BaseLinuxPath, ClangPath, null!, Logger);
 
-			string DevPath = ""; //@"include-what-you-use-0.19\vs_projects\bin\RelWithDebInfo";
+				// starting with clang 16.x the directory naming changed to include major version only
+				string ClangVersionString = (CompilerToolChainInfo.ClangVersion.Major >= 16) ? CompilerToolChainInfo.ClangVersion.Major.ToString() : CompilerToolChainInfo.ClangVersion.ToString();
+				DirectoryReference SystemPath = DirectoryReference.Combine(BaseLinuxPath!, "lib", "clang", ClangVersionString, "include");
+				CrossCompilingArguments.Add(GetSystemIncludePathArgument(SystemPath, CrossCompileEnvironment));
+			}
 
-			FileReference IWYUPath = FileReference.Combine(Unreal.EngineDirectory, RelativePathToIWYUDirectory, DevPath, @"include-what-you-use.exe");
+			FileReference IWYUPath = FileReference.Combine(Unreal.EngineDirectory, RelativePathToIWYUDirectory, @"include-what-you-use.exe");
+			//FileReference IWYUPath = FileReference.FromString(@"e:\dev\fn\Engine\Restricted\NotForLicensees\Source\ThirdParty\IWYU\iwyu-0.23\vs_projects\bin\RelWithDebInfo\include-what-you-use.exe");
 			return new ClangToolChainInfo(IWYUPath.Directory, IWYUPath, IWYUPath, Logger);
 		}
 
@@ -74,68 +81,6 @@ namespace UnrealBuildTool
 			return FileName + Extension;
 		}
 
-		protected override CPPOutput CompileCPPFiles(CppCompileEnvironment CompileEnvironment, IEnumerable<FileItem> InputFiles, DirectoryReference OutputDir, string ModuleName, IActionGraphBuilder Graph)
-		{
-			if (!InputFiles.Any())
-			{
-				return new CPPOutput();
-			}
-
-			Dictionary<FileItem, List<FileItem>> GenCppLookup = CompileEnvironment.FileInlineGenCPPMap;
-
-			List<string> GlobalArguments = new();
-			if (CompileEnvironment.UserIncludePaths.Count != 0)
-			{
-				GetCompileArguments_Global(CompileEnvironment, GlobalArguments);
-			}
-
-			// Need to add these on cmd line.. IWYU is not parsing commands after response file expansion.
-			string CommonCommandLineArgs = " -Xiwyu --mapping_file=" + IWYUMappingFile.FullName +
-										   " -Xiwyu --prefix_header_includes=keep" +
-										   //										   " -Xiwyu --transitive_includes_only" +  // Since we are building headers separately (not together with their cpp) we don't need this
-										   " -Xiwyu --no_check_matching_header" +
-										   " -Xiwyu --cxx17ns";
-
-			// Create a compile action for each source file.
-			List<FileItem> IwyuFiles = new List<FileItem>();
-			foreach (FileItem SourceFile in InputFiles)
-			{
-				//if (SourceFile.HasExtension(".cpp") || !SourceFile.FullName.Contains("\\Runtime\\Engine\\"))
-				//	continue;
-
-				Action Action = CompileCPPFile(CompileEnvironment, SourceFile, OutputDir, ModuleName, Graph, GlobalArguments, new CPPOutput());
-
-				FileItem IwyuItem = Action.ProducedItems.First(i => i.Name.EndsWith(".iwyu"));
-
-				string CommandLineArgs = CommonCommandLineArgs + " -Xiwyu --write_json_path=\"" + IwyuItem + "\" ";
-				if (SourceFile.HasExtension(".cpp"))
-				{
-					List<FileItem>? InlinedFiles;
-					if (GenCppLookup.TryGetValue(SourceFile, out InlinedFiles))
-					{
-						foreach (FileItem InlinedFile in InlinedFiles)
-						{
-							CommandLineArgs += $"-Xiwyu --check_also={InlinedFile.FullName.Replace('\\', '/')} ";
-						}
-					}
-				}
-				else
-				{
-					// We want to build header files first
-					Action.bIsHighPriority = true;
-				}
-
-				Action.CommandArguments = CommandLineArgs + Action.CommandArguments;
-				IwyuFiles.Add(IwyuItem);
-			}
-
-			return new CPPOutput() { ObjectFiles = IwyuFiles };
-		}
-
-        // Include paths must be absolute because we have code gen that can produce includes that looks like this: "../../../Source/Foo/Bar.h"
-        // IWYU combines include path and include and then try to normalize the result. This could result in a path that can't be normalized because there are more
-        // ".." than there are directories to reduce. This silently fails in odd ways
-
         /// <inheritdoc/>
         public override string NormalizeCommandLinePath(FileSystemReference Reference, CppRootPaths RootPath)
 		{
@@ -164,60 +109,172 @@ namespace UnrealBuildTool
 
 		protected override void GetCompileArguments_Global(CppCompileEnvironment CompileEnvironment, List<string> Arguments)
 		{
+			if (CompileEnvironment.bHasSharedResponseFile)
+			{
+				return;
+			}
+
 			Arguments.Add(GetPreprocessorDefinitionArgument("UE_DIRECT_HEADER_COMPILE=1"));
 			Arguments.Add(GetPreprocessorDefinitionArgument("PLATFORM_COMPILER_IWYU=1"));
 
-			base.GetCompileArguments_Global(CompileEnvironment, Arguments);
+			Arguments.Add("-Xiwyu --mapping_file=" + IWYUMappingFile);
+			Arguments.Add("-Xiwyu --no_default_mappings");
+			//Arguments.Add("-Xiwyu --delayed_template_parsing");
+			Arguments.Add("-Xiwyu --prefix_header_includes=keep");
+			// Arguments.Add("-Xiwyu --transitive_includes_only");  // Since we are building headers separately (not together with their cpp) we don't need this
+			// Arguments.Add("-Xiwyu --no_check_matching_header");
+			Arguments.Add("-Xiwyu --cxx17ns");
 
-			if (ShouldUseLibcxx(CompileEnvironment.Architecture))
+			Arguments.Add("-Wno-delete-incomplete");
+			Arguments.Add("-Wno-unknown-warning-option");
+
+			if (TargetPlatform == UnrealTargetPlatform.Win64)
 			{
-				Arguments.Add("-nostdinc++");
-				Arguments.Add(GetSystemIncludePathArgument(DirectoryReference.Combine(Unreal.EngineSourceDirectory, "ThirdParty", "Unix", "LibCxx", "include"), CompileEnvironment));
-				Arguments.Add(GetSystemIncludePathArgument(DirectoryReference.Combine(Unreal.EngineSourceDirectory, "ThirdParty", "Unix", "LibCxx", "include", "c++", "v1"), CompileEnvironment));
+				Arguments.Add("--driver-mode=cl");
+				VCToolChain vcToolChain = (VCToolChain)InnerToolChain;
+				vcToolChain.AppendCLArguments_GlobalIWYU(CompileEnvironment, Arguments);
 			}
+			else
+			{
+				base.GetCompileArguments_Global(CompileEnvironment, Arguments);
 
-			Arguments.Add("-fno-math-errno");               // do not assume that math ops have side effects
+				if (TargetPlatform == UnrealTargetPlatform.Linux)
+				{
+					if (ShouldUseLibcxx(CompileEnvironment.Architecture))
+					{
+						Arguments.Add("-nostdinc++");
+						Arguments.Add(GetSystemIncludePathArgument(DirectoryReference.Combine(Unreal.EngineSourceDirectory, "ThirdParty", "Unix", "LibCxx", "include"), CompileEnvironment));
+						Arguments.Add(GetSystemIncludePathArgument(DirectoryReference.Combine(Unreal.EngineSourceDirectory, "ThirdParty", "Unix", "LibCxx", "include", "c++", "v1"), CompileEnvironment));
+					}
+
+					Arguments.Add("-fno-math-errno");               // do not assume that math ops have side effects
 
-			Arguments.Add(GetRTTIFlag(CompileEnvironment)); // flag for run-time type info
+					Arguments.Add(GetRTTIFlag(CompileEnvironment)); // flag for run-time type info
+
+					if (true)//CrossCompiling())
+					{
+						Arguments.Add($"-target {CompileEnvironment.Architecture.LinuxName}");        // Set target triple
+						Arguments.AddRange(CrossCompilingArguments);
+					}
 
-			if (true)//CrossCompiling())
+					if (CompileEnvironment.bHideSymbolsByDefault)
+					{
+						Arguments.Add("-fvisibility-ms-compat");
+						Arguments.Add("-fvisibility-inlines-hidden");
+					}
+				}
+			}
+		}
+
+		protected override void GetCompileArguments_WarningsAndErrors(CppCompileEnvironment CompileEnvironment, List<string> Arguments)
+		{
+			if (TargetPlatform == UnrealTargetPlatform.Linux)
 			{
-				Arguments.Add($"-target {CompileEnvironment.Architecture.LinuxName}");        // Set target triple
-				Arguments.AddRange(CrossCompilingArguments);
+				CompileEnvironment.CppCompileWarnings.ShadowVariableWarningLevel = WarningLevel.Off;
+				base.GetCompileArguments_WarningsAndErrors(CompileEnvironment, Arguments);
+				Arguments.Add("-Wno-undefined-bool-conversion");
+				Arguments.Add("-Wno-deprecated-anon-enum-enum-conversion");
+				Arguments.Add("-Wno-ambiguous-reversed-operator");
+				ClangWarnings.GetHeaderDisabledWarnings(Arguments);
 			}
+		}
+
+		// IWYU can't link
+		public override FileItem LinkFiles(LinkEnvironment LinkEnvironment, bool bBuildImportLibraryOnly, IActionGraphBuilder Graph)
+		{
+			throw new BuildException("Unable to link with IWYU toolchain.");
+		}
 
-			if (CompileEnvironment.bHideSymbolsByDefault)
+		protected override void GetCompileArguments_H(CppCompileEnvironment CompileEnvironment, List<string> Arguments)
+		{
+			if (TargetPlatform == UnrealTargetPlatform.Linux)
 			{
-				Arguments.Add("-fvisibility-ms-compat");
-				Arguments.Add("-fvisibility-inlines-hidden");
+				base.GetCompileArguments_H(CompileEnvironment, Arguments);
+			}
+			else
+			{
+				((VCToolChain)InnerToolChain).AppendCLArguments_H(CompileEnvironment, Arguments);
 			}
 		}
 
-		protected override void GetCompileArguments_WarningsAndErrors(CppCompileEnvironment CompileEnvironment, List<string> Arguments)
+		protected override void GetCompileArguments_CPP(CppCompileEnvironment CompileEnvironment, List<string> Arguments)
 		{
-			CompileEnvironment.CppCompileWarnings.ShadowVariableWarningLevel = WarningLevel.Off;
-			base.GetCompileArguments_WarningsAndErrors(CompileEnvironment, Arguments);
-			Arguments.Add("-Wno-undefined-bool-conversion");
-			Arguments.Add("-Wno-deprecated-anon-enum-enum-conversion");
-			Arguments.Add("-Wno-ambiguous-reversed-operator");
-			ClangWarnings.GetHeaderDisabledWarnings(Arguments);
+			if (TargetPlatform == UnrealTargetPlatform.Linux)
+			{
+				base.GetCompileArguments_CPP(CompileEnvironment, Arguments);
+			}
+			else
+			{
+				((VCToolChain)InnerToolChain).AppendCLArguments_CPP(CompileEnvironment, Arguments);
+			}
 		}
 
-		// Skip ISPC headers
-		protected override CPPOutput GenerateISPCHeaders(CppCompileEnvironment CompileEnvironment, IEnumerable<FileItem> InputFiles, DirectoryReference OutputDir, IActionGraphBuilder Graph)
+		protected override void GetCompileArguments_C(CppCompileEnvironment CompileEnvironment, List<string> Arguments)
 		{
-			return new CPPOutput();
+			if (TargetPlatform == UnrealTargetPlatform.Linux)
+			{
+				base.GetCompileArguments_C(CompileEnvironment, Arguments);
+			}
+			else
+			{
+				((VCToolChain)InnerToolChain).AppendCLArguments_C(CompileEnvironment, Arguments);
+			}
 		}
 
-		protected override CPPOutput CompileISPCFiles(CppCompileEnvironment CompileEnvironment, IEnumerable<FileItem> InputFiles, DirectoryReference OutputDir, IActionGraphBuilder Graph)
+		protected override string GetForceIncludeFileArgument(FileReference ForceIncludeFile, CppCompileEnvironment CompileEnvironment)
 		{
-			return new CPPOutput();
+			if (TargetPlatform == UnrealTargetPlatform.Linux)
+			{
+				return GetForceIncludeFileArgument(ForceIncludeFile, CompileEnvironment);
+			}
+			else
+			{
+				return $"/FI\"{NormalizeCommandLinePath(ForceIncludeFile, CompileEnvironment.RootPaths)}\"";
+			}
 		}
 
-		// IWYU can't link
-		public override FileItem LinkFiles(LinkEnvironment LinkEnvironment, bool bBuildImportLibraryOnly, IActionGraphBuilder Graph)
+		protected override FileItem GetCompileArguments_FileType(CppCompileEnvironment CompileEnvironment, FileItem SourceFile, DirectoryReference OutputDir, List<string> Arguments, Action CompileAction, CPPOutput CompileResult)
 		{
-			throw new BuildException("Unable to link with IWYU toolchain.");
+			FileItem fileItem = base.GetCompileArguments_FileType(CompileEnvironment, SourceFile, OutputDir, Arguments, CompileAction, CompileResult);
+
+			Arguments.Add("-Xiwyu --write_json_path=\"" + fileItem.AbsolutePath.Replace('\\', '/') + "\" ");
+			if (SourceFile.HasExtension(".cpp"))
+			{
+				List<FileItem>? InlinedFiles;
+				if (CompileEnvironment.FileInlineGenCPPMap.TryGetValue(SourceFile, out InlinedFiles))
+				{
+					foreach (FileItem InlinedFile in InlinedFiles)
+					{
+						Arguments.Add($"-Xiwyu --check_also={InlinedFile.FullName.Replace('\\', '/')} ");
+					}
+				}
+			}
+			else
+			{
+				// We want to build header files first
+				CompileAction.bIsHighPriority = true;
+			}
+			return fileItem;
 		}
+
+		protected override string GetDepencenciesListFileArgument(FileItem DependencyListFile, CppCompileEnvironment CompileEnvironment)
+		{
+			string result = base.GetDepencenciesListFileArgument(DependencyListFile, CompileEnvironment);
+
+			if (TargetPlatform == UnrealTargetPlatform.Win64)
+			{
+				result = result.Replace("-MD", "/clang:-MD").Replace("-MF", "/clang:-MF");
+			}
+			return result;
+		}
+
+		public override void SetEnvironmentVariables() => InnerToolChain.SetEnvironmentVariables();
+		public override void GenerateTypeLibraryHeader(CppCompileEnvironment CompileEnvironment, ModuleRules.TypeLibrary TypeLibrary, FileReference OutputFile, FileReference? OutputHeader, IActionGraphBuilder Graph) => InnerToolChain.GenerateTypeLibraryHeader(CompileEnvironment, TypeLibrary, OutputFile, OutputHeader, Graph);
+		public override List<string> GetISPCCompileTargets(UnrealTargetPlatform Platform, UnrealArch Arch) => InnerToolChain.GetISPCCompileTargets(Platform, Arch);
+		public override string GetISPCOSTarget(UnrealTargetPlatform Platform) => InnerToolChain.GetISPCOSTarget(Platform);
+		public override string GetISPCArchTarget(UnrealTargetPlatform Platform, UnrealArch Arch) => InnerToolChain.GetISPCArchTarget(Platform, Arch);
+		public override string? GetISPCCpuTarget(UnrealTargetPlatform Platform) => InnerToolChain.GetISPCCpuTarget(Platform);
+		public override string GetISPCObjectFileFormat(UnrealTargetPlatform Platform) => InnerToolChain.GetISPCObjectFileFormat(Platform);
+		public override string GetISPCObjectFileSuffix(UnrealTargetPlatform Platform) => InnerToolChain.GetISPCObjectFileSuffix(Platform);
 	}
-}
\ No newline at end of file
+}
