Add UE_FORCEINLINE_HINT which can be conditionally defined to inline by targets.

Removing all FORCEINLINEs from the code base causes perf regressions when we're not building with PGO or LTO (eg. the Editor).
The purpose of this is to leave critical functions as FORCEINLINE when we're not using LTO/PGO, but have them be inline when we are so the compiler can make the best decisions for perf/size.
This will get us back the other 50% of the total program size savings we see when not using FORCEINLINE without the perf hit.

The intent of this macro is for it to be an advanced feature primarily used within Core as that's where most the program size cost comes from.
This isn't meant to replace FORCEINLINE everwhere.

#rnx

diff --git a/Engine/Source/Programs/UnrealBuildTool/Configuration/ReadOnlyTargetRules.cs b/Engine/Source/Programs/UnrealBuildTool/Configuration/ReadOnlyTargetRules.cs
--- a/Engine/Source/Programs/UnrealBuildTool/Configuration/ReadOnlyTargetRules.cs
+++ b/Engine/Source/Programs/UnrealBuildTool/Configuration/ReadOnlyTargetRules.cs
@@ -265,6 +265,8 @@ namespace UnrealBuildTool
 
 		public OptimizationMode OptimizationLevel => Inner.OptimizationLevel;
 
+		public bool bDefineForceInlineHintToInline => Inner.bDefineForceInlineHintToInline;
+
 		public FPSemanticsMode FPSemantics => Inner.FPSemantics;
 
 		public bool bRetainFramePointers => Inner.bRetainFramePointers;
diff --git a/Engine/Source/Programs/UnrealBuildTool/Configuration/TargetRules.cs b/Engine/Source/Programs/UnrealBuildTool/Configuration/TargetRules.cs
--- a/Engine/Source/Programs/UnrealBuildTool/Configuration/TargetRules.cs
+++ b/Engine/Source/Programs/UnrealBuildTool/Configuration/TargetRules.cs
@@ -1494,6 +1494,19 @@ namespace UnrealBuildTool
 		[CommandLine("-OptimizationLevel=")]
 		public OptimizationMode OptimizationLevel { get; set; } = OptimizationMode.Speed;
 
+		/// <summary>
+		/// By default UE_FORCEINLINE_HINT is defined to FORCEINLINE. If this flag is true it will instead be defined to 'inline'.
+		/// This flag should only be true for targets where code size is a concern and/or is being compiled with PGO or LTO.
+		/// Note that if this flag is true when PGOOptimize is enabled, it should also be true when PGOProfile is enabled.
+		/// </summary>
+		[RequiresUniqueBuildEnvironment]
+		public bool bDefineForceInlineHintToInline
+		{
+			get => bDefineForceInlineHintToInlinePrivate ?? (bPGOOptimize || bPGOProfile || bAllowLTCG);
+			set => bDefineForceInlineHintToInlinePrivate = value;
+		}
+		private bool? bDefineForceInlineHintToInlinePrivate;
+
 		/// <summary>
 		/// Allows setting the FP semantics.
 		/// </summary>
diff --git a/Engine/Source/Programs/UnrealBuildTool/Configuration/UEBuildTarget.cs b/Engine/Source/Programs/UnrealBuildTool/Configuration/UEBuildTarget.cs
--- a/Engine/Source/Programs/UnrealBuildTool/Configuration/UEBuildTarget.cs
+++ b/Engine/Source/Programs/UnrealBuildTool/Configuration/UEBuildTarget.cs
@@ -5984,6 +5984,15 @@ namespace UnrealBuildTool
 				GlobalCompileEnvironment.Definitions.Add("WITH_PROCESS_PRIORITY_CONTROL=0");
 			}
 
+			if (Rules.bDefineForceInlineHintToInline)
+			{
+				GlobalCompileEnvironment.Definitions.Add("UE_DEFINE_FORCEINLINE_HINT_TO_INLINE=1");
+			}
+			else
+			{
+				GlobalCompileEnvironment.Definitions.Add("UE_DEFINE_FORCEINLINE_HINT_TO_INLINE=0");
+			}
+
 			// Define the custom config if specified, this should not be set in an editor build
 			if (!bUseSharedBuildEnvironment && !String.IsNullOrEmpty(Rules.CustomConfig))
 			{
diff --git a/Engine/Source/Runtime/Core/Public/HAL/Platform.h b/Engine/Source/Runtime/Core/Public/HAL/Platform.h
--- a/Engine/Source/Runtime/Core/Public/HAL/Platform.h
+++ b/Engine/Source/Runtime/Core/Public/HAL/Platform.h
@@ -703,6 +703,21 @@
 	#define RESTRICT __restrict						/* no alias hint */
 #endif
 
+#ifndef UE_DEFINE_FORCEINLINE_HINT_TO_INLINE
+	#define UE_DEFINE_FORCEINLINE_HINT_TO_INLINE 0
+#endif
+
+// A configurable alternative to FORCEINLINE than can be configured per target.
+// The main use of this is for core functionality that is performance sensitive but has a large code size cost.
+// This allows keeping code force inlined by default to improve performance in cases where it wouldn't normally be inlined,
+// but also allow disabling it for targets where the compiler can make more informed decisions such as when using PGO.
+// This ultimately should result in a better size/speed trade off in those cases.
+#if UE_DEFINE_FORCEINLINE_HINT_TO_INLINE
+	#define UE_FORCEINLINE_HINT inline
+#else
+	#define UE_FORCEINLINE_HINT FORCEINLINE
+#endif
+
 // UE_REWRITE should be placed on functions which rewrite a function call as another expression. Like a macro,
 // but without using the preprocessor.
 //
