UnrealBuildTool: Invalidate makefile if config.ini files are updated

#rnx

diff --git a/Engine/Source/Programs/UnrealBuildTool/Configuration/UEBuildTarget.cs b/Engine/Source/Programs/UnrealBuildTool/Configuration/UEBuildTarget.cs
--- a/Engine/Source/Programs/UnrealBuildTool/Configuration/UEBuildTarget.cs
+++ b/Engine/Source/Programs/UnrealBuildTool/Configuration/UEBuildTarget.cs
@@ -3013,6 +3013,14 @@ namespace UnrealBuildTool
 			}
 			Makefile.ExternalDependencies.UnionWith(Makefile.PluginFiles);
 
+			// Track config files
+			Makefile.ExternalDependencies.UnionWith(ConfigCache.ReadHierarchy(ConfigHierarchyType.Engine, Rules.ProjectFile?.Directory, Rules.Platform).FileReferences.Select(FileItem.GetItemByFileReference));
+			Makefile.ExternalDependencies.UnionWith(ConfigCache.ReadHierarchy(ConfigHierarchyType.Game, Rules.ProjectFile?.Directory, Rules.Platform).FileReferences.Select(FileItem.GetItemByFileReference));
+			if (Rules.Type == TargetType.Editor)
+			{
+				Makefile.ExternalDependencies.UnionWith(ConfigCache.ReadHierarchy(ConfigHierarchyType.Editor, Rules.ProjectFile?.Directory, Rules.Platform).FileReferences.Select(FileItem.GetItemByFileReference));
+			}
+
 			// Also track the version file
 			Makefile.ExternalDependencies.Add(FileItem.GetItemByFileReference(BuildVersion.GetDefaultFileName()));
 
diff --git a/Engine/Source/Programs/UnrealBuildTool/System/ConfigCache.cs b/Engine/Source/Programs/UnrealBuildTool/System/ConfigCache.cs
--- a/Engine/Source/Programs/UnrealBuildTool/System/ConfigCache.cs
+++ b/Engine/Source/Programs/UnrealBuildTool/System/ConfigCache.cs
@@ -373,7 +373,9 @@ namespace UnrealBuildTool
 			{
 				// Find all the input files
 				List<ConfigFile> Files = new List<ConfigFile>();
-				foreach (FileReference IniFileName in ConfigHierarchy.EnumerateConfigFileLocations(Type, ProjectDir, Platform, CustomConfig, IncludePluginsForTargetType))
+				IEnumerable<FileReference> IniFileNames = ConfigHierarchy.EnumerateConfigFileLocations(Type, ProjectDir, Platform, CustomConfig, IncludePluginsForTargetType);
+				List<FileReference> ReadFiles = [];
+				foreach (FileReference IniFileName in IniFileNames)
 				{
 					ConfigFile? File;
 					Logger.LogInformation($"Trying to load: {IniFileName}");
@@ -381,6 +383,7 @@ namespace UnrealBuildTool
 					{
 						Logger.LogInformation($"    Found!");
 						Files.Add(File);
+						ReadFiles.Add(IniFileName);
 					}
 
 					if (HotfixDir != null)
@@ -391,6 +394,7 @@ namespace UnrealBuildTool
 						{
 							Logger.LogInformation($"    Hotfixed with '{HotfixFile}'");
 							Files.Add(File);
+							ReadFiles.Add(HotfixFile);
 						}
 					}
 				}
@@ -403,7 +407,7 @@ namespace UnrealBuildTool
 				}
 
 				// Create the hierarchy
-				return new ConfigHierarchy(Files);
+				return new ConfigHierarchy(Files, ReadFiles);
 			});
 		}
 
diff --git a/Engine/Source/Programs/UnrealBuildTool/System/ConfigHierarchy.cs b/Engine/Source/Programs/UnrealBuildTool/System/ConfigHierarchy.cs
--- a/Engine/Source/Programs/UnrealBuildTool/System/ConfigHierarchy.cs
+++ b/Engine/Source/Programs/UnrealBuildTool/System/ConfigHierarchy.cs
@@ -178,25 +178,32 @@ namespace UnrealBuildTool
 		/// <summary>
 		/// Array of 
 		/// </summary>
-		ConfigFile[] Files;
+		readonly ConfigFile[] Files;
+
+		/// <summary>
+		/// Set of file references
+		/// </summary>
+		public IEnumerable<FileReference> FileReferences { get; }
 
 		/// <summary>
 		/// Cache of requested config sections
 		/// </summary>
-		Dictionary<string, ConfigHierarchySection> NameToSection = new Dictionary<string, ConfigHierarchySection>(StringComparer.InvariantCultureIgnoreCase);
+		readonly Dictionary<string, ConfigHierarchySection> NameToSection = new(StringComparer.InvariantCultureIgnoreCase);
 
 		/// <summary>
 		/// Lock for NameToSection
 		/// </summary>
-		System.Threading.ReaderWriterLockSlim NameToSectionLock = new System.Threading.ReaderWriterLockSlim();
+		readonly System.Threading.ReaderWriterLockSlim NameToSectionLock = new();
 
 		/// <summary>
 		/// Construct a config hierarchy from the given files
 		/// </summary>
 		/// <param name="Files">Set of files to include (in order)</param>
-		public ConfigHierarchy(IEnumerable<ConfigFile> Files)
+		/// <param name="FileReferences">Set of file paths, for tracking</param>
+		public ConfigHierarchy(IEnumerable<ConfigFile> Files, IEnumerable<FileReference>? FileReferences = null)
 		{
-			this.Files = Files.ToArray();
+			this.Files = [.. Files];
+			this.FileReferences = FileReferences ?? [];
 		}
 
 		#region IDisposible
@@ -211,24 +218,7 @@ namespace UnrealBuildTool
 		/// Names of all sections in all config files
 		/// </summary>
 		/// <returns></returns>
-		public HashSet<string> SectionNames
-		{
-			get
-			{
-				HashSet<string> Result = new HashSet<string>();
-				foreach (ConfigFile File in Files)
-				{
-					foreach (string SectionName in File.SectionNames)
-					{
-						if (!Result.Contains(SectionName))
-						{
-							Result.Add(SectionName);
-						}
-					}
-				}
-				return Result;
-			}
-		}
+		public HashSet<string> SectionNames => [.. Files.SelectMany(x => x.SectionNames)];
 
 		/// <summary>
 		/// Finds a config section with the given name
