UnrealBuildTool: Various VFS path fixes

#rnx

diff --git a/Engine/Source/Programs/UnrealBuildTool/Platform/Linux/LinuxToolChain.cs b/Engine/Source/Programs/UnrealBuildTool/Platform/Linux/LinuxToolChain.cs
--- a/Engine/Source/Programs/UnrealBuildTool/Platform/Linux/LinuxToolChain.cs
+++ b/Engine/Source/Programs/UnrealBuildTool/Platform/Linux/LinuxToolChain.cs
@@ -787,15 +787,15 @@ namespace UnrealBuildTool
 			{
 				Arguments.Add($"-target {LinkEnvironment.Architecture.LinuxName}");        // Set target triple
 				DirectoryReference SysRootPath = LinuxInfo.BaseLinuxPath!;
-				Arguments.Add($"--sysroot=\"{NormalizeCommandLinePath(SysRootPath)}\"");
+				Arguments.Add($"--sysroot=\"{NormalizeCommandLinePath(SysRootPath, LinkEnvironment.RootPaths)}\"");
 
 				// Linking with the toolchain on linux appears to not search usr/
 				if (BuildHostPlatform.Current.Platform == UnrealTargetPlatform.Linux)
 				{
-					Arguments.Add($"-B\"{NormalizeCommandLinePath(DirectoryReference.Combine(SysRootPath, "usr", "lib"))}\"");
-					Arguments.Add($"-B\"{NormalizeCommandLinePath(DirectoryReference.Combine(SysRootPath, "usr", "lib64"))}\"");
-					Arguments.Add($"-L\"{NormalizeCommandLinePath(DirectoryReference.Combine(SysRootPath, "usr", "lib"))}\"");
-					Arguments.Add($"-L\"{NormalizeCommandLinePath(DirectoryReference.Combine(SysRootPath, "usr", "lib64"))}\"");
+					Arguments.Add($"-B\"{NormalizeCommandLinePath(DirectoryReference.Combine(SysRootPath, "usr", "lib"), LinkEnvironment.RootPaths)}\"");
+					Arguments.Add($"-B\"{NormalizeCommandLinePath(DirectoryReference.Combine(SysRootPath, "usr", "lib64"), LinkEnvironment.RootPaths)}\"");
+					Arguments.Add($"-L\"{NormalizeCommandLinePath(DirectoryReference.Combine(SysRootPath, "usr", "lib"), LinkEnvironment.RootPaths)}\"");
+					Arguments.Add($"-L\"{NormalizeCommandLinePath(DirectoryReference.Combine(SysRootPath, "usr", "lib64"), LinkEnvironment.RootPaths)}\"");
 				}
 			}
 
@@ -1091,18 +1091,18 @@ namespace UnrealBuildTool
 			{
 				if (InputFile.HasExtension(".ldscript"))
 				{
-					ResponseLines.Add($"--version-script=\"{NormalizeCommandLinePath(InputFile)}\"");
+					ResponseLines.Add($"--version-script=\"{NormalizeCommandLinePath(InputFile, LinkEnvironment.RootPaths)}\"");
 				}
 				else
 				{
-					ResponseLines.Add($"\"{NormalizeCommandLinePath(InputFile)}\"");
+					ResponseLines.Add($"\"{NormalizeCommandLinePath(InputFile, LinkEnvironment.RootPaths)}\"");
 				}
 				LinkAction.PrerequisiteItems.Add(InputFile);
 			}
 
 			if (LinkEnvironment.bIsBuildingDLL)
 			{
-				ResponseLines.Add(String.Format(" -soname=\"{0}\"", OutputFile.Location.GetFileName()));
+				ResponseLines.Add($"-soname=\"{OutputFile.Location.GetFileName()}\"");
 			}
 
 			// Start with the configured LibraryPaths and also add paths to any libraries that
diff --git a/Engine/Source/Programs/UnrealBuildTool/ToolChain/IWYUToolChain.cs b/Engine/Source/Programs/UnrealBuildTool/ToolChain/IWYUToolChain.cs
--- a/Engine/Source/Programs/UnrealBuildTool/ToolChain/IWYUToolChain.cs
+++ b/Engine/Source/Programs/UnrealBuildTool/ToolChain/IWYUToolChain.cs
@@ -53,7 +53,7 @@ namespace UnrealBuildTool
 				SourceFileMetadataCache MetadataCache = SourceFileMetadataCache.CreateHierarchy(null, Logger);
 				CppCompileEnvironment CrossCompileEnvironment = new CppCompileEnvironment(UnrealTargetPlatform.Linux, CppConfiguration.Development, new(LinuxPlatform.DefaultHostArchitecture), MetadataCache);
 
-				CrossCompilingArguments.Add($"--sysroot=\"{NormalizeCommandLinePath(BaseLinuxPath!)}\"");
+				CrossCompilingArguments.Add($"--sysroot=\"{NormalizeCommandLinePath(BaseLinuxPath!, CrossCompileEnvironment.RootPaths)}\"");
 
 				FileReference ClangPath = FileReference.Combine(BaseLinuxPath!, "bin", $"clang++{BuildHostPlatform.Current.BinarySuffix}");
 				ClangToolChainInfo CompilerToolChainInfo = new ClangToolChainInfo(BaseLinuxPath, ClangPath, null!, Logger);
