UnrealBuildTool: Allow /LinkRepro with uba

#rnx

diff --git a/Engine/Source/Programs/UnrealBuildTool/Platform/Windows/VCToolChain.cs b/Engine/Source/Programs/UnrealBuildTool/Platform/Windows/VCToolChain.cs
--- a/Engine/Source/Programs/UnrealBuildTool/Platform/Windows/VCToolChain.cs
+++ b/Engine/Source/Programs/UnrealBuildTool/Platform/Windows/VCToolChain.cs
@@ -3324,10 +3324,10 @@ namespace UnrealBuildTool
 				DirectoryReference LinkReproRoot = new DirectoryReference(Target.WindowsPlatform.LinkReproDir);
 				DirectoryReference LinkReproPath = DirectoryReference.Combine(LinkReproRoot, OutputFile.Name);
 				DirectoryReference.CreateDirectory(LinkReproPath);
-				LinkAction.CommandArguments = $"{LinkAction.CommandArguments} /LINKREPRO:{Utils.MakePathSafeToUseWithCommandLine(LinkReproPath.FullName)}";
+				LinkAction.CommandArguments = $"{LinkAction.CommandArguments} /LINKREPRO:{Utils.MakePathSafeToUseWithCommandLine(NormalizeCommandLinePath(LinkReproPath, RootPaths))}";
 
+				LinkAction.ArtifactMode = ArtifactMode.None;
 				LinkAction.bCanExecuteRemotely = false;
-				LinkAction.bCanExecuteInUBA = false;
 			}
 
 			if (Target.bAllowUbaCompression)
@@ -3724,6 +3724,12 @@ namespace UnrealBuildTool
 			GlobalCompileEnvironment.RootPaths[CppRootPathFolder.WinSDK] = EnvVars.WindowsSdkDir;
 			GlobalLinkEnvironment.RootPaths[CppRootPathFolder.WinSDK] = EnvVars.WindowsSdkDir;
 
+			if (Target.WindowsPlatform.LinkReproDir != null)
+			{
+				DirectoryReference LinkReproRoot = new(Target.WindowsPlatform.LinkReproDir);
+				GlobalLinkEnvironment.RootPaths.AddExtraPath(("LinkRepro", LinkReproRoot.FullName));
+			}
+
 			// This is used by uba to create one input key for all files under dirs
 			EpicGames.UBA.Utils.RegisterPathHash(EnvVars.WindowsSdkDir.FullName, EnvVars.WindowsSdkVersion.ToString());
 			EpicGames.UBA.Utils.RegisterPathHash(EnvVars.CompilerDir.FullName, EnvVars.CompilerVersion.ToString());
