UnrealBuildTool: Apple toolchain does not support --thinlto-jobs=all

#rnx

diff --git a/Engine/Source/Programs/UnrealBuildTool/ToolChain/AppleToolChain.cs b/Engine/Source/Programs/UnrealBuildTool/ToolChain/AppleToolChain.cs
--- a/Engine/Source/Programs/UnrealBuildTool/ToolChain/AppleToolChain.cs
+++ b/Engine/Source/Programs/UnrealBuildTool/ToolChain/AppleToolChain.cs
@@ -846,6 +846,19 @@ namespace UnrealBuildTool
 			return OutputFiles;
 		}
 
+		protected override void GetLinkArguments_Optimizations(LinkEnvironment linkEnvironment, List<string> arguments)
+		{
+			base.GetLinkArguments_Optimizations(linkEnvironment, arguments);
+
+			bool bLTOEnabled = Options.HasFlag(ClangToolChainOptions.EnableLinkTimeOptimization);
+			bool bThinLTOEnabled = Options.HasFlag(ClangToolChainOptions.EnableThinLTO);
+
+			if (bLTOEnabled && bThinLTOEnabled)
+			{
+				arguments.Remove("-Wl,--thinlto-jobs=all"); // unsupported
+			}
+		}
+
 		protected override void GetLinkArguments_Global(LinkEnvironment LinkEnvironment, List<string> Arguments)
 		{
 			base.GetLinkArguments_Global(LinkEnvironment, Arguments);
