UnrealBuildTool: Apple clang does not support thinlto cache

#rnx

diff --git a/Engine/Source/Programs/UnrealBuildTool/ToolChain/AppleToolChain.cs b/Engine/Source/Programs/UnrealBuildTool/ToolChain/AppleToolChain.cs
--- a/Engine/Source/Programs/UnrealBuildTool/ToolChain/AppleToolChain.cs
+++ b/Engine/Source/Programs/UnrealBuildTool/ToolChain/AppleToolChain.cs
@@ -856,6 +856,11 @@ namespace UnrealBuildTool
 			if (bLTOEnabled && bThinLTOEnabled)
 			{
 				arguments.Remove("-Wl,--thinlto-jobs=all"); // unsupported
+
+				if (linkEnvironment.ThinLTOCacheDirectory != null)
+				{
+					arguments.RemoveAll(x => x.StartsWith("-Wl,--thinlto-cache", StringComparison.Ordinal));
+				}
 			}
 		}
 
