UnrealBuildTool: Enable UBA with VFS for linking MSVC PGO

#rnx

diff --git a/Engine/Source/Programs/UnrealBuildTool/Platform/Windows/VCToolChain.cs b/Engine/Source/Programs/UnrealBuildTool/Platform/Windows/VCToolChain.cs
--- a/Engine/Source/Programs/UnrealBuildTool/Platform/Windows/VCToolChain.cs
+++ b/Engine/Source/Programs/UnrealBuildTool/Platform/Windows/VCToolChain.cs
@@ -2964,13 +2964,6 @@ namespace UnrealBuildTool
 		{
 			CppRootPaths RootPaths = LinkEnvironment.RootPaths;
 
-			// Remove when we fix so pgo linking run in uba
-			if (LinkEnvironment.bPGOOptimize || LinkEnvironment.bPGOProfile)
-			{
-				RootPaths = new(RootPaths);
-				RootPaths.bUseVfs = false;
-			}
-
 			if (LinkEnvironment.bIsBuildingDotNetAssembly)
 			{
 				return FileItem.GetItemByFileReference(LinkEnvironment.OutputFilePath);
@@ -3119,7 +3112,7 @@ namespace UnrealBuildTool
 			List<string> InputFileNames = new List<string>();
 			foreach (FileItem InputFile in LinkEnvironment.InputFiles)
 			{
-				InputFileNames.Add(String.Format("\"{0}\"", NormalizeCommandLinePath(InputFile, RootPaths)));
+				InputFileNames.Add($"\"{NormalizeCommandLinePath(InputFile, RootPaths)}\"");
 				PrerequisiteItems.Add(InputFile);
 			}
 
@@ -3127,19 +3120,18 @@ namespace UnrealBuildTool
 			{
 				foreach (FileReference Library in LinkEnvironment.Libraries)
 				{
-					InputFileNames.Add(String.Format("\"{0}\"", NormalizeCommandLinePath(Library, RootPaths)));
+					InputFileNames.Add($"\"{NormalizeCommandLinePath(Library, RootPaths)}\"");
 					PrerequisiteItems.Add(FileItem.GetItemByFileReference(Library));
 				}
 				foreach (string SystemLibrary in LinkEnvironment.SystemLibraries)
 				{
-					InputFileNames.Add(String.Format("\"{0}\"", SystemLibrary));
+					InputFileNames.Add($"\"{SystemLibrary}\"");
 				}
 
 				foreach (FileItem NatvisFile in LinkEnvironment.DebuggerVisualizerFiles)
 				{
 					PrerequisiteItems.Add(NatvisFile);
-					Arguments.Add(String.Format("/NATVIS:\"{0}\"",
-						NormalizeCommandLinePath(NatvisFile, RootPaths)));
+					Arguments.Add($"/NATVIS:\"{NormalizeCommandLinePath(NatvisFile, RootPaths)}\"");
 				}
 
 				if (Target.WindowsPlatform.ManifestFile != null && !LinkEnvironment.bIsBuildingDLL)
@@ -3205,7 +3197,7 @@ namespace UnrealBuildTool
 					{
 						FileReference MAPFilePath = FileReference.Combine(LinkEnvironment.OutputDirectory!, Path.GetFileNameWithoutExtension(OutputFile.AbsolutePath) + ".map");
 						FileItem MAPFile = FileItem.GetItemByFileReference(MAPFilePath);
-						Arguments.Add($"/MAP:\"{NormalizeCommandLinePath(MAPFilePath)}\"");
+						Arguments.Add($"/MAP:\"{NormalizeCommandLinePath(MAPFilePath, RootPaths)}\"");
 						ProducedItems.Add(MAPFile);
 
 						// Export a list of object file paths, so we can locate the object files referenced by the map file
@@ -3308,11 +3300,6 @@ namespace UnrealBuildTool
 				LinkAction.bCanExecuteRemotely = false;
 			}
 
-			if (LinkEnvironment.bPGOOptimize || LinkEnvironment.bPGOProfile)
-			{
-				LinkAction.bCanExecuteInUBA = false; // Disabled for now. Should revisit to see why it is not working
-			}
-
 			if (EnvVars.Compiler.IsClang() && Target.WindowsPlatform.bSampleBasedPGO && LinkEnvironment.bPGOOptimize)
 			{
 				LinkAction.PrerequisiteItems.Add(FileItem.GetItemByFileReference(GetClangLinkProfDataFilename(LinkEnvironment)));
