UnrealBuildTool: Fix VFS root path for thin-lto cache

#rnx

diff --git a/Engine/Source/Programs/UnrealBuildTool/ToolChain/UEToolChain.cs b/Engine/Source/Programs/UnrealBuildTool/ToolChain/UEToolChain.cs
--- a/Engine/Source/Programs/UnrealBuildTool/ToolChain/UEToolChain.cs
+++ b/Engine/Source/Programs/UnrealBuildTool/ToolChain/UEToolChain.cs
@@ -324,7 +324,7 @@ namespace UnrealBuildTool
 			GlobalCompileEnvironment.RootPaths[CppRootPathFolder.Root] = Unreal.RootDirectory;
 			GlobalLinkEnvironment.RootPaths[CppRootPathFolder.Root] = Unreal.RootDirectory;
 
-			DirectoryReference? thinLtoCache = DirectoryReference.FromString(GlobalLinkEnvironment.ThinLTOCacheDirectory);
+			DirectoryReference? thinLtoCache = DirectoryReference.FromString(Target.ThinLTOCacheDirectory);
 			if (thinLtoCache != null)
 			{
 				GlobalLinkEnvironment.RootPaths.AddExtraPath(("ThinLTOCache", thinLtoCache.FullName));
