Disable Dynamic Debugging on arm64ec as it's not currently supported

#rnx

diff --git a/Engine/Source/Programs/UnrealBuildTool/Platform/Windows/VCToolChain.cs b/Engine/Source/Programs/UnrealBuildTool/Platform/Windows/VCToolChain.cs
--- a/Engine/Source/Programs/UnrealBuildTool/Platform/Windows/VCToolChain.cs
+++ b/Engine/Source/Programs/UnrealBuildTool/Platform/Windows/VCToolChain.cs
@@ -65,6 +65,7 @@ namespace UnrealBuildTool
 			&& !(Target.WindowsPlatform.Environment != null && Target.WindowsPlatform.Environment.CompilerVersion < new VersionNumber(14, 44, 34918))
 			&& !(Target.bAllowLTCG || Target.bPGOOptimize || Target.bPGOProfile)
 			&& !(Target.Configuration == UnrealTargetConfiguration.Debug || Target.Configuration == UnrealTargetConfiguration.DebugGame)
+			&& !(Target.Architecture == UnrealArch.Arm64ec)
 			&& !(Target.WindowsPlatform.bEnableAddressSanitizer)
 			&& !(Target.bSupportEditAndContinue);
 
@@ -99,6 +100,10 @@ namespace UnrealBuildTool
 				{
 					Lines.Add($" * Target.{nameof(Target.Configuration)} '{Target.Configuration}' is Debug or DebugGame. Development, Test, or Shipping configuration is recommended when Dynamic Debugging is requested");
 				}
+				if (Target.Architecture == UnrealArch.Arm64ec)
+				{
+					Lines.Add($" * Target.{nameof(Target.Architecture)} '{Target.Architecture}' does not currently support Dynamic Debugging.");
+				}
 				if (Target.WindowsPlatform.bEnableAddressSanitizer)
 				{
 					Lines.Add($" * Target.{nameof(Target.WindowsPlatform)}.{nameof(Target.WindowsPlatform.bEnableAddressSanitizer)} is enabled");
