Improve merged modules report for unmergeable modules

#rnx

diff --git a/Engine/Source/Programs/UnrealBuildTool/Configuration/ReadOnlyTargetRules.cs b/Engine/Source/Programs/UnrealBuildTool/Configuration/ReadOnlyTargetRules.cs
--- a/Engine/Source/Programs/UnrealBuildTool/Configuration/ReadOnlyTargetRules.cs
+++ b/Engine/Source/Programs/UnrealBuildTool/Configuration/ReadOnlyTargetRules.cs
@@ -668,8 +668,8 @@ namespace UnrealBuildTool
 		/// <summary>
 		/// Report merged modules structure to the target so that the game may use this information.
 		/// </summary>
-		public virtual void ReportMergedModules(Dictionary<string, IEnumerable<ModuleRules>> ModuleGroups, Dictionary<string, IEnumerable<string>> PluginUsage, Dictionary<string, IEnumerable<string>> RootPluginUsage)
-			=> Inner.ReportMergedModules(ModuleGroups, PluginUsage, RootPluginUsage);
+		public virtual void ReportMergedModules(Dictionary<string, IEnumerable<ModuleRules>> ModuleGroups, Dictionary<string, Dictionary<string, List<string>>> RootPluginUsage)
+			=> Inner.ReportMergedModules(ModuleGroups, RootPluginUsage);
 
 		/// <summary>
 		/// Determines if the automation tests should be compiled based on the current configuration and optional forced settings.
diff --git a/Engine/Source/Programs/UnrealBuildTool/Configuration/TargetRules.cs b/Engine/Source/Programs/UnrealBuildTool/Configuration/TargetRules.cs
--- a/Engine/Source/Programs/UnrealBuildTool/Configuration/TargetRules.cs
+++ b/Engine/Source/Programs/UnrealBuildTool/Configuration/TargetRules.cs
@@ -2822,7 +2822,7 @@ namespace UnrealBuildTool
 		/// <summary>
 		/// Validate that the merged modules setup is appropriate
 		/// </summary>
-		public virtual void ReportMergedModules(Dictionary<string, IEnumerable<ModuleRules>> ModuleGroups, Dictionary<string, IEnumerable<string>> PluginUsage, Dictionary<string, IEnumerable<string>> RootPluginUsage) { }
+		public virtual void ReportMergedModules(Dictionary<string, IEnumerable<ModuleRules>> ModuleGroups, Dictionary<string, Dictionary<string, List<string>>> RootPluginUsage) { }
 
 		/// <summary>
 		/// Checks if a property has been set with the RequiresUniqueBuildEnvironmentAttribute, and is different from it's base
diff --git a/Engine/Source/Programs/UnrealBuildTool/Configuration/UEBuildTarget.cs b/Engine/Source/Programs/UnrealBuildTool/Configuration/UEBuildTarget.cs
--- a/Engine/Source/Programs/UnrealBuildTool/Configuration/UEBuildTarget.cs
+++ b/Engine/Source/Programs/UnrealBuildTool/Configuration/UEBuildTarget.cs
@@ -3616,7 +3616,11 @@ namespace UnrealBuildTool
 
 			bool CanPluginBeMerged(UEBuildPlugin plugin)
 			{
-				if (Rules.MergeOnlyGameFeatures)
+				if (Rules.MergePluginsExcluded.Contains(plugin.Name))
+				{
+					return false;
+				}
+				else if (Rules.MergeOnlyGameFeatures)
 				{
 					DirectoryReference gameFeaturesDir = DirectoryReference.Combine(ProjectDirectory, "Plugins", "GameFeatures");
 					return plugin.File.IsUnderDirectory(gameFeaturesDir);
@@ -3868,8 +3872,7 @@ namespace UnrealBuildTool
 				GetModulesDependenciesRecursively(allLeftoverModules.OfType<UEBuildModuleCPP>(), ref allLeftoverDependencies);
 
 				Dictionary<string, IEnumerable<ModuleRules>> moduleGroups = new();
-				Dictionary<string, IEnumerable<string>> pluginUsageMap = new();
-				Dictionary<string, IEnumerable<string>> rootPluginUsageMap = new();
+				Dictionary<string, Dictionary<string, List<string>>> rootPluginUsageMap = new();
 				if (Rules.bReportModuleMergingData)
 				{ 
 					// prepare the merged modules group data for validation by the game
@@ -3882,29 +3885,34 @@ namespace UnrealBuildTool
 							List<string> userNames = new();
 							UEBuildPlugin plugin = pluginCandidates.First();
 
-							if (globalPluginUserMap.ContainsKey(plugin))
+							void FindRootUsers(UEBuildPlugin plugin, ref Dictionary<string, List<string>> roots, List<string> path)
 							{
-								pluginUsageMap.Add(module.Name, globalPluginUserMap[plugin].Select(x => x.Name));
-							}
+								List<string> newPath = new(path);
+								newPath.Add(plugin.Name);
 
-							void FindRootUsers(UEBuildPlugin plugin, ref HashSet<UEBuildPlugin> roots)
-							{
 								if (globalPluginUserMap.ContainsKey(plugin))
 								{
 									foreach (UEBuildPlugin user in globalPluginUserMap[plugin])
 									{
-										FindRootUsers(user, ref roots);
+										FindRootUsers(user, ref roots, newPath);
 									}
 								}
-								else if (!roots.Contains(plugin))
+								else
 								{
-									roots.Add(plugin);
+									if (!roots.ContainsKey(plugin.Name))
+									{
+										roots.Add(plugin.Name, newPath);
+									}
+									else if (path.Count() < roots[plugin.Name].Count())
+									{
+										roots[plugin.Name] = newPath;
+									}
 								}
 							}
 
-							HashSet<UEBuildPlugin> roots = new();
-							FindRootUsers(plugin, ref roots);
-							rootPluginUsageMap.Add(module.Name, roots.Select(x => x.Name));
+							Dictionary<string, List<string>> roots = new();
+							FindRootUsers(plugin, ref roots, new());
+							rootPluginUsageMap.Add(module.Name, roots);
 						}
 					}
 				}
@@ -3926,7 +3934,7 @@ namespace UnrealBuildTool
 
 				if (Rules.bReportModuleMergingData)
 				{
-					Rules.ReportMergedModules(moduleGroups, pluginUsageMap, rootPluginUsageMap);
+					Rules.ReportMergedModules(moduleGroups, rootPluginUsageMap);
 				}
 			}
 
