UnrealBuildTool: Intellisense warnings

#rnx

diff --git a/Engine/Source/Programs/UnrealBuildTool/System/Action.cs b/Engine/Source/Programs/UnrealBuildTool/System/Action.cs
--- a/Engine/Source/Programs/UnrealBuildTool/System/Action.cs
+++ b/Engine/Source/Programs/UnrealBuildTool/System/Action.cs
@@ -4,7 +4,6 @@ using System;
 using System.Collections.Generic;
 using System.Diagnostics;
 using System.Linq;
-using System.Runtime.InteropServices.Marshalling;
 using System.Threading;
 using EpicGames.Core;
 using Microsoft.Extensions.Logging;
@@ -242,21 +241,21 @@ namespace UnrealBuildTool
 		/// <summary>
 		/// Every file this action depends on.  These files need to exist and be up to date in order for this action to even be considered
 		/// </summary>
-		public SortedSet<FileItem> PrerequisiteItems { get; set; } = new();
+		public SortedSet<FileItem> PrerequisiteItems { get; set; } = [];
 
 		/// <summary>
 		/// The files that this action produces after completing
 		/// </summary>
-		public SortedSet<FileItem> ProducedItems { get; set; } = new();
+		public SortedSet<FileItem> ProducedItems { get; set; } = [];
 
 		/// <summary>
 		/// Items that should be deleted before running this action
 		/// </summary>
-		public SortedSet<FileItem> DeleteItems { get; set; } = new();
+		public SortedSet<FileItem> DeleteItems { get; set; } = [];
 
 		/// <summary>
 		/// Root paths for this action (generally engine root project root, toolchain root, sdk root)
-        /// Order needs to be mainted, so a list is used
+        /// Order needs to be maintained, so a list is used
 		/// </summary>
 		public CppRootPaths RootPaths { get; set; } = new();
 
@@ -311,7 +310,7 @@ namespace UnrealBuildTool
 		public bool bCanExecuteRemotelyWithXGE { get; set; } = true;
 
 		/// <summary>
-		/// True if this action can be executed in box
+		/// True if this action can be executed in uba
 		/// </summary>
 		public bool bCanExecuteInUBA { get; set; } = true;
 
@@ -382,9 +381,9 @@ namespace UnrealBuildTool
 		{
 			ActionType = InOther.ActionType;
 			ArtifactMode = InOther.ArtifactMode;
-			PrerequisiteItems = new(InOther.PrerequisiteItems);
-			ProducedItems = new(InOther.ProducedItems);
-			DeleteItems = new(InOther.DeleteItems);
+			PrerequisiteItems = [.. InOther.PrerequisiteItems];
+			ProducedItems = [.. InOther.ProducedItems];
+			DeleteItems = [.. InOther.DeleteItems];
 			RootPaths = new(InOther.RootPaths);
 			DependencyListFile = InOther.DependencyListFile;
 			WorkingDirectory = InOther.WorkingDirectory;
@@ -483,7 +482,7 @@ namespace UnrealBuildTool
 		/// <param name="Object">The object to parse</param>
 		public static Action ImportJson(JsonObject Object)
 		{
-			Action Action = new Action(Object.GetEnumField<ActionType>("Type"));
+			Action Action = new(Object.GetEnumField<ActionType>("Type"));
 
 			if (Object.TryGetEnumField("ArtifactMode", out ArtifactMode ArtifactMode))
 			{
@@ -713,15 +712,13 @@ namespace UnrealBuildTool
 		}
 
 		/// <summary>
-		/// Determine if this Action ignores output conficts.
+		/// Determine if this Action ignores output conflicts.
 		/// </summary>
 		/// <param name="Action"> The Action to check</param>
 		/// <returns>True if conflicts are ignored, else false.</returns>
-		public static bool IgnoreConflicts(this IExternalAction Action)
-		{
-			return Action.ActionType == ActionType.WriteMetadata ||
-				Action.ActionType == ActionType.CreateAppBundle;
-		}
+		public static bool IgnoreConflicts(this IExternalAction Action) =>
+			Action.ActionType == ActionType.WriteMetadata ||
+			Action.ActionType == ActionType.CreateAppBundle;
 	}
 
 	/// <summary>
@@ -730,16 +727,10 @@ namespace UnrealBuildTool
 	class DefaultActionSerializer : ActionSerializerBase<Action>
 	{
 		/// <inheritdoc/>
-		public override Action Read(BinaryArchiveReader Reader)
-		{
-			return new Action(Reader);
-		}
+		public override Action Read(BinaryArchiveReader Reader) => new Action(Reader);
 
 		/// <inheritdoc/>
-		public override void Write(BinaryArchiveWriter Writer, Action Action)
-		{
-			Action.Write(Writer);
-		}
+		public override void Write(BinaryArchiveWriter Writer, Action Action) => Action.Write(Writer);
 	}
 
 	/// <summary>
@@ -782,7 +773,7 @@ namespace UnrealBuildTool
 		/// <summary>
 		/// If set, will be output whenever the group differs to the last executed action. Set when executing multiple targets at once.
 		/// </summary>
-		public SortedSet<string> GroupNames = new SortedSet<string>();
+		public SortedSet<string> GroupNames = [];
 
 		#region Wrapper implementation of IAction
 
diff --git a/Engine/Source/Programs/UnrealBuildTool/System/CppCompileEnvironment.cs b/Engine/Source/Programs/UnrealBuildTool/System/CppCompileEnvironment.cs
--- a/Engine/Source/Programs/UnrealBuildTool/System/CppCompileEnvironment.cs
+++ b/Engine/Source/Programs/UnrealBuildTool/System/CppCompileEnvironment.cs
@@ -6,7 +6,6 @@ using System.IO.Enumeration;
 using EpicGames.Core;
 using UnrealBuildBase;
 using System;
-using System.Diagnostics.CodeAnalysis;
 
 namespace UnrealBuildTool
 {
@@ -189,11 +188,11 @@ namespace UnrealBuildTool
 	/// </summary>
 	class CPPOutput
 	{
-		public List<FileItem> ObjectFiles = new List<FileItem>();
-		public List<FileItem> CompiledModuleInterfaces = new List<FileItem>();
-		public List<FileItem> GeneratedHeaderFiles = new List<FileItem>();
+		public List<FileItem> ObjectFiles = [];
+		public List<FileItem> CompiledModuleInterfaces = [];
+		public List<FileItem> GeneratedHeaderFiles = [];
 		public FileItem? PrecompiledHeaderFile = null;
-		public Dictionary<UnrealArch, FileItem> PerArchPrecompiledHeaderFiles = new();
+		public Dictionary<UnrealArch, FileItem> PerArchPrecompiledHeaderFiles = [];
 
 		public void Merge(CPPOutput Other, UnrealArch Architecture)
 		{
@@ -340,7 +339,7 @@ namespace UnrealBuildTool
 		/// <summary>
 		/// AutoRTFM external open -> closed mapping files
 		/// </summary>
-		public HashSet<FileItem> AutoRTFMExternalMappingFiles = new();
+		public HashSet<FileItem> AutoRTFMExternalMappingFiles = [];
 
 		/// <summary>
 		/// Whether we should run an extra LLVM verification pass after the AutoRTFM compiler pass.
@@ -427,27 +426,27 @@ namespace UnrealBuildTool
 		/// The static analyzer rulesets that should be used to filter warnings. This is only supported for MSVC.
 		/// See https://learn.microsoft.com/en-us/cpp/code-quality/using-rule-sets-to-specify-the-cpp-rules-to-run
 		/// </summary>
-		public HashSet<FileReference> StaticAnalyzerRulesets = new HashSet<FileReference>();
+		public HashSet<FileReference> StaticAnalyzerRulesets = [];
 
 		/// <summary>
 		/// The static analyzer checkers that should be enabled rather than the defaults. This is only supported for Clang.
 		/// </summary>
-		public HashSet<string> StaticAnalyzerCheckers = new HashSet<string>();
+		public HashSet<string> StaticAnalyzerCheckers = [];
 
 		/// <summary>
 		/// The static analyzer default checkers that should be disabled. Unused if StaticAnalyzerCheckers is populated. This is only supported for Clang.
 		/// </summary>
-		public HashSet<string> StaticAnalyzerDisabledCheckers = new HashSet<string>();
+		public HashSet<string> StaticAnalyzerDisabledCheckers = [];
 
 		/// <summary>
 		/// The static analyzer non-default checkers that should be enabled. Unused if StaticAnalyzerCheckers is populated. This is only supported for Clang.
 		/// </summary>
-		public HashSet<string> StaticAnalyzerAdditionalCheckers = new HashSet<string>();
+		public HashSet<string> StaticAnalyzerAdditionalCheckers = [];
 
 		/// <summary>
 		/// The PVS Studio analysis warnings that should be disabled.
 		/// </summary>
-		public HashSet<string> StaticAnalyzerPVSDisabledErrors = new();
+		public HashSet<string> StaticAnalyzerPVSDisabledErrors = [];
 
 		/// <summary>
 		/// True if compiler optimizations should be enabled. This setting is distinct from the configuration (see CPPTargetConfiguration).
@@ -622,22 +621,22 @@ namespace UnrealBuildTool
 		/// <summary>
 		/// List of header files to force include
 		/// </summary>
-		public List<FileItem> ForceIncludeFiles = new List<FileItem>();
+		public List<FileItem> ForceIncludeFiles = [];
 
 		/// <summary>
 		/// List of files that need to be up to date before compile can proceed
 		/// </summary>
-		public List<FileItem> AdditionalPrerequisites = new List<FileItem>();
+		public List<FileItem> AdditionalPrerequisites = [];
 
 		/// <summary>
 		/// A dictionary of the source file items and the inlined gen.cpp files contained in it
 		/// </summary>
-		public Dictionary<FileItem, List<FileItem>> FileInlineGenCPPMap = new();
+		public Dictionary<FileItem, List<FileItem>> FileInlineGenCPPMap = [];
 		
 		/// <summary>
 		/// Non-default naming conventions that generated CPP files can have (this uses a filter to check the file name for matches) 
 		/// </summary>
-		public List<string> ExtraGeneratedCPPFileTypes = new List<string>();
+		public List<string> ExtraGeneratedCPPFileTypes = [];
 
 		/// <summary>
 		/// FileItems with colliding names. (Which means they would overwrite each other in intermediate folder
@@ -647,12 +646,12 @@ namespace UnrealBuildTool
 		/// <summary>
 		/// The C++ preprocessor definitions to use.
 		/// </summary>
-		public List<string> Definitions = new List<string>();
+		public List<string> Definitions = [];
 
 		/// <summary>
 		/// Additional response files that will be used by main response file
 		/// </summary>
-		public List<FileItem> AdditionalResponseFiles = new();
+		public List<FileItem> AdditionalResponseFiles = [];
 
 		/// <summary>
 		/// Whether the compile environment has a response file in AdditionalResponseFiles that contains global compiler arguments.
@@ -667,7 +666,7 @@ namespace UnrealBuildTool
 		/// <summary>
 		/// A list of additional frameworks whose include paths are needed.
 		/// </summary>
-		public List<UEBuildFramework> AdditionalFrameworks = new List<UEBuildFramework>();
+		public List<UEBuildFramework> AdditionalFrameworks = [];
 
 		/// <summary>
 		/// A dictionary of PCH files for multiple architectures
@@ -797,7 +796,7 @@ namespace UnrealBuildTool
 		/// File to use as an LLVM FPass Plugin:
 		///    'clang -fpass-plugin=[FPassPlugin]' 
 		/// </summary>
-		public List<string> FPassPlugins = new();
+		public List<string> FPassPlugins = [];
 
 		/// <summary>
 		/// Default constructor.
@@ -809,12 +808,12 @@ namespace UnrealBuildTool
 			this.Architectures = Architectures;
 			this.MetadataCache = MetadataCache;
 			RootPaths = new();
-			SharedPCHs = new List<PrecompiledHeaderTemplate>();
-			UserIncludePaths = new HashSet<DirectoryReference>();
-			SystemIncludePaths = new HashSet<DirectoryReference>();
-			SharedUserIncludePaths = new HashSet<DirectoryReference>();
-			SharedSystemIncludePaths = new HashSet<DirectoryReference>();
-			ModuleInterfacePaths = new HashSet<DirectoryReference>();
+			SharedPCHs = [];
+			UserIncludePaths = [];
+			SystemIncludePaths = [];
+			SharedUserIncludePaths = [];
+			SharedSystemIncludePaths = [];
+			ModuleInterfacePaths = [];
 			CppCompileWarnings = new CppCompileWarnings();
 		}
 
@@ -851,17 +850,17 @@ namespace UnrealBuildTool
 			CppCompileWarnings = CppCompileWarnings.CreateShallowCopy(Other.CppCompileWarnings);
 			bWarningsAsErrors = Other.bWarningsAsErrors;
 			bDisableStaticAnalysis = Other.bDisableStaticAnalysis;
-			StaticAnalyzerCheckers = new HashSet<string>(Other.StaticAnalyzerCheckers);
-			StaticAnalyzerDisabledCheckers = new HashSet<string>(Other.StaticAnalyzerDisabledCheckers);
-			StaticAnalyzerAdditionalCheckers = new HashSet<string>(Other.StaticAnalyzerAdditionalCheckers);
-			StaticAnalyzerPVSDisabledErrors = new HashSet<string>(Other.StaticAnalyzerPVSDisabledErrors);
+			StaticAnalyzerCheckers = [.. Other.StaticAnalyzerCheckers];
+			StaticAnalyzerDisabledCheckers = [.. Other.StaticAnalyzerDisabledCheckers];
+			StaticAnalyzerAdditionalCheckers = [.. Other.StaticAnalyzerAdditionalCheckers];
+			StaticAnalyzerPVSDisabledErrors = [.. Other.StaticAnalyzerPVSDisabledErrors];
 			bStaticAnalyzerExtensions = Other.bStaticAnalyzerExtensions;
-			StaticAnalyzerRulesets = new HashSet<FileReference>(Other.StaticAnalyzerRulesets);
+			StaticAnalyzerRulesets = [.. Other.StaticAnalyzerRulesets];
 			bOptimizeCode = Other.bOptimizeCode;
 			bUseAutoRTFMCompiler = Other.bUseAutoRTFMCompiler;
 			bDisableAutoRTFMInstrumentation = Other.bDisableAutoRTFMInstrumentation;
 			bEnableAutoRTFMVerification = Other.bEnableAutoRTFMVerification;
-			AutoRTFMExternalMappingFiles = new HashSet<FileItem>(Other.AutoRTFMExternalMappingFiles);
+			AutoRTFMExternalMappingFiles = [.. Other.AutoRTFMExternalMappingFiles];
 			bAutoRTFMVerify = Other.bAutoRTFMVerify;
 			bAutoRTFMClosedStaticLinkage = Other.bAutoRTFMClosedStaticLinkage;
 			bCodeCoverage = Other.bCodeCoverage;
@@ -891,17 +890,17 @@ namespace UnrealBuildTool
 			bAllowRemotelyCompiledPCHs = Other.bAllowRemotelyCompiledPCHs;
 			bHasSharedResponseFile = Other.bHasSharedResponseFile;
 			bUseHeaderUnitsForPch = Other.bUseHeaderUnitsForPch;
-			UserIncludePaths = new HashSet<DirectoryReference>(Other.UserIncludePaths);
-			SystemIncludePaths = new HashSet<DirectoryReference>(Other.SystemIncludePaths);
-			SharedUserIncludePaths = new HashSet<DirectoryReference>(Other.SharedUserIncludePaths);
-			SharedSystemIncludePaths = new HashSet<DirectoryReference>(Other.SharedSystemIncludePaths);
-			ModuleInterfacePaths = new HashSet<DirectoryReference>(Other.ModuleInterfacePaths);
+			UserIncludePaths = [.. Other.UserIncludePaths];
+			SystemIncludePaths = [.. Other.SystemIncludePaths];
+			SharedUserIncludePaths = [.. Other.SharedUserIncludePaths];
+			SharedSystemIncludePaths = [.. Other.SharedSystemIncludePaths];
+			ModuleInterfacePaths = [.. Other.ModuleInterfacePaths];
 			bCheckSystemHeadersForModification = Other.bCheckSystemHeadersForModification;
 			ForceIncludeFiles.AddRange(Other.ForceIncludeFiles);
 			AdditionalPrerequisites.AddRange(Other.AdditionalPrerequisites);
 			ExtraGeneratedCPPFileTypes = Other.ExtraGeneratedCPPFileTypes;
 			CollidingNames = Other.CollidingNames;
-			FileInlineGenCPPMap = new Dictionary<FileItem, List<FileItem>>(Other.FileInlineGenCPPMap);
+			FileInlineGenCPPMap = new(Other.FileInlineGenCPPMap);
 			Definitions.AddRange(Other.Definitions);
 			AdditionalResponseFiles.AddRange(Other.AdditionalResponseFiles);
 			AdditionalArguments = Other.AdditionalArguments;
diff --git a/Engine/Source/Programs/UnrealBuildTool/System/LinkEnvironment.cs b/Engine/Source/Programs/UnrealBuildTool/System/LinkEnvironment.cs
--- a/Engine/Source/Programs/UnrealBuildTool/System/LinkEnvironment.cs
+++ b/Engine/Source/Programs/UnrealBuildTool/System/LinkEnvironment.cs
@@ -61,74 +61,67 @@ namespace UnrealBuildTool
 		/// <summary>
 		/// The file path for the executable file that is output by the linker.
 		/// </summary>
-		public List<FileReference> OutputFilePaths = new List<FileReference>();
+		public List<FileReference> OutputFilePaths = [];
 
 		/// <summary>
 		/// Returns the OutputFilePath is there is only one entry in OutputFilePaths
 		/// </summary>
-		public FileReference OutputFilePath
-		{
-			get
-			{
-				if (OutputFilePaths.Count != 1)
-				{
-					throw new BuildException("Attempted to use LinkEnvironmentConfiguration.OutputFilePath property, but there are multiple (or no) OutputFilePaths. You need to handle multiple in the code that called this (size = {0})", OutputFilePaths.Count);
-				}
-				return OutputFilePaths[0];
-			}
-		}
+		public FileReference OutputFilePath =>
+			OutputFilePaths.Count != 1
+				? throw new BuildException("Attempted to use LinkEnvironmentConfiguration.OutputFilePath property, but there are multiple (or no) OutputFilePaths. You need to handle multiple in the code that called this (size = {0})", OutputFilePaths.Count)
+				: OutputFilePaths[0];
 
 		/// <summary>
 		/// List of libraries to link in
 		/// </summary>
-		public List<FileReference> Libraries = new List<FileReference>();
+		public List<FileReference> Libraries = [];
 
 		/// <summary>
 		/// A list of additional libraries to link in (using the list of library search paths)
 		/// </summary>
-		public List<string> SystemLibraries = new List<string>();
+		public List<string> SystemLibraries = [];
 
 		/// <summary>
 		/// A list of the paths used to find libraries.
 		/// </summary>
-		public List<DirectoryReference> SystemLibraryPaths = new List<DirectoryReference>();
+		public List<DirectoryReference> SystemLibraryPaths = [];
 
 		/// <summary>
 		/// A list of libraries to exclude from linking.
 		/// </summary>
-		public List<string> ExcludedLibraries = new List<string>();
+		public List<string> ExcludedLibraries = [];
 
 		/// <summary>
 		/// Paths to add as search paths for runtime libraries
 		/// </summary>
-		public List<string> RuntimeLibraryPaths = new List<string>();
+		public List<string> RuntimeLibraryPaths = [];
 
 		/// <summary>
 		/// A list of additional frameworks to link in.
 		/// </summary>
-		public List<UEBuildFramework> AdditionalFrameworks = new List<UEBuildFramework>();
+		public List<UEBuildFramework> AdditionalFrameworks = [];
 
 		/// <summary>
 		/// The iOS/Mac frameworks to link in
 		/// </summary>
-		public List<string> Frameworks = new List<string>();
-		public List<string> WeakFrameworks = new List<string>();
+		public List<string> Frameworks = [];
+		public List<string> WeakFrameworks = [];
 
 		/// <summary>
 		/// iOS/Mac resources that should be copied to the app bundle
 		/// </summary>
-		public List<UEBuildBundleResource> AdditionalBundleResources = new List<UEBuildBundleResource>();
+		public List<UEBuildBundleResource> AdditionalBundleResources = [];
 
 		/// <summary>
 		/// A list of the dynamically linked libraries that shouldn't be loaded until they are first called
 		/// into.
 		/// </summary>
-		public List<string> DelayLoadDLLs = new List<string>();
+		public List<string> DelayLoadDLLs = [];
 
 		/// <summary>
 		/// Per-architecture lists of dependencies for linking to ignore (useful when building for multiple architectures, and a lib only is needed for one architecture), it's up to the Toolchain to use this
 		/// </summary>
-		public Dictionary<string, HashSet<UnrealArch>> DependenciesToSkipPerArchitecture = new();
+		public Dictionary<string, HashSet<UnrealArch>> DependenciesToSkipPerArchitecture = [];
 
 		/// <summary>
 		/// Additional arguments to pass to the linker.
@@ -334,27 +327,27 @@ namespace UnrealBuildTool
 		/// <summary>
 		/// A list of the object files to be linked.
 		/// </summary>
-		public List<FileItem> InputFiles = new List<FileItem>();
+		public List<FileItem> InputFiles = [];
 
 		/// <summary>
 		/// The default resource file to link in to every binary if a custom one is not provided
 		/// </summary>
-		public List<FileItem> DefaultResourceFiles = new List<FileItem>();
+		public List<FileItem> DefaultResourceFiles = [];
 
 		/// <summary>
 		/// Resource files which should be compiled into every binary
 		/// </summary>
-		public List<FileItem> CommonResourceFiles = new List<FileItem>();
+		public List<FileItem> CommonResourceFiles = [];
 
 		/// <summary>
 		/// List of functions that should be exported from this module
 		/// </summary>
-		public List<string> IncludeFunctions = new List<string>();
+		public List<string> IncludeFunctions = [];
 
 		/// <summary>
 		/// Debugger visualizer files to build into debug info for this binary. 
 		/// </summary>
-		public List<FileItem> DebuggerVisualizerFiles = new List<FileItem>();
+		public List<FileItem> DebuggerVisualizerFiles = [];
 
 		/// <summary>
 		/// Provides a Module Definition File (.def) to the linker to describe various attributes of a DLL.
@@ -365,12 +358,12 @@ namespace UnrealBuildTool
 		/// <summary>
 		/// All the additional properties from the modules linked into this binary
 		/// </summary>
-		public List<ReceiptProperty> AdditionalProperties = new List<ReceiptProperty>();
+		public List<ReceiptProperty> AdditionalProperties = [];
 
 		/// <summary>
 		/// List of all of the runtime dependencies.
 		/// </summary>
-		public List<ModuleRules.RuntimeDependency> RuntimeDependencies = new List<ModuleRules.RuntimeDependency>();
+		public List<ModuleRules.RuntimeDependency> RuntimeDependencies = [];
 
 		/// <summary>
 		/// Default constructor.
diff --git a/Engine/Source/Programs/UnrealBuildTool/ToolChain/ISPCToolChain.cs b/Engine/Source/Programs/UnrealBuildTool/ToolChain/ISPCToolChain.cs
--- a/Engine/Source/Programs/UnrealBuildTool/ToolChain/ISPCToolChain.cs
+++ b/Engine/Source/Programs/UnrealBuildTool/ToolChain/ISPCToolChain.cs
@@ -68,7 +68,7 @@ namespace UnrealBuildTool
 		/// <returns>List of instruction set targets passed to ISPC compiler</returns>
 		public virtual List<string> GetISPCCompileTargets(UnrealTargetPlatform Platform, UnrealArch Arch)
 		{
-			List<string> ISPCTargets = new List<string>();
+			List<string> ISPCTargets = [];
 
 			if (UEBuildPlatform.IsPlatformInGroup(Platform, UnrealPlatformGroup.Windows) ||
 				UEBuildPlatform.IsPlatformInGroup(Platform, UnrealPlatformGroup.Unix) ||
@@ -76,7 +76,7 @@ namespace UnrealBuildTool
 			{
 				if (Arch.bIsX64 || Arch == UnrealArch.Arm64ec)
 				{
-					ISPCTargets.AddRange(new string[] { "avx512skx-i32x8", "avx2", "avx", "sse4" });
+					ISPCTargets.AddRange(["avx512skx-i32x8", "avx2", "avx", "sse4"]);
 				}
 				else
 				{
@@ -216,7 +216,7 @@ namespace UnrealBuildTool
 			return null;
 		}
 
-		static readonly Dictionary<string, string> s_ISPCCompilerVersions = new Dictionary<string, string>();
+		static readonly Dictionary<string, string> s_ISPCCompilerVersions = [];
 
 		/// <summary>
 		/// Returns the version of the ISPC compiler for the specified platform. If GetISPCHostCompilerPath() doesn't return a valid path
@@ -284,7 +284,7 @@ namespace UnrealBuildTool
 			return String.Empty;
 		}
 
-		private string EscapeDefinitionForISPC(string Definition)
+		private static string EscapeDefinitionForISPC(string Definition)
 		{
 			// See: [REDACTED_PROFILE_URL]/ispc/blob/4ee767560cd752eaf464c124eb7ef1b0fd37f1df/src/main.cpp#L264 for ispc's argument parsing code, which does the following (and does not support escaping):
 			// Argument      Parses as 
@@ -329,7 +329,7 @@ namespace UnrealBuildTool
 				else
 				{
 					// Treat all special chars in the value string as needing explicit extra quoting. Thoroughly clumsy.
-					StringBuilder Requoted = new StringBuilder();
+					StringBuilder Requoted = new();
 					foreach (char c in UnquotedValue)
 					{
 						if (c == '"')
@@ -479,7 +479,7 @@ namespace UnrealBuildTool
 
 		protected override CPPOutput GenerateISPCHeaders(CppCompileEnvironment CompileEnvironment, IEnumerable<FileItem> InputFiles, DirectoryReference OutputDir, IActionGraphBuilder Graph)
 		{
-			CPPOutput Result = new CPPOutput();
+			CPPOutput Result = new();
 
 			if (!CompileEnvironment.bCompileISPC)
 			{
@@ -490,7 +490,7 @@ namespace UnrealBuildTool
 
 			List<string> CompileTargets = GetISPCCompileTargets(CompileEnvironment.Platform, CompileEnvironment.Architecture);
 
-			List<string> GlobalArguments = new List<string>();
+			List<string> GlobalArguments = [];
 
 			// Build target string. No comma on last
 			string TargetString = String.Join(',', CompileTargets);
@@ -568,10 +568,11 @@ namespace UnrealBuildTool
 				// TODO: Remove, might work
 				CompileAction.bCanExecuteRemotelyWithSNDBS = false;
 
-				List<string> Arguments = new List<string>();
-
-				// Add the ISPC obj file as a prerequisite of the action.
-				Arguments.Add($"\"{NormalizeCommandLinePath(ISPCFile, RootPaths)}\"");
+				List<string> Arguments =
+				[
+					// Add the ISPC obj file as a prerequisite of the action.
+					$"\"{NormalizeCommandLinePath(ISPCFile, RootPaths)}\"",
+				];
 
 				// Add the ISPC h file to the produced item list.
 				FileItem ISPCIncludeHeaderFile = FileItem.GetItemByFileReference(
@@ -649,7 +650,7 @@ namespace UnrealBuildTool
 
 		protected override CPPOutput CompileISPCFiles(CppCompileEnvironment CompileEnvironment, IEnumerable<FileItem> InputFiles, DirectoryReference OutputDir, IActionGraphBuilder Graph)
 		{
-			CPPOutput Result = new CPPOutput();
+			CPPOutput Result = new();
 
 			if (!CompileEnvironment.bCompileISPC)
 			{
@@ -660,7 +661,7 @@ namespace UnrealBuildTool
 
 			List<string> CompileTargets = GetISPCCompileTargets(CompileEnvironment.Platform, CompileEnvironment.Architecture);
 
-			List<string> GlobalArguments = new List<string>();
+			List<string> GlobalArguments = [];
 
 			// Build target string. No comma on last
 			string TargetString = "";
@@ -693,7 +694,7 @@ namespace UnrealBuildTool
 
 			bool bByteCodeOutput = (PlatformObjectFileFormat == "llvm");
 
-			List<string> CommonArgs = new List<string>();
+			List<string> CommonArgs = [];
 			if (CompileEnvironment.Configuration == CppConfiguration.Debug)
 			{
 				if (CompileEnvironment.Platform == UnrealTargetPlatform.Mac)
@@ -778,12 +779,13 @@ namespace UnrealBuildTool
 				CompileAction.ArtifactMode = ArtifactMode.Enabled;
 				CompileAction.CacheBucket = GetISPCCacheBucket(TargetRules, CompileEnvironment);
 
-				List<string> Arguments = new List<string>();
-
-				// Add the ISPC file to be compiled.
-				Arguments.Add($"\"{NormalizeCommandLinePath(ISPCFile, RootPaths)}\"");
+				List<string> Arguments =
+				[
+					// Add the ISPC file to be compiled.
+					$"\"{NormalizeCommandLinePath(ISPCFile, RootPaths)}\"",
+				];
 
-				List<FileItem> CompiledISPCObjFiles = new List<FileItem>();
+				List<FileItem> CompiledISPCObjFiles = [];
 
 				string FileName = Path.GetFileName(ISPCFile.AbsolutePath);
 
@@ -880,7 +882,7 @@ namespace UnrealBuildTool
 					string? ByteCodeCompilerPath = GetISPCHostBytecodeCompilerPath(BuildHostPlatform.Current.Platform);
 					if (ByteCodeCompilerPath != null)
 					{
-						List<FileItem> FinalObjectFiles = new List<FileItem>();
+						List<FileItem> FinalObjectFiles = [];
 						foreach (FileItem CompiledBytecodeObjFile in CompiledISPCObjFiles)
 						{
 							string FileNameWithoutExtension = Path.GetFileNameWithoutExtension(CompiledBytecodeObjFile.AbsolutePath);
@@ -893,11 +895,13 @@ namespace UnrealBuildTool
 
 							Action PostCompileAction = Graph.CreateAction(ActionType.Compile);
 
-							List<string> PostCompileArgs = new List<string>();
-							PostCompileArgs.Add($"\"{NormalizeCommandLinePath(CompiledBytecodeObjFile, RootPaths)}\"");
-							PostCompileArgs.Add("-c");
-							PostCompileArgs.AddRange(CommonArgs);
-							PostCompileArgs.Add($"-o \"{NormalizeCommandLinePath(FinalCompiledISPCObjFile, RootPaths)}\"");
+							List<string> PostCompileArgs =
+							[
+								$"\"{NormalizeCommandLinePath(CompiledBytecodeObjFile, RootPaths)}\"",
+								"-c",
+								.. CommonArgs,
+								$"-o \"{NormalizeCommandLinePath(FinalCompiledISPCObjFile, RootPaths)}\"",
+							];
 
 							if (bMergeModules)
 							{
diff --git a/Engine/Source/Programs/UnrealBuildTool/ToolChain/UEToolChain.cs b/Engine/Source/Programs/UnrealBuildTool/ToolChain/UEToolChain.cs
--- a/Engine/Source/Programs/UnrealBuildTool/ToolChain/UEToolChain.cs
+++ b/Engine/Source/Programs/UnrealBuildTool/ToolChain/UEToolChain.cs
@@ -122,9 +122,8 @@ namespace UnrealBuildTool
 					DirectoryReference ArchOutputDir = new(OutputDir.FullName.Replace(PlatformArchitecturesString, ArchConfig.GetFolderNameForArchitecture(Arch)));
 
 					CppCompileEnvironment ArchEnvironment = new(CompileEnvironment, Arch);
-					if (ArchEnvironment.UserIncludePaths.Contains(OutputDir))
+					if (ArchEnvironment.UserIncludePaths.Remove(OutputDir))
 					{
-						ArchEnvironment.UserIncludePaths.Remove(OutputDir);
 						ArchEnvironment.UserIncludePaths.Add(ArchOutputDir);
 					}
 					CPPOutput ArchResult = CompileCPPFiles(ArchEnvironment, InputFiles, ArchOutputDir, ModuleName, Graph);
@@ -141,7 +140,7 @@ namespace UnrealBuildTool
 
 		public virtual CPPOutput CompileRCFiles(CppCompileEnvironment CompileEnvironment, IEnumerable<FileItem> InputFiles, DirectoryReference OutputDir, IActionGraphBuilder Graph)
 		{
-			CPPOutput Result = new CPPOutput();
+			CPPOutput Result = new();
 			return Result;
 		}
 
@@ -149,13 +148,11 @@ namespace UnrealBuildTool
 
 		public CPPOutput CompileAllISPCFiles(CppCompileEnvironment CompileEnvironment, IEnumerable<FileItem> InputFiles, DirectoryReference OutputDir, IActionGraphBuilder Graph)
 		{
-			CPPOutput Result;
-
 			UnrealArchitectureConfig ArchConfig = UnrealArchitectureConfig.ForPlatform(CompileEnvironment.Platform);
 			// compile architectures separately if needed
 			if (ArchConfig.Mode == UnrealArchitectureMode.SingleTargetCompileSeparately || ArchConfig.Mode == UnrealArchitectureMode.SingleTargetLinkSeparately)
 			{
-				Result = new CPPOutput();
+				CPPOutput Result = new();
 				foreach (UnrealArch Arch in CompileEnvironment.Architectures.Architectures)
 				{
 					string PlatformArchitecturesString = ArchConfig.GetFolderNameForArchitectures(CompileEnvironment.Architectures);
@@ -165,26 +162,21 @@ namespace UnrealBuildTool
 					CPPOutput ArchResult = CompileISPCFiles(ArchEnvironment, InputFiles, ArchOutputDir, Graph);
 					Result.Merge(ArchResult, Arch);
 				}
-			}
-			else
-			{
-				Result = CompileISPCFiles(CompileEnvironment, InputFiles, OutputDir, Graph);
+				return Result;
 			}
 
-			return Result;
+			return CompileISPCFiles(CompileEnvironment, InputFiles, OutputDir, Graph);
 		}
 
 		protected abstract CPPOutput GenerateISPCHeaders(CppCompileEnvironment CompileEnvironment, IEnumerable<FileItem> InputFiles, DirectoryReference OutputDir, IActionGraphBuilder Graph);
 
 		public CPPOutput GenerateAllISPCHeaders(CppCompileEnvironment CompileEnvironment, IEnumerable<FileItem> InputFiles, DirectoryReference OutputDir, IActionGraphBuilder Graph)
 		{
-			CPPOutput Result;
-
 			UnrealArchitectureConfig ArchConfig = UnrealArchitectureConfig.ForPlatform(CompileEnvironment.Platform);
 			// compile architectures separately if needed
 			if (ArchConfig.Mode == UnrealArchitectureMode.SingleTargetCompileSeparately || ArchConfig.Mode == UnrealArchitectureMode.SingleTargetLinkSeparately)
 			{
-				Result = new CPPOutput();
+				CPPOutput Result = new();
 				foreach (UnrealArch Arch in CompileEnvironment.Architectures.Architectures)
 				{
 					string PlatformArchitecturesString = ArchConfig.GetFolderNameForArchitectures(CompileEnvironment.Architectures);
@@ -195,12 +187,8 @@ namespace UnrealBuildTool
 					Result.Merge(ArchResult, Arch);
 				}
 			}
-			else
-			{
-				Result = GenerateISPCHeaders(CompileEnvironment, InputFiles, OutputDir, Graph);
-			}
-
-			return Result;
+			
+			return GenerateISPCHeaders(CompileEnvironment, InputFiles, OutputDir, Graph);
 		}
 
 		public virtual void GenerateTypeLibraryHeader(CppCompileEnvironment CompileEnvironment, ModuleRules.TypeLibrary TypeLibrary, FileReference OutputFile, FileReference? OutputHeader, IActionGraphBuilder Graph)
@@ -216,14 +204,13 @@ namespace UnrealBuildTool
 		/// <returns></returns>
 		public virtual FileItem[] LinkImportLibrary(LinkEnvironment LinkEnvironment, IActionGraphBuilder Graph)
 		{
-			// by default doing nothing
-			return Array.Empty<FileItem>();
+			return [];
 		}
 
 		public abstract FileItem? LinkFiles(LinkEnvironment LinkEnvironment, bool bBuildImportLibraryOnly, IActionGraphBuilder Graph);
 		public virtual FileItem[] LinkAllFiles(LinkEnvironment LinkEnvironment, bool bBuildImportLibraryOnly, IActionGraphBuilder Graph)
 		{
-			List<FileItem> Result = new();
+			List<FileItem> Result = [];
 
 			// compile architectures separately if needed
 			UnrealArchitectureConfig ArchConfig = UnrealArchitectureConfig.ForPlatform(LinkEnvironment.Platform);
@@ -231,14 +218,14 @@ namespace UnrealBuildTool
 			{
 				foreach (UnrealArch Arch in LinkEnvironment.Architectures.Architectures)
 				{
-					LinkEnvironment ArchEnvironment = new LinkEnvironment(LinkEnvironment, Arch);
+					LinkEnvironment ArchEnvironment = new(LinkEnvironment, Arch);
 
 					// determine the output location of intermediates (so, if OutputDir had the arch name in it, like Intermediate/x86+arm64, we would replace it with either emptry string
 					// or a single arch name
 					//string ArchReplacement = Arch == ArchitectureWithoutMarkup() ? "" : ArchConfig.GetFolderNameForArchitecture(Arch);
 					string PlatformArchitecturesString = ArchConfig.GetFolderNameForArchitectures(LinkEnvironment.Architectures);
 
-					ArchEnvironment.OutputFilePaths = LinkEnvironment.OutputFilePaths.Select(x => new FileReference(x.FullName.Replace(PlatformArchitecturesString, ArchConfig.GetFolderNameForArchitecture(Arch)))).ToList();
+					ArchEnvironment.OutputFilePaths = [.. LinkEnvironment.OutputFilePaths.Select(x => new FileReference(x.FullName.Replace(PlatformArchitecturesString, ArchConfig.GetFolderNameForArchitecture(Arch))))];
 
 					FileItem? LinkFile = LinkFiles(ArchEnvironment, bBuildImportLibraryOnly, Graph);
 					if (LinkFile != null)
@@ -255,22 +242,22 @@ namespace UnrealBuildTool
 					Result.Add(LinkFile);
 				}
 			}
-			return Result.ToArray();
+			return [.. Result];
 		}
 
 		public virtual IEnumerable<string> GetGlobalCommandLineArgs(CppCompileEnvironment CompileEnvironment)
 		{
-			return Array.Empty<string>();
+			return [];
 		}
 
 		public virtual IEnumerable<string> GetCPPCommandLineArgs(CppCompileEnvironment CompileEnvironment)
 		{
-			return Array.Empty<string>();
+			return [];
 		}
 
 		public virtual IEnumerable<string> GetCCommandLineArgs(CppCompileEnvironment CompileEnvironment)
 		{
-			return Array.Empty<string>();
+			return [];
 		}
 
 		public virtual CppCompileEnvironment CreateSharedResponseFile(CppCompileEnvironment CompileEnvironment, FileReference OutResponseFile, IActionGraphBuilder Graph)
@@ -283,7 +270,7 @@ namespace UnrealBuildTool
 		}
 
 		// Construct a relative path for the intermediate response file
-		private FileReference GetResponseFileName(FileReference Location) => Location.ChangeExtension($"{Location.GetExtension()}{ResponseExt}");
+		private static FileReference GetResponseFileName(FileReference Location) => Location.ChangeExtension($"{Location.GetExtension()}{ResponseExt}");
 
 		/// <summary>
 		/// Get the name of the response file for the current compile environment and output file
@@ -291,7 +278,7 @@ namespace UnrealBuildTool
 		/// <param name="CompileEnvironment"></param>
 		/// <param name="OutputFile"></param>
 		/// <returns></returns>
-		public FileReference GetResponseFileName(CppCompileEnvironment CompileEnvironment, FileItem OutputFile) => GetResponseFileName(OutputFile.Location);
+		public static FileReference GetResponseFileName(CppCompileEnvironment CompileEnvironment, FileItem OutputFile) => GetResponseFileName(OutputFile.Location);
 
 		/// <summary>
 		/// Get the name of the response file for the current linker environment and output file
@@ -299,7 +286,7 @@ namespace UnrealBuildTool
 		/// <param name="LinkEnvironment"></param>
 		/// <param name="OutputFile"></param>
 		/// <returns></returns>
-		public FileReference GetResponseFileName(LinkEnvironment LinkEnvironment, FileItem OutputFile)
+		public static FileReference GetResponseFileName(LinkEnvironment LinkEnvironment, FileItem OutputFile)
 		{
 			// Construct a relative path for the intermediate response file
 			return GetResponseFileName(FileReference.Combine(LinkEnvironment.IntermediateDirectory ?? OutputFile.Directory.Location, OutputFile.Location.GetFileName()));
@@ -307,7 +294,7 @@ namespace UnrealBuildTool
 
 		public virtual ICollection<FileItem> PostBuild(ReadOnlyTargetRules Target, FileItem Executable, LinkEnvironment ExecutableLinkEnvironment, IActionGraphBuilder Graph)
 		{
-			return new List<FileItem>();
+			return [];
 		}
 
 		public virtual ICollection<FileItem> PostBuild(ReadOnlyTargetRules Target, IEnumerable<FileItem> Executables, LinkEnvironment ExecutableLinkEnvironment, IActionGraphBuilder Graph)
@@ -390,7 +377,7 @@ namespace UnrealBuildTool
 
 		public virtual IEnumerable<string> GetExtraLinkFileAdditionalSymbols(UEBuildBinaryType binaryType)
 		{
-			return Enumerable.Empty<string>();
+			return [];
 		}
 
 		/// <summary>
@@ -400,7 +387,7 @@ namespace UnrealBuildTool
 		/// <param name="ToolArg">Argument that will be passed to the tool</param>
 		/// <param name="Expression">null, or a Regular expression to capture in the output</param>
 		/// <returns></returns>
-		protected string? RunToolAndCaptureOutput(FileReference Command, string ToolArg, string? Expression = null)
+		protected static string? RunToolAndCaptureOutput(FileReference Command, string ToolArg, string? Expression = null)
 		{
 			string ProcessOutput = Utils.RunLocalProcessAndReturnStdOut(Command.FullName, ToolArg, null);
 
@@ -424,9 +411,7 @@ namespace UnrealBuildTool
 		{
 			string? ProcessOutput = RunToolAndCaptureOutput(Command, VersionArg, VersionExpression);
 
-			Version? ToolVersion;
-
-			if (Version.TryParse(ProcessOutput, out ToolVersion))
+			if (Version.TryParse(ProcessOutput, out Version? ToolVersion))
 			{
 				return ToolVersion;
 			}
