UnrealBuildTool: Remove unused old warning property

#rnx

diff --git a/Engine/Source/Programs/UnrealBuildTool/System/CppCompileEnvironment.cs b/Engine/Source/Programs/UnrealBuildTool/System/CppCompileEnvironment.cs
--- a/Engine/Source/Programs/UnrealBuildTool/System/CppCompileEnvironment.cs
+++ b/Engine/Source/Programs/UnrealBuildTool/System/CppCompileEnvironment.cs
@@ -118,7 +118,7 @@ namespace UnrealBuildTool
 	public enum MinimumCpuArchitectureX64
 	{
 		/// <summary>
-		/// No minimum architecure
+		/// No minimum architecture
 		/// </summary>
 		None,
 
@@ -401,7 +401,7 @@ namespace UnrealBuildTool
 		public WarningLevel DefaultWarningLevel = WarningLevel.Warning;
 
 		/// <summary>
-		/// The compiler warnings associacted with the compile environment.
+		/// The compiler warnings associated with the compile environment.
 		/// </summary>
 		public CppCompileWarnings CppCompileWarnings;
 
@@ -759,19 +759,10 @@ namespace UnrealBuildTool
 		public EngineIncludeOrderVersion IncludeOrderVersion = EngineIncludeOrderVersion.Latest;
 
 		/// <summary>
-		/// Set flags for determinstic compiles.
+		/// Set flags for deterministic compiles.
 		/// </summary>
 		public bool bDeterministic;
 
-		/// <summary>
-		/// Set flags for determinstic compile warnings.
-		/// </summary>
-		public WarningLevel DeterministicWarningLevel
-		{
-			get => CppCompileWarnings.DeterministicWarningLevel;
-			set => CppCompileWarnings.DeterministicWarningLevel = value;
-		}
-
 		/// <summary>
 		/// Emits compilation errors for incorrect UE_LOG format strings.
 		/// </summary>
