UnrealBuildTool: Fix arm64ec toolchain check

#rnx

diff --git a/Engine/Source/Programs/UnrealBuildTool/Platform/Windows/MicrosoftPlatformSDK.cs b/Engine/Source/Programs/UnrealBuildTool/Platform/Windows/MicrosoftPlatformSDK.cs
--- a/Engine/Source/Programs/UnrealBuildTool/Platform/Windows/MicrosoftPlatformSDK.cs
+++ b/Engine/Source/Programs/UnrealBuildTool/Platform/Windows/MicrosoftPlatformSDK.cs
@@ -1112,16 +1112,16 @@ namespace UnrealBuildTool
 			int FamilyRank = PreferredVisualCppVersions.TakeWhile(x => !x.Contains(Family)).Count();
 
 			string? Error = null;
-			if (Version < MinimumVisualCppVersion)
-			{
-				Error = $"UnrealBuildTool requires at minimum the MSVC {MinimumVisualCppVersion} toolchain. Please install a later toolchain such as {PreferredVisualCppVersions.Select(x => x.Min).Max()} from the Visual Studio installer.";
-			}
 
 			VersionNumberRange? Banned = BannedVisualCppVersions.FirstOrDefault(x => x.Contains(Version));
 			if (Banned != null)
 			{
 				Error = $"UnrealBuildTool has banned the MSVC {Banned} toolchains due to compiler issues. Please install a different toolchain such as {PreferredVisualCppVersions.Select(x => x.Min).Max()} by opening the generated solution and installing recommended components or from the Visual Studio installer.";
 			}
+			else if (Version < MinimumVisualCppVersion)
+			{
+				Error = $"UnrealBuildTool requires at minimum the MSVC {MinimumVisualCppVersion} toolchain. Please install a later toolchain such as {PreferredVisualCppVersions.Select(x => x.Min).Max()} from the Visual Studio installer.";
+			}
 
 			if (HasX64ToolChain(ToolChainDir, HostArchitecture))
 			{
@@ -1134,11 +1134,8 @@ namespace UnrealBuildTool
 				Logger.LogDebug("Found Visual Studio toolchain: {ToolChainDir} (Family={Family}, FamilyRank={FamilyRank}, Version={Version}, HostArchitecture={HostArchitecture}, ReleaseChannel={ReleaseChannel}, Architecture={Arch}, Error={Error}, Redist={RedistDir})", ToolChainDir, Family, FamilyRank, Version, HostArchitecture, ReleaseChannel, UnrealArch.Arm64.ToString(), Error != null, RedistDir);
 				ToolChains.Add(new ToolChainInstallation(Family, FamilyRank, Version, HostArchitecture, ReleaseChannel, UnrealArch.Arm64, Error, ToolChainDir, RedistDir, IsAutoSdk));
 
-				if (HasArm64ECToolChain(ToolChainDir))
-				{
-					Logger.LogDebug("Found Visual Studio toolchain: {ToolChainDir} (Family={Family}, FamilyRank={FamilyRank}, Version={Version}, HostArchitecture={HostArchitecture}, ReleaseChannel={ReleaseChannel}, Architecture={Arch}, Error={Error}, Redist={RedistDir})", ToolChainDir, Family, FamilyRank, Version, HostArchitecture, ReleaseChannel, UnrealArch.Arm64ec.ToString(), Error != null, RedistDir);
-					ToolChains.Add(new ToolChainInstallation(Family, FamilyRank, Version, HostArchitecture, ReleaseChannel, UnrealArch.Arm64ec, Error, ToolChainDir, RedistDir, IsAutoSdk));
-				}
+				Logger.LogDebug("Found Visual Studio toolchain: {ToolChainDir} (Family={Family}, FamilyRank={FamilyRank}, Version={Version}, HostArchitecture={HostArchitecture}, ReleaseChannel={ReleaseChannel}, Architecture={Arch}, Error={Error}, Redist={RedistDir})", ToolChainDir, Family, FamilyRank, Version, HostArchitecture, ReleaseChannel, UnrealArch.Arm64ec.ToString(), Error != null, RedistDir);
+				ToolChains.Add(new ToolChainInstallation(Family, FamilyRank, Version, HostArchitecture, ReleaseChannel, UnrealArch.Arm64ec, Error, ToolChainDir, RedistDir, IsAutoSdk));
 			}
 		}
 
@@ -1408,16 +1405,6 @@ namespace UnrealBuildTool
 			return FileReference.Exists(FileReference.Combine(ToolChainDir, "bin", UnrealArchToMSVCHostDirectoryName(HostArchitecture), "arm64", "cl.exe"));
 		}
 
-		/// <summary>
-		/// Checks if the given directory contains an arm64ec toolchain.
-		/// </summary>
-		/// <param name="ToolChainDir">Directory to check</param>
-		/// <returns>True if the given directory contains the arm64ec toolchain</returns>
-		static bool HasArm64ECToolChain(DirectoryReference ToolChainDir)
-		{
-			return FileReference.Exists(FileReference.Combine(ToolChainDir, "lib", "arm64ec", "binmode.obj"));
-		}
-
 		/// <summary>
 		/// Determines if an IDE for the given compiler is installed.
 		/// </summary>
