UnrealBuildTool: rename bIsGCCCompiler -> bIsClangCompiler

#rnx

diff --git a/Engine/Source/Programs/UnrealBuildTool/Executors/SNDBS.cs b/Engine/Source/Programs/UnrealBuildTool/Executors/SNDBS.cs
--- a/Engine/Source/Programs/UnrealBuildTool/Executors/SNDBS.cs
+++ b/Engine/Source/Programs/UnrealBuildTool/Executors/SNDBS.cs
@@ -274,7 +274,7 @@ namespace UnrealBuildTool
 					}
 
 					// Add PCH source file dependencies for clang or clang-cl
-					if (a.bIsGCCCompiler)
+					if (a.bIsClangCompiler)
 					{
 						// Look for any prerequisite actions that produce .pch files and add any .cpp files they depend on
 						List<Dictionary<string, object>>? ExplicitInputFiles = Job["explicit_input_files"] as List<Dictionary<string, object>>;
diff --git a/Engine/Source/Programs/UnrealBuildTool/Executors/UnrealBuildAccelerator/UBAExecutor.cs b/Engine/Source/Programs/UnrealBuildTool/Executors/UnrealBuildAccelerator/UBAExecutor.cs
--- a/Engine/Source/Programs/UnrealBuildTool/Executors/UnrealBuildAccelerator/UBAExecutor.cs
+++ b/Engine/Source/Programs/UnrealBuildTool/Executors/UnrealBuildAccelerator/UBAExecutor.cs
@@ -950,7 +950,7 @@ namespace UnrealBuildTool
 				Priority = priority,
 				UserData = action,
 				Description = description,
-				Configuration = action.bIsGCCCompiler ? EpicGames.UBA.ProcessStartInfo.CommonProcessConfigs.CompileClang : EpicGames.UBA.ProcessStartInfo.CommonProcessConfigs.CompileMsvc,
+				Configuration = action.bIsClangCompiler ? EpicGames.UBA.ProcessStartInfo.CommonProcessConfigs.CompileClang : EpicGames.UBA.ProcessStartInfo.CommonProcessConfigs.CompileMsvc,
 				LogFile = UBAConfig.bLogEnabled ? (action.Inner.ProducedItems.First().Location.GetFileName() + ".log") : null,
 			};
 
diff --git a/Engine/Source/Programs/UnrealBuildTool/Executors/XGE.cs b/Engine/Source/Programs/UnrealBuildTool/Executors/XGE.cs
--- a/Engine/Source/Programs/UnrealBuildTool/Executors/XGE.cs
+++ b/Engine/Source/Programs/UnrealBuildTool/Executors/XGE.cs
@@ -529,7 +529,7 @@ namespace UnrealBuildTool
 						ToolElement.SetAttribute("VCCompiler", "true");
 					}
 				}
-				if (Action.bIsGCCCompiler)
+				if (Action.bIsClangCompiler)
 				{
 					ToolElement.SetAttribute("AutoReserveMemory", "*.gch");
 				}
diff --git a/Engine/Source/Programs/UnrealBuildTool/Platform/Windows/VCCompileAction.cs b/Engine/Source/Programs/UnrealBuildTool/Platform/Windows/VCCompileAction.cs
--- a/Engine/Source/Programs/UnrealBuildTool/Platform/Windows/VCCompileAction.cs
+++ b/Engine/Source/Programs/UnrealBuildTool/Platform/Windows/VCCompileAction.cs
@@ -119,11 +119,6 @@ namespace UnrealBuildTool
 		/// </summary>
 		public FileItem? PchThroughHeaderFile { get; set; }
 
-		/// <summary>
-		/// Response file for the compiler
-		/// </summary>
-		public FileItem? VfsOverlayFile { get; set; }
-
 		/// <summary>
 		/// List of additional response paths
 		/// </summary>
@@ -248,11 +243,11 @@ namespace UnrealBuildTool
 				return $"Compile [{Architecture}]";
 			}
 		}
-		bool IExternalAction.bIsGCCCompiler => false;
+		bool IExternalAction.bIsClangCompiler => false;
 		bool IExternalAction.bDeleteProducedItemsOnError => CompilerType.IsClang() && bIsAnalyzing && bWarningsAsErrors;
 		bool IExternalAction.bForceWarningsAsError => CompilerType.IsClang() && bIsAnalyzing && bWarningsAsErrors;
 		bool IExternalAction.bProducesImportLibrary => false;
-		string IExternalAction.StatusDescription => (SourceFile == null) ? "Compiling" : SourceFile.Location.GetFileName();
+		string IExternalAction.StatusDescription => SourceFile?.Location.GetFileName() ?? "Compiling";
 		bool IExternalAction.bShouldOutputStatusDescription => CompilerType.IsClang();
 
 		/// <inheritdoc/>
@@ -272,10 +267,6 @@ namespace UnrealBuildTool
 				{
 					yield return UsingPchFile;
 				}
-				if (VfsOverlayFile != null)
-				{
-					yield return VfsOverlayFile;
-				}
 				foreach (FileItem additionalResponseFile in AdditionalResponseFiles)
 				{
 					yield return additionalResponseFile;
@@ -396,7 +387,6 @@ namespace UnrealBuildTool
 			CreatePchFile = inAction.CreatePchFile;
 			UsingPchFile = inAction.UsingPchFile;
 			PchThroughHeaderFile = inAction.PchThroughHeaderFile;
-			VfsOverlayFile = inAction.VfsOverlayFile;
 			AdditionalResponseFiles = [.. inAction.AdditionalResponseFiles];
 			IncludePaths = [.. inAction.IncludePaths];
 			SystemIncludePaths = [.. inAction.SystemIncludePaths];
@@ -444,7 +434,6 @@ namespace UnrealBuildTool
 			CreatePchFile = reader.ReadFileItem();
 			UsingPchFile = reader.ReadFileItem();
 			PchThroughHeaderFile = reader.ReadFileItem();
-			VfsOverlayFile = reader.ReadFileItem();
 			AdditionalResponseFiles = reader.ReadList(reader.ReadFileItem)!;
 			IncludePaths = reader.ReadList(reader.ReadCompactDirectoryReference)!;
 			SystemIncludePaths = reader.ReadList(reader.ReadCompactDirectoryReference)!;
@@ -491,7 +480,6 @@ namespace UnrealBuildTool
 			writer.WriteFileItem(CreatePchFile);
 			writer.WriteFileItem(UsingPchFile);
 			writer.WriteFileItem(PchThroughHeaderFile);
-			writer.WriteFileItem(VfsOverlayFile);
 			writer.WriteList(AdditionalResponseFiles, writer.WriteFileItem);
 			writer.WriteList(IncludePaths, writer.WriteCompactDirectoryReference);
 			writer.WriteList(SystemIncludePaths, writer.WriteCompactDirectoryReference);
diff --git a/Engine/Source/Programs/UnrealBuildTool/System/Action.cs b/Engine/Source/Programs/UnrealBuildTool/System/Action.cs
--- a/Engine/Source/Programs/UnrealBuildTool/System/Action.cs
+++ b/Engine/Source/Programs/UnrealBuildTool/System/Action.cs
@@ -161,9 +161,9 @@ namespace UnrealBuildTool
 		bool bCanExecuteInUBACrossArchitecture { get; }
 
 		/// <summary>
-		/// True if this action is using the GCC compiler.  Some build systems may be able to optimize for this case.
+		/// True if this action is using the Clang compiler.  Some build systems may be able to optimize for this case.
 		/// </summary>
-		bool bIsGCCCompiler { get; }
+		bool bIsClangCompiler { get; }
 
 		/// <summary>
 		/// True if the action executor should delete output files on an error.
@@ -320,9 +320,9 @@ namespace UnrealBuildTool
 		public bool bCanExecuteInUBACrossArchitecture { get; set; } = true;
 
 		/// <summary>
-		/// True if this action is using the GCC compiler.  Some build systems may be able to optimize for this case.
+		/// True if this action is using the Clang compiler.  Some build systems may be able to optimize for this case.
 		/// </summary>
-		public bool bIsGCCCompiler { get; set; } = false;
+		public bool bIsClangCompiler { get; set; } = false;
 
 		/// <summary>
 		/// True if the action executor should delete output files on an error.
@@ -397,7 +397,7 @@ namespace UnrealBuildTool
 			bCanExecuteRemotelyWithXGE = InOther.bCanExecuteRemotelyWithXGE;
 			bCanExecuteInUBA = InOther.bCanExecuteInUBA;
 			bCanExecuteInUBACrossArchitecture = InOther.bCanExecuteInUBACrossArchitecture;
-			bIsGCCCompiler = InOther.bIsGCCCompiler;
+			bIsClangCompiler = InOther.bIsClangCompiler;
 			bDeleteProducedItemsOnError = InOther.bDeleteProducedItemsOnError;
 			bForceWarningsAsError = InOther.bForceWarningsAsError;
 			bShouldOutputStatusDescription = InOther.bShouldOutputStatusDescription;
@@ -424,7 +424,7 @@ namespace UnrealBuildTool
 			bCanExecuteRemotelyWithXGE = Reader.ReadBool();
 			bCanExecuteInUBA = Reader.ReadBool();
 			bCanExecuteInUBACrossArchitecture = Reader.ReadBool();
-			bIsGCCCompiler = Reader.ReadBool();
+			bIsClangCompiler = Reader.ReadBool();
 			bDeleteProducedItemsOnError = Reader.ReadBool();
 			bForceWarningsAsError = Reader.ReadBool();
 			bShouldOutputStatusDescription = Reader.ReadBool();
@@ -459,7 +459,7 @@ namespace UnrealBuildTool
 			Writer.WriteBool(bCanExecuteRemotelyWithXGE);
 			Writer.WriteBool(bCanExecuteInUBA);
 			Writer.WriteBool(bCanExecuteInUBACrossArchitecture);
-			Writer.WriteBool(bIsGCCCompiler);
+			Writer.WriteBool(bIsClangCompiler);
 			Writer.WriteBool(bDeleteProducedItemsOnError);
 			Writer.WriteBool(bForceWarningsAsError);
 			Writer.WriteBool(bShouldOutputStatusDescription);
@@ -544,9 +544,9 @@ namespace UnrealBuildTool
 				Action.bCanExecuteInUBACrossArchitecture = bCanExecuteInUBACrossArchitecture;
 			}
 
-			if (Object.TryGetBoolField("bIsGCCCompiler", out bool bIsGCCCompiler))
+			if (Object.TryGetBoolField("bIsClangCompiler", out bool bIsClangCompiler))
 			{
-				Action.bIsGCCCompiler = bIsGCCCompiler;
+				Action.bIsClangCompiler = bIsClangCompiler;
 			}
 
 			if (Object.TryGetBoolField("bDeleteProducedItemsOnError", out bool bDeleteProducedItemsOnError))
@@ -671,7 +671,7 @@ namespace UnrealBuildTool
 			Writer.WriteValue("bCanExecuteRemotelyWithXGE", Action.bCanExecuteRemotelyWithXGE);
 			Writer.WriteValue("bCanExecuteInUBA", Action.bCanExecuteInUBA);
 			Writer.WriteValue("bCanExecuteInUBACrossArchitecture", Action.bCanExecuteInUBACrossArchitecture);
-			Writer.WriteValue("bIsGCCCompiler", Action.bIsGCCCompiler);
+			Writer.WriteValue("bIsClangCompiler", Action.bIsClangCompiler);
 			Writer.WriteValue("bDeleteProducedItemsOnError", Action.bDeleteProducedItemsOnError);
 			Writer.WriteValue("bForceWarningsAsError", Action.bForceWarningsAsError);			
 			Writer.WriteValue("bShouldOutputStatusDescription", Action.bShouldOutputStatusDescription);
@@ -795,7 +795,7 @@ namespace UnrealBuildTool
 		public bool bCanExecuteRemotelyWithXGE => Inner.bCanExecuteRemotelyWithXGE;
 		public bool bCanExecuteInUBA => Inner.bCanExecuteInUBA;
 		public bool bCanExecuteInUBACrossArchitecture => Inner.bCanExecuteInUBACrossArchitecture;
-		public bool bIsGCCCompiler => Inner.bIsGCCCompiler;
+		public bool bIsClangCompiler => Inner.bIsClangCompiler;
 		public bool bDeleteProducedItemsOnError => Inner.bDeleteProducedItemsOnError;
 		public bool bForceWarningsAsError => Inner.bForceWarningsAsError;
 		public bool bShouldOutputStatusDescription => Inner.bShouldOutputStatusDescription;
diff --git a/Engine/Source/Programs/UnrealBuildTool/ToolChain/AppleToolChain.cs b/Engine/Source/Programs/UnrealBuildTool/ToolChain/AppleToolChain.cs
--- a/Engine/Source/Programs/UnrealBuildTool/ToolChain/AppleToolChain.cs
+++ b/Engine/Source/Programs/UnrealBuildTool/ToolChain/AppleToolChain.cs
@@ -721,7 +721,7 @@ namespace UnrealBuildTool
 				CompileAction.CommandDescription += $" [{ReadableArch}]";
 			}
 			CompileAction.StatusDescription = string.Join(" ", InputFiles.Select(x => x.Name));
-			CompileAction.bIsGCCCompiler = true;
+			CompileAction.bIsClangCompiler = true;
 			CompileAction.bCanExecuteRemotely = true;
 			CompileAction.bCanExecuteInUBA = false;
 
@@ -768,7 +768,7 @@ namespace UnrealBuildTool
 				}
 
 				HeaderAction.WorkingDirectory = CompileAction.WorkingDirectory;
-				HeaderAction.bIsGCCCompiler = CompileAction.bIsGCCCompiler;
+				HeaderAction.bIsClangCompiler = CompileAction.bIsClangCompiler;
 				HeaderAction.bCanExecuteRemotely = CompileAction.bCanExecuteRemotely;
 				HeaderAction.bCanExecuteInUBA = CompileAction.bCanExecuteInUBA;
 			}
diff --git a/Engine/Source/Programs/UnrealBuildTool/ToolChain/ClangToolChain.cs b/Engine/Source/Programs/UnrealBuildTool/ToolChain/ClangToolChain.cs
--- a/Engine/Source/Programs/UnrealBuildTool/ToolChain/ClangToolChain.cs
+++ b/Engine/Source/Programs/UnrealBuildTool/ToolChain/ClangToolChain.cs
@@ -1399,7 +1399,7 @@ namespace UnrealBuildTool
 				CompileAction.CommandDescription += $" [{ReadableArch}]";
 			}
 			CompileAction.StatusDescription = Path.GetFileName(SourceFile.AbsolutePath);
-			CompileAction.bIsGCCCompiler = true;
+			CompileAction.bIsClangCompiler = true;
 
 			// Don't farm out creation of pre-compiled headers as it is the critical path task.
 			if (CompileEnvironment.PrecompiledHeaderAction == PrecompiledHeaderAction.Create)
@@ -1435,7 +1435,7 @@ namespace UnrealBuildTool
 				PrepassAction.PrerequisiteItems.Remove(CompilerResponseFileItem);
 				PrepassAction.CommandDescription = "Preprocess Depends";
 				PrepassAction.StatusDescription = CompileAction.StatusDescription;
-				PrepassAction.bIsGCCCompiler = true;
+				PrepassAction.bIsClangCompiler = true;
 				PrepassAction.bCanExecuteRemotely = false;
 				PrepassAction.bShouldOutputStatusDescription = true;
 				PrepassAction.CommandPath = CompileAction.CommandPath;
