UnrealBuildTool: Fix missing return

#rnx

diff --git a/Engine/Source/Programs/UnrealBuildTool/ToolChain/UEToolChain.cs b/Engine/Source/Programs/UnrealBuildTool/ToolChain/UEToolChain.cs
--- a/Engine/Source/Programs/UnrealBuildTool/ToolChain/UEToolChain.cs
+++ b/Engine/Source/Programs/UnrealBuildTool/ToolChain/UEToolChain.cs
@@ -104,13 +104,11 @@ namespace UnrealBuildTool
 
 		public CPPOutput CompileAllCPPFiles(CppCompileEnvironment CompileEnvironment, IEnumerable<FileItem> InputFiles, DirectoryReference OutputDir, string ModuleName, IActionGraphBuilder Graph)
 		{
-			CPPOutput Result;
-
 			UnrealArchitectureConfig ArchConfig = UnrealArchitectureConfig.ForPlatform(CompileEnvironment.Platform);
 			// compile architectures separately if needed
 			if (ArchConfig.Mode == UnrealArchitectureMode.SingleTargetCompileSeparately || ArchConfig.Mode == UnrealArchitectureMode.SingleTargetLinkSeparately)
 			{
-				Result = new CPPOutput();
+				CPPOutput Result = new();
 				foreach (UnrealArch Arch in CompileEnvironment.Architectures.Architectures)
 				{
 					// determine the output location of intermediates (so, if OutputDir had the arch name in it, like Intermediate/x86+arm64, we would replace it with either emptry string
@@ -129,13 +127,10 @@ namespace UnrealBuildTool
 					CPPOutput ArchResult = CompileCPPFiles(ArchEnvironment, InputFiles, ArchOutputDir, ModuleName, Graph);
 					Result.Merge(ArchResult, Arch);
 				}
-			}
-			else
-			{
-				Result = CompileCPPFiles(CompileEnvironment, InputFiles, OutputDir, ModuleName, Graph);
+				return Result;
 			}
 
-			return Result;
+			return CompileCPPFiles(CompileEnvironment, InputFiles, OutputDir, ModuleName, Graph);
 		}
 
 		public virtual CPPOutput CompileRCFiles(CppCompileEnvironment CompileEnvironment, IEnumerable<FileItem> InputFiles, DirectoryReference OutputDir, IActionGraphBuilder Graph)
@@ -186,6 +181,7 @@ namespace UnrealBuildTool
 					CPPOutput ArchResult = GenerateISPCHeaders(ArchEnvironment, InputFiles, ArchOutputDir, Graph);
 					Result.Merge(ArchResult, Arch);
 				}
+				return Result;
 			}
 			
 			return GenerateISPCHeaders(CompileEnvironment, InputFiles, OutputDir, Graph);
