UnrealBuildTool: Run correct arch of UbaObjTool for arm64

#rnx

diff --git a/Engine/Source/Programs/UnrealBuildTool/Platform/Windows/VCEnvironment.cs b/Engine/Source/Programs/UnrealBuildTool/Platform/Windows/VCEnvironment.cs
--- a/Engine/Source/Programs/UnrealBuildTool/Platform/Windows/VCEnvironment.cs
+++ b/Engine/Source/Programs/UnrealBuildTool/Platform/Windows/VCEnvironment.cs
@@ -271,7 +271,7 @@ namespace UnrealBuildTool
 			if (bAllowClangLinker && Compiler.IsClang())
 			{
 				// Since obj files could be LLVM IR Stream file format (when building with ltcg) we can't use link.exe. UbaObjTool support all formats and produce identical lib files as link.exe when running non-ltcg
-				DirectoryReference Dir = DirectoryReference.Combine(UnrealBuildBase.Unreal.EngineDirectory, "Binaries", "Win64", "UnrealBuildAccelerator", Architecture.ToString());
+				DirectoryReference Dir = DirectoryReference.Combine(UnrealBuildBase.Unreal.EngineDirectory, "Binaries", "Win64", "UnrealBuildAccelerator", RuntimeInformation.ProcessArchitecture.ToString().ToLowerInvariant());
 				return FileReference.Combine(Dir, "UbaObjTool.exe");
 			}
 			return FileReference.Combine(GetVCToolPath(ToochainDir, Architecture), "link.exe"); // We add /LIB to cmd line so we can use link.exe directly instead of going via lib.exe
