UnrealBuildTool: Add remote horde config for max workers and max idle time

#rnx

diff --git a/Engine/Source/Programs/UnrealBuildTool/Executors/UnrealBuildAccelerator/UBAAgentCoordinatorHorde.cs b/Engine/Source/Programs/UnrealBuildTool/Executors/UnrealBuildAccelerator/UBAAgentCoordinatorHorde.cs
--- a/Engine/Source/Programs/UnrealBuildTool/Executors/UnrealBuildAccelerator/UBAAgentCoordinatorHorde.cs
+++ b/Engine/Source/Programs/UnrealBuildTool/Executors/UnrealBuildAccelerator/UBAAgentCoordinatorHorde.cs
@@ -452,8 +452,9 @@ namespace UnrealBuildTool
 				encryption = Encryption.Ssl;
 			}
 
-			logger.LogInformation("Horde URL: {Server}, Pool: {Pool}, Cluster: {Cluster}, Condition: {Condition}, Connection: {Connection}, HordeEncryption: {Encryption}",
-				server, hordeConfig.HordePool ?? "(none)", hordeConfig.HordeCluster ?? "(none)", hordeConfig.HordeCondition ?? "(none)", connectionMode?.ToString() ?? "(none)", encryption?.ToString() ?? "(none)");
+			logger.LogInformation("Horde URL: {Server}, Pool: {Pool}, Cluster: {Cluster}, Condition: {Condition}, Connection: {Connection}, Encryption: {Encryption}, MaxCores: {MaxCores}, MaxWorkers: {MaxWorkers}, MaxIdle: {MaxIdle}s",
+				server, hordeConfig.HordePool ?? "(none)", hordeConfig.HordeCluster ?? "(none)", hordeConfig.HordeCondition ?? "(none)", connectionMode?.ToString() ?? "(none)", encryption?.ToString() ?? "(none)",
+				hordeConfig.HordeMaxCores, hordeConfig.HordeMaxWorkers, hordeConfig.HordeMaxIdle);
 			try
 			{
 				bool allowWine = hordeConfig.bHordeAllowWine && OperatingSystem.IsWindows();
@@ -584,8 +585,8 @@ namespace UnrealBuildTool
 						arguments.Add("-MaxCpu=%UE_HORDE_CPU_COUNT%");
 						arguments.Add("-MulCpu=%UE_HORDE_CPU_MULTIPLIER%");
 						arguments.Add("-MaxCon=8");
-						arguments.Add("-MaxWorkers=64");
-						arguments.Add("-MaxIdle=15");
+						arguments.Add($"-MaxWorkers={hordeConfig.HordeMaxWorkers}");
+						arguments.Add($"-MaxIdle={hordeConfig.HordeMaxIdle}");
 						arguments.Add("-UseCrawler");
 						arguments.Add("-ProxyUseLocalStorage"); // If set then agent's cas storage will be used by proxy
 						//arguments.Add("-UseIocp=4");
diff --git a/Engine/Source/Programs/UnrealBuildTool/Executors/UnrealBuildAccelerator/UnrealBuildAcceleratorHordeConfig.cs b/Engine/Source/Programs/UnrealBuildTool/Executors/UnrealBuildAccelerator/UnrealBuildAcceleratorHordeConfig.cs
--- a/Engine/Source/Programs/UnrealBuildTool/Executors/UnrealBuildAccelerator/UnrealBuildAcceleratorHordeConfig.cs
+++ b/Engine/Source/Programs/UnrealBuildTool/Executors/UnrealBuildAccelerator/UnrealBuildAcceleratorHordeConfig.cs
@@ -146,6 +146,22 @@ namespace UnrealBuildTool
 		[CommandLine("-UBAHordeMaxCores")]
 		public int HordeMaxCores { get; set; } = 576;
 
+		/// <summary>
+		/// Max workers allowed to be used by build session
+		/// </summary>
+		[ConfigFile(ConfigHierarchyType.Engine, "Uba.Provider.Horde", "MaxWorkers")]
+		[XmlConfigFile(Category = "Horde", Name = "MaxWorkers")]
+		[CommandLine("-UBAHordeMaxWorkers")]
+		public int HordeMaxWorkers { get; set; } = 64;
+
+		/// <summary>
+		/// Max idle time in seconds to maintain a connection to a worker when no pending actions are currently available due to waiting on dependencies
+		/// </summary>
+		[ConfigFile(ConfigHierarchyType.Engine, "Uba.Provider.Horde", "MaxIdle")]
+		[XmlConfigFile(Category = "Horde", Name = "MaxIdle")]
+		[CommandLine("-UBAHordeMaxIdle")]
+		public int HordeMaxIdle { get; set; } = 15;
+
 		/// <summary>
 		/// How long UBT should wait to ask for help. Useful in build configs where machine can delay remote work and still get same wall time results (pch dependencies etc)
 		/// </summary>
