UnrealBuildTool: Resolve intellisense warnings for LinuxToolChain

#rnx

diff --git a/Engine/Source/Programs/UnrealBuildTool/Platform/Linux/LinuxToolChain.cs b/Engine/Source/Programs/UnrealBuildTool/Platform/Linux/LinuxToolChain.cs
--- a/Engine/Source/Programs/UnrealBuildTool/Platform/Linux/LinuxToolChain.cs
+++ b/Engine/Source/Programs/UnrealBuildTool/Platform/Linux/LinuxToolChain.cs
@@ -57,10 +57,7 @@ namespace UnrealBuildTool
 			// prevent unknown clangs since the build is likely to fail on too old or too new compilers
 			if ((CompilerVersionLessThan(18, 0, 0) || CompilerVersionGreaterOrEqual(19, 0, 0)) && !Options.HasFlag(ClangToolChainOptions.UseAutoRTFMCompiler))
 			{
-				throw new BuildException(
-					String.Format("This version of the Unreal Engine can only be compiled with clang 18.x. clang {0} may not build it - please use a different version.",
-						Info.ClangVersion)
-					);
+				throw new BuildLogEventException("This version of the Unreal Engine can only be compiled with clang 18.x. clang {Version} may not build it - please use a different version.", Info.ClangVersion);
 			}
 		}
 
@@ -157,18 +154,18 @@ namespace UnrealBuildTool
 		{
 			bool bUseCmdExe = BuildHostPlatform.Current.Platform == UnrealTargetPlatform.Win64;
 			string DumpCommand = bUseCmdExe ? "\"{0}\" \"{1}\" \"{2}\" 2>NUL" : "\"{0}\" -c -o \"{2}\" \"{1}\"";
-			
-			Func<string, string> DeleteCommand = fileName =>
+
+			string DeleteCommand(string fileName)
+			{
+				if (bUseCmdExe)
 				{
-					if (bUseCmdExe)
-					{
-						return $"if exist \"{fileName}\" del /F \"{fileName}\"";
-					}
-					else
-					{
-						return $"rm \"{fileName}\"";
-					}
-				};
+					return $"if exist \"{fileName}\" del /F \"{fileName}\"";
+				}
+				else
+				{
+					return $"rm \"{fileName}\"";
+				}
+			}
 
 			FileItem EncodedBinarySymbolsFile = FileItem.GetItemByPath(Path.Combine(LinkEnvironment.OutputDirectory!.FullName, OutputFile.Location.GetFileNameWithoutExtension() + ".sym"));
 			FileItem SymbolsFile = FileItem.GetItemByPath(Path.Combine(LinkEnvironment.IntermediateDirectory!.FullName, OutputFile.Location.GetFileName() + ".psym"));
@@ -181,8 +178,10 @@ namespace UnrealBuildTool
 				SymbolsFile = FileItem.GetItemByPath(Path.Combine(LinkEnvironment.OutputDirectory.FullName, OutputFile.Location.GetFileNameWithoutExtension() + ".psym"));
 			}
 
-			StringWriter Out = new StringWriter();
-			Out.NewLine = bUseCmdExe ? "\r\n" : "\n";
+			StringWriter Out = new()
+			{
+				NewLine = bUseCmdExe ? "\r\n" : "\n"
+			};
 
 			if (!Options.HasFlag(ClangToolChainOptions.DisableDumpSyms) || Options.HasFlag(ClangToolChainOptions.PreservePSYM))
 			{
@@ -194,11 +193,7 @@ namespace UnrealBuildTool
 				);
 
 				// encode breakpad symbols
-				Out.WriteLine("\"{0}\" \"{1}\" \"{2}\"",
-					LinuxInfo.BreakpadEncoder,
-					SymbolsFile.AbsolutePath,
-					EncodedBinarySymbolsFile.AbsolutePath
-				);
+				Out.WriteLine($"\"{LinuxInfo.BreakpadEncoder}\" \"{SymbolsFile.AbsolutePath}\" \"{EncodedBinarySymbolsFile.AbsolutePath}\"");
 
 				if (preservePsym)
 				{
@@ -212,7 +207,7 @@ namespace UnrealBuildTool
 			else
 			{
 				// we have to create dummy files to prevent packaging errors
-				Out.WriteLine("echo DummySyms>> \"{0}\"", EncodedBinarySymbolsFile.AbsolutePath);
+				Out.WriteLine($"echo DummySyms>> \"{EncodedBinarySymbolsFile.AbsolutePath}\"");
 			}
 
 			action.ProducedItems.Add(EncodedBinarySymbolsFile);
@@ -221,34 +216,19 @@ namespace UnrealBuildTool
 			{
 				if (MaxBinarySizeOverrideForObjcopy > 0 && bUseCmdExe)
 				{
-					Out.WriteLine("for /F \"tokens=*\" %%F in (\"{0}\") DO set size=%%~zF",
-						OutputFile.AbsolutePath
-					);
+					Out.WriteLine($"for /F \"tokens=*\" %%F in (\"{OutputFile.AbsolutePath}\") DO set size=%%~zF");
 
-					Out.WriteLine("if %size% LSS {0} (", MaxBinarySizeOverrideForObjcopy);
+					Out.WriteLine($"if %size% LSS {MaxBinarySizeOverrideForObjcopy} (");
 				}
 
 				// objcopy stripped file
-				Out.WriteLine("\"{0}\" --strip-all \"{1}\" \"{2}\"",
-					LinuxInfo.Objcopy,
-					OutputFile.AbsolutePath,
-					StrippedFile.AbsolutePath
-				);
+				Out.WriteLine($"\"{LinuxInfo.Objcopy}\" --strip-all \"{OutputFile.AbsolutePath}\" \"{StrippedFile.AbsolutePath}\"");
 
 				// objcopy debug file
-				Out.WriteLine("\"{0}\" --only-keep-debug \"{1}\" \"{2}\"",
-					LinuxInfo.Objcopy,
-					OutputFile.AbsolutePath,
-					DebugFile.AbsolutePath
-				);
+				Out.WriteLine($"\"{LinuxInfo.Objcopy}\" --only-keep-debug \"{OutputFile.AbsolutePath}\" \"{DebugFile.AbsolutePath}\"");
 
 				// objcopy link debug file to final so
-				Out.WriteLine("\"{0}\" --add-gnu-debuglink=\"{1}\" \"{2}\" \"{3}\"",
-					LinuxInfo.Objcopy,
-					DebugFile.AbsolutePath,
-					StrippedFile.AbsolutePath,
-					OutputFile.AbsolutePath
-				);
+				Out.WriteLine($"\"{LinuxInfo.Objcopy}\" --add-gnu-debuglink=\"{DebugFile.AbsolutePath}\" \"{StrippedFile.AbsolutePath}\" \"{OutputFile.AbsolutePath}\"");
 
 				Out.WriteLine(DeleteCommand(StrippedFile.AbsolutePath));
 
@@ -258,24 +238,20 @@ namespace UnrealBuildTool
 					{
 						// If we have an override size, then we need to create a dummy file if that size is exceeded
 						Out.WriteLine(") ELSE (");
-						Out.WriteLine("echo DummyDebug >> \"{0}\"", DebugFile.AbsolutePath);
+						Out.WriteLine($"echo DummyDebug >> \"{DebugFile.AbsolutePath}\"");
 						Out.WriteLine(")");
 					}
 				}
 				else
 				{
 					// Change the debug file to normal permissions. It was taking on the +x rights from the output file
-					Out.WriteLine("chmod 644 \"{0}\"",
-						DebugFile.AbsolutePath
-					);
+					Out.WriteLine($"chmod 644 \"{DebugFile.AbsolutePath}\"");
 				}
 			}
 			else
 			{
 				// If we have disabled objcopy then we need to create a dummy debug file
-				Out.WriteLine("echo DummyDebug >> \"{0}\"",
-					DebugFile.AbsolutePath
-				);
+				Out.WriteLine($"echo DummyDebug >> \"{DebugFile.AbsolutePath}\"");
 			}
 
 			action.ProducedItems.Add(DebugFile);
@@ -588,7 +564,7 @@ namespace UnrealBuildTool
 			}
 			if (!String.IsNullOrEmpty(Value))
 			{
-				if (!Value.StartsWith("\"") && (Value.Contains(' ') || Value.Contains('$')))
+				if (!Value.StartsWith('\"') && (Value.Contains(' ') || Value.Contains('$')))
 				{
 					Value = Value.Trim('\"');       // trim any leading or trailing quotes
 					Value = "\"" + Value + "\"";    // ensure wrap string with double quotes
@@ -599,8 +575,8 @@ namespace UnrealBuildTool
 			}
 
 			return Value == null
-				? String.Format("{0}", Key)
-				: String.Format("{0}={1}", Key, Value);
+				? $"{Key}"
+				: $"{Key}={Value}";
 		}
 
 		public override void PrepareRuntimeDependencies(List<RuntimeDependency> RuntimeDependencies, Dictionary<FileReference, FileReference> TargetFileToSourceFile, DirectoryReference ExeDir)
@@ -614,8 +590,7 @@ namespace UnrealBuildTool
 				Options.HasFlag(ClangToolChainOptions.EnableLibFuzzer))
 			{
 				bool bInternalBuild = false;
-				BuildVersion? Version;
-				if (BuildVersion.TryRead(BuildVersion.GetDefaultFileName(), out Version))
+				if (BuildVersion.TryRead(BuildVersion.GetDefaultFileName(), out BuildVersion? Version))
 				{
 					bInternalBuild = !Version.IsLicenseeVersion;
 				}
@@ -625,7 +600,7 @@ namespace UnrealBuildTool
 					string? InternalSdkPath = UEBuildPlatform.GetSDK(UnrealTargetPlatform.Linux)!.GetInternalSDKPath();
 					if (InternalSdkPath != null)
 					{
-						DirectoryReference InternalSdkPathRef = new DirectoryReference(InternalSdkPath);
+						DirectoryReference InternalSdkPathRef = new(InternalSdkPath);
 
 						FileReference SymbolizerSourcePath = FileReference.Combine(InternalSdkPathRef, "bin/llvm-symbolizer");
 						FileReference SymbolizerTargetPath = FileReference.Combine(ExeDir, "llvm-symbolizer");
@@ -810,7 +785,7 @@ namespace UnrealBuildTool
 			}
 		}
 
-		string GetArchiveArguments(LinkEnvironment LinkEnvironment)
+		static string GetArchiveArguments(LinkEnvironment LinkEnvironment)
 		{
 			return " rcs";
 		}
@@ -821,7 +796,7 @@ namespace UnrealBuildTool
 		/// <summary>
 		/// Contains converter from output file to temporary output file (the one that will be relinked)
 		/// </summary>
-		private Dictionary<FileItem, FileItem> LinkToPrelinkFileList = new();
+		private readonly Dictionary<FileItem, FileItem> LinkToPrelinkFileList = [];
 
 
 		/// <summary>
@@ -863,7 +838,7 @@ namespace UnrealBuildTool
 				Logger.LogInformation("Compressing debug files");
 			}
 
-			Logger.LogInformation("Targeted minimum CPU architecture: {0}", (CompileEnvironment.Architecture == UnrealArch.X64) ? CompileEnvironment.MinCpuArchX64 : CompileEnvironment.MinArm64CpuTarget);
+			Logger.LogInformation("Targeted minimum CPU architecture: {Architecture}", (CompileEnvironment.Architecture == UnrealArch.X64) ? CompileEnvironment.MinCpuArchX64 : CompileEnvironment.MinArm64CpuTarget);
 
 			if (CompileEnvironment.bPGOOptimize)
 			{
@@ -914,7 +889,7 @@ namespace UnrealBuildTool
 				return new CPPOutput();
 			}
 
-			List<string> GlobalArguments = new();
+			List<string> GlobalArguments = [];
 			GetCompileArguments_Global(CompileEnvironment, GlobalArguments);
 
 			//var BuildPlatform = UEBuildPlatform.GetBuildPlatform(CompileEnvironment.Platform);
@@ -927,7 +902,7 @@ namespace UnrealBuildTool
 			}
 
 			// Create a compile action for each source file.
-			CPPOutput Result = new CPPOutput();
+			CPPOutput Result = new();
 			foreach (FileItem SourceFile in InputFiles)
 			{
 				CompileCPPFile(CompileEnvironment, SourceFile, OutputDir, ModuleName, Graph, GlobalArguments, Result);
@@ -936,7 +911,7 @@ namespace UnrealBuildTool
 			return Result;
 		}
 
-		public FileItem GetPrelinkOutput(DirectoryReference intermediateDir, FileItem output)
+		public static FileItem GetPrelinkOutput(DirectoryReference intermediateDir, FileItem output)
 		{
 			return FileItem.GetItemByFileReference(FileReference.Combine(intermediateDir, "Prelink", output.Name));
 		}
@@ -961,14 +936,14 @@ namespace UnrealBuildTool
 			ArchiveAction.ProducedItems.Add(OutputFile);
 			ArchiveAction.CommandDescription = "Archive";
 			ArchiveAction.StatusDescription = Path.GetFileName(OutputFile.AbsolutePath);
-			ArchiveAction.CommandArguments += String.Format("{0} \"{1}\"", GetArchiveArguments(LinkEnvironment), OutputFile.AbsolutePath);
+			ArchiveAction.CommandArguments += $"{GetArchiveArguments(LinkEnvironment)} \"{OutputFile.AbsolutePath}\"";
 
 			// Add the input files to a response file, and pass the response file on the command-line.
-			List<string> InputFileNames = new List<string>();
+			List<string> InputFileNames = [];
 			foreach (FileItem InputFile in LinkEnvironment.InputFiles)
 			{
 				string InputAbsolutePath = InputFile.AbsolutePath.Replace("\\", "/");
-				InputFileNames.Add(String.Format("\"{0}\"", InputAbsolutePath));
+				InputFileNames.Add($"\"{InputAbsolutePath}\"");
 				ArchiveAction.PrerequisiteItems.Add(InputFile);
 			}
 
@@ -979,7 +954,7 @@ namespace UnrealBuildTool
 				FileItem ResponseFileItem = Graph.CreateIntermediateTextFile(ResponsePath, InputFileNames);
 				ArchiveAction.PrerequisiteItems.Add(ResponseFileItem);
 			}
-			ArchiveAction.CommandArguments += String.Format(" @\"{0}\"", ResponsePath.FullName);
+			ArchiveAction.CommandArguments += $" @\"{ResponsePath.FullName}\"";
 
 			// Add the additional arguments specified by the environment.
 			ArchiveAction.CommandArguments += LinkEnvironment.AdditionalArguments;
@@ -1004,7 +979,7 @@ namespace UnrealBuildTool
 		{
 			Debug.Assert(!bBuildImportLibraryOnly);
 
-			List<string> RPaths = new List<string>();
+			List<string> RPaths = [];
 
 			if (LinkEnvironment.bIsBuildingLibrary || bBuildImportLibraryOnly)
 			{
@@ -1020,7 +995,7 @@ namespace UnrealBuildTool
 			LinkCommandString = "\"" + Info.Clang + "\"";
 
 			// Get link arguments.
-			List<string> LinkArguments = new List<string>();
+			List<string> LinkArguments = [];
 			GetLinkArguments_Global(LinkEnvironment, LinkArguments);
 			LinkCommandString += " " + String.Join(' ', LinkArguments);
 
@@ -1031,7 +1006,7 @@ namespace UnrealBuildTool
 			// Add profdata as a prerequisite for pgo optimize
 			if (LinkEnvironment.bPGOOptimize && LinkEnvironment.PGODirectory != null && LinkEnvironment.PGOFilenamePrefix != null)
 			{
-				DirectoryReference PGODir = new DirectoryReference(LinkEnvironment.PGODirectory);
+				DirectoryReference PGODir = new(LinkEnvironment.PGODirectory);
 				LinkAction.PrerequisiteItems.Add(FileItem.GetItemByFileReference(FileReference.Combine(PGODir, LinkEnvironment.PGOFilenamePrefix)));
 			}
 
@@ -1083,10 +1058,10 @@ namespace UnrealBuildTool
 			LinkAction.ArtifactMode = ArtifactMode.Enabled;
 
 			// Add the output file to the command-line.
-			LinkCommandString += String.Format(" -o \"{0}\"", OutputFile.AbsolutePath);
+			LinkCommandString += $" -o \"{OutputFile.AbsolutePath}\"";
 
 			// Add the input files to a response file, and pass the response file on the command-line.
-			List<string> ResponseLines = new List<string>();
+			List<string> ResponseLines = [];
 			foreach (FileItem InputFile in LinkEnvironment.InputFiles)
 			{
 				if (InputFile.HasExtension(".ldscript"))
@@ -1116,7 +1091,7 @@ namespace UnrealBuildTool
 				if (!String.IsNullOrEmpty(PathToLib))
 				{
 					// make path absolute, because FixDependencies script may be executed in a different directory
-					DirectoryReference AbsolutePathToLib = new DirectoryReference(PathToLib);
+					DirectoryReference AbsolutePathToLib = new(PathToLib);
 					if (!AllLibraryPaths.Contains(AbsolutePathToLib))
 					{
 						AllLibraryPaths.Add(AbsolutePathToLib);
@@ -1145,13 +1120,13 @@ namespace UnrealBuildTool
 						// do not replace directly, but take care to avoid potential double slashes or missed slashes
 						string PathFromRootDir = RelativePath.Replace(Unreal.RootDirectory.FullName, "");
 						// Path.Combine doesn't combine these properly
-						RelativePath = ((PathFromRootDir.StartsWith("\\") || PathFromRootDir.StartsWith("/")) ? "..\\..\\.." : "..\\..\\..\\") + PathFromRootDir;
+						RelativePath = ((PathFromRootDir.StartsWith('\\') || PathFromRootDir.StartsWith('/')) ? "..\\..\\.." : "..\\..\\..\\") + PathFromRootDir;
 					}
 
 					if (!RPaths.Contains(RelativePath))
 					{
 						RPaths.Add(RelativePath);
-						ResponseLines.Add(String.Format(" -rpath=\"${{ORIGIN}}/{0}\"", RelativePath.Replace('\\', '/')));
+						ResponseLines.Add($" -rpath=\"${{ORIGIN}}/{RelativePath.Replace('\\', '/')}\"");
 					}
 				}
 			}
@@ -1160,7 +1135,7 @@ namespace UnrealBuildTool
 			{
 				string RelativePath = RuntimeLibaryPath;
 
-				if (!RelativePath.StartsWith("$"))
+				if (!RelativePath.StartsWith('$'))
 				{
 					if (LinkEnvironment.bIsBuildingDLL)
 					{
@@ -1188,27 +1163,27 @@ namespace UnrealBuildTool
 					// do not replace directly, but take care to avoid potential double slashes or missed slashes
 					string PathFromRootDir = RelativePath.Replace(Unreal.RootDirectory.FullName, "");
 					// Path.Combine doesn't combine these properly
-					RelativePath = ((PathFromRootDir.StartsWith("\\") || PathFromRootDir.StartsWith("/")) ? "..\\..\\.." : "..\\..\\..\\") + PathFromRootDir;
+					RelativePath = ((PathFromRootDir.StartsWith('\\') || PathFromRootDir.StartsWith('/')) ? "..\\..\\.." : "..\\..\\..\\") + PathFromRootDir;
 				}
 
 				if (!RPaths.Contains(RelativePath))
 				{
 					RPaths.Add(RelativePath);
-					ResponseLines.Add(String.Format(" -rpath=\"${{ORIGIN}}/{0}\"", RelativePath.Replace('\\', '/')));
+					ResponseLines.Add($" -rpath=\"${{ORIGIN}}/{RelativePath.Replace('\\', '/')}\"");
 				}
 			}
 
-			ResponseLines.Add(String.Format(" -rpath-link=\"{0}\"", Path.GetDirectoryName(RealOutputFile.AbsolutePath)));
+			ResponseLines.Add($" -rpath-link=\"{Path.GetDirectoryName(RealOutputFile.AbsolutePath)}\"");
 
 			// Add the library paths to the argument list.
 			foreach (DirectoryReference LibraryPath in AllLibraryPaths)
 			{
 				// use absolute paths because of FixDependencies script again
-				ResponseLines.Add(String.Format(" -L\"{0}\"", LibraryPath.FullName.Replace('\\', '/')));
+				ResponseLines.Add($" -L\"{LibraryPath.FullName.Replace('\\', '/')}\"");
 			}
 
-			List<string> EngineAndGameLibrariesLinkFlags = new List<string>();
-			List<FileItem> EngineAndGameLibrariesFiles = new List<FileItem>();
+			List<string> EngineAndGameLibrariesLinkFlags = [];
+			List<FileItem> EngineAndGameLibrariesFiles = [];
 
 			// Pre-2.25 ld has symbol resolution problems when .so are mixed with .a in a single --start-group/--end-group
 			// when linking with --as-needed.
@@ -1224,7 +1199,7 @@ namespace UnrealBuildTool
 				if (String.IsNullOrEmpty(Path.GetDirectoryName(AdditionalLibrary)))
 				{
 					// library was passed just like "jemalloc", turn it into -ljemalloc
-					ExternalLibraries += String.Format(" -l{0}", AdditionalLibrary);
+					ExternalLibraries += $" -l{AdditionalLibrary}";
 				}
 				else if (Path.GetExtension(AdditionalLibrary) == ".a")
 				{
@@ -1232,7 +1207,7 @@ namespace UnrealBuildTool
 					string AbsoluteAdditionalLibrary = Path.GetFullPath(AdditionalLibrary);
 					if (AbsoluteAdditionalLibrary.Contains(' '))
 					{
-						AbsoluteAdditionalLibrary = String.Format("\"{0}\"", AbsoluteAdditionalLibrary);
+						AbsoluteAdditionalLibrary = $"\"{AbsoluteAdditionalLibrary}\"";
 					}
 					AbsoluteAdditionalLibrary = AbsoluteAdditionalLibrary.Replace('\\', '/');
 
@@ -1264,16 +1239,16 @@ namespace UnrealBuildTool
 					if (LibName.StartsWith("lib"))
 					{
 						// Remove lib prefix
-						LibName = LibName.Remove(0, 3);
+						LibName = LibName[3..];
 					}
 					else if (LibraryDependency.Exists)
 					{
 						// The library exists, but it is not prefixed with "lib", so force the
 						// linker to find it without that prefix by prepending a colon to
 						// the file name.
-						LibName = String.Format(":{0}", LibraryDependency.Name);
+						LibName = $":{LibraryDependency.Name}";
 					}
-					string LibLinkFlag = String.Format(" -l{0}", LibName);
+					string LibLinkFlag = $" -l{LibName}";
 
 					if (needsRelink)
 					{
@@ -1300,7 +1275,7 @@ namespace UnrealBuildTool
 			FileReference ResponseFileName = GetResponseFileName(LinkEnvironment, OutputFile);
 			FileItem ResponseFileItem = Graph.CreateIntermediateTextFile(ResponseFileName, ResponseLines);
 
-			LinkCommandString += String.Format(" -Wl,@\"{0}\"", ResponseFileName);
+			LinkCommandString += $" -Wl,@\"{ResponseFileName}\"";
 			LinkAction.PrerequisiteItems.Add(ResponseFileItem);
 
 			LinkCommandString += " -Wl,--start-group";
@@ -1390,7 +1365,7 @@ namespace UnrealBuildTool
 					LinkWriter.WriteLine("set /a Retries+=1");
 					LinkWriter.WriteLine("goto linkloop");
 					LinkWriter.WriteLine(":failedtorelink");
-					LinkWriter.WriteLine("echo Failed to link {0} after %Retries% retries", OutputFile.AbsolutePath);
+					LinkWriter.WriteLine($"echo Failed to link {OutputFile.AbsolutePath} after %Retries% retries");
 					LinkWriter.WriteLine("exit 1");
 				}
 				else
@@ -1485,7 +1460,7 @@ namespace UnrealBuildTool
 						RelinkWriter.WriteLine("set /a Retries+=1");
 						RelinkWriter.WriteLine("goto relinkloop");
 						RelinkWriter.WriteLine(":failedtorelink");
-						RelinkWriter.WriteLine("echo Failed to relink {0} after %Retries% retries", OutputFile.AbsolutePath);
+						RelinkWriter.WriteLine($"echo Failed to relink {OutputFile.AbsolutePath} after %Retries% retries");
 						RelinkWriter.WriteLine("exit 1");
 					}
 					else
@@ -1534,11 +1509,13 @@ namespace UnrealBuildTool
 				File.Copy(SourceFile.FullName, TargetFile.FullName, true);
 			}
 
-			ProcessStartInfo StartInfo = new ProcessStartInfo();
-			StartInfo.FileName = LinuxInfo.Objcopy.FullName;
-			StartInfo.Arguments = "--strip-debug \"" + TargetFile.FullName + "\"";
-			StartInfo.UseShellExecute = false;
-			StartInfo.CreateNoWindow = true;
+			ProcessStartInfo StartInfo = new()
+			{
+				FileName = LinuxInfo.Objcopy.FullName,
+				Arguments = "--strip-debug \"" + TargetFile.FullName + "\"",
+				UseShellExecute = false,
+				CreateNoWindow = true
+			};
 			Utils.RunLocalProcessAndLogOutput(StartInfo, Logger);
 		}
 
@@ -1551,8 +1528,8 @@ namespace UnrealBuildTool
 			{
 				// starting with clang 16.x the directory naming changed to include major version only
 				string ClangVersionString = (Info.ClangVersion.Major >= 16) ? Info.ClangVersion.Major.ToString() : Info.ClangVersion.ToString();
-				ExtraArguments.Add(String.Format("-isystem {0}", System.IO.Path.Combine(InternalSdkPath, "lib", "clang", ClangVersionString, "include").Replace("\\", "/")));
-				ExtraArguments.Add(String.Format("-isystem {0}", System.IO.Path.Combine(InternalSdkPath, "usr", "include").Replace("\\", "/")));
+				ExtraArguments.Add($"-isystem {System.IO.Path.Combine(InternalSdkPath, "lib", "clang", ClangVersionString, "include").Replace("\\", "/")}");
+				ExtraArguments.Add($"-isystem {System.IO.Path.Combine(InternalSdkPath, "usr", "include").Replace("\\", "/")}");
 			}
 		}
 
