UnrealBuildTool: Remove support for -PreprocessDepends, this was part of an old distribution experiment that is no longer needed with UBA

#rnx

diff --git a/Engine/Source/Programs/UnrealBuildTool/Configuration/ReadOnlyTargetRules.cs b/Engine/Source/Programs/UnrealBuildTool/Configuration/ReadOnlyTargetRules.cs
--- a/Engine/Source/Programs/UnrealBuildTool/Configuration/ReadOnlyTargetRules.cs
+++ b/Engine/Source/Programs/UnrealBuildTool/Configuration/ReadOnlyTargetRules.cs
@@ -387,8 +387,6 @@ namespace UnrealBuildTool
 
 		public bool bPreprocessOnly => Inner.bPreprocessOnly;
 
-		public bool bPreprocessDepends => Inner.bPreprocessDepends;
-
 		public bool bWithAssembly => Inner.bWithAssembly;
 
 		public bool bAllowUbaCompression => Inner.bAllowUbaCompression;
diff --git a/Engine/Source/Programs/UnrealBuildTool/Configuration/TargetRules.cs b/Engine/Source/Programs/UnrealBuildTool/Configuration/TargetRules.cs
--- a/Engine/Source/Programs/UnrealBuildTool/Configuration/TargetRules.cs
+++ b/Engine/Source/Programs/UnrealBuildTool/Configuration/TargetRules.cs
@@ -1994,6 +1994,7 @@ namespace UnrealBuildTool
 		/// </summary>
 		[CommandLine("-PreprocessDepends")]
 		[XmlConfigFile(Category = "BuildConfiguration")]
+		[Obsolete("bPreprocessDepends is no longer supported")]
 		public bool bPreprocessDepends { get; set; }
 
 		/// <summary>
diff --git a/Engine/Source/Programs/UnrealBuildTool/ToolChain/ClangToolChain.cs b/Engine/Source/Programs/UnrealBuildTool/ToolChain/ClangToolChain.cs
--- a/Engine/Source/Programs/UnrealBuildTool/ToolChain/ClangToolChain.cs
+++ b/Engine/Source/Programs/UnrealBuildTool/ToolChain/ClangToolChain.cs
@@ -215,7 +215,6 @@ namespace UnrealBuildTool
 		protected const int ClangCmdlineDangerZone = 30 * 1024;
 
 		// Target settings
-		protected bool PreprocessDepends = false;
 		protected bool ShowIncludes = false;
 		protected StaticAnalyzer StaticAnalyzer = StaticAnalyzer.None;
 		protected StaticAnalyzerMode StaticAnalyzerMode = StaticAnalyzerMode.Deep;
@@ -285,7 +284,6 @@ namespace UnrealBuildTool
 				EpicGames.UBA.Utils.RegisterPathHash(LazyInfo.Value.BasePath.FullName, pathHash);
 			}
 
-			PreprocessDepends = Target.bPreprocessDepends;
 			ShowIncludes = Target.bShowIncludes;
 			StaticAnalyzer = Target.StaticAnalyzer;
 			StaticAnalyzerMode = Target.StaticAnalyzerMode;
@@ -1189,7 +1187,6 @@ namespace UnrealBuildTool
 			Arguments.Add(GetSourceFileArgument(SourceFile, CompileEnvironment));
 
 			// Generate the included header dependency list
-			if (!PreprocessDepends)
 			{
 				FileItem DependencyListFile = FileItem.GetItemByFileReference(FileReference.Combine(OutputDir, GetFileNameFromExtension(FileName, ".d")));
 				Arguments.Add(GetDepencenciesListFileArgument(DependencyListFile, CompileEnvironment));
@@ -1426,64 +1423,6 @@ namespace UnrealBuildTool
 				CompileAction.bDeleteProducedItemsOnError = true;
 			}
 
-			// Two-pass compile where the preprocessor is run first to output the dependency list
-			if (PreprocessDepends)
-			{
-				Action PrepassAction = Graph.CreateAction(ActionType.Compile);
-				PrepassAction.RootPaths = CompileAction.RootPaths;
-				PrepassAction.PrerequisiteItems.UnionWith(CompileAction.PrerequisiteItems);
-				PrepassAction.PrerequisiteItems.Remove(CompilerResponseFileItem);
-				PrepassAction.CommandDescription = "Preprocess Depends";
-				PrepassAction.StatusDescription = CompileAction.StatusDescription;
-				PrepassAction.bIsClangCompiler = true;
-				PrepassAction.bCanExecuteRemotely = false;
-				PrepassAction.bShouldOutputStatusDescription = true;
-				PrepassAction.CommandPath = CompileAction.CommandPath;
-				PrepassAction.CommandVersion = CompileAction.CommandVersion;
-				PrepassAction.WorkingDirectory = CompileAction.WorkingDirectory;
-
-				List<string> PreprocessGlobalArguments = new(GlobalArguments);
-				List<string> PreprocessFileArguments = new(FileArguments);
-				PreprocessGlobalArguments.Remove("-c");
-
-				FileItem DependencyListFile = FileItem.GetItemByFileReference(FileReference.Combine(OutputDir, Path.GetFileName(SourceFile.AbsolutePath) + ".d"));
-				PreprocessFileArguments.Add(GetPreprocessDepencenciesListFileArgument(DependencyListFile, CompileEnvironment));
-				PrepassAction.DependencyListFile = DependencyListFile;
-				PrepassAction.ProducedItems.Add(DependencyListFile);
-
-				PreprocessFileArguments.Remove("-ftime-trace");
-				PreprocessFileArguments.Remove(GetOutputFileArgument(CompileAction.ProducedItems.First(), CompileEnvironment));
-
-				PrepassAction.DeleteItems.UnionWith(PrepassAction.ProducedItems);
-
-				// Gets the target file so we can get the correct output path.
-				FileItem PreprocessTargetFile = PrepassAction.DependencyListFile;
-
-				// Creates the path to the response file using the name of the output file and creates its contents.
-				FileReference PreprocessResponseFileName = GetResponseFileName(CompileEnvironment, PreprocessTargetFile);
-				List<string> PreprocessResponseFileContents = new();
-				PreprocessResponseFileContents.AddRange(PreprocessGlobalArguments);
-				PreprocessResponseFileContents.AddRange(PreprocessFileArguments);
-
-				if (OperatingSystem.IsWindows())
-				{
-					int CmdLineLength = Info.Clang.ToString().Length + String.Join(' ', ResponseFileContents).Length;
-					bool bIsInDangerZone = CmdLineLength >= ClangCmdlineDangerZone && CmdLineLength <= ClangCmdLineMaxSize;
-					if (bIsInDangerZone)
-					{
-						PreprocessResponseFileContents.Add(ClangDummyDefine);
-					}
-				}
-
-				// Adds the response file to the compiler input.
-				FileItem PreprocessResponseFileItem = Graph.CreateIntermediateTextFile(PreprocessResponseFileName, PreprocessResponseFileContents);
-				PrepassAction.PrerequisiteItems.Add(PreprocessResponseFileItem);
-
-				PrepassAction.CommandArguments = GetResponseFileArgument(PreprocessResponseFileItem, CompileEnvironment.RootPaths);
-				CompileAction.DependencyListFile = DependencyListFile;
-				CompileAction.PrerequisiteItems.Add(DependencyListFile);
-			}
-
 			return CompileAction;
 		}
 
