UnrealBuildTool: Intellisense suggestions for ClangToolChain.cs

#rnx

diff --git a/Engine/Source/Programs/UnrealBuildTool/ToolChain/ClangToolChain.cs b/Engine/Source/Programs/UnrealBuildTool/ToolChain/ClangToolChain.cs
--- a/Engine/Source/Programs/UnrealBuildTool/ToolChain/ClangToolChain.cs
+++ b/Engine/Source/Programs/UnrealBuildTool/ToolChain/ClangToolChain.cs
@@ -139,7 +139,7 @@ namespace UnrealBuildTool
 
 	abstract partial class ClangToolChain : ISPCToolChain
 	{
-		protected class ClangToolChainInfo
+		protected partial class ClangToolChainInfo
 		{
 			protected ILogger Logger { get; init; }
 			public DirectoryReference? BasePath { get; init; }
@@ -180,7 +180,7 @@ namespace UnrealBuildTool
 			/// <exception cref="BuildException"></exception>
 			protected virtual Version QueryClangVersion()
 			{
-				Match MatchClangVersion = Regex.Match(ClangVersionString, @"clang version (?<full>(?<major>\d+)\.(?<minor>\d+)\.(?<patch>\d+))");
+				Match MatchClangVersion = ClangVersionRegex().Match(ClangVersionString);
 				if (MatchClangVersion.Success && Version.TryParse(MatchClangVersion.Groups["full"].Value, out Version? ClangVersion))
 				{
 					return ClangVersion;
@@ -199,10 +199,17 @@ namespace UnrealBuildTool
 			/// </summary>
 			/// <returns>The standard output when running Archiver --version</returns>
 			protected virtual string QueryArchiverVersionString() => Utils.RunLocalProcessAndReturnStdOut(Archiver.FullName, "--version", null);
+
+			/// <summary>
+			/// Regex to parse clang version from console output
+			/// </summary>
+			/// <returns></returns>
+			[GeneratedRegex(@"clang version (?<full>(?<major>\d+)\.(?<minor>\d+)\.(?<patch>\d+))")]
+			private static partial Regex ClangVersionRegex();
 		}
 
 		// The Clang version being used to compile
-		Lazy<ClangToolChainInfo> LazyInfo;
+		readonly Lazy<ClangToolChainInfo> LazyInfo;
 		protected ClangToolChainInfo Info => LazyInfo.Value;
 
 		protected ClangToolChainOptions Options;
@@ -296,8 +303,8 @@ namespace UnrealBuildTool
 			if (Target.bPrintToolChainTimingInfo && Target.bParseTimingInfoForTracing)
 			{
 				TargetMakefile Makefile = MakefileBuilder.Makefile;
-				List<IExternalAction> CompileActions = Makefile.Actions.Where(x => x.ActionType == ActionType.Compile && x.ProducedItems.Any(i => i.HasExtension(".json"))).ToList();
-				List<FileItem> TimingJsonFiles = CompileActions.SelectMany(a => a.ProducedItems.Where(i => i.HasExtension(".json"))).ToList();
+				List<IExternalAction> CompileActions = [.. Makefile.Actions.Where(x => x.ActionType == ActionType.Compile && x.ProducedItems.Any(i => i.HasExtension(".json")))];
+				List<FileItem> TimingJsonFiles = [.. CompileActions.SelectMany(a => a.ProducedItems.Where(i => i.HasExtension(".json")))];
 
 				// Handing generating aggregate timing information if we compiled more than one file.
 				if (TimingJsonFiles.Count > 0)
@@ -308,16 +315,16 @@ namespace UnrealBuildTool
 					{
 						DirectoryReference.CreateDirectory(ManifestFile.Directory);
 					}
-					File.WriteAllLines(ManifestFile.FullName, TimingJsonFiles.Select(f => f.FullName.Remove(f.FullName.Length - ".json".Length)));
+					File.WriteAllLines(ManifestFile.FullName, TimingJsonFiles.Select(f => f.FullName[..^".json".Length]));
 
 					FileItem AggregateOutputFile = FileItem.GetItemByFileReference(FileReference.Combine(Makefile.ProjectIntermediateDirectory, $"{Target.Name}.trace.csv"));
 					FileItem HeadersOutputFile = FileItem.GetItemByFileReference(FileReference.Combine(Makefile.ProjectIntermediateDirectory, $"{Target.Name}.headers.csv"));
-					List<string> AggregateActionArgs = new List<string>()
-					{
+					List<string> AggregateActionArgs =
+					[
 						$"-ManifestFile={ManifestFile.FullName}",
 						$"-AggregateFile={AggregateOutputFile.FullName}",
 						$"-HeadersFile={HeadersOutputFile.FullName}",
-					};
+					];
 
 					Action AggregateTimingInfoAction = MakefileBuilder.CreateRecursiveAction<AggregateClangTimingInfo>(ActionType.ParseTimingInfo, String.Join(" ", AggregateActionArgs));
 					AggregateTimingInfoAction.StatusDescription = $"Aggregating {TimingJsonFiles.Count} Timing File(s)";
@@ -328,11 +335,11 @@ namespace UnrealBuildTool
 					Makefile.OutputItems.AddRange(AggregateTimingInfoAction.ProducedItems);
 
 					FileItem ArchiveOutputFile = FileItem.GetItemByFileReference(FileReference.Combine(Makefile.ProjectIntermediateDirectory, $"{Target.Name}.traces.zip"));
-					List<string> ArchiveActionArgs = new List<string>()
-					{
+					List<string> ArchiveActionArgs =
+					[
 						$"-ManifestFile={ManifestFile.FullName}",
 						$"-ArchiveFile={ArchiveOutputFile.FullName}",
-					};
+					];
 
 					Action ArchiveTimingInfoAction = MakefileBuilder.CreateRecursiveAction<AggregateClangTimingInfo>(ActionType.ParseTimingInfo, String.Join(" ", ArchiveActionArgs));
 					ArchiveTimingInfoAction.StatusDescription = $"Archiving {TimingJsonFiles.Count} Timing File(s)";
@@ -658,7 +665,7 @@ namespace UnrealBuildTool
 			Arguments.Add("-Wall");                                     // https://clang.llvm.org/docs/DiagnosticsReference.html#wall
 			Arguments.Add("-Werror");                                   // https://clang.llvm.org/docs/UsersManual.html#cmdoption-werror
 			
-			VersionNumber clangVersionNumber = new VersionNumber(Info.ClangVersion.Major, Info.ClangVersion.Minor, Info.ClangVersion.Build);
+			VersionNumber clangVersionNumber = new(Info.ClangVersion.Major, Info.ClangVersion.Minor, Info.ClangVersion.Build);
 			ClangWarnings.GetWarnings(CompileEnvironment, clangVersionNumber, Arguments);
 
 			// always use absolute paths for errors, this can help IDEs go to the error properly
@@ -960,7 +967,7 @@ namespace UnrealBuildTool
 				Arguments.Add("-Xclang -analyzer-config -Xclang mode=shallow");
 			}
 
-			VersionNumber ClangVersion = new VersionNumber(Info.ClangVersion.Major, Info.ClangVersion.Minor, Info.ClangVersion.Build);
+			VersionNumber ClangVersion = new(Info.ClangVersion.Major, Info.ClangVersion.Minor, Info.ClangVersion.Build);
 
 			if (CompileEnvironment.StaticAnalyzerCheckers.Count > 0)
 			{
@@ -1069,8 +1076,7 @@ namespace UnrealBuildTool
 			CompileAction.PrerequisiteItems.UnionWith(CompileEnvironment.AdditionalPrerequisites);
 			CompileAction.PrerequisiteItems.Add(SourceFile);
 
-			List<FileItem>? InlinedFiles;
-			if (CompileEnvironment.FileInlineGenCPPMap.TryGetValue(SourceFile, out InlinedFiles))
+			if (CompileEnvironment.FileInlineGenCPPMap.TryGetValue(SourceFile, out List<FileItem>? InlinedFiles))
 			{
 				CompileAction.PrerequisiteItems.UnionWith(InlinedFiles);
 			}
@@ -1122,7 +1128,7 @@ namespace UnrealBuildTool
 			string FileName = SourceFile.Name;
 			if (CompileEnvironment.CollidingNames != null && CompileEnvironment.CollidingNames.Contains(SourceFile))
 			{
-				string HashString = ContentHash.MD5(SourceFile.AbsolutePath.Substring(Unreal.RootDirectory.FullName.Length)).GetHashCode().ToString("X4");
+				string HashString = ContentHash.MD5(SourceFile.AbsolutePath[Unreal.RootDirectory.FullName.Length..]).GetHashCode().ToString("X4");
 				FileName = Path.GetFileNameWithoutExtension(FileName) + "_" + HashString + Path.GetExtension(FileName);
 			}
 
@@ -1165,7 +1171,7 @@ namespace UnrealBuildTool
 			// Add profdata as a prerequisite for pgo optimize
 			if (CompileEnvironment.bPGOOptimize && CompileEnvironment.PGODirectory != null && CompileEnvironment.PGOFilenamePrefix != null)
 			{
-				DirectoryReference PGODir = new DirectoryReference(CompileEnvironment.PGODirectory);
+				DirectoryReference PGODir = new(CompileEnvironment.PGODirectory);
 				CompileAction.PrerequisiteItems.Add(FileItem.GetItemByFileReference(FileReference.Combine(PGODir, CompileEnvironment.PGOFilenamePrefix)));
 			}
 
@@ -1231,7 +1237,7 @@ namespace UnrealBuildTool
 				{
 					if (NewList == ResponseFileContents)
 					{
-						NewList = new List<string>(ResponseFileContents);
+						NewList = [.. ResponseFileContents];
 					}
 				}
 				NewList[I] = NewLine;
@@ -1241,35 +1247,35 @@ namespace UnrealBuildTool
 
 		public override IEnumerable<string> GetGlobalCommandLineArgs(CppCompileEnvironment CompileEnvironment)
 		{
-			List<string> Arguments = new();
+			List<string> Arguments = [];
 			GetCompileArguments_Global(new CppCompileEnvironment(CompileEnvironment), Arguments);
 			return Arguments;
 		}
 
 		public override IEnumerable<string> GetCPPCommandLineArgs(CppCompileEnvironment CompileEnvironment)
 		{
-			List<string> Arguments = new();
+			List<string> Arguments = [];
 			GetCompileArguments_CPP(CompileEnvironment, Arguments);
 			return Arguments;
 		}
 
 		public override IEnumerable<string> GetCCommandLineArgs(CppCompileEnvironment CompileEnvironment)
 		{
-			List<string> Arguments = new();
+			List<string> Arguments = [];
 			GetCompileArguments_C(CompileEnvironment, Arguments);
 			return Arguments;
 		}
 
 		public override CppCompileEnvironment CreateSharedResponseFile(CppCompileEnvironment CompileEnvironment, FileReference OutResponseFile, IActionGraphBuilder Graph)
 		{
-			CppCompileEnvironment NewCompileEnvironment = new CppCompileEnvironment(CompileEnvironment);
-			List<string> Arguments = new List<string>();
+			CppCompileEnvironment NewCompileEnvironment = new(CompileEnvironment);
+			List<string> Arguments = [];
 
 			GetCompileArguments_Global(CompileEnvironment, Arguments);
 
 			// Stash shared include paths for validation purposes
-			NewCompileEnvironment.SharedUserIncludePaths = new(NewCompileEnvironment.UserIncludePaths);
-			NewCompileEnvironment.SharedSystemIncludePaths = new(NewCompileEnvironment.SystemIncludePaths);
+			NewCompileEnvironment.SharedUserIncludePaths = [.. NewCompileEnvironment.UserIncludePaths];
+			NewCompileEnvironment.SharedSystemIncludePaths = [.. NewCompileEnvironment.SystemIncludePaths];
 
 			NewCompileEnvironment.UserIncludePaths.Clear();
 			NewCompileEnvironment.SystemIncludePaths.Clear();
@@ -1294,7 +1300,7 @@ namespace UnrealBuildTool
 				return;
 			}
 
-			List<string> GlobalArguments = new();
+			List<string> GlobalArguments = [];
 			if (!CompileEnvironment.bHasSharedResponseFile)
 			{
 				GetCompileArguments_Global(CompileEnvironment, GlobalArguments);
@@ -1321,7 +1327,7 @@ namespace UnrealBuildTool
 			CompileAction.Weight = CompileActionWeight;
 
 			// copy the global arguments into the file arguments, so GetCompileArguments_FileType can remove entries if needed (special case but can be important)
-			List<string> FileArguments = new(GlobalArguments);
+			List<string> FileArguments = [.. GlobalArguments];
 
 			// If we are using the AutoRTFM compiler, we make the compile action depend on the version of the compiler itself.
 			// This lets us update the compiler (which might not cause a version update of the compiler, which instead tracks
@@ -1433,7 +1439,7 @@ namespace UnrealBuildTool
 				return new CPPOutput();
 			}
 
-			List<string> GlobalArguments = new();
+			List<string> GlobalArguments = [];
 
 			if (!CompileEnvironment.bHasSharedResponseFile)
 			{
@@ -1441,7 +1447,7 @@ namespace UnrealBuildTool
 			}
 
 			// Create a compile action for each source file.
-			CPPOutput Result = new CPPOutput();
+			CPPOutput Result = new();
 			foreach (FileItem SourceFile in InputFiles)
 			{
 				CompileCPPFile(CompileEnvironment, SourceFile, OutputDir, ModuleName, Graph, GlobalArguments, Result);
