Enable global destructors in Test & Shipping when merged modules are active

#rnx

diff --git a/Engine/Source/Programs/UnrealBuildTool/Platform/IOS/IOSToolChain.cs b/Engine/Source/Programs/UnrealBuildTool/Platform/IOS/IOSToolChain.cs
--- a/Engine/Source/Programs/UnrealBuildTool/Platform/IOS/IOSToolChain.cs
+++ b/Engine/Source/Programs/UnrealBuildTool/Platform/IOS/IOSToolChain.cs
@@ -191,7 +191,7 @@ namespace UnrealBuildTool
 				Arguments.Add("-fno-unwind-tables");
 				Arguments.Add("-fno-asynchronous-unwind-tables");
 
-				if (!CompileEnvironment.bPGOProfile)
+				if (!CompileEnvironment.bPGOProfile && !bMergeModules)
 				{
 					Arguments.Add("-fno-use-cxa-atexit");
 				}
