Fix issue in UBT where the modules files could contain old plugin references if the BinariesSubFolder changed.

diff --git a/Engine/Source/Programs/UnrealBuildTool/Modes/WriteMetadataMode.cs b/Engine/Source/Programs/UnrealBuildTool/Modes/WriteMetadataMode.cs
--- a/Engine/Source/Programs/UnrealBuildTool/Modes/WriteMetadataMode.cs
+++ b/Engine/Source/Programs/UnrealBuildTool/Modes/WriteMetadataMode.cs
@@ -316,9 +316,12 @@ namespace UnrealBuildTool
 		{
 			foreach (KeyValuePair<string, string> ModulePair in SourceManifest.ModuleNameToFileName)
 			{
-				if (!TargetManifest.ModuleNameToFileName.ContainsKey(ModulePair.Key))
+				if (FileReference.Exists(new FileReference(ModulePair.Value)))
 				{
-					TargetManifest.ModuleNameToFileName.Add(ModulePair.Key, ModulePair.Value);
+					if (!TargetManifest.ModuleNameToFileName.ContainsKey(ModulePair.Key))
+					{
+						TargetManifest.ModuleNameToFileName.Add(ModulePair.Key, ModulePair.Value);
+					}
 				}
 			}
 		}
