Change target settings for NoInlineLineTables so that line tables for inline functions can be blanket removed and then selectively reintroduced.

diff --git a/Engine/Source/Programs/UnrealBuildTool/Configuration/ReadOnlyTargetRules.cs b/Engine/Source/Programs/UnrealBuildTool/Configuration/ReadOnlyTargetRules.cs
--- a/Engine/Source/Programs/UnrealBuildTool/Configuration/ReadOnlyTargetRules.cs
+++ b/Engine/Source/Programs/UnrealBuildTool/Configuration/ReadOnlyTargetRules.cs
@@ -360,8 +360,8 @@ namespace UnrealBuildTool
 		public IReadOnlySet<string> DebugInfoLineTablesOnlyPlugins => Inner.DebugInfoLineTablesOnlyPlugins;
 
 		public DebugInfoMode DebugInfoNoInlineLineTables => Inner.DebugInfoNoInlineLineTables;
-		public IReadOnlySet<string> DebugInfoNoInlineLineTablesModules => Inner.DebugInfoNoInlineLineTablesModules;
-		public IReadOnlySet<string> DebugInfoNoInlineLineTablesPlugins => Inner.DebugInfoNoInlineLineTablesPlugins;
+		public IReadOnlySet<string> DebugInfoInlineLineTablesModules => Inner.DebugInfoInlineLineTablesModules;
+		public IReadOnlySet<string> DebugInfoInlineLineTablesPlugins => Inner.DebugInfoInlineLineTablesPlugins;
 		
 		public DebugInfoMode DebugInfoSimpleTemplateNames => Inner.DebugInfoSimpleTemplateNames;
 		public IReadOnlySet<string> DebugInfoSimpleTemplateNamesModules => Inner.DebugInfoSimpleTemplateNamesModules;
diff --git a/Engine/Source/Programs/UnrealBuildTool/Configuration/TargetRules.cs b/Engine/Source/Programs/UnrealBuildTool/Configuration/TargetRules.cs
--- a/Engine/Source/Programs/UnrealBuildTool/Configuration/TargetRules.cs
+++ b/Engine/Source/Programs/UnrealBuildTool/Configuration/TargetRules.cs
@@ -317,6 +317,8 @@ namespace UnrealBuildTool
 
 	/// <summary>
 	/// Debug info mode for compiler settings to determine how much debug info is available
+	/// Note that in some cases this flag is used for features which strip certain debug info, in which case the meaning is inverted.
+	/// i.e. 'Full' removes that type of debug info from all modules
 	/// </summary>
 	[Flags]
 	public enum DebugInfoMode
@@ -1882,16 +1884,16 @@ namespace UnrealBuildTool
 		public DebugInfoMode DebugInfoNoInlineLineTables { get; set; } = DebugInfoMode.None;
 
 		/// <summary>
-		/// Modules that should omit line number tables for inline functions for compilers that support doing so. 
+		/// Modules that should include line number tables for inline functions for compilers that support doing so, even if DebugInfoNoInlineLineTables would otherwise remove them.
 		/// </summary>
-		[CommandLine("-DebugInfoNoInlineLineTablesModules=", ListSeparator = '+')]
-		public HashSet<string> DebugInfoNoInlineLineTablesModules { get; } = [];
+		[CommandLine("-DebugInfoInlineLineTablesModules=", ListSeparator = '+')]
+		public HashSet<string> DebugInfoInlineLineTablesModules { get; } = [];
 
 		/// <summary>
-		/// Plugins that should omit line number tables for inline functions for compilers that support doing so. 
+		/// Modules that should include line number tables for inline functions for compilers that support doing so, even if DebugInfoNoInlineLineTables would otherwise remove them.
 		/// </summary>
-		[CommandLine("-DebugInfoNoInlineLineTablesPlugins=", ListSeparator = '+')]
-		public HashSet<string> DebugInfoNoInlineLineTablesPlugins { get; } = [];
+		[CommandLine("-DebugInfoInlineLineTablesPlugins=", ListSeparator = '+')]
+		public HashSet<string> DebugInfoInlineLineTablesPlugins { get; } = [];
 
 		/// <summary>
 		/// True if template type names should be simplified to reduce the size of debug info for compilers that support doing so.
diff --git a/Engine/Source/Programs/UnrealBuildTool/Configuration/UEBuildModuleCPP.cs b/Engine/Source/Programs/UnrealBuildTool/Configuration/UEBuildModuleCPP.cs
--- a/Engine/Source/Programs/UnrealBuildTool/Configuration/UEBuildModuleCPP.cs
+++ b/Engine/Source/Programs/UnrealBuildTool/Configuration/UEBuildModuleCPP.cs
@@ -2288,13 +2288,15 @@ namespace UnrealBuildTool
 				Result.bDebugLineTablesOnly = true;
 			}
 
-			if (Rules.Plugin != null && Rules.Target.DebugInfoNoInlineLineTablesPlugins.Contains(Rules.Plugin.Name))
+			// Note that these rules are the opposite of DebugInfoLineTablesOnly* - they selectively reintroduce line tables for 
+			// inlined functions even if the blanket rules above would remove them.
+			if (Rules.Plugin != null && Rules.Target.DebugInfoInlineLineTablesPlugins.Contains(Rules.Plugin.Name))
 			{
-				Result.bDebugNoInlineLineTables = true;
+				Result.bDebugNoInlineLineTables = false;
 			}
-			else if (Rules.Target.DebugInfoNoInlineLineTablesModules.Contains(Name))
+			else if (Rules.Target.DebugInfoInlineLineTablesModules.Contains(Name))
 			{
-				Result.bDebugNoInlineLineTables = true;
+				Result.bDebugNoInlineLineTables = false;
 			}
 
 			if (Rules.Plugin != null && Rules.Target.DebugInfoSimpleTemplateNamesPlugins.Contains(Rules.Plugin.Name))
