UnrealBuildTool: Remove visual studio project dependency on UBT as the build script will ensure it is up to date

#rnx

diff --git a/Engine/Source/Programs/UnrealBuildTool/ProjectFiles/VisualStudio/VCProjectFileGenerator.cs b/Engine/Source/Programs/UnrealBuildTool/ProjectFiles/VisualStudio/VCProjectFileGenerator.cs
--- a/Engine/Source/Programs/UnrealBuildTool/ProjectFiles/VisualStudio/VCProjectFileGenerator.cs
+++ b/Engine/Source/Programs/UnrealBuildTool/ProjectFiles/VisualStudio/VCProjectFileGenerator.cs
@@ -769,29 +769,19 @@ namespace UnrealBuildTool
 					// Use the existing project's GUID that's already known to us
 					string ProjectGUID = CurProject.ProjectGUID.ToString("B").ToUpperInvariant();
 
-					VCSolutionFileContent.AppendLine("Project(\"" + ProjectTypeGUID + "\") = \"" + ProjectNameInSolution + "\", \"" + CurProject.ProjectFilePath.MakeRelativeTo(ProjectFileGenerator.PrimaryProjectPath) + "\", \"" + ProjectGUID + "\"");
+					VCSolutionFileContent.AppendLine($"Project(\"{ProjectTypeGUID}\") = \"{ProjectNameInSolution}\", \"{CurProject.ProjectFilePath.MakeRelativeTo(ProjectFileGenerator.PrimaryProjectPath)}\", \"{ProjectGUID}\"");
 
-					// Setup dependency on UnrealBuildTool, if we need that.  This makes sure that UnrealBuildTool is
-					// freshly compiled before kicking off any build operations on this target project
 					if (!CurProject.IsStubProject)
 					{
-						List<ProjectFile> Dependencies = new List<ProjectFile>();
-						if (CurProject.IsGeneratedProject && UBTProject != null && CurProject != UBTProject)
-						{
-							Dependencies.Add(UBTProject);
-							Dependencies.AddRange(UBTProject.DependsOnProjects);
-						}
-						Dependencies.AddRange(CurProject.DependsOnProjects);
-
-						if (Dependencies.Count > 0)
+						if (CurProject.DependsOnProjects.Count > 0)
 						{
 							VCSolutionFileContent.AppendLine("\tProjectSection(ProjectDependencies) = postProject");
 
 							// Setup any addition dependencies this project has...
-							foreach (ProjectFile DependsOnProject in Dependencies)
+							foreach (ProjectFile DependsOnProject in CurProject.DependsOnProjects)
 							{
 								string DependsOnProjectGUID = ((MSBuildProjectFile)DependsOnProject).ProjectGUID.ToString("B").ToUpperInvariant();
-								VCSolutionFileContent.AppendLine("\t\t" + DependsOnProjectGUID + " = " + DependsOnProjectGUID);
+								VCSolutionFileContent.AppendLine($"\t\t{DependsOnProjectGUID} = {DependsOnProjectGUID}");
 							}
 
 							VCSolutionFileContent.AppendLine("\tEndProjectSection");
