- If the FinalizeWithXcode stuff fails, run again but with more logging to help discover potential cause of errors

diff --git a/Engine/Source/Programs/UnrealBuildTool/ToolChain/AppleToolChain.cs b/Engine/Source/Programs/UnrealBuildTool/ToolChain/AppleToolChain.cs
--- a/Engine/Source/Programs/UnrealBuildTool/ToolChain/AppleToolChain.cs
+++ b/Engine/Source/Programs/UnrealBuildTool/ToolChain/AppleToolChain.cs
@@ -997,58 +997,82 @@ namespace UnrealBuildTool
 			string MutexName = GlobalSingleInstanceMutex.GetUniqueMutexForPath("UnrealBuildTool_XcodeBuild", Unreal.RootDirectory);
 			using (new GlobalSingleInstanceMutex(MutexName, true))
 			{
-				string DestinationParam;
-				// use generic:
-				//   - for Mac, the destination seems to want name="My Mac", not the real name of the Mac
-				//   - when no destinations were supplied
-				//   - when trying to build iOS app to run on local mac 
-				//   - when using manual codesigning because the main use of this is to update provisioning information (reducing potential for issues)
-				if (Platform == UnrealTargetPlatform.Mac || DestinationIds == null || DestinationIds.Count == 0 || DestinationIds[0] == "" || DestinationIds[0].Contains(AppleToolChainSettings.LocalMacUUID) || !bUseAutomaticCodeSigning)
+				for (int Pass = 0; Pass < 2; Pass++)
 				{
-					DestinationParam = $"-destination generic/platform=\"{AppleExports.GetDestinationPlatform(Platform, Architectures)}\"";
-				}
-				else
-				{
-					string DestPlatform = AppleExports.GetDestinationPlatform(Platform, Architectures);
-					DestinationParam = string.Join(" ", DestinationIds.Select(x => $"-destination \"platform={DestPlatform},{IdOrNameParam(x)}\""));
-				}
-				List<string> Arguments = new()
-				{
-					"UBT_NO_POST_DEPLOY=true",
-					FileReference.Combine(ApplePlatformSDK.DeveloperDir, "usr/bin/xcodebuild").FullName,
-					Action,
-					$"-workspace \"{XcodeProject.FullName}\"",
-					$"-scheme \"{SchemeName}\"",
-					$"-configuration \"{Configuration}\"",
-					DestinationParam,
-					"-hideShellScriptEnvironment",
-					// xcode gets confused it we _just_ wrote out entitlements while generating the temp project, and it thinks it was modified _during_ building
-					// but it wasn't, it was written before the build started
-					"CODE_SIGN_ALLOW_ENTITLEMENTS_MODIFICATION=YES",
-					ExtraOptions,
-					//$"-sdk {SDKName}",
-				};
-
-				Console.WriteLine("params: {0}", String.Join(" ", Arguments));
-
-				Process LocalProcess = new Process();
-				LocalProcess.StartInfo = new ProcessStartInfo("/usr/bin/env", String.Join(" ", Arguments));
-				LocalProcess.OutputDataReceived += (Sender, Args) => { LocalProcessOutput(Args, false, Logger); };
-				LocalProcess.ErrorDataReceived += (Sender, Args) =>
-				{
-					if (Args != null && Args.Data != null
-					&& Args.Data.Contains("Failed to load profile") && Args.Data.Contains("<stdin>"))
+					string DestinationParam;
+					// use generic:
+					//   - for Mac, the destination seems to want name="My Mac", not the real name of the Mac
+					//   - when no destinations were supplied
+					//   - when trying to build iOS app to run on local mac 
+					//   - when using manual codesigning because the main use of this is to update provisioning information (reducing potential for issues)
+					if (Platform == UnrealTargetPlatform.Mac || DestinationIds == null || DestinationIds.Count == 0 || DestinationIds[0] == "" || DestinationIds[0].Contains(AppleToolChainSettings.LocalMacUUID) || !bUseAutomaticCodeSigning)
 					{
-						Logger.LogInformation("Silencing the following provision profile error, it is not affecting code signing:");
-						LocalProcessOutput(Args, false, Logger);
+						DestinationParam = $"-destination generic/platform=\"{AppleExports.GetDestinationPlatform(Platform, Architectures)}\"";
 					}
 					else
 					{
-						LocalProcessOutput(Args, true, Logger);
+						string DestPlatform = AppleExports.GetDestinationPlatform(Platform, Architectures);
+						DestinationParam = string.Join(" ", DestinationIds.Select(x => $"-destination \"platform={DestPlatform},{IdOrNameParam(x)}\""));
+					}
+					List<string> Arguments = new()
+					{
+						"UBT_NO_POST_DEPLOY=true",
+						FileReference.Combine(ApplePlatformSDK.DeveloperDir, "usr/bin/xcodebuild").FullName,
+						Action,
+						$"-workspace \"{XcodeProject.FullName}\"",
+						$"-scheme \"{SchemeName}\"",
+						$"-configuration \"{Configuration}\"",
+						DestinationParam,
+						// xcode gets confused it we _just_ wrote out entitlements while generating the temp project, and it thinks it was modified _during_ building
+						// but it wasn't, it was written before the build started
+						"CODE_SIGN_ALLOW_ENTITLEMENTS_MODIFICATION=YES",
+						ExtraOptions,
+						//$"-sdk {SDKName}",
+					};
+
+					if (Pass == 0)
+					{
+						Arguments.Add("-hideShellScriptEnvironment");
 					}
-				};
-				return Utils.RunLocalProcess(LocalProcess);
+
+
+					Console.WriteLine("params: {0}", String.Join(" ", Arguments));
+
+					Process LocalProcess = new Process();
+					LocalProcess.StartInfo = new ProcessStartInfo("/usr/bin/env", String.Join(" ", Arguments));
+					LocalProcess.OutputDataReceived += (Sender, Args) => { LocalProcessOutput(Args, false, Logger); };
+					LocalProcess.ErrorDataReceived += (Sender, Args) =>
+					{
+						if (Args != null && Args.Data != null
+						&& Args.Data.Contains("Failed to load profile") && Args.Data.Contains("<stdin>"))
+						{
+							Logger.LogInformation("Silencing the following provision profile error, it is not affecting code signing:");
+							LocalProcessOutput(Args, false, Logger);
+						}
+						else
+						{
+							LocalProcessOutput(Args, true, Logger);
+						}
+					};
+
+
+					// if first pass succeeded, return 0, otherwise return whatever the exit code is
+					int ExitCode = Utils.RunLocalProcess(LocalProcess);
+					if (ExitCode == 0 || Pass > 0)
+					{
+						return ExitCode;
+					}
+
+					Logger.LogWarning("");
+					Logger.LogWarning("");
+					Logger.LogWarning("**************************************************************************************************************");
+					Logger.LogWarning("* App finalization failed, running again with all environment variables displayed to help diagnose the error *");
+					Logger.LogWarning("**************************************************************************************************************");
+					Logger.LogWarning("");
+					Logger.LogWarning("");
+				}
 			}
+			return -1;
 		}
 
 		static void LocalProcessOutput(DataReceivedEventArgs? Args, bool bIsError, ILogger Logger)
