Address this error for -PGOProfile builds on Clang 19 for Windows

clang-cl: error: argument unused during compilation: '-fno-eliminate-unused-debug-types' [-Werror,-Wunused-command-line-argument]

diff --git a/Engine/Source/Programs/UnrealBuildTool/Platform/Windows/VCToolChain.cs b/Engine/Source/Programs/UnrealBuildTool/Platform/Windows/VCToolChain.cs
--- a/Engine/Source/Programs/UnrealBuildTool/Platform/Windows/VCToolChain.cs
+++ b/Engine/Source/Programs/UnrealBuildTool/Platform/Windows/VCToolChain.cs
@@ -1156,7 +1156,7 @@ namespace UnrealBuildTool
 
 					// https://clang.llvm.org/docs/UsersManual.html#cmdoption-feliminate-unused-debug-types
 					// Intel ICX 2024.2 does not have this option, but reports Clang 19, so turn it off here
-					if (ClangVersion >= new VersionNumber(19) && !Target.WindowsPlatform.Compiler.IsIntel())
+					if (ClangVersion >= new VersionNumber(19) && !Target.WindowsPlatform.Compiler.IsIntel() && !CompileEnvironment.bDebugLineTablesOnly)
 					{
 						Arguments.Add("-fno-eliminate-unused-debug-types");
 					}
