[UBT]
* Improved error  message when finding same inlined gen cpp multiple times

diff --git a/Engine/Source/Programs/UnrealBuildTool/Configuration/UEBuildModuleCPP.cs b/Engine/Source/Programs/UnrealBuildTool/Configuration/UEBuildModuleCPP.cs
--- a/Engine/Source/Programs/UnrealBuildTool/Configuration/UEBuildModuleCPP.cs
+++ b/Engine/Source/Programs/UnrealBuildTool/Configuration/UEBuildModuleCPP.cs
@@ -661,7 +661,7 @@ namespace UnrealBuildTool
 							}
 							else
 							{
-								Logger.LogError("'{CPPFileItem}' is looking for a generated cpp with named '{HeaderFile}.gen.cpp'", CPPFileItem.AbsolutePath, ListOfInlinedGenCppsItem);
+								Logger.LogError("'{CPPFileItem}' is looking for a generated cpp with named '{HeaderFile}.gen.cpp' (Found {FoundCount}) time)", CPPFileItem.AbsolutePath, ListOfInlinedGenCppsItem, ListOfInlinedGenCpps.Count);
 							}
 						}
 					}
