[UBT]
* Reduced the number of times CppDependencyCache calls GetCurrentDirectory since it creates lots of thread contention. Added todo for adding this from the outside somewhere

diff --git a/Engine/Source/Programs/UnrealBuildTool/System/CppDependencyCache.cs b/Engine/Source/Programs/UnrealBuildTool/System/CppDependencyCache.cs
--- a/Engine/Source/Programs/UnrealBuildTool/System/CppDependencyCache.cs
+++ b/Engine/Source/Programs/UnrealBuildTool/System/CppDependencyCache.cs
@@ -447,10 +447,20 @@ namespace UnrealBuildTool
 
 				TokenIdx += 2;
 
+				// TODO: Move this to the provided from the outside somewhere.
+				string CurrentDir = string.Empty;
+
 				List<FileItem> NewDependencyFiles = new List<FileItem>();
 				for (; TokenIdx < Tokens.Count && Tokens[TokenIdx] != "\n"; TokenIdx++)
 				{
-					NewDependencyFiles.Add(FileItem.GetItemByPath(Tokens[TokenIdx]));
+					string FilePath = Tokens[TokenIdx];
+					if (!Path.IsPathFullyQualified(FilePath))
+					{
+						if (CurrentDir == string.Empty)
+							CurrentDir = Directory.GetCurrentDirectory(); // Get current directory has crazy contention. So we only want to call it once
+						FilePath = Path.Combine(CurrentDir, FilePath);
+					}
+					NewDependencyFiles.Add(FileItem.GetItemByPath(FilePath));
 				}
 
 				while (TokenIdx < Tokens.Count && Tokens[TokenIdx] == "\n")
