PipInstall: Update ubt pip install mode to also support Arm64/X64 arch tags.

diff --git a/Engine/Source/Programs/UnrealBuildTool/Modes/PipInstallMode.cs b/Engine/Source/Programs/UnrealBuildTool/Modes/PipInstallMode.cs
--- a/Engine/Source/Programs/UnrealBuildTool/Modes/PipInstallMode.cs
+++ b/Engine/Source/Programs/UnrealBuildTool/Modes/PipInstallMode.cs
@@ -123,7 +123,6 @@ namespace UnrealBuildTool.Modes
 
 			DirectoryReference ProjectDir = DirectoryReference.FromFile(TargetDescriptor.ProjectFile);
 			DirectoryReference InstallDir = DirectoryReference.Combine(ProjectDir, "Intermediate", "PipInstall");
-			UnrealTargetPlatform Platform = TargetDescriptor.Platform;
 
 			// Let env variable override pip install path
 			DirectoryReference? EnvInstallPath = DirectoryReference.FromString(Environment.GetEnvironmentVariable("UE_PIPINSTALL_PATH"));
@@ -132,7 +131,7 @@ namespace UnrealBuildTool.Modes
 				InstallDir = EnvInstallPath;
 			}
 
-			PipEnv Pip = new(InstallDir, Platform, Logger, ProgressWriter.bWriteMarkup);
+			PipEnv Pip = new(InstallDir, TargetDescriptor, Logger, ProgressWriter.bWriteMarkup);
 			if ((Action & (PipAction)ActionBits.GenReqs) != 0)
 			{
 				// Make sure the virtual environment used for installs is compatible with python interpreter version
@@ -201,6 +200,7 @@ namespace UnrealBuildTool.Modes
 		private const string ExtraUrlsFilename = "extra_urls.txt";
 		private const string MergedRequirementsFilename = "merged_requirements.txt";
 
+		private UnrealArch TargetArch;
 		private UnrealTargetPlatform TargetPlatform;
 		private DirectoryReference InstallDir;
 		private FileReference PythonVenvExe;
@@ -209,11 +209,17 @@ namespace UnrealBuildTool.Modes
 		private ILogger Logger;
 		private IBaseProgressLogFactory LoggerFactory;
 
-		public PipEnv(DirectoryReference InInstallDir, UnrealTargetPlatform InPlatform, ILogger InLogger, bool UseProgressWriter = false)
+		public PipEnv(DirectoryReference InInstallDir, TargetDescriptor InTargetDesc, ILogger InLogger, bool UseProgressWriter = false)
 		{
+			if (InTargetDesc.Architectures.bIsMultiArch)
+			{
+				InLogger.LogWarning("PipInstall expects single-arch target(host): {arch}", InTargetDesc.Architectures.ToString());
+			}
+			
 			Logger = InLogger;
 			InstallDir = InInstallDir;
-			TargetPlatform = InPlatform;
+			TargetPlatform = InTargetDesc.Platform;
+			TargetArch = InTargetDesc.Architectures.SingleArchitecture;
 
 			LoggerFactory = (UseProgressWriter) ? new PipProgressLogCreator(Logger) : new SimpleCmdLogCreator(Logger);
 
@@ -257,10 +263,9 @@ namespace UnrealBuildTool.Modes
 		public void WritePluginsListing(UEBuildTarget Target, ILogger Logger, bool bAllPlugins = false)
 		{
 			// TODO: The path listing file won't match the .pth generated in-engine.
-			// In particular the additional paths setting
+			//		 In particular the additional paths setting
 
 			FileReference PluginsListingFile = FileReference.Combine(InstallDir, PluginsListingFilename);
-
 			DirectoryReference PipSitePackagesPath = DirectoryReference.Combine(InstallDir, "Lib", "site-packages");
 			FileReference PyPluginsSitePackageFile = FileReference.Combine(PipSitePackagesPath, PluginsSitePackageFilename);
 
@@ -287,23 +292,29 @@ namespace UnrealBuildTool.Modes
 				return;
 			}
 
+			// May need to handle this differently
+			string ArchStr = MakeArchString(TargetArch);
+
 			List<string> PluginsSitePackages = new List<string>();
 			List<string> PluginsList = new List<string>();
+
+			string[] SitePkgSubdirs = { 
+				Path.Combine("Lib", Target.Platform.ToString(), ArchStr, "site-packages"),
+				Path.Combine("Lib", Target.Platform.ToString(), "site-packages"),
+				Path.Combine("Lib", "site-packages")
+			};
+			
 			foreach (PluginInfo Plugin in CheckPlugins)
 			{
+				// Add all existing plugin site-packages subdirs to be written to pth file listing
 				DirectoryReference PythonContentPath = DirectoryReference.Combine(Plugin.Directory, "Content", "Python");
-				DirectoryReference PluginPlatformSitePackagesPath = DirectoryReference.Combine(PythonContentPath, "Lib", Target.Platform.ToString(), "site-packages");
-				DirectoryReference PluginGeneralSitePackagesPath = DirectoryReference.Combine(PythonContentPath, "Lib", "site-packages");
-
-				// Write platform/general site-packages paths per-plugin to .pth file to account for packaged python dependencies during pip install
-				if (DirectoryReference.Exists(PluginPlatformSitePackagesPath))
+				foreach (string CheckSubdir in SitePkgSubdirs)
 				{
-					PluginsSitePackages.Add(PluginPlatformSitePackagesPath.ToString());
-				}
-
-				if (DirectoryReference.Exists(PluginGeneralSitePackagesPath))
-				{
-					PluginsSitePackages.Add(PluginPlatformSitePackagesPath.ToString());
+					DirectoryReference CheckPluginPath = DirectoryReference.Combine(PythonContentPath, CheckSubdir);
+					if (DirectoryReference.Exists(CheckPluginPath))
+					{
+						PluginsSitePackages.Add(CheckPluginPath.ToString());
+					}
 				}
 
 				if (!JsonObject.TryRead(Plugin.File, out JsonObject? PluginJson))
@@ -312,7 +323,7 @@ namespace UnrealBuildTool.Modes
 					continue;
 				}
 
-				foreach (JsonObject PlatformReqs in PipEnv.CompatibleRequirements(PluginJson, Target.Platform))
+				foreach (JsonObject PlatformReqs in PipEnv.CompatibleRequirements(PluginJson, TargetPlatform, TargetArch))
 				{
 					PluginsList.Add(Plugin.File.ToString());
 					break;
@@ -438,7 +449,7 @@ namespace UnrealBuildTool.Modes
 			}
 		}
 
-		public static IEnumerable<JsonObject> CompatibleRequirements(JsonObject PluginJson, UnrealTargetPlatform Platform)
+		public static IEnumerable<JsonObject> CompatibleRequirements(JsonObject PluginJson, UnrealTargetPlatform Platform, UnrealArch Arch)
 		{
 			// Check for python requirements field
 			if (!PluginJson.TryGetObjectArrayField("PythonRequirements", out JsonObject[]? RequirementsJson))
@@ -453,6 +464,12 @@ namespace UnrealBuildTool.Modes
 				{
 					continue;
 				}
+				
+				PlatformReqs.TryGetStringField("Architecture", out string? ArchField);
+				if (!CompatibleArch(ArchField, Arch))
+				{
+					continue;
+				}
 
 				yield return PlatformReqs;
 			}
@@ -465,6 +482,28 @@ namespace UnrealBuildTool.Modes
 				|| String.Equals(PlatformField, "All", StringComparison.InvariantCultureIgnoreCase);
 		}
 
+		private static bool CompatibleArch(string? ArchField, UnrealArch Arch)
+		{
+			return (ArchField == null)
+				|| String.Equals(ArchField, MakeArchString(Arch), StringComparison.InvariantCultureIgnoreCase)
+				|| String.Equals(ArchField, "All", StringComparison.InvariantCultureIgnoreCase);
+		}
+
+		private static string MakeArchString(UnrealArch Arch)
+		{
+			if (Arch == UnrealArch.X64)
+			{
+				return "X64";
+			}
+
+			if (Arch == UnrealArch.Arm64 || Arch == UnrealArch.Arm64ec)
+			{
+				return "Arm64";
+			}
+
+			return "Unknown";
+		}
+		
 		private string? ParseVenvVersion(DirectoryReference VenvDir)
 		{
 			FileReference VenvConfig = FileReference.Combine(VenvDir, "pyvenv.cfg");
@@ -574,7 +613,7 @@ namespace UnrealBuildTool.Modes
 					continue;
 				}
 
-				foreach (JsonObject PlatformReqs in CompatibleRequirements(PluginJson, TargetPlatform))
+				foreach (JsonObject PlatformReqs in CompatibleRequirements(PluginJson, TargetPlatform, TargetArch))
 				{
 					if (!PlatformReqs.TryGetStringArrayField("Requirements", out string[]? Reqs))
 					{
