Fixed a prior change the removed some diagnostic logging when UHT is run as a mode of UBT.

#rnx

diff --git a/Engine/Source/Programs/Shared/EpicGames.UHT/Utils/UhtSession.cs b/Engine/Source/Programs/Shared/EpicGames.UHT/Utils/UhtSession.cs
--- a/Engine/Source/Programs/Shared/EpicGames.UHT/Utils/UhtSession.cs
+++ b/Engine/Source/Programs/Shared/EpicGames.UHT/Utils/UhtSession.cs
@@ -2702,7 +2702,6 @@ namespace EpicGames.UHT.Utils
 		private void StepExport()
 		{
 			HashSet<string> externalDependencies = new();
-			int totalWrittenFiles = 0;
 			Try(null, () =>
 			{
 				Logger.LogDebug("Step - Exports");
@@ -2742,7 +2741,7 @@ namespace EpicGames.UHT.Utils
 						{
 							if (output.Saved)
 							{
-								totalWrittenFiles++;
+								TotalFilesWritten++;
 							}
 						}
 						foreach (string dep in factory.ExternalDependencies)
@@ -2766,8 +2765,6 @@ namespace EpicGames.UHT.Utils
 					}
 				}
 			});
-
-			TotalFilesWritten = totalWrittenFiles;
 		}
 		#endregion
 
diff --git a/Engine/Source/Programs/UnrealBuildTool/Modes/UnrealHeaderToolMode.cs b/Engine/Source/Programs/UnrealBuildTool/Modes/UnrealHeaderToolMode.cs
--- a/Engine/Source/Programs/UnrealBuildTool/Modes/UnrealHeaderToolMode.cs
+++ b/Engine/Source/Programs/UnrealBuildTool/Modes/UnrealHeaderToolMode.cs
@@ -1,17 +1,18 @@
 // Copyright Epic Games, Inc. All Rights Reserved.
 
-using System;
-using System.Collections.Generic;
-using System.IO;
-using System.Reflection;
-using System.Text;
-using System.Threading.Tasks;
 using EpicGames.Core;
 using EpicGames.UHT.Tables;
 using EpicGames.UHT.Tokenizer;
 using EpicGames.UHT.Types;
 using EpicGames.UHT.Utils;
 using Microsoft.Extensions.Logging;
+using System;
+using System.Collections.Generic;
+using System.Diagnostics;
+using System.IO;
+using System.Reflection;
+using System.Text;
+using System.Threading.Tasks;
 using UnrealBuildBase;
 
 namespace UnrealBuildTool.Modes
@@ -1133,6 +1134,10 @@ namespace UnrealBuildTool.Modes
 			try
 			{
 
+				// Start a time to track runtime
+				Stopwatch stopwatch = new();
+				stopwatch.Start();
+
 				// Initialize the attributes
 				UhtTables Tables = new();
 
@@ -1270,8 +1275,16 @@ namespace UnrealBuildTool.Modes
 
 				// Read and parse
 				Session.Run(ManifestPath!, Arguments);
-				TotalFilesWritten = Session.TotalFilesWritten;
 				Session.LogMessages();
+
+				// Generate summary message
+				stopwatch.Stop();
+				if (!Session.HasErrors || Session.TotalFilesWritten != 0)
+				{
+					string ActualTargetName = Session.Manifest != null ? Session.Manifest.TargetName : "UnknownTarget";
+					Logger.LogInformation("UHT processed {TargetName} in {Time} seconds ({TotalFilesWritten} files written)", ActualTargetName, stopwatch.Elapsed.TotalSeconds, Session.TotalFilesWritten);
+				}
+
 				return (int)(Session.HasErrors ? CompilationResult.OtherCompilationError : CompilationResult.Succeeded);
 			}
 			catch (Exception Ex)
@@ -1282,7 +1295,5 @@ namespace UnrealBuildTool.Modes
 				return (int)CompilationResult.OtherCompilationError;
 			}
 		}
-
-		public int TotalFilesWritten { get; private set; }
 	}
 }
diff --git a/Engine/Source/Programs/UnrealBuildTool/System/ExternalExecution.cs b/Engine/Source/Programs/UnrealBuildTool/System/ExternalExecution.cs
--- a/Engine/Source/Programs/UnrealBuildTool/System/ExternalExecution.cs
+++ b/Engine/Source/Programs/UnrealBuildTool/System/ExternalExecution.cs
@@ -1298,13 +1298,10 @@ namespace UnrealBuildTool
 					Logger.LogInformation("  Running Internal UnrealHeaderTool {CmdLine}", CmdLine);
 
 					// Run UHT
-					Stopwatch s = new Stopwatch();
-					s.Start();
 					IScope Timer = GlobalTracer.Instance.BuildSpan("Executing UnrealHeaderTool").StartActive();
 					UnrealHeaderToolMode UHTTool = new UnrealHeaderToolMode();
 					CompilationResult UHTResult = (CompilationResult)await UHTTool.ExecuteAsync(Arguments, Logger);
 					Timer.Span.Finish();
-					s.Stop();
 
 					if (UHTResult != CompilationResult.Succeeded)
 					{
@@ -1327,8 +1324,6 @@ namespace UnrealBuildTool
 						throw new CompilationResultException(UHTResult);
 					}
 
-					Logger.LogInformation("Reflection code generated for {TargetName} in {Time} seconds ({TotalFilesWritten} files written)", ActualTargetName, s.Elapsed.TotalSeconds, UHTTool.TotalFilesWritten);
-
 					// Update the tool info file
 					DirectoryReference.CreateDirectory(ToolInfoFile.Directory);
 					FileReference.WriteAllText(ToolInfoFile, UhtAssemblyPath);
