[UBT]
* Disabled -ProxyUseLocalStorage for horde-spawned UbaAgents due to suspecting cause of seen hang

diff --git a/Engine/Source/Programs/UnrealBuildTool/Executors/UnrealBuildAccelerator/UBAAgentCoordinatorHorde.cs b/Engine/Source/Programs/UnrealBuildTool/Executors/UnrealBuildAccelerator/UBAAgentCoordinatorHorde.cs
--- a/Engine/Source/Programs/UnrealBuildTool/Executors/UnrealBuildAccelerator/UBAAgentCoordinatorHorde.cs
+++ b/Engine/Source/Programs/UnrealBuildTool/Executors/UnrealBuildAccelerator/UBAAgentCoordinatorHorde.cs
@@ -588,7 +588,7 @@ namespace UnrealBuildTool
 						arguments.Add($"-MaxWorkers={hordeConfig.HordeMaxWorkers}");
 						arguments.Add($"-MaxIdle={hordeConfig.HordeMaxIdle}");
 						arguments.Add("-UseCrawler");
-						arguments.Add("-ProxyUseLocalStorage"); // If set then agent's cas storage will be used by proxy
+						// arguments.Add("-ProxyUseLocalStorage"); // If set then agent's cas storage will be used by proxy. [honk] This is disabled because we suspect a hang in this feature
 						//arguments.Add("-UseIocp=4");
 						//arguments.Add("-NoStore"); // This means that no cas files will be written to disk.. so no state is stored between runs
 
