Add support for compiling protobuf files in modules stored in a Proto directory, and customize compilation options for generated source files to disable unnecessary checks and analysis (warnings, address sanitization, etc.)

diff --git a/Engine/Source/Programs/UnrealBuildTool/Configuration/UEBuildModuleCPP.cs b/Engine/Source/Programs/UnrealBuildTool/Configuration/UEBuildModuleCPP.cs
--- a/Engine/Source/Programs/UnrealBuildTool/Configuration/UEBuildModuleCPP.cs
+++ b/Engine/Source/Programs/UnrealBuildTool/Configuration/UEBuildModuleCPP.cs
@@ -37,6 +37,7 @@ namespace UnrealBuildTool
 			public readonly List<FileItem> SwiftFiles = new List<FileItem>();
 			public readonly List<FileItem> RCFiles = new List<FileItem>();
 			public readonly List<FileItem> ISPCFiles = new List<FileItem>();
+			public readonly List<FileItem> ProtoFiles = new List<FileItem>();
 
 			public IEnumerable<FileItem> AllHeaderFiles => HeaderFiles.Concat(ISPCHeaderFiles);
 			public IEnumerable<FileItem> AllCompiledFiles => IXXFiles.Concat(CPPFiles).Concat(CFiles).Concat(CCFiles).Concat(CXXFiles).Concat(MFiles).Concat(MMFiles).Concat(SwiftFiles).Concat(RCFiles).Concat(ISPCFiles);
@@ -62,6 +63,11 @@ namespace UnrealBuildTool
 		/// </summary>
 		public bool bHasUObjects;
 
+		/// <summary>
+		/// If any proto files were found in the module's source files
+		/// </summary>
+		public bool bHasProto;
+
 		/// <summary>
 		/// The directory for this module's generated code
 		/// </summary>
@@ -77,6 +83,11 @@ namespace UnrealBuildTool
 		/// </summary>
 		public DirectoryReference? GeneratedCodeDirectoryVNI => GeneratedCodeDirectory != null ? DirectoryReference.Combine(GeneratedCodeDirectory!, "VNI") : null;
 
+		/// <summary>
+		/// The directory for this module's generated Protobuf code
+		/// </summary>
+		public DirectoryReference? GeneratedCodeDirectoryProto => GeneratedCodeDirectory != null ? DirectoryReference.Combine(GeneratedCodeDirectory!, "Proto") : null;
+
 		/// <summary>
 		/// Global override to force all include paths to be always added
 		/// </summary>
@@ -258,6 +269,13 @@ namespace UnrealBuildTool
 				AssociatedVerseDirectory = MaybeVerseDirectory;
 				bDependsOnVerse = true;
 			}
+			
+			// Check for a Proto directory next to the rules file
+			DirectoryReference MaybeProtoDirectory = DirectoryReference.Combine(Rules.File.Directory, "Proto");
+			if (IsValidProtoDirectory(MaybeProtoDirectory))
+			{
+				bHasProto = true;
+			}
 
 			if (Rules.bAddDefaultIncludePaths)
 			{
@@ -438,6 +456,11 @@ namespace UnrealBuildTool
 				{
 					IncludePaths.Add(GeneratedCodeDirectoryVNI!);
 				}
+				
+				if (bHasProto || bForceAddIncludePath)
+				{
+					IncludePaths.Add(GeneratedCodeDirectoryProto!);
+				}
 
 				foreach ((string Subdirectory, Action<ILogger, DirectoryReference> _) in Rules.GenerateHeaderFuncs)
 				{
@@ -910,6 +933,34 @@ namespace UnrealBuildTool
 				}
 			}
 
+			if (bHasProto && GeneratedCodeDirectoryProto != null && DirectoryReference.Exists(GeneratedCodeDirectoryProto))
+			{
+				CppCompileEnvironment ProtoCompileEnvironment = new CppCompileEnvironment(CppCompileEnvironment);
+				ProtoCompileEnvironment.PrecompiledHeaderAction = PrecompiledHeaderAction.None;
+				ProtoCompileEnvironment.bUseUnity = false;
+				ProtoCompileEnvironment.bWarningsAsErrors = false;
+				// Disable dynamic analysis on generated proto code as it causes issues when the proto/grpc libs are not
+				// compiled the same way. This can cause false positive in some situations, for example with address
+				// sanitize enabled (see: [REDACTED_PROFILE_URL]/protobuf/issues/13115).
+				ProtoCompileEnvironment.bDisableDynamicAnalysis = true;
+
+				IEnumerable<FileItem> ProtoCPPFiles = from x in DirectoryReference.EnumerateFiles(GeneratedCodeDirectoryProto, "*.pb.cc", SearchOption.AllDirectories) select FileItem.GetItemByFileReference(x);
+				CPPOutput Output = ToolChain.CompileAllCPPFiles(ProtoCompileEnvironment, ProtoCPPFiles, IntermediateDirectory, Name, Graph);
+
+				LinkInputFiles.AddRange(Output.ObjectFiles);
+
+				// TODO(jamal.fanaian): Temporary hack to deal with `verify` macro causing issues with absl having a `verify` method
+				foreach (FileReference HeaderFile in DirectoryReference.EnumerateFiles(GeneratedCodeDirectoryProto, "*.pb.h", SearchOption.AllDirectories))
+				{
+					// TODO(jamal.fanaian): Use a StreamReader so we don't load the entire file in memory
+					string Content = FileReference.ReadAllText(HeaderFile);
+					if (!Content.StartsWith("#undef verify"))
+					{
+						FileReference.WriteAllText(HeaderFile, "#undef verify" + Environment.NewLine + Content);
+					}
+				}
+			}
+
 			// Compile CPP files
 			if (CPPFiles.Count > 0)
 			{
@@ -2742,6 +2793,7 @@ namespace UnrealBuildTool
 			".cpp", ".cxx", ".ixx", ".c", ".cc",
 			".m", ".mm", ".swift",
 			".rc", ".ispc",
+			".proto",
 		};
 
 		/// <summary>
@@ -2829,6 +2881,10 @@ namespace UnrealBuildTool
 					SourceFiles.Add(InputFile);
 					InputFiles.ISPCFiles.Add(InputFile);
 				}
+				else if (InputFile.HasExtension(".proto"))
+				{
+					InputFiles.ProtoFiles.Add(InputFile);
+				}
 				// if you add any extension checks here, make sure to update KnownInputFileExtensions
 			}
 			return SourceFiles.ToArray();
@@ -2902,5 +2958,19 @@ namespace UnrealBuildTool
 
 			return false;
 		}
+		
+
+		/// <summary>
+		/// Checks a given directory path whether it exists and if it contains any Proto source files
+		/// </summary>
+		public static bool IsValidProtoDirectory(DirectoryReference MaybeProtoDirectory)
+		{
+			if (!DirectoryReference.Exists(MaybeProtoDirectory))
+			{
+				return false;
+			}
+
+			return Directory.EnumerateFiles(MaybeProtoDirectory.FullName, "*.proto", SearchOption.AllDirectories).Any();
+		}
 	}
 }
diff --git a/Engine/Source/Programs/UnrealBuildTool/Configuration/UEBuildTarget.cs b/Engine/Source/Programs/UnrealBuildTool/Configuration/UEBuildTarget.cs
--- a/Engine/Source/Programs/UnrealBuildTool/Configuration/UEBuildTarget.cs
+++ b/Engine/Source/Programs/UnrealBuildTool/Configuration/UEBuildTarget.cs
@@ -2611,6 +2611,11 @@ namespace UnrealBuildTool
 			{
 				return Makefile;
 			}
+			
+			using (GlobalTracer.Instance.BuildSpan("ProtoExecution.SetupProtoModules()").StartActive())
+			{
+				ProtoExecution.SetupProtoModules(ModulesToGenerateHeadersFor, Makefile, Logger);
+			}
 
 #if __VPROJECT_AVAILABLE__
 			// Copy Verse BPVM usage flag
@@ -2645,6 +2650,13 @@ namespace UnrealBuildTool
 				// Create all the shared PCH instances before processing the modules
 				CreateSharedPCHInstances(Rules, TargetToolChain, OriginalBinaries, GlobalCompileEnvironment, MakefileBuilder, Logger);
 			}
+			
+			if (Makefile.ProtoModules.Count > 0)
+			{
+				//GlobalCompileEnvironment.ExtraGeneratedCPPFileTypes.Add("*.pb.cc");
+				// ProtoExecution.CreateProtocActions(MakefileBuilder);
+				ProtoExecution.ExecuteProtocActions(Makefile, Logger);
+			}
 
 			// Can probably be moved further down?
 #if __VPROJECT_AVAILABLE__
diff --git a/Engine/Source/Programs/UnrealBuildTool/Platform/Windows/VCToolChain.cs b/Engine/Source/Programs/UnrealBuildTool/Platform/Windows/VCToolChain.cs
--- a/Engine/Source/Programs/UnrealBuildTool/Platform/Windows/VCToolChain.cs
+++ b/Engine/Source/Programs/UnrealBuildTool/Platform/Windows/VCToolChain.cs
@@ -791,7 +791,8 @@ namespace UnrealBuildTool
 			}
 
 			// Address sanitizer
-			if (Target.WindowsPlatform.bEnableAddressSanitizer)
+			// TODO(jamal.fanaian): We should probably not make this a global property instead?
+			if (Target.WindowsPlatform.bEnableAddressSanitizer && !CompileEnvironment.bDisableDynamicAnalysis)
 			{
 				// Enable address sanitizer. This also requires companion libraries at link time.
 				// Works for clang too.
diff --git a/Engine/Source/Programs/UnrealBuildTool/System/CppCompileEnvironment.cs b/Engine/Source/Programs/UnrealBuildTool/System/CppCompileEnvironment.cs
--- a/Engine/Source/Programs/UnrealBuildTool/System/CppCompileEnvironment.cs
+++ b/Engine/Source/Programs/UnrealBuildTool/System/CppCompileEnvironment.cs
@@ -448,6 +448,11 @@ namespace UnrealBuildTool
 		/// </summary>
 		public HashSet<string> StaticAnalyzerPVSDisabledErrors = [];
 
+		/// <summary>
+		/// Whether to disable all dynamic analysis (such as Address Sanitizer)
+		/// </summary>
+		public bool bDisableDynamicAnalysis = false;
+
 		/// <summary>
 		/// True if compiler optimizations should be enabled. This setting is distinct from the configuration (see CPPTargetConfiguration).
 		/// </summary>
diff --git a/Engine/Source/Programs/UnrealBuildTool/System/ProtoExecution.cs b/Engine/Source/Programs/UnrealBuildTool/System/ProtoExecution.cs
new file mode 100644
--- /dev/null
+++ b/Engine/Source/Programs/UnrealBuildTool/System/ProtoExecution.cs
@@ -0,0 +1,306 @@
+// Copyright Epic Games, Inc. All Rights Reserved.
+
+using System;
+using System.Collections.Generic;
+using System.Data;
+using System.Diagnostics;
+using System.Diagnostics.CodeAnalysis;
+using System.IO;
+using System.Linq;
+using System.Reflection;
+using System.Text;
+using System.Text.Json;
+using System.Threading.Tasks;
+using EpicGames.Core;
+using EpicGames.UHT.Utils;
+using Microsoft.Extensions.Logging;
+using OpenTracing;
+using OpenTracing.Util;
+using UnrealBuildBase;
+using UnrealBuildTool.Modes;
+
+namespace UnrealBuildTool
+{
+	/// <summary>
+	/// Information about a module that needs to be passed to protoc for code generation
+	/// </summary>
+	public class ProtoModuleInfo
+	{
+		/// <summary>
+		/// Module name
+		/// </summary>
+		public required string ModuleName;
+
+		/// <summary>
+		/// Path to the module rules file
+		/// </summary>
+		public required FileReference ModuleRulesFile;
+		
+		/// <summary>
+		/// Directory containing generated code
+		/// </summary>
+		public required DirectoryItem GeneratedCodeDirectory;
+		
+		/// <summary>
+		/// Directory containing proto source files
+		/// </summary>
+		public required DirectoryReference ProtoDirectory;
+		
+		/// <summary>
+		/// List of source files to compile
+		/// </summary>
+		public List<FileItem> ProtoFiles = new List<FileItem>();
+		
+		/// <summary>
+		/// Object initializer for ProtoModuleInfo
+		/// </summary>
+		public ProtoModuleInfo()
+		{
+		}
+		
+		/// <summary>
+		/// Load ProtoModuleInfo from a binary archive
+		/// </summary>
+		/// <param name="Reader"></param>
+		[SetsRequiredMembers]
+		public ProtoModuleInfo(BinaryArchiveReader Reader)
+		{
+			ModuleName = Reader.ReadString()!;
+			ModuleRulesFile = Reader.ReadFileReference();
+			GeneratedCodeDirectory = Reader.ReadDirectoryItem()!;
+			ProtoDirectory = Reader.ReadDirectoryReference()!;
+			ProtoFiles = Reader.ReadList(() => Reader.ReadFileItem()!) ?? new List<FileItem>();
+		}
+
+		/// <summary>
+		/// Write ProtoModuleInfo as a binary archive
+		/// </summary>
+		/// <param name="Writer"></param>
+		public void Write(BinaryArchiveWriter Writer)
+		{
+			Writer.WriteString(ModuleName);
+			Writer.WriteFileReference(ModuleRulesFile);
+			Writer.WriteDirectoryItem(GeneratedCodeDirectory);
+			Writer.WriteDirectoryReference(ProtoDirectory);
+			Writer.WriteList(ProtoFiles, x => Writer.WriteFileItem(x));
+		}
+	}
+
+	class ProtocAction : Action
+	{
+		public ProtocAction(ProtoModuleInfo Module, FileItem ProtoFile) : base(ActionType.Compile)
+		{
+			DirectoryReference ProtobufIncludeDir = DirectoryReference.Combine(ProtoExecution.GetProtobufDirectory(), "include");
+			DirectoryReference ProtoDir = DirectoryReference.Combine(Module.ModuleRulesFile.Directory, "Proto");
+
+			PrerequisiteItems.Add(ProtoFile);
+			string ProducedName = Path.GetRelativePath(DirectoryReference.Combine(Module.ModuleRulesFile.Directory, "Proto").ToString(), ProtoFile.FullName);
+			ProducedItems.Add(FileItem.GetItemByPath(Path.ChangeExtension(ProducedName, "pb.h")));
+			ProducedItems.Add(FileItem.GetItemByPath(Path.ChangeExtension(ProducedName, "pb.cc")));
+			CommandPath = ProtoExecution.GetProtocPath();
+			CommandArguments = $" -I {ProtobufIncludeDir}";
+			CommandArguments += $" -I {ProtoDir}";
+			CommandArguments += $" --cpp_out={Module.GeneratedCodeDirectory}";
+		}
+
+		public ProtocAction(BinaryArchiveReader Reader) : base(ActionType.Compile)
+		{
+			PrerequisiteItems = Reader.ReadSortedSet(Reader.ReadFileItem)!;
+			ProducedItems = Reader.ReadSortedSet(Reader.ReadFileItem)!;
+			CommandPath = Reader.ReadFileReference()!;
+			CommandArguments = Reader.ReadString()!;
+		}
+
+		/// <inheritdoc/>
+		public new void Write(BinaryArchiveWriter Writer)
+		{
+			Writer.WriteSortedSet(PrerequisiteItems, Writer.WriteFileItem);
+			Writer.WriteSortedSet(ProducedItems, Writer.WriteFileItem);
+			Writer.WriteFileReference(CommandPath);
+			Writer.WriteString(CommandArguments);
+		}
+	}
+
+	/// <summary>
+	/// Serializer for <see cref="ProtocAction"/> instances
+	/// </summary>
+	class ProtocActionSerializer : ActionSerializerBase<ProtocAction>
+	{
+		/// <inheritdoc/>
+		public override ProtocAction Read(BinaryArchiveReader reader)
+		{
+			return new(reader);
+		}
+
+		/// <inheritdoc/>
+		public override void Write(BinaryArchiveWriter writer, ProtocAction action)
+		{
+			action.Write(writer);
+		}
+	}
+
+
+	enum ProtocOutput
+	{
+		Cpp,
+		GrpcCpp,
+	}
+
+	static class ProtocOutputExtension
+	{
+		public static string GetArgs(this ProtocOutput Output, DirectoryReference OutputLocation)
+		{
+			switch (Output)
+			{
+				case ProtocOutput.Cpp:
+					return $"--cpp_out={OutputLocation}";
+				case ProtocOutput.GrpcCpp:
+					return $"--plugin=\"protoc-gen-grpc={ProtoExecution.GetGrpcCppPluginPath()}\" --grpc_out={OutputLocation}";
+			}
+
+			return "";
+		}
+	}
+
+	class ProtoExecution
+	{
+		public static readonly string ProtobufVersion = "31.1";
+		public static readonly string GrpcVersion = "1.72.1";
+
+		public static void SetupProtoModules(
+			IEnumerable<UEBuildModuleCPP> Modules,
+			TargetMakefile Makefile,
+			ILogger Logger)
+		{
+			foreach (UEBuildModuleCPP Module in Modules)
+			{
+				Makefile.ProtoModules.AddRange(SetupProtoModulesRecursively(Module));
+			}
+		}
+
+		static List<ProtoModuleInfo> SetupProtoModulesRecursively(UEBuildModuleCPP Module)
+		{
+			List<ProtoModuleInfo> ProtoModules = new List<ProtoModuleInfo>();
+
+			if (Module.bHasProto)
+			{
+				DirectoryItem GeneratedCodeDirectory = DirectoryItem.GetItemByDirectoryReference(Module.GeneratedCodeDirectoryProto!);
+
+				ProtoModuleInfo ProtoModule = new ProtoModuleInfo
+				{
+					ModuleName = Module.Name,
+					ModuleRulesFile = Module.Rules.File,
+					GeneratedCodeDirectory = GeneratedCodeDirectory,
+					ProtoDirectory = DirectoryReference.Combine(Module.Rules.File.Directory, "Proto"),
+				};
+
+				ProtoModules.Add(ProtoModule);
+
+				// Module.GeneratedCppDirectories ??= new List<string>();
+				// Module.GeneratedCppDirectories.Add(GeneratedCodeDirectory.FullName);
+			}
+
+			// if (Module.PublicDependencyModules != null)
+			// {
+			// 	foreach (UEBuildModuleCPP DependencyModule in Module.PublicDependencyModules.OfType<UEBuildModuleCPP>())
+			// 	{
+			// 		ProtoModules.AddRange(SetupProtoModulesRecursively(DependencyModule));
+			// 	}
+			// }
+
+			// if (Module.PrivateDependencyModules != null)
+			// {
+			// 	foreach (UEBuildModuleCPP DependencyModule in Module.PrivateDependencyModules.OfType<UEBuildModuleCPP>())
+			// 	{
+			// 		ProtoModules.AddRange(SetupProtoModulesRecursively(DependencyModule));
+			// 	}
+			// }
+
+			return ProtoModules;
+		}
+
+		public static void CreateProtocActions(TargetMakefileBuilder MakefileBuilder)
+		{
+			if (MakefileBuilder.Makefile.ProtoModules.Count > 0)
+			{
+				foreach (ProtoModuleInfo ProtoModule in MakefileBuilder.Makefile.ProtoModules)
+				{
+					DirectoryItem ProtoDirectory = DirectoryItem.GetItemByDirectoryReference(ProtoModule.ProtoDirectory);
+					foreach (FileReference ProtoFile in DirectoryReference.EnumerateFiles(ProtoModule.ProtoDirectory, "*.proto", SearchOption.AllDirectories))
+					{
+						MakefileBuilder.AddAction(new ProtocAction(ProtoModule, FileItem.GetItemByFileReference(ProtoFile)));
+					}
+				}
+			}
+		}
+
+		public static void ExecuteProtocActions(TargetMakefile Makefile, ILogger Logger)
+		{
+			if (Makefile.ProtoModules.Count > 0)
+			{
+				foreach (ProtoModuleInfo ProtoModule in Makefile.ProtoModules)
+				{
+					DirectoryReference.CreateDirectory(ProtoModule.GeneratedCodeDirectory.Location);
+					foreach (FileReference ProtoFile in DirectoryReference.EnumerateFiles(ProtoModule.ProtoDirectory, "*.proto", SearchOption.AllDirectories))
+					{
+						// TODO: Compile multiple files in one exec
+						ExecuteProtocAction(ProtoModule, ProtoFile, ProtocOutput.Cpp, Logger);
+						// TODO: Make the check for gRPC more efficient
+						if (FileReference.ReadAllText(ProtoFile).Contains("service"))
+						{
+							ExecuteProtocAction(ProtoModule, ProtoFile, ProtocOutput.GrpcCpp, Logger);
+						}
+					}
+				}
+			}
+		}
+
+		private static void ExecuteProtocAction(ProtoModuleInfo Module, FileReference ProtoFile, ProtocOutput Output, ILogger Logger)
+		{
+			DirectoryReference ProtobufIncludeDir = DirectoryReference.Combine(ProtoExecution.GetProtobufDirectory(), "include");
+			DirectoryReference ProtoDir = DirectoryReference.Combine(Module.ModuleRulesFile.Directory, "Proto");
+
+			FileReference CommandPath = ProtoExecution.GetProtocPath();
+			string Args = $" -I {ProtobufIncludeDir}";
+			Args += $" -I {ProtoDir}";
+			Args += $" {ProtoFile}";
+			Args += $" {Output.GetArgs(Module.GeneratedCodeDirectory.Location)}";
+
+			// TODO: Hide warnings
+			// TODO: Run async and parallel?
+			int Result = ExternalExecution.RunExternalNativeExecutable(CommandPath, Args, Logger);
+			// TODO: Handle Result
+		}
+
+		public static DirectoryReference GetProtobufDirectory()
+		{
+			return DirectoryReference.Combine(Unreal.EngineDirectory, "Source", "ThirdParty", "Protobuf", ProtobufVersion);
+		}
+
+		public static DirectoryReference GetGrpcDirectory()
+		{
+			return DirectoryReference.Combine(Unreal.EngineDirectory, "Source", "ThirdParty", "gRPC", GrpcVersion);
+		}
+
+		public static FileReference GetProtocPath()
+		{
+			// TODO: Support other platforms
+			string PlatformDir = "Win64";
+			string ArchDir = "x64";
+			string ExecutableName = "protoc.exe";
+
+			return FileReference.Combine(GetProtobufDirectory(), "bin", PlatformDir, ArchDir, "Release", ExecutableName);
+		}
+
+		public static FileReference GetGrpcCppPluginPath()
+		{
+			// TODO: Support other platforms
+			string PlatformDir = "Win64";
+			string ArchDir = "x64";
+			string ExecutableName = "grpc_cpp_plugin.exe";
+
+			DirectoryReference GrpcDirectory = DirectoryReference.Combine(Unreal.EngineDirectory, "Source", "ThirdParty", "gRPC", GrpcVersion);
+			return FileReference.Combine(GrpcDirectory, "bin", PlatformDir, ArchDir, "Release", ExecutableName);
+		}
+	}
+}
diff --git a/Engine/Source/Programs/UnrealBuildTool/System/TargetMakefile.cs b/Engine/Source/Programs/UnrealBuildTool/System/TargetMakefile.cs
--- a/Engine/Source/Programs/UnrealBuildTool/System/TargetMakefile.cs
+++ b/Engine/Source/Programs/UnrealBuildTool/System/TargetMakefile.cs
@@ -348,6 +348,11 @@ namespace UnrealBuildTool
 		public bool bUseVerseBPVM = true;
 #endif
 
+		/// <summary>
+		/// All build modules containing Verse code
+		/// </summary>
+		public List<ProtoModuleInfo> ProtoModules;
+
 		/// <summary>
 		/// List of all plugin names. The makefile will be considered invalid if any of these changes, or new plugins are added.
 		/// </summary>
@@ -457,6 +462,7 @@ namespace UnrealBuildTool
 #if __VPROJECT_AVAILABLE__
 			VNIModules = new List<VNIModuleInfo>();
 #endif
+			ProtoModules = new List<ProtoModuleInfo>();
 			PluginFiles = new HashSet<FileItem>();
 			ExternalDependencies = new HashSet<FileItem>();
 			InternalDependencies = new HashSet<FileItem>();
@@ -511,6 +517,7 @@ namespace UnrealBuildTool
 			VNIModules = Reader.ReadList(() => new VNIModuleInfo(Reader))!;
 			bUseVerseBPVM = Reader.ReadBool();
 #endif
+			ProtoModules = Reader.ReadList(() => new ProtoModuleInfo(Reader))!;
 			PluginFiles = Reader.ReadHashSet(() => Reader.ReadFileItem())!;
 			ExternalDependencies = Reader.ReadHashSet(() => Reader.ReadFileItem())!;
 			InternalDependencies = Reader.ReadHashSet(() => Reader.ReadFileItem())!;
@@ -585,6 +592,7 @@ namespace UnrealBuildTool
 			Writer.WriteList(VNIModules, e => e.Write(Writer));
 			Writer.WriteBool(bUseVerseBPVM);
 #endif
+			Writer.WriteList(ProtoModules, e => e.Write(Writer));
 			Writer.WriteHashSet(PluginFiles, x => Writer.WriteFileItem(x));
 			Writer.WriteHashSet(ExternalDependencies, x => Writer.WriteFileItem(x));
 			Writer.WriteHashSet(InternalDependencies, x => Writer.WriteFileItem(x));
