[UBT]
* Fixed bug in merged modules path where not all exi files were registered as outputs
* Touched cache bucket calculation for merged modules just to avoid invalid cache entries

diff --git a/Engine/Source/Programs/UnrealBuildTool/ToolChain/ISPCToolChain.cs b/Engine/Source/Programs/UnrealBuildTool/ToolChain/ISPCToolChain.cs
--- a/Engine/Source/Programs/UnrealBuildTool/ToolChain/ISPCToolChain.cs
+++ b/Engine/Source/Programs/UnrealBuildTool/ToolChain/ISPCToolChain.cs
@@ -386,7 +386,7 @@ namespace UnrealBuildTool
 			IEnumerable<string> strings = [
 				target.CacheSalt ?? String.Empty,			// Optional salt
 				target.BuildEnvironment.ToString(),			// Unique or Share environment
-				target.bMergeModules ? "MM" : String.Empty	// Merged modules will modify all object files
+				target.bMergeModules ? "MM2" : String.Empty	// Merged modules will modify all object files
 			];
 			return strings.Aggregate(GetStringHash(String.Empty), (acc, str) => GetStringHash(str, acc));
 		}
@@ -868,6 +868,12 @@ namespace UnrealBuildTool
 					// EXTRACTEXPORTS can only be interpreted by UBA.. so this action won't build outside uba
 					AdditionalArguments = " /EXTRACTEXPORTS";
 					CompileAction.ProducedItems.Add(FileItem.GetItemByFileReference(FileReference.Combine(OutputDir, FileName + ".exi")));
+
+					foreach (FileItem file in CompiledISPCObjFiles)
+					{
+						CompileAction.ProducedItems.Add(FileItem.GetItemByFileReference(file.Location.ChangeExtension(".exi")));
+					}
+
 				}
 
 				CompileAction.CommandArguments = $"@\"{NormalizeCommandLinePath(ResponseFileName, RootPaths)}\"{AdditionalArguments}";
