[UBT]
* Added support for getting cache server from horde. Bit of a wip implementation with some hacks

diff --git a/Engine/Source/Programs/Shared/EpicGames.UBA/ICacheClient.cs b/Engine/Source/Programs/Shared/EpicGames.UBA/ICacheClient.cs
--- a/Engine/Source/Programs/Shared/EpicGames.UBA/ICacheClient.cs
+++ b/Engine/Source/Programs/Shared/EpicGames.UBA/ICacheClient.cs
@@ -18,12 +18,18 @@ namespace EpicGames.UBA
 	public interface ICacheClient : IBaseInterface
 	{
 		/// <summary>
-		/// Connect to cache client
+		/// Connect to cache server
 		/// </summary>
 		/// <param name="host">Cache server address</param>
 		/// <param name="port">Cache server port</param>
+		/// <param name="desiredConnectionCount">Number of tcp connections we want</param>
 		/// <returns>True if successful</returns>
-		public abstract bool Connect(string host, int port);
+		public abstract bool Connect(string host, int port, int desiredConnectionCount = 1);
+
+		/// <summary>
+		/// Disconnect from cache server
+		/// </summary>
+		public abstract void Disconnect();
 
 		/// <summary>
 		/// Register path with string that will be hashed. All files under this path will be ignored and instead refer to this path hash
diff --git a/Engine/Source/Programs/Shared/EpicGames.UBA/Impl/CacheClientImpl.cs b/Engine/Source/Programs/Shared/EpicGames.UBA/Impl/CacheClientImpl.cs
--- a/Engine/Source/Programs/Shared/EpicGames.UBA/Impl/CacheClientImpl.cs
+++ b/Engine/Source/Programs/Shared/EpicGames.UBA/Impl/CacheClientImpl.cs
@@ -3,6 +3,7 @@
 using System;
 using System.Collections.Generic;
 using System.Runtime.InteropServices;
+using System.Threading.Tasks;
 
 namespace EpicGames.UBA.Impl
 {
@@ -26,6 +27,9 @@ namespace EpicGames.UBA.Impl
 		[DllImport("UbaHost", CharSet = CharSet.Auto)]
 		static extern byte CacheClient_Connect(nint cacheClient, string host, int port);
 
+		[DllImport("UbaHost", CharSet = CharSet.Auto)]
+		static extern void CacheClient_Disconnect(nint cacheClient);
+
 		[DllImport("UbaHost", CharSet = CharSet.Auto)]
 		static extern byte CacheClient_WriteToCache2(nint cacheClient, uint bucket, nint info, byte[] inputs, uint inputsSize, byte[] outputs, uint outputsSize);
 
@@ -69,6 +73,10 @@ namespace EpicGames.UBA.Impl
 
 			if (_handle != IntPtr.Zero)
 			{
+				if (_connectingTasks != null)
+				{
+					Task.WaitAll(_connectingTasks.ToArray());
+				}
 				CacheClient_Destroy(_handle);
 				_handle = IntPtr.Zero;
 			}
@@ -81,7 +89,7 @@ namespace EpicGames.UBA.Impl
 			_handle = CacheClient_Create(_sessionServer.GetHandle(), (byte)(reportMissReason?1:0), crypto, hint);
 		}
 
-		public bool Connect(string host, int port)
+		public bool Connect(string host, int port, int desiredConnectionCount)
 		{
 			if (CacheClient_Connect(_handle, host, port) == 0)
 			{
@@ -93,9 +101,31 @@ namespace EpicGames.UBA.Impl
 				CacheClient_RegisterPathHash(_handle, item.Key, item.Value);
 			}
 
+			if (desiredConnectionCount > 1)
+			{
+				_connectingTasks = new();
+				for (int i = 1; i < desiredConnectionCount; ++i)
+				{
+					_connectingTasks.Add(Task.Run(() =>
+					{
+						CacheClient_Connect(_handle, host, port);
+					}));
+				}
+			}
+
 			return true;
 		}
 
+		public void Disconnect()
+		{
+			if (_connectingTasks != null)
+			{
+				Task.WaitAll(_connectingTasks.ToArray());
+			}
+
+			CacheClient_Disconnect(_handle);
+		}
+
 		public bool RegisterPathHash(string path, string hash)
 		{
 			return CacheClient_RegisterPathHash(_handle, path, hash) != 0;
@@ -139,5 +169,7 @@ namespace EpicGames.UBA.Impl
 		}
 
 		public nint GetHandle() => _handle;
+
+		private List<Task>? _connectingTasks = null;
 	}
 }
diff --git a/Engine/Source/Programs/Shared/EpicGames.UBA/Impl/SchedulerImpl.cs b/Engine/Source/Programs/Shared/EpicGames.UBA/Impl/SchedulerImpl.cs
--- a/Engine/Source/Programs/Shared/EpicGames.UBA/Impl/SchedulerImpl.cs
+++ b/Engine/Source/Programs/Shared/EpicGames.UBA/Impl/SchedulerImpl.cs
@@ -133,6 +133,7 @@ namespace EpicGames.UBA.Impl
 
 			using IConfig config = IConfig.LoadConfig("");
 			config.AddValue("Scheduler", "MaxLocalProcessors", maxLocalProcessors.ToString());
+			config.AddValue("Scheduler", "MaxParallelCacheQueries", "100");
 			if (forceRemote)
 			{
 				config.AddValue("Scheduler", "ForceRemote", "True");
diff --git a/Engine/Source/Programs/UnrealBuildTool/Executors/UnrealBuildAccelerator/UBAAgentCoordinator.cs b/Engine/Source/Programs/UnrealBuildTool/Executors/UnrealBuildAccelerator/UBAAgentCoordinator.cs
--- a/Engine/Source/Programs/UnrealBuildTool/Executors/UnrealBuildAccelerator/UBAAgentCoordinator.cs
+++ b/Engine/Source/Programs/UnrealBuildTool/Executors/UnrealBuildAccelerator/UBAAgentCoordinator.cs
@@ -1,6 +1,7 @@
 // Copyright Epic Games, Inc. All Rights Reserved.
 
 using System;
+using System.Threading;
 using System.Threading.Tasks;
 using EpicGames.Core;
 using EpicGames.UBA;
@@ -23,6 +24,8 @@ namespace UnrealBuildTool
 
 		Task InitAsync(UBAExecutor executor);
 
+		UnrealBuildAcceleratorCacheConfig? RequestCacheServer(CancellationToken cancellationToken);
+
 		void Start(IUbaAgentCoordinatorScheduler scheduler, Func<LinkedAction, bool> canRunRemotely, StatusUpdateAction updateStatus);
 
 		void Stop();
diff --git a/Engine/Source/Programs/UnrealBuildTool/Executors/UnrealBuildAccelerator/UBAAgentCoordinatorHorde.cs b/Engine/Source/Programs/UnrealBuildTool/Executors/UnrealBuildAccelerator/UBAAgentCoordinatorHorde.cs
--- a/Engine/Source/Programs/UnrealBuildTool/Executors/UnrealBuildAccelerator/UBAAgentCoordinatorHorde.cs
+++ b/Engine/Source/Programs/UnrealBuildTool/Executors/UnrealBuildAccelerator/UBAAgentCoordinatorHorde.cs
@@ -178,6 +178,23 @@ namespace UnrealBuildTool
 			_logger.LogInformation("Created tool bundle with locator {UbaAgentLocator}", _ubaAgentLocator.ToString());
 		}
 
+		public UnrealBuildAcceleratorCacheConfig? RequestCacheServer(CancellationToken cancellationToken, int desiredConnectionCount)
+		{
+			// REALLY UGLY. Remove this loop. We know InitAsync has been called before this one so right now it is a hack to just loop.
+			while (_client == null)
+			{
+				Thread.Sleep(10);
+			}
+
+			var clusterId = new ClusterId(UnrealBuildAcceleratorHordeConfig.ClusterDefault);
+			UbaConfig ubaConfig = _client.AllocateUbaCacheServerAsync(clusterId, cancellationToken).Result;
+			if (String.IsNullOrEmpty(ubaConfig.CacheEndpoint))
+			{
+				return null;
+			}
+			return new UnrealBuildAcceleratorCacheConfig() { CacheServer = ubaConfig.CacheEndpoint, bRequireVfs = true, WriteCache = ubaConfig.WriteAccess.ToString(), Crypto = ubaConfig.CacheSessionKey, DesiredConnectionCount = desiredConnectionCount };
+		}
+
 		async Task<BlobLocator> CreateToolAsync(DirectoryReference baseDir, IEnumerable<FileReference> files, CancellationToken cancellationToken)
 		{
 			// TODO: should drive this off API version reported by server
@@ -801,8 +818,33 @@ namespace UnrealBuildTool
 			executor.AgentCoordinatorInitialized(this, _hordeSessionTask.IsCompletedSuccessfully && _hordeSessionTask.Result.Item1 != null);
 		}
 
+		public UnrealBuildAcceleratorCacheConfig? RequestCacheServer(CancellationToken cancellationToken)
+		{
+			if (!_hordeConfig.bCacheEnabled)
+			{
+				return null;
+			}
+
+			UBAHordeSession? hordeSession = null;
+			string reason = "Unknown";
+			if (_hordeSessionTask != null)
+			{
+				(hordeSession, reason) = _hordeSessionTask.Result;
+			}
+			if (hordeSession == null)
+			{
+				return null;
+			}
+
+			return hordeSession.RequestCacheServer(cancellationToken, _hordeConfig.CacheDesiredConnectionCount);
+		}
+
 		public void Start(IUbaAgentCoordinatorScheduler scheduler, Func<LinkedAction, bool> canRunRemotely, StatusUpdateAction updateStatus)
 		{
+			if (!_hordeConfig.bAgentsEnabled)
+			{
+				return;
+			}
 			_updateStatus = updateStatus;
 
 			int timerPeriod = 5000;
diff --git a/Engine/Source/Programs/UnrealBuildTool/Executors/UnrealBuildAccelerator/UBAExecutor.cs b/Engine/Source/Programs/UnrealBuildTool/Executors/UnrealBuildAccelerator/UBAExecutor.cs
--- a/Engine/Source/Programs/UnrealBuildTool/Executors/UnrealBuildAccelerator/UBAExecutor.cs
+++ b/Engine/Source/Programs/UnrealBuildTool/Executors/UnrealBuildAccelerator/UBAExecutor.cs
@@ -385,14 +385,30 @@ namespace UnrealBuildTool
 
 							OnSessionCreated?.Invoke(session);
 
-							if (_cacheableActionsLeft != 0 && _cacheConfigs.Any())
+							List<UnrealBuildAcceleratorCacheConfig> cacheConfigs = new();
+							if (_cacheableActionsLeft != 0)
+							{
+								if (!UBAConfig.bDisableRemote) // Need to match the InitAsync
+								{
+									foreach (IUBAAgentCoordinator coordinator in _agentCoordinators)
+									{
+										UnrealBuildAcceleratorCacheConfig? cacheConfig = coordinator.RequestCacheServer(new CancellationToken());
+										if (cacheConfig != null)
+											cacheConfigs.Add(cacheConfig);
+									}
+								}
+								if (_cacheConfigs.Any())
+									cacheConfigs.AddRange(_cacheConfigs);
+							}
+
+							if (cacheConfigs.Count > 0)
 							{
 								session.UpdateStatus(1, 1, "Cache", LogEntryType.Info, null);
 								session.UpdateStatus(1, 6, "Connecting...", LogEntryType.Info);
 
-								foreach (UnrealBuildAcceleratorCacheConfig cacheConfig in _cacheConfigs)
+								foreach (UnrealBuildAcceleratorCacheConfig cacheConfig in cacheConfigs)
 								{
-									ICacheClient cacheClient = ICacheClient.CreateCacheClient(session, UBAConfig.bReportCacheMissReason, "", _cacheableActionsLeft != 0 ? GetCacheClientHint() : "");
+									ICacheClient cacheClient = ICacheClient.CreateCacheClient(session, UBAConfig.bReportCacheMissReason, cacheConfig.Crypto, GetCacheClientHint());
 									string[] nameAndPort = cacheConfig.CacheServer.Split(':');
 									int port = 1347;
 									if (nameAndPort.Length > 1)
@@ -401,7 +417,7 @@ namespace UnrealBuildTool
 									}
 
 									System.Diagnostics.Stopwatch stopwatch = System.Diagnostics.Stopwatch.StartNew();
-									bool cacheSuccess = cacheClient.Connect(nameAndPort[0], port);
+									bool cacheSuccess = cacheClient.Connect(nameAndPort[0], port, cacheConfig.DesiredConnectionCount);
 									long totalMs = stopwatch.ElapsedMilliseconds;
 									string successText = cacheSuccess ? "Connected to" : "Failed to connect to";
 									logger.LogInformation("UbaCache - {SuccessText} {Name}:{Port} ({Seconds}.{Milliseconds}s)", successText, nameAndPort[0], port, totalMs / 1000, totalMs % 1000);
@@ -600,6 +616,8 @@ namespace UnrealBuildTool
 				if (_bIsCancelled)
 					return;
 				_bIsCancelled = true;
+				foreach (ICacheClient client in cache)
+					client.Disconnect();
 				scheduler.Cancel();
 				server?.StopServer(); // Make sure all remove processes are returned. We can't have any callbacks after this
 				session?.CancelAll(); // Cancel all processes native side
@@ -993,6 +1011,11 @@ namespace UnrealBuildTool
 				return 0;
 			}
 
+			if (!UBAConfig.bCacheLinkActions && action.ActionType == ActionType.Link)
+			{
+				return 0;
+			}
+
 			uint bucket = 0;
 			if (action.CacheBucket != 0)
 			{
@@ -1083,6 +1106,11 @@ namespace UnrealBuildTool
 				return true;
 			}
 
+			if (!UBAConfig.bCacheLinkActions && action.ActionType == ActionType.Link)
+			{
+				return true;
+			}
+
 			// Collect all inputs for action
 			// We use prerequisite items plus what we find in dependency list file if it exists.
 
diff --git a/Engine/Source/Programs/UnrealBuildTool/Executors/UnrealBuildAccelerator/UnrealBuildAcceleratorCacheConfig.cs b/Engine/Source/Programs/UnrealBuildTool/Executors/UnrealBuildAccelerator/UnrealBuildAcceleratorCacheConfig.cs
--- a/Engine/Source/Programs/UnrealBuildTool/Executors/UnrealBuildAccelerator/UnrealBuildAcceleratorCacheConfig.cs
+++ b/Engine/Source/Programs/UnrealBuildTool/Executors/UnrealBuildAccelerator/UnrealBuildAcceleratorCacheConfig.cs
@@ -42,6 +42,20 @@ namespace UnrealBuildTool
 		/// If the cache can be written to
 		/// </summary>
 		public bool CanWrite => WriteCache.Equals("True", StringComparison.OrdinalIgnoreCase) || (Unreal.IsBuildMachine() && WriteCache.Equals("BuildMachineOnly", StringComparison.OrdinalIgnoreCase));
+
+		/// <summary>
+		/// Crypto used to connect to cache server.
+		/// </summary>
+		[XmlConfigFile(Category = "UnrealBuildAccelerator", Name = "CacheCrypto")]
+		[ConfigFile(ConfigHierarchyType.Engine, "Uba.CacheProvider", "Crypto")]
+		public string Crypto { get; set; } = String.Empty;
+
+		/// <summary>
+		/// Desired connection count. When there is latency it might be a good idea to have multiple tcp connections
+		/// </summary>
+		[XmlConfigFile(Category = "UnrealBuildAccelerator", Name = "CacheDesiredConnectionCount")]
+		[ConfigFile(ConfigHierarchyType.Engine, "Uba.CacheProvider", "DesiredConnectionCount")]
+		public int DesiredConnectionCount { get; set; } = 1;
 	}
 
 	#region Extension methods
diff --git a/Engine/Source/Programs/UnrealBuildTool/Executors/UnrealBuildAccelerator/UnrealBuildAcceleratorConfig.cs b/Engine/Source/Programs/UnrealBuildTool/Executors/UnrealBuildAccelerator/UnrealBuildAcceleratorConfig.cs
--- a/Engine/Source/Programs/UnrealBuildTool/Executors/UnrealBuildAccelerator/UnrealBuildAcceleratorConfig.cs
+++ b/Engine/Source/Programs/UnrealBuildTool/Executors/UnrealBuildAccelerator/UnrealBuildAcceleratorConfig.cs
@@ -233,6 +233,7 @@ namespace UnrealBuildTool
 		[XmlConfigFile(Category = "UnrealBuildAccelerator")]
 		[CommandLine("-UBAStoreObjFilesCompressed", Value = "true")]
 		[CommandLine("-UBADisableStoreObjFilesCompressed", Value = "false")]
+		[CommandLine("-nocompress", Value = "false")]
 		public bool bStoreObjFilesCompressed { get; set; }
 
 		/// <summary>
@@ -312,6 +313,12 @@ namespace UnrealBuildTool
 		[CommandLine("-UBAReportCacheMissReason")]
 		public bool bReportCacheMissReason { get; set; }
 
+		/// <summary>
+		/// Control if link actions should be downloaded from cache or not
+		/// </summary>
+		[XmlConfigFile(Category = "UnrealBuildAccelerator", Name = "CacheLinkActions")]
+		public bool bCacheLinkActions { get; set; } = true;
+
 		/// <summary>
 		/// Compression level. Options are None, SuperFast, VeryFast, Fast, Normal, Optimal1 to Optimal5
 		/// </summary>
diff --git a/Engine/Source/Programs/UnrealBuildTool/Executors/UnrealBuildAccelerator/UnrealBuildAcceleratorHordeConfig.cs b/Engine/Source/Programs/UnrealBuildTool/Executors/UnrealBuildAccelerator/UnrealBuildAcceleratorHordeConfig.cs
--- a/Engine/Source/Programs/UnrealBuildTool/Executors/UnrealBuildAccelerator/UnrealBuildAcceleratorHordeConfig.cs
+++ b/Engine/Source/Programs/UnrealBuildTool/Executors/UnrealBuildAccelerator/UnrealBuildAcceleratorHordeConfig.cs
@@ -216,6 +216,24 @@ namespace UnrealBuildTool
 		/// </summary>
 		[ConfigFile(ConfigHierarchyType.Engine, "Uba.Provider.Horde", "UbaEnabled")]
 		public string? HordeEnabled { get; set; } = null;
+
+		/// <summary>
+		/// Agents enabled parameter for use by INI config, expects one of [True, False, BuildMachineOnly]
+		/// </summary>
+		[ConfigFile(ConfigHierarchyType.Engine, "Uba.Provider.Horde", "UbaAgentsEnabled")]
+		public bool bAgentsEnabled { get; set; } = true;
+
+		/// <summary>
+		/// Enabled parameter for use by INI config, expects one of [True, False, BuildMachineOnly]
+		/// </summary>
+		[ConfigFile(ConfigHierarchyType.Engine, "Uba.Provider.Horde", "UbaCacheEnabled")]
+		public bool bCacheEnabled { get; set; } = false;
+
+		/// <summary>
+		/// Number of desired connections when connecting to cache
+		/// </summary>
+		[ConfigFile(ConfigHierarchyType.Engine, "Uba.Provider.Horde", "UbaCacheDesiredConnectionCount")]
+		public int CacheDesiredConnectionCount { get; set; } = 1;
 	}
 
 	#region Extension methods
