UnrealBuildTool: Add clang lib path when the clang linker is in use

#rnx

diff --git a/Engine/Source/Programs/UnrealBuildTool/Platform/Windows/VCEnvironment.cs b/Engine/Source/Programs/UnrealBuildTool/Platform/Windows/VCEnvironment.cs
--- a/Engine/Source/Programs/UnrealBuildTool/Platform/Windows/VCEnvironment.cs
+++ b/Engine/Source/Programs/UnrealBuildTool/Platform/Windows/VCEnvironment.cs
@@ -315,35 +315,19 @@ namespace UnrealBuildTool
 		}
 
 		/// <summary>
-		/// Return the standard Visual C++ library path.
+		/// Return the standard Clang C++ library path.
 		/// </summary>
-		protected virtual DirectoryReference GetToolChainLibsDir()
+		protected virtual DirectoryReference GetClangLibsDir()
 		{
-			string ArchFolder = Architecture.WindowsSystemLibDir;
-
-			// Add the standard Visual C++ library paths
-			if (Compiler.IsIntel() && bAllowClangLinker)
-			{
-				VersionNumber ClangVersion = MicrosoftPlatformSDK.GetClangVersionForIntelCompiler(CompilerPath);
-				return DirectoryReference.Combine(CompilerDir, "lib", "clang", ClangVersion.GetComponent(0).ToString(), "lib", "windows");
-			}
-			else if (ToolChain.IsMSVC())
-			{
-				return DirectoryReference.Combine(ToolChainDir, "lib", ArchFolder);
-			}
-			else
-			{
-				DirectoryReference LibsPath = DirectoryReference.Combine(ToolChainDir, "LIB");
-
-				if (Architecture == UnrealArch.X64)
-				{
-					LibsPath = DirectoryReference.Combine(LibsPath, "amd64");
-				}
-
-				return LibsPath;
-			}
+			VersionNumber ClangVersion = Compiler.IsIntel() ? MicrosoftPlatformSDK.GetClangVersionForIntelCompiler(CompilerPath) : CompilerVersion;
+			return DirectoryReference.Combine(CompilerDir, "lib", "clang", ClangVersion.GetComponent(0).ToString(), "lib", "windows");
 		}
 
+		/// <summary>
+		/// Return the standard Visual C++ library path.
+		/// </summary>
+		protected virtual DirectoryReference GetToolChainLibsDir() => DirectoryReference.Combine(ToolChainDir, "lib", Architecture.WindowsSystemLibDir);
+
 		/// <summary>
 		/// Sets up the standard compile environment for the toolchain
 		/// </summary>
@@ -355,7 +339,13 @@ namespace UnrealBuildTool
 			// Add the standard Visual C++ include paths
 			IncludePaths.Add(DirectoryReference.Combine(ToolChainDir, "INCLUDE"));
 
-			// Add the standard Visual C++ library paths
+			if (Compiler.IsClang() && (bAllowClangLinker || bAllowRadLinker))
+			{
+				// Add the standard Clang C++ library path before Visual CPP path
+				LibraryPaths.Add(GetClangLibsDir());
+			}
+
+			// Add the standard Visual C++ library path
 			LibraryPaths.Add(GetToolChainLibsDir());
 
 			// If we're on >= Visual Studio 2015 and using pre-Windows 10 SDK, we need to find a Windows 10 SDK and add the UCRT include paths
