Always include Clang's library directory when compiling with Clang on Windows, regardless of linker.

Also expand a bit on the comment.

#rnx

diff --git a/Engine/Source/Programs/UnrealBuildTool/Platform/Windows/VCEnvironment.cs b/Engine/Source/Programs/UnrealBuildTool/Platform/Windows/VCEnvironment.cs
--- a/Engine/Source/Programs/UnrealBuildTool/Platform/Windows/VCEnvironment.cs
+++ b/Engine/Source/Programs/UnrealBuildTool/Platform/Windows/VCEnvironment.cs
@@ -339,9 +339,9 @@ namespace UnrealBuildTool
 			// Add the standard Visual C++ include paths
 			IncludePaths.Add(DirectoryReference.Combine(ToolChainDir, "INCLUDE"));
 
-			if (Compiler.IsClang() && (bAllowClangLinker || bAllowRadLinker))
+			if (Compiler.IsClang())
 			{
-				// Add the standard Clang C++ library path before Visual CPP path
+				// Add the standard Clang C++ library path before the Visual C++ path. This includes compiler libraries required for PGO, ASAN, etc.
 				LibraryPaths.Add(GetClangLibsDir());
 			}
 
