UnrealBuildTool: Remove erroneous exception when verifying pgo compiler version

#rnx

diff --git a/Engine/Source/Programs/UnrealBuildTool/Platform/Windows/VCToolChain.cs b/Engine/Source/Programs/UnrealBuildTool/Platform/Windows/VCToolChain.cs
--- a/Engine/Source/Programs/UnrealBuildTool/Platform/Windows/VCToolChain.cs
+++ b/Engine/Source/Programs/UnrealBuildTool/Platform/Windows/VCToolChain.cs
@@ -3757,7 +3757,6 @@ namespace UnrealBuildTool
 							Logger.LogInformation("Profile: {ProfileCompiler} ({ProfileVersion}) Optimize: {OptimizeCompiler} ({OptimizeVersion})", ProfileCompiler, ProfileVersion, EnvVars.Compiler, CompilerVersion);
 							GlobalCompileEnvironment.bPGOOptimize = false;
 							GlobalLinkEnvironment.bPGOOptimize = false;
-							throw new BuildException("test");
 						}
 					}
 				}
