[UBT]
* Added -UBAHordeCacheEnabled that can be used to enable request of cache server from horde from command line

diff --git a/Engine/Source/Programs/UnrealBuildTool/Executors/UnrealBuildAccelerator/UnrealBuildAcceleratorHordeConfig.cs b/Engine/Source/Programs/UnrealBuildTool/Executors/UnrealBuildAccelerator/UnrealBuildAcceleratorHordeConfig.cs
--- a/Engine/Source/Programs/UnrealBuildTool/Executors/UnrealBuildAccelerator/UnrealBuildAcceleratorHordeConfig.cs
+++ b/Engine/Source/Programs/UnrealBuildTool/Executors/UnrealBuildAccelerator/UnrealBuildAcceleratorHordeConfig.cs
@@ -227,6 +227,7 @@ namespace UnrealBuildTool
 		/// Enabled parameter for use by INI config, expects one of [True, False, BuildMachineOnly]
 		/// </summary>
 		[ConfigFile(ConfigHierarchyType.Engine, "Uba.Provider.Horde", "UbaCacheEnabled")]
+		[CommandLine("-UBAHordeCacheEnabled")]
 		public bool bCacheEnabled { get; set; } = false;
 
 		/// <summary>
