Refactor Proto execution to be properly added to the task graph during compilation

diff --git a/Engine/Source/Programs/UnrealBuildTool/Configuration/CppCompileWarnings.cs b/Engine/Source/Programs/UnrealBuildTool/Configuration/CppCompileWarnings.cs
--- a/Engine/Source/Programs/UnrealBuildTool/Configuration/CppCompileWarnings.cs
+++ b/Engine/Source/Programs/UnrealBuildTool/Configuration/CppCompileWarnings.cs
@@ -173,6 +173,22 @@ namespace UnrealBuildTool
 		}
 		private WarningLevel _switchUnhandledEnumeratorWarningLevel;
 
+		/// <summary>
+		/// The compiler detected a conversion from size_t to a smaller type.
+		/// MSVC - 
+		///		https://learn.microsoft.com/en-us/cpp/error-messages/compiler-warnings/compiler-warning-level-3-c4267
+		///		4267 - 'var' : conversion from 'size_t' to 'type', possible loss of data
+		/// </summary>
+		[XmlConfigFile(Category = "BuildConfiguration")]
+		[WarningsVCToolChain(["/wd4267"], ["/w4267"], ["/we4267"])]
+		[BasicWarningLevelDefault(WarningLevel.Off)]
+		public WarningLevel ShortenSizeTToIntWarningLevel
+		{
+			get => ResolveWarning(ParentCppCompileWarnings?.UnsafeTypeCastWarningLevel, _shortenSizeTToIntWarningLevel);
+			set => _shortenSizeTToIntWarningLevel = value;
+		}
+		private WarningLevel _shortenSizeTToIntWarningLevel;
+
 		#region -- Clang Compiler --
 
 		/// <summary>
@@ -858,10 +874,14 @@ namespace UnrealBuildTool
 
 		/// <summary>
 		/// Indicates what warning/error level to treat function parameters that are declared, but not used within the function body.
+		/// MSVC - 
+		///		https://learn.microsoft.com/en-us/cpp/error-messages/compiler-warnings/compiler-warning-level-4-c4100
+		///		4100 - 'identifier' : unreferenced formal parameter
 		/// Clang -
 		///		https://clang.llvm.org/docs/DiagnosticsReference.html#wunused-parameter
 		/// </summary>
 		/// <remarks>Unused function parameter. A lot are named 'bUnused'...</remarks>
+		[WarningsMSCV(["/wd4100"], ["/w4100"], ["/we4100"])]
 		[WarningsVCClang(["-Wno-unused-parameter"], ["-Wunused-parameter", "-Wno-error=unused-parameter"], ["-Wunused-parameter"])]
 		[BasicWarningLevelDefault(WarningLevel.Off)]
 		public WarningLevel UnusedParameterWarningLevel
diff --git a/Engine/Source/Programs/UnrealBuildTool/Configuration/ReadOnlyCppCompileWarnings.cs b/Engine/Source/Programs/UnrealBuildTool/Configuration/ReadOnlyCppCompileWarnings.cs
--- a/Engine/Source/Programs/UnrealBuildTool/Configuration/ReadOnlyCppCompileWarnings.cs
+++ b/Engine/Source/Programs/UnrealBuildTool/Configuration/ReadOnlyCppCompileWarnings.cs
@@ -43,6 +43,7 @@ namespace UnrealBuildTool
 		public WarningLevel UndefinedIdentifierWarningLevel => _inner.UndefinedIdentifierWarningLevel;
 		public WarningLevel UnsafeTypeCastWarningLevel => _inner.UnsafeTypeCastWarningLevel;
 		public WarningLevel SwitchUnhandledEnumeratorWarningLevel => _inner.SwitchUnhandledEnumeratorWarningLevel;
+		public WarningLevel ShortenSizeTToIntWarningLevel => _inner.ShortenSizeTToIntWarningLevel;
 
 		#region -- Clang Compiler Disables --
 
diff --git a/Engine/Source/Programs/UnrealBuildTool/Configuration/UEBuildModuleCPP.cs b/Engine/Source/Programs/UnrealBuildTool/Configuration/UEBuildModuleCPP.cs
--- a/Engine/Source/Programs/UnrealBuildTool/Configuration/UEBuildModuleCPP.cs
+++ b/Engine/Source/Programs/UnrealBuildTool/Configuration/UEBuildModuleCPP.cs
@@ -83,11 +83,6 @@ namespace UnrealBuildTool
 		/// </summary>
 		public DirectoryReference? GeneratedCodeDirectoryVNI => GeneratedCodeDirectory != null ? DirectoryReference.Combine(GeneratedCodeDirectory!, "VNI") : null;
 
-		/// <summary>
-		/// The directory for this module's generated Protobuf code
-		/// </summary>
-		public DirectoryReference? GeneratedCodeDirectoryProto => GeneratedCodeDirectory != null ? DirectoryReference.Combine(GeneratedCodeDirectory!, "Proto") : null;
-
 		/// <summary>
 		/// Global override to force all include paths to be always added
 		/// </summary>
@@ -456,11 +451,6 @@ namespace UnrealBuildTool
 				{
 					IncludePaths.Add(GeneratedCodeDirectoryVNI!);
 				}
-				
-				if (bHasProto || bForceAddIncludePath)
-				{
-					IncludePaths.Add(GeneratedCodeDirectoryProto!);
-				}
 
 				foreach ((string Subdirectory, Action<ILogger, DirectoryReference> _) in Rules.GenerateHeaderFuncs)
 				{
@@ -620,6 +610,34 @@ namespace UnrealBuildTool
 				ModuleCompileEnvironment.UserIncludePaths.Add(IntermediateDirectory);
 			}
 
+			if (bHasProto)
+			{
+				DirectoryReference ProtoDirectory = DirectoryReference.Combine(Rules.File.Directory, "Proto");
+				ProtocOutput OutputFiles = ProtoExecution.CompileProtoFiles(InputFiles.ProtoFiles, ProtoDirectory, IntermediateDirectory, Graph);
+				CppCompileEnvironment.AdditionalPrerequisites.AddRange(OutputFiles.HeaderFiles);
+				ModuleCompileEnvironment.AdditionalPrerequisites.AddRange(OutputFiles.HeaderFiles);
+				CppCompileEnvironment.UserIncludePaths.Add(IntermediateDirectory);
+				ModuleCompileEnvironment.UserIncludePaths.Add(IntermediateDirectory);
+
+				CppCompileEnvironment ProtoCompileEnvironment = new CppCompileEnvironment(CppCompileEnvironment);
+				ProtoCompileEnvironment.PrecompiledHeaderAction = PrecompiledHeaderAction.None;
+				ProtoCompileEnvironment.bUseUnity = false;
+				ProtoCompileEnvironment.bWarningsAsErrors = false;
+				// Disable dynamic analysis on generated proto code as it causes issues when the proto/grpc libs are not
+				// compiled the same way. This can cause false positive in some situations, for example with address
+				// sanitize enabled (see: [REDACTED_PROFILE_URL]/protobuf/issues/13115).
+				ProtoCompileEnvironment.bDisableDynamicAnalysis = true;
+				// Disable warnings that are common in genereated proto cpp
+				ProtoCompileEnvironment.CppCompileWarnings.DeprecationWarningLevel = WarningLevel.Off;
+				ProtoCompileEnvironment.CppCompileWarnings.ShortenSizeTToIntWarningLevel = WarningLevel.Off;
+				ProtoCompileEnvironment.CppCompileWarnings.UnusedParameterWarningLevel = WarningLevel.Off;
+				ProtoCompileEnvironment.CppCompileWarnings.DeprecationWarningLevel = WarningLevel.Off;
+
+				CPPOutput Output = ToolChain.CompileAllCPPFiles(ProtoCompileEnvironment, OutputFiles.CppFiles, IntermediateDirectory, Name, Graph);
+
+				LinkInputFiles.AddRange(Output.ObjectFiles);
+			}
+
 			// Write all the definitions to a separate file
 			CreateHeaderForDefinitions(CppCompileEnvironment, IntermediateDirectory, null, Graph);
 
@@ -933,34 +951,6 @@ namespace UnrealBuildTool
 				}
 			}
 
-			if (bHasProto && GeneratedCodeDirectoryProto != null && DirectoryReference.Exists(GeneratedCodeDirectoryProto))
-			{
-				CppCompileEnvironment ProtoCompileEnvironment = new CppCompileEnvironment(CppCompileEnvironment);
-				ProtoCompileEnvironment.PrecompiledHeaderAction = PrecompiledHeaderAction.None;
-				ProtoCompileEnvironment.bUseUnity = false;
-				ProtoCompileEnvironment.bWarningsAsErrors = false;
-				// Disable dynamic analysis on generated proto code as it causes issues when the proto/grpc libs are not
-				// compiled the same way. This can cause false positive in some situations, for example with address
-				// sanitize enabled (see: [REDACTED_PROFILE_URL]/protobuf/issues/13115).
-				ProtoCompileEnvironment.bDisableDynamicAnalysis = true;
-
-				IEnumerable<FileItem> ProtoCPPFiles = from x in DirectoryReference.EnumerateFiles(GeneratedCodeDirectoryProto, "*.pb.cc", SearchOption.AllDirectories) select FileItem.GetItemByFileReference(x);
-				CPPOutput Output = ToolChain.CompileAllCPPFiles(ProtoCompileEnvironment, ProtoCPPFiles, IntermediateDirectory, Name, Graph);
-
-				LinkInputFiles.AddRange(Output.ObjectFiles);
-
-				// TODO(jamal.fanaian): Temporary hack to deal with `verify` macro causing issues with absl having a `verify` method
-				foreach (FileReference HeaderFile in DirectoryReference.EnumerateFiles(GeneratedCodeDirectoryProto, "*.pb.h", SearchOption.AllDirectories))
-				{
-					// TODO(jamal.fanaian): Use a StreamReader so we don't load the entire file in memory
-					string Content = FileReference.ReadAllText(HeaderFile);
-					if (!Content.StartsWith("#undef verify"))
-					{
-						FileReference.WriteAllText(HeaderFile, "#undef verify" + Environment.NewLine + Content);
-					}
-				}
-			}
-
 			// Compile CPP files
 			if (CPPFiles.Count > 0)
 			{
diff --git a/Engine/Source/Programs/UnrealBuildTool/Configuration/UEBuildTarget.cs b/Engine/Source/Programs/UnrealBuildTool/Configuration/UEBuildTarget.cs
--- a/Engine/Source/Programs/UnrealBuildTool/Configuration/UEBuildTarget.cs
+++ b/Engine/Source/Programs/UnrealBuildTool/Configuration/UEBuildTarget.cs
@@ -2612,11 +2612,6 @@ namespace UnrealBuildTool
 				return Makefile;
 			}
 			
-			using (GlobalTracer.Instance.BuildSpan("ProtoExecution.SetupProtoModules()").StartActive())
-			{
-				ProtoExecution.SetupProtoModules(ModulesToGenerateHeadersFor, Makefile, Logger);
-			}
-
 #if __VPROJECT_AVAILABLE__
 			// Copy Verse BPVM usage flag
 			Makefile.bUseVerseBPVM = Rules.bUseVerseBPVM;
@@ -2650,13 +2645,6 @@ namespace UnrealBuildTool
 				// Create all the shared PCH instances before processing the modules
 				CreateSharedPCHInstances(Rules, TargetToolChain, OriginalBinaries, GlobalCompileEnvironment, MakefileBuilder, Logger);
 			}
-			
-			if (Makefile.ProtoModules.Count > 0)
-			{
-				//GlobalCompileEnvironment.ExtraGeneratedCPPFileTypes.Add("*.pb.cc");
-				// ProtoExecution.CreateProtocActions(MakefileBuilder);
-				ProtoExecution.ExecuteProtocActions(Makefile, Logger);
-			}
 
 			// Can probably be moved further down?
 #if __VPROJECT_AVAILABLE__
diff --git a/Engine/Source/Programs/UnrealBuildTool/System/ProtoExecution.cs b/Engine/Source/Programs/UnrealBuildTool/System/ProtoExecution.cs
--- a/Engine/Source/Programs/UnrealBuildTool/System/ProtoExecution.cs
+++ b/Engine/Source/Programs/UnrealBuildTool/System/ProtoExecution.cs
@@ -2,105 +2,50 @@
 
 using System;
 using System.Collections.Generic;
-using System.Data;
-using System.Diagnostics;
-using System.Diagnostics.CodeAnalysis;
 using System.IO;
 using System.Linq;
-using System.Reflection;
-using System.Text;
-using System.Text.Json;
-using System.Threading.Tasks;
 using EpicGames.Core;
-using EpicGames.UHT.Utils;
-using Microsoft.Extensions.Logging;
-using OpenTracing;
-using OpenTracing.Util;
 using UnrealBuildBase;
-using UnrealBuildTool.Modes;
 
 namespace UnrealBuildTool
 {
 	/// <summary>
-	/// Information about a module that needs to be passed to protoc for code generation
+	/// Result from protoc action
 	/// </summary>
-	public class ProtoModuleInfo
+	public class ProtocOutput
 	{
 		/// <summary>
-		/// Module name
+		/// C++ files generated by a ProtocAction
 		/// </summary>
-		public required string ModuleName;
+		public List<FileItem> CppFiles = [];
 
 		/// <summary>
-		/// Path to the module rules file
+		/// Header files generated by a ProtocAction
 		/// </summary>
-		public required FileReference ModuleRulesFile;
-		
-		/// <summary>
-		/// Directory containing generated code
-		/// </summary>
-		public required DirectoryItem GeneratedCodeDirectory;
-		
-		/// <summary>
-		/// Directory containing proto source files
-		/// </summary>
-		public required DirectoryReference ProtoDirectory;
-		
-		/// <summary>
-		/// List of source files to compile
-		/// </summary>
-		public List<FileItem> ProtoFiles = new List<FileItem>();
-		
-		/// <summary>
-		/// Object initializer for ProtoModuleInfo
-		/// </summary>
-		public ProtoModuleInfo()
-		{
-		}
-		
-		/// <summary>
-		/// Load ProtoModuleInfo from a binary archive
-		/// </summary>
-		/// <param name="Reader"></param>
-		[SetsRequiredMembers]
-		public ProtoModuleInfo(BinaryArchiveReader Reader)
-		{
-			ModuleName = Reader.ReadString()!;
-			ModuleRulesFile = Reader.ReadFileReference();
-			GeneratedCodeDirectory = Reader.ReadDirectoryItem()!;
-			ProtoDirectory = Reader.ReadDirectoryReference()!;
-			ProtoFiles = Reader.ReadList(() => Reader.ReadFileItem()!) ?? new List<FileItem>();
-		}
-
-		/// <summary>
-		/// Write ProtoModuleInfo as a binary archive
-		/// </summary>
-		/// <param name="Writer"></param>
-		public void Write(BinaryArchiveWriter Writer)
-		{
-			Writer.WriteString(ModuleName);
-			Writer.WriteFileReference(ModuleRulesFile);
-			Writer.WriteDirectoryItem(GeneratedCodeDirectory);
-			Writer.WriteDirectoryReference(ProtoDirectory);
-			Writer.WriteList(ProtoFiles, x => Writer.WriteFileItem(x));
-		}
+		public List<FileItem> HeaderFiles = [];
 	}
 
 	class ProtocAction : Action
 	{
-		public ProtocAction(ProtoModuleInfo Module, FileItem ProtoFile) : base(ActionType.Compile)
+		public ProtocAction(List<FileItem> ProtoFiles, DirectoryReference ProtoModuleDir, DirectoryReference OutputDir, ProtocOutputType OutputType) : base(ActionType.Compile)
 		{
 			DirectoryReference ProtobufIncludeDir = DirectoryReference.Combine(ProtoExecution.GetProtobufDirectory(), "include");
-			DirectoryReference ProtoDir = DirectoryReference.Combine(Module.ModuleRulesFile.Directory, "Proto");
 
-			PrerequisiteItems.Add(ProtoFile);
-			string ProducedName = Path.GetRelativePath(DirectoryReference.Combine(Module.ModuleRulesFile.Directory, "Proto").ToString(), ProtoFile.FullName);
-			ProducedItems.Add(FileItem.GetItemByPath(Path.ChangeExtension(ProducedName, "pb.h")));
-			ProducedItems.Add(FileItem.GetItemByPath(Path.ChangeExtension(ProducedName, "pb.cc")));
+			PrerequisiteItems.UnionWith(ProtoFiles);
+			WorkingDirectory = ProtoModuleDir;
 			CommandPath = ProtoExecution.GetProtocPath();
 			CommandArguments = $" -I {ProtobufIncludeDir}";
-			CommandArguments += $" -I {ProtoDir}";
-			CommandArguments += $" --cpp_out={Module.GeneratedCodeDirectory}";
+			CommandArguments += $" -I {ProtoModuleDir}";
+			CommandArguments += $" {OutputType.GetArgs(OutputDir)}";
+
+			foreach (FileItem ProtoFile in ProtoFiles)
+			{
+				CommandArguments += $" {ProtoFile.AbsolutePath}";
+
+				string ProducedName = Path.GetRelativePath(ProtoModuleDir.ToString(), ProtoFile.FullName);
+				ProducedItems.Add(ProtoExecution.GetProducedFileItem(OutputDir, ProducedName, OutputType, ".h"));
+				ProducedItems.Add(ProtoExecution.GetProducedFileItem(OutputDir, ProducedName, OutputType, ".cc"));
+			}
 		}
 
 		public ProtocAction(BinaryArchiveReader Reader) : base(ActionType.Compile)
@@ -140,7 +85,7 @@ namespace UnrealBuildTool
 	}
 
 
-	enum ProtocOutput
+	enum ProtocOutputType
 	{
 		Cpp,
 		GrpcCpp,
@@ -148,128 +93,73 @@ namespace UnrealBuildTool
 
 	static class ProtocOutputExtension
 	{
-		public static string GetArgs(this ProtocOutput Output, DirectoryReference OutputLocation)
+		public static string GetArgs(this ProtocOutputType Output, DirectoryReference OutputLocation)
 		{
 			switch (Output)
 			{
-				case ProtocOutput.Cpp:
+				case ProtocOutputType.Cpp:
 					return $"--cpp_out={OutputLocation}";
-				case ProtocOutput.GrpcCpp:
+				case ProtocOutputType.GrpcCpp:
 					return $"--plugin=\"protoc-gen-grpc={ProtoExecution.GetGrpcCppPluginPath()}\" --grpc_out={OutputLocation}";
 			}
 
 			return "";
 		}
+
+		public static string GetExtensionPrefix(this ProtocOutputType Output)
+		{
+			switch (Output)
+			{
+				case ProtocOutputType.Cpp:
+					return "pb";
+				case ProtocOutputType.GrpcCpp:
+					return "grpc.pb";
+			}
+
+			return "";
+		}
 	}
 
 	class ProtoExecution
 	{
-		public static readonly string ProtobufVersion = "31.1";
+		public static readonly string ProtobufVersion = "30.0";
 		public static readonly string GrpcVersion = "1.72.1";
 
-		public static void SetupProtoModules(
-			IEnumerable<UEBuildModuleCPP> Modules,
-			TargetMakefile Makefile,
-			ILogger Logger)
+		public static ProtocOutput CompileProtoFiles(IEnumerable<FileItem> ProtoFiles, DirectoryReference ProtoModuleDir, DirectoryReference OutputDir, IActionGraphBuilder Graph)
 		{
-			foreach (UEBuildModuleCPP Module in Modules)
-			{
-				Makefile.ProtoModules.AddRange(SetupProtoModulesRecursively(Module));
-			}
-		}
+			ProtocOutput Result = new ProtocOutput();
 
-		static List<ProtoModuleInfo> SetupProtoModulesRecursively(UEBuildModuleCPP Module)
-		{
-			List<ProtoModuleInfo> ProtoModules = new List<ProtoModuleInfo>();
+			List<FileItem> ProtoCppFiles = new List<FileItem>();
+			List<FileItem> ProtoGrpcFiles = new List<FileItem>();
 
-			if (Module.bHasProto)
+			foreach (FileItem ProtoFile in ProtoFiles)
 			{
-				DirectoryItem GeneratedCodeDirectory = DirectoryItem.GetItemByDirectoryReference(Module.GeneratedCodeDirectoryProto!);
+				ProtoCppFiles.Add(ProtoFile);
 
-				ProtoModuleInfo ProtoModule = new ProtoModuleInfo
-				{
-					ModuleName = Module.Name,
-					ModuleRulesFile = Module.Rules.File,
-					GeneratedCodeDirectory = GeneratedCodeDirectory,
-					ProtoDirectory = DirectoryReference.Combine(Module.Rules.File.Directory, "Proto"),
-				};
+				string ProducedName = Path.GetRelativePath(ProtoModuleDir.ToString(), ProtoFile.FullName);
+				Result.HeaderFiles.Add(GetProducedFileItem(OutputDir, ProducedName, ProtocOutputType.Cpp, ".h"));
+				Result.CppFiles.Add(GetProducedFileItem(OutputDir, ProducedName, ProtocOutputType.Cpp, ".cc"));
 
-				ProtoModules.Add(ProtoModule);
+				if (FileReference.ReadAllText(ProtoFile.Location).Contains("service"))
+				{
+					ProtoGrpcFiles.Add(ProtoFile);
 
-				// Module.GeneratedCppDirectories ??= new List<string>();
-				// Module.GeneratedCppDirectories.Add(GeneratedCodeDirectory.FullName);
+					Result.HeaderFiles.Add(GetProducedFileItem(OutputDir, ProducedName, ProtocOutputType.GrpcCpp, ".h"));
+					Result.CppFiles.Add(GetProducedFileItem(OutputDir, ProducedName, ProtocOutputType.GrpcCpp, ".cc"));
+				}
 			}
 
-			// if (Module.PublicDependencyModules != null)
-			// {
-			// 	foreach (UEBuildModuleCPP DependencyModule in Module.PublicDependencyModules.OfType<UEBuildModuleCPP>())
-			// 	{
-			// 		ProtoModules.AddRange(SetupProtoModulesRecursively(DependencyModule));
-			// 	}
-			// }
-
-			// if (Module.PrivateDependencyModules != null)
-			// {
-			// 	foreach (UEBuildModuleCPP DependencyModule in Module.PrivateDependencyModules.OfType<UEBuildModuleCPP>())
-			// 	{
-			// 		ProtoModules.AddRange(SetupProtoModulesRecursively(DependencyModule));
-			// 	}
-			// }
-
-			return ProtoModules;
-		}
-
-		public static void CreateProtocActions(TargetMakefileBuilder MakefileBuilder)
-		{
-			if (MakefileBuilder.Makefile.ProtoModules.Count > 0)
+			if (ProtoCppFiles.Count > 0)
 			{
-				foreach (ProtoModuleInfo ProtoModule in MakefileBuilder.Makefile.ProtoModules)
-				{
-					DirectoryItem ProtoDirectory = DirectoryItem.GetItemByDirectoryReference(ProtoModule.ProtoDirectory);
-					foreach (FileReference ProtoFile in DirectoryReference.EnumerateFiles(ProtoModule.ProtoDirectory, "*.proto", SearchOption.AllDirectories))
-					{
-						MakefileBuilder.AddAction(new ProtocAction(ProtoModule, FileItem.GetItemByFileReference(ProtoFile)));
-					}
-				}
+				Graph.AddAction(new ProtocAction(ProtoCppFiles, ProtoModuleDir, OutputDir, ProtocOutputType.Cpp));
 			}
-		}
 
-		public static void ExecuteProtocActions(TargetMakefile Makefile, ILogger Logger)
-		{
-			if (Makefile.ProtoModules.Count > 0)
+			if (ProtoGrpcFiles.Count > 0)
 			{
-				foreach (ProtoModuleInfo ProtoModule in Makefile.ProtoModules)
-				{
-					DirectoryReference.CreateDirectory(ProtoModule.GeneratedCodeDirectory.Location);
-					foreach (FileReference ProtoFile in DirectoryReference.EnumerateFiles(ProtoModule.ProtoDirectory, "*.proto", SearchOption.AllDirectories))
-					{
-						// TODO: Compile multiple files in one exec
-						ExecuteProtocAction(ProtoModule, ProtoFile, ProtocOutput.Cpp, Logger);
-						// TODO: Make the check for gRPC more efficient
-						if (FileReference.ReadAllText(ProtoFile).Contains("service"))
-						{
-							ExecuteProtocAction(ProtoModule, ProtoFile, ProtocOutput.GrpcCpp, Logger);
-						}
-					}
-				}
+				Graph.AddAction(new ProtocAction(ProtoGrpcFiles, ProtoModuleDir, OutputDir, ProtocOutputType.GrpcCpp));
 			}
-		}
 
-		private static void ExecuteProtocAction(ProtoModuleInfo Module, FileReference ProtoFile, ProtocOutput Output, ILogger Logger)
-		{
-			DirectoryReference ProtobufIncludeDir = DirectoryReference.Combine(ProtoExecution.GetProtobufDirectory(), "include");
-			DirectoryReference ProtoDir = DirectoryReference.Combine(Module.ModuleRulesFile.Directory, "Proto");
-
-			FileReference CommandPath = ProtoExecution.GetProtocPath();
-			string Args = $" -I {ProtobufIncludeDir}";
-			Args += $" -I {ProtoDir}";
-			Args += $" {ProtoFile}";
-			Args += $" {Output.GetArgs(Module.GeneratedCodeDirectory.Location)}";
-
-			// TODO: Hide warnings
-			// TODO: Run async and parallel?
-			int Result = ExternalExecution.RunExternalNativeExecutable(CommandPath, Args, Logger);
-			// TODO: Handle Result
+			return Result;
 		}
 
 		public static DirectoryReference GetProtobufDirectory()
@@ -284,23 +174,48 @@ namespace UnrealBuildTool
 
 		public static FileReference GetProtocPath()
 		{
-			// TODO: Support other platforms
-			string PlatformDir = "Win64";
-			string ArchDir = "x64";
-			string ExecutableName = "protoc.exe";
+			string ExecutableName = "protoc" + BuildHostPlatform.Current.BinarySuffix;
+
+			if (BuildHostPlatform.Current.Platform == UnrealTargetPlatform.Win64)
+			{
+				return FileReference.Combine(GetProtobufDirectory(), "bin", "Win64", UnrealArch.Host.Value.WindowsName, "Release", ExecutableName);
+			}
+			else if (BuildHostPlatform.Current.Platform == UnrealTargetPlatform.Mac)
+			{
+				return FileReference.Combine(GetProtobufDirectory(), "bin", "Mac", "Release", ExecutableName);
+			}
+			else if (BuildHostPlatform.Current.Platform.IsInGroup(UnrealPlatformGroup.Unix))
+			{
+				return FileReference.Combine(GetProtobufDirectory(), "bin", "Unix", UnrealArch.Host.Value.LinuxName, "Release", ExecutableName);
+			}
+
+			throw new BuildException($"Unsupported platform {BuildHostPlatform.Current.Platform} for Protobuf compilation");
 
-			return FileReference.Combine(GetProtobufDirectory(), "bin", PlatformDir, ArchDir, "Release", ExecutableName);
 		}
 
 		public static FileReference GetGrpcCppPluginPath()
 		{
-			// TODO: Support other platforms
-			string PlatformDir = "Win64";
-			string ArchDir = "x64";
-			string ExecutableName = "grpc_cpp_plugin.exe";
+			string ExecutableName = "grpc_cpp_plugin" + BuildHostPlatform.Current.BinarySuffix;
 
-			DirectoryReference GrpcDirectory = DirectoryReference.Combine(Unreal.EngineDirectory, "Source", "ThirdParty", "gRPC", GrpcVersion);
-			return FileReference.Combine(GrpcDirectory, "bin", PlatformDir, ArchDir, "Release", ExecutableName);
+			if (BuildHostPlatform.Current.Platform == UnrealTargetPlatform.Win64)
+			{
+				return FileReference.Combine(GetGrpcDirectory(), "bin", "Win64", UnrealArch.Host.Value.WindowsName, "Release", ExecutableName);
+			}
+			else if (BuildHostPlatform.Current.Platform == UnrealTargetPlatform.Mac)
+			{
+				return FileReference.Combine(GetGrpcDirectory(), "bin", "Mac", "Release", ExecutableName);
+			}
+			else if (BuildHostPlatform.Current.Platform.IsInGroup(UnrealPlatformGroup.Unix))
+			{
+				return FileReference.Combine(GetGrpcDirectory(), "bin", "Unix", UnrealArch.Host.Value.LinuxName, "Release", ExecutableName);
+			}
+
+			throw new BuildException($"Unsupported platform {BuildHostPlatform.Current.Platform} for Protobuf compilation");
+		}
+
+		internal static FileItem GetProducedFileItem(DirectoryReference OutputDir, string BaseName, ProtocOutputType OutputType, string extension)
+		{
+			return FileItem.GetItemByFileReference(FileReference.Combine(OutputDir, Path.ChangeExtension(BaseName, OutputType.GetExtensionPrefix() + extension)));
 		}
 	}
 }
diff --git a/Engine/Source/Programs/UnrealBuildTool/System/TargetMakefile.cs b/Engine/Source/Programs/UnrealBuildTool/System/TargetMakefile.cs
--- a/Engine/Source/Programs/UnrealBuildTool/System/TargetMakefile.cs
+++ b/Engine/Source/Programs/UnrealBuildTool/System/TargetMakefile.cs
@@ -348,11 +348,6 @@ namespace UnrealBuildTool
 		public bool bUseVerseBPVM = true;
 #endif
 
-		/// <summary>
-		/// All build modules containing Verse code
-		/// </summary>
-		public List<ProtoModuleInfo> ProtoModules;
-
 		/// <summary>
 		/// List of all plugin names. The makefile will be considered invalid if any of these changes, or new plugins are added.
 		/// </summary>
@@ -462,7 +457,6 @@ namespace UnrealBuildTool
 #if __VPROJECT_AVAILABLE__
 			VNIModules = new List<VNIModuleInfo>();
 #endif
-			ProtoModules = new List<ProtoModuleInfo>();
 			PluginFiles = new HashSet<FileItem>();
 			ExternalDependencies = new HashSet<FileItem>();
 			InternalDependencies = new HashSet<FileItem>();
@@ -517,7 +511,6 @@ namespace UnrealBuildTool
 			VNIModules = Reader.ReadList(() => new VNIModuleInfo(Reader))!;
 			bUseVerseBPVM = Reader.ReadBool();
 #endif
-			ProtoModules = Reader.ReadList(() => new ProtoModuleInfo(Reader))!;
 			PluginFiles = Reader.ReadHashSet(() => Reader.ReadFileItem())!;
 			ExternalDependencies = Reader.ReadHashSet(() => Reader.ReadFileItem())!;
 			InternalDependencies = Reader.ReadHashSet(() => Reader.ReadFileItem())!;
@@ -592,7 +585,6 @@ namespace UnrealBuildTool
 			Writer.WriteList(VNIModules, e => e.Write(Writer));
 			Writer.WriteBool(bUseVerseBPVM);
 #endif
-			Writer.WriteList(ProtoModules, e => e.Write(Writer));
 			Writer.WriteHashSet(PluginFiles, x => Writer.WriteFileItem(x));
 			Writer.WriteHashSet(ExternalDependencies, x => Writer.WriteFileItem(x));
 			Writer.WriteHashSet(InternalDependencies, x => Writer.WriteFileItem(x));
