[UBA]
* Fixed hang when pressing ctrl-c while doing lots of cache related operations
* Added task scopes around requesting horde cache and connecting to cache

diff --git a/Engine/Source/Programs/UnrealBuildTool/Executors/UnrealBuildAccelerator/UBAExecutor.cs b/Engine/Source/Programs/UnrealBuildTool/Executors/UnrealBuildAccelerator/UBAExecutor.cs
--- a/Engine/Source/Programs/UnrealBuildTool/Executors/UnrealBuildAccelerator/UBAExecutor.cs
+++ b/Engine/Source/Programs/UnrealBuildTool/Executors/UnrealBuildAccelerator/UBAExecutor.cs
@@ -390,6 +390,7 @@ namespace UnrealBuildTool
 							{
 								if (!UBAConfig.bDisableRemote) // Need to match the InitAsync
 								{
+									using SessionTaskScope s = new(session, "Requesting Horde cache", "");
 									foreach (IUBAAgentCoordinator coordinator in _agentCoordinators)
 									{
 										UnrealBuildAcceleratorCacheConfig? cacheConfig = coordinator.RequestCacheServer(new CancellationToken());
@@ -406,6 +407,7 @@ namespace UnrealBuildTool
 								session.UpdateStatus(1, 1, "Cache", LogEntryType.Info, null);
 								session.UpdateStatus(1, 6, "Connecting...", LogEntryType.Info);
 
+								using SessionTaskScope s = new(session, "Connecting to cache", "");
 								foreach (UnrealBuildAcceleratorCacheConfig cacheConfig in cacheConfigs)
 								{
 									ICacheClient cacheClient = ICacheClient.CreateCacheClient(session, UBAConfig.bReportCacheMissReason, cacheConfig.Crypto, GetCacheClientHint());
@@ -616,9 +618,9 @@ namespace UnrealBuildTool
 				if (_bIsCancelled)
 					return;
 				_bIsCancelled = true;
+				scheduler.Cancel();
 				foreach (ICacheClient client in cache)
 					client.Disconnect();
-				scheduler.Cancel();
 				server?.StopServer(); // Make sure all remove processes are returned. We can't have any callbacks after this
 				session?.CancelAll(); // Cancel all processes native side
 				storage?.SaveCasTable();
