[UBA]
* Added -UBAHordeCacheDesiredConnectionCount=<x> that can be used to set desired connection count from the command line

diff --git a/Engine/Source/Programs/UnrealBuildTool/Executors/UnrealBuildAccelerator/UnrealBuildAcceleratorHordeConfig.cs b/Engine/Source/Programs/UnrealBuildTool/Executors/UnrealBuildAccelerator/UnrealBuildAcceleratorHordeConfig.cs
--- a/Engine/Source/Programs/UnrealBuildTool/Executors/UnrealBuildAccelerator/UnrealBuildAcceleratorHordeConfig.cs
+++ b/Engine/Source/Programs/UnrealBuildTool/Executors/UnrealBuildAccelerator/UnrealBuildAcceleratorHordeConfig.cs
@@ -234,6 +234,7 @@ namespace UnrealBuildTool
 		/// Number of desired connections when connecting to cache
 		/// </summary>
 		[ConfigFile(ConfigHierarchyType.Engine, "Uba.Provider.Horde", "UbaCacheDesiredConnectionCount")]
+		[CommandLine("-UBAHordeCacheDesiredConnectionCount=")]
 		public int CacheDesiredConnectionCount { get; set; } = 1;
 	}
 
