[UBT]
* Changed so PerModuleInline.gen.cpp is always created regardless if linking is disabled or not if unity files are enabled. Because if we skip this then the unity file will look different when linking is disabled and will not match cache

diff --git a/Engine/Source/Programs/UnrealBuildTool/Configuration/UEBuildModuleCPP.cs b/Engine/Source/Programs/UnrealBuildTool/Configuration/UEBuildModuleCPP.cs
--- a/Engine/Source/Programs/UnrealBuildTool/Configuration/UEBuildModuleCPP.cs
+++ b/Engine/Source/Programs/UnrealBuildTool/Configuration/UEBuildModuleCPP.cs
@@ -784,7 +784,7 @@ namespace UnrealBuildTool
 
 			// The operator new/delete overrides need to exist in all modules.
 			// We put these overrides in a separate file in order to be able to ensure only one file is added per binary for the primary module
-			if (Target.LinkType == TargetLinkType.Modular && bCreatePerModuleFile && !Target.bDisableLinking)
+			if (Target.LinkType == TargetLinkType.Modular && bCreatePerModuleFile && (!Target.bDisableLinking || bMergeUnityFiles))
 			{
 				FileReference NewDeleteOverrides = FileReference.Combine(IntermediateDirectory, "PerModuleInline.gen.cpp");
 				List<string> Content =
