Adding verse settings to the UHT manifest

#rnx

diff --git a/Engine/Source/Programs/Shared/EpicGames.Core/UHTTypes.cs b/Engine/Source/Programs/Shared/EpicGames.Core/UHTTypes.cs
--- a/Engine/Source/Programs/Shared/EpicGames.Core/UHTTypes.cs
+++ b/Engine/Source/Programs/Shared/EpicGames.Core/UHTTypes.cs
@@ -126,6 +126,9 @@ namespace EpicGames.Core
 			public UHTVerseScope VerseScope { get; set; } = UHTVerseScope.PublicAPI;
 			public bool HasVerse { get; set; } = false;
 			public string VerseMountPoint { get; set; } = "";
+			public string VersePackageName { get; set; } = "";
+			public string VerseDirectoryPath { get; set; } = "";
+			public List<string> VerseDependencies { get; set; } = [];
 			public bool AlwaysExportStructs { get; set; } = true;
 			public bool AlwaysExportEnums { get; set; } = true;
 			public bool AllowUETypesInNamespaces { get; set; } = false;
diff --git a/Engine/Source/Programs/UnrealBuildTool/Configuration/UEBuildTarget.cs b/Engine/Source/Programs/UnrealBuildTool/Configuration/UEBuildTarget.cs
--- a/Engine/Source/Programs/UnrealBuildTool/Configuration/UEBuildTarget.cs
+++ b/Engine/Source/Programs/UnrealBuildTool/Configuration/UEBuildTarget.cs
@@ -2593,7 +2593,15 @@ namespace UnrealBuildTool
 			// Prepare cached data for UHT header generation
 			using (GlobalTracer.Instance.BuildSpan("ExternalExecution.SetupUObjectModules()").StartActive())
 			{
-				ExternalExecution.SetupUObjectModules(ModulesToGenerateHeadersFor, Rules.Platform, ProjectDescriptor, Makefile.UObjectModules, Makefile.UObjectModuleHeaders, Rules.GeneratedCodeVersion, MetadataCache, Logger);
+				ExternalExecution.SetupUObjectModules(
+					ModulesToGenerateHeadersFor, 
+					Rules.Platform,
+					ProjectDescriptor, 
+					Makefile.UObjectModules, 
+					Makefile.UObjectModuleHeaders, 
+					Makefile.ExecutableFile.Directory, 
+					Rules.GeneratedCodeVersion, 
+					MetadataCache);
 			}
 
 			if (Rules.NativePointerMemberBehaviorOverride != null)
@@ -2623,7 +2631,7 @@ namespace UnrealBuildTool
 			// files that are injected as top level prerequisites.  If UHT only emitted included header files, we wouldn't need to run it during the Gather phase at all.
 			if (Makefile.UObjectModules.Count > 0)
 			{
-				await ExternalExecution.ExecuteHeaderToolIfNecessaryAsync(BuildConfiguration, ProjectFile, Makefile, TargetName, WorkingSet, Logger);
+				await ExternalExecution.ExecuteHeaderToolIfNecessaryAsync(BuildConfiguration, ProjectFile, Makefile, TargetName, Logger);
 			}
 			
 			GatherExtraGeneratedCPPFileTypes(Makefile, TargetName, ref GlobalCompileEnvironment);
diff --git a/Engine/Source/Programs/UnrealBuildTool/Modes/BuildMode.cs b/Engine/Source/Programs/UnrealBuildTool/Modes/BuildMode.cs
--- a/Engine/Source/Programs/UnrealBuildTool/Modes/BuildMode.cs
+++ b/Engine/Source/Programs/UnrealBuildTool/Modes/BuildMode.cs
@@ -1176,7 +1176,7 @@ namespace UnrealBuildTool
 				// If the target needs UHT to be run, we'll go ahead and do that now
 				if (Makefile.UObjectModules.Count > 0)
 				{
-					await ExternalExecution.ExecuteHeaderToolIfNecessaryAsync(BuildConfiguration, TargetDescriptor.ProjectFile, Makefile, TargetDescriptor.Name, WorkingSet, Logger);
+					await ExternalExecution.ExecuteHeaderToolIfNecessaryAsync(BuildConfiguration, TargetDescriptor.ProjectFile, Makefile, TargetDescriptor.Name, Logger);
 				}
 
 #if __VPROJECT_AVAILABLE__
diff --git a/Engine/Source/Programs/UnrealBuildTool/System/ExternalExecution.cs b/Engine/Source/Programs/UnrealBuildTool/System/ExternalExecution.cs
--- a/Engine/Source/Programs/UnrealBuildTool/System/ExternalExecution.cs
+++ b/Engine/Source/Programs/UnrealBuildTool/System/ExternalExecution.cs
@@ -1,6 +1,7 @@
 // Copyright Epic Games, Inc. All Rights Reserved.
 
 using System;
+using System.Collections;
 using System.Collections.Generic;
 using System.Data;
 using System.Diagnostics;
@@ -23,59 +24,58 @@ namespace UnrealBuildTool
 {
 	static class UHTModuleTypeExtensions
 	{
-		public static UHTModuleType? EngineModuleTypeFromHostType(ModuleHostType ModuleType)
+		public static UHTModuleType? GetEngineModuleType(this ModuleHostType ModuleType)
 		{
-			switch (ModuleType)
+			return ModuleType switch
 			{
-				case ModuleHostType.Program:
-					return UHTModuleType.Program;
-				case ModuleHostType.Runtime:
-				case ModuleHostType.RuntimeNoCommandlet:
-				case ModuleHostType.RuntimeAndProgram:
-				case ModuleHostType.CookedOnly:
-				case ModuleHostType.ServerOnly:
-				case ModuleHostType.ClientOnly:
-				case ModuleHostType.ClientOnlyNoCommandlet:
-					return UHTModuleType.EngineRuntime;
-				case ModuleHostType.Developer:
-				case ModuleHostType.DeveloperTool:
-					return UHTModuleType.EngineDeveloper;
-				case ModuleHostType.Editor:
-				case ModuleHostType.EditorNoCommandlet:
-				case ModuleHostType.EditorAndProgram:
-					return UHTModuleType.EngineEditor;
-				case ModuleHostType.UncookedOnly:
-					return UHTModuleType.EngineUncooked;
-				default:
-					return null;
-			}
+				ModuleHostType.Program => UHTModuleType.Program,
+
+				ModuleHostType.Runtime or 
+				ModuleHostType.RuntimeNoCommandlet or 
+				ModuleHostType.RuntimeAndProgram or 
+				ModuleHostType.CookedOnly or 
+				ModuleHostType.ServerOnly or 
+				ModuleHostType.ClientOnly or 
+				ModuleHostType.ClientOnlyNoCommandlet => UHTModuleType.EngineRuntime,
+
+				ModuleHostType.Developer or 
+				ModuleHostType.DeveloperTool => UHTModuleType.EngineDeveloper,
+
+				ModuleHostType.Editor or
+				ModuleHostType.EditorNoCommandlet or 
+				ModuleHostType.EditorAndProgram => UHTModuleType.EngineEditor,
+
+				ModuleHostType.UncookedOnly => UHTModuleType.EngineUncooked,
+
+				_ => null,
+			};
 		}
-		public static UHTModuleType? GameModuleTypeFromHostType(ModuleHostType ModuleType)
+
+		public static UHTModuleType? GetGameModuleType(this ModuleHostType ModuleType)
 		{
-			switch (ModuleType)
+			return ModuleType switch
 			{
-				case ModuleHostType.Program:
-					return UHTModuleType.Program;
-				case ModuleHostType.Runtime:
-				case ModuleHostType.RuntimeNoCommandlet:
-				case ModuleHostType.RuntimeAndProgram:
-				case ModuleHostType.CookedOnly:
-				case ModuleHostType.ServerOnly:
-				case ModuleHostType.ClientOnly:
-				case ModuleHostType.ClientOnlyNoCommandlet:
-					return UHTModuleType.GameRuntime;
-				case ModuleHostType.Developer:
-				case ModuleHostType.DeveloperTool:
-					return UHTModuleType.GameDeveloper;
-				case ModuleHostType.Editor:
-				case ModuleHostType.EditorNoCommandlet:
-				case ModuleHostType.EditorAndProgram:
-					return UHTModuleType.GameEditor;
-				case ModuleHostType.UncookedOnly:
-					return UHTModuleType.GameUncooked;
-				default:
-					return null;
-			}
+				ModuleHostType.Program => UHTModuleType.Program,
+
+				ModuleHostType.Runtime or 
+				ModuleHostType.RuntimeNoCommandlet or 
+				ModuleHostType.RuntimeAndProgram or 
+				ModuleHostType.CookedOnly or 
+				ModuleHostType.ServerOnly or 
+				ModuleHostType.ClientOnly or 
+				ModuleHostType.ClientOnlyNoCommandlet => UHTModuleType.GameRuntime,
+
+				ModuleHostType.Developer or 
+				ModuleHostType.DeveloperTool => UHTModuleType.GameDeveloper,
+
+				ModuleHostType.Editor or 
+				ModuleHostType.EditorNoCommandlet or 
+				ModuleHostType.EditorAndProgram => UHTModuleType.GameEditor,
+
+				ModuleHostType.UncookedOnly => UHTModuleType.GameUncooked,
+
+				_ => null,
+			};
 		}
 	}
 
@@ -174,6 +174,26 @@ namespace UnrealBuildTool
 		/// </summary>
 		public required string VerseMountPoint;
 
+		/// <summary>
+		/// Verse module name
+		/// </summary>
+		public required string VersePackageName;
+
+		/// <summary>
+		/// Relative location of the verse code
+		/// </summary>
+		public required string VerseDirectoryPath;
+
+		/// <summary>
+		/// Collection of verse packages this package depends on
+		/// </summary>
+		public HashSet<string> VersePublicDependencies = [];
+
+		/// <summary>
+		/// Collection of verse packages this package depends on
+		/// </summary>
+		public HashSet<string> VersePrivateDependencies = [];
+
 		/// <summary>
 		/// If true, verse scripts were found
 		/// </summary>
@@ -203,48 +223,6 @@ namespace UnrealBuildTool
 		{
 		}
 
-		[Obsolete("Prefer to use an object initializer to avoid issues with argument ordering")]
-		public UHTModuleInfo(
-			string ModuleName, 
-			FileReference ModuleRulesFile, 
-			DirectoryReference[] ModuleDirectories, 
-			DirectoryReference[] ModuleIncludePaths, 
-			UHTModuleType ModuleType, 
-			DirectoryItem GeneratedCodeDirectory, 
-			EGeneratedCodeVersion GeneratedCodeVersion, 
-			bool bIsReadOnly, 
-			ModuleRules.PackageOverrideType OverrideType,
-			string VersePath,
-			VerseScope VerseScope,
-			bool HasVerse,
-			string VerseMountPoint,
-			bool AlwaysExportStructs,
-			bool AlwaysExportEnums,
-			bool AllowUETypesInNamespaces)
-		{
-			this.ModuleName = ModuleName;
-			this.ModuleRulesFile = ModuleRulesFile;
-			this.ModuleDirectories = ModuleDirectories;
-			this.ModuleIncludePaths = ModuleIncludePaths;
-			this.ModuleType = ModuleType.ToString();
-			OverrideModuleType = OverrideType.ToString();
-			PublicUObjectClassesHeaders = new List<FileItem>();
-			PublicUObjectHeaders = new List<FileItem>();
-			InternalUObjectHeaders = new List<FileItem>();
-			PrivateUObjectHeaders = new List<FileItem>();
-			PublicDefines = new List<string>();
-			this.GeneratedCodeDirectory = GeneratedCodeDirectory;
-			this.GeneratedCodeVersion = GeneratedCodeVersion;
-			this.bIsReadOnly = bIsReadOnly;
-			this.VersePath = VersePath;
-			this.VerseScope = VerseScope;
-			this.HasVerse = HasVerse;
-			this.VerseMountPoint = VerseMountPoint;
-			this.AlwaysExportStructs = AlwaysExportStructs;
-			this.AlwaysExportEnums = AlwaysExportEnums;
-			this.AllowUETypesInNamespaces = AllowUETypesInNamespaces;
-		}
-
 		[SetsRequiredMembers]
 		public UHTModuleInfo(BinaryArchiveReader Reader)
 		{
@@ -267,6 +245,10 @@ namespace UnrealBuildTool
 			VerseScope = (VerseScope)Reader.ReadInt();
 			HasVerse = Reader.ReadBool();
 			VerseMountPoint = Reader.ReadString()!;
+			VersePackageName = Reader.ReadString()!;
+			VerseDirectoryPath = Reader.ReadString()!;
+			VersePublicDependencies = Reader.ReadHashSet(() => Reader.ReadString())!;
+			VersePrivateDependencies = Reader.ReadHashSet(() => Reader.ReadString())!;
 			AlwaysExportStructs = Reader.ReadBool();
 			AlwaysExportEnums = Reader.ReadBool();
 			AllowUETypesInNamespaces = Reader.ReadBool();
@@ -294,6 +276,10 @@ namespace UnrealBuildTool
 			Writer.WriteInt((int)VerseScope);
 			Writer.WriteBool(HasVerse);
 			Writer.WriteString(VerseMountPoint);
+			Writer.WriteString(VersePackageName);
+			Writer.WriteString(VerseDirectoryPath);
+			Writer.WriteHashSet(VersePublicDependencies, Item => Writer.WriteString(Item));
+			Writer.WriteHashSet(VersePrivateDependencies, Item => Writer.WriteString(Item));
 			Writer.WriteBool(AlwaysExportStructs);
 			Writer.WriteBool(AlwaysExportEnums);
 			Writer.WriteBool(AllowUETypesInNamespaces);
@@ -341,22 +327,12 @@ namespace UnrealBuildTool
 	{
 		static UHTModuleType GetEngineModuleTypeFromDescriptor(ModuleDescriptor Module)
 		{
-			UHTModuleType? Type = UHTModuleTypeExtensions.EngineModuleTypeFromHostType(Module.Type);
-			if (Type == null)
-			{
-				throw new BuildException("Unhandled engine module type {0} for {1}", Module.Type.ToString(), Module.Name);
-			}
-			return Type.GetValueOrDefault();
+			return Module.Type.GetEngineModuleType() ?? throw new BuildException("Unhandled engine module type {0} for {1}", Module.Type.ToString(), Module.Name);
 		}
 
 		static UHTModuleType GetGameModuleTypeFromDescriptor(ModuleDescriptor Module)
 		{
-			UHTModuleType? Type = UHTModuleTypeExtensions.GameModuleTypeFromHostType(Module.Type);
-			if (Type == null)
-			{
-				throw new BuildException("Unhandled game module type {0}", Module.Type.ToString());
-			}
-			return Type.GetValueOrDefault();
+			return Module.Type.GetGameModuleType() ?? throw new BuildException("Unhandled game module type {0} for {1}", Module.Type.ToString(), Module.Name);
 		}
 
 		/// <summary>
@@ -385,7 +361,6 @@ namespace UnrealBuildTool
 			return Module.Rules.Context.Scope.Name == "Project" ? "Game" : "Engine";
 		}
 
-
 		/// <summary>
 		/// Gets the module type for a given rules object
 		/// </summary>
@@ -418,7 +393,7 @@ namespace UnrealBuildTool
 					ModuleDescriptor? Module = ProjectDescriptor.Modules.FirstOrDefault(x => x.Name == RulesObject.Name);
 					if (Module != null)
 					{
-						return UHTModuleTypeExtensions.GameModuleTypeFromHostType(Module.Type) ?? UHTModuleType.GameRuntime;
+						return Module.Type.GetGameModuleType() ?? UHTModuleType.GameRuntime;
 					}
 				}
 				return UHTModuleType.GameRuntime;
@@ -476,25 +451,42 @@ namespace UnrealBuildTool
 			UnrealTargetPlatform Platform, 
 			ProjectDescriptor? ProjectDescriptor, 
 			List<UHTModuleInfo> UObjectModules,
-			List<UHTModuleHeaderInfo> UObjectModuleHeaders, 
+			List<UHTModuleHeaderInfo> UObjectModuleHeaders,
+			DirectoryReference ExecutableDirectory,
 			EGeneratedCodeVersion GeneratedCodeVersion, 
-			SourceFileMetadataCache MetadataCache, 
-			ILogger Logger)
+			SourceFileMetadataCache MetadataCache)
 		{
 			// Find the type of each module
-			Dictionary<UEBuildModuleCPP, UHTModuleType> ModuleToType = new Dictionary<UEBuildModuleCPP, UHTModuleType>();
+			Dictionary<UEBuildModuleCPP, UHTModuleType> ModuleToType = [];
 			foreach (UEBuildModuleCPP Module in ModulesToGenerateHeadersFor)
 			{
 				ModuleToType[Module] = GetModuleType(Module.Rules, ProjectDescriptor);
 			}
 
 			// Sort modules by type, then by dependency
-			List<UEBuildModuleCPP> ModulesSortedByType = ModulesToGenerateHeadersFor.OrderBy(c => ModuleToType[c]).ToList();
-			ModulesSortedByType = UEBuildModule.StableTopologicalSort(ModulesSortedByType.Cast<UEBuildModule>().ToList()).Cast<UEBuildModuleCPP>().ToList();
+			List<UEBuildModuleCPP> ModulesSortedByType = [.. ModulesToGenerateHeadersFor.OrderBy(c => ModuleToType[c])];
+			ModulesSortedByType = [.. UEBuildModule.StableTopologicalSort([.. ModulesSortedByType.Cast<UEBuildModule>()]).Cast<UEBuildModuleCPP>()];
+
+			// Create a module to index lookup
+			Dictionary<UEBuildModule, int> ModuleToIndex = [];
+			for (int Idx = 0; Idx < ModulesSortedByType.Count; Idx++)
+			{
+				ModuleToIndex.Add(ModulesSortedByType[Idx], Idx);
+			}
+
+			// We neeed the base dir to make a relative path for the Verse source location
+			DirectoryReference BaseDir = ExecutableDirectory;
+			string? RelativeBaseDir = UEBuildTarget.GetRelativeBaseDir(ExecutableDirectory, Platform);
+			if (RelativeBaseDir != null)
+			{
+				// Copy the logic from UEBuildTarget.SetupGlobalEnvironment ~ln 3645 that defines UE_RELATIVE_BASE_DIR to make the BaseDir relative to the Engine directory
+				// This ensures that the source file gathering can always use FPlatformProcess:BaseDir reliably since it has UE_RELATIVE_BASE_DIR baked in
+				BaseDir = DirectoryReference.Combine(ExecutableDirectory, RelativeBaseDir);
+			}
 
 			// Create the info for each module in parallel
 			UHTModuleInfo[] ModuleInfoArray = new UHTModuleInfo[ModulesSortedByType.Count];
-			using (ThreadPoolWorkQueue Queue = new ThreadPoolWorkQueue())
+			using (ThreadPoolWorkQueue Queue = new())
 			{
 				IReadOnlySet<string> ExcludedFolders = UEBuildPlatform.GetBuildPlatform(Platform).GetExcludedFolderNames();
 				for (int Idx = 0; Idx < ModulesSortedByType.Count; Idx++)
@@ -503,9 +495,20 @@ namespace UnrealBuildTool
 
 					DirectoryItem GeneratedCodeDirectory = DirectoryItem.GetItemByDirectoryReference(Module.GeneratedCodeDirectoryUHT!);
 
-					DirectoryReference[] ModuleIncludePaths = Module.PublicIncludePaths.Union(Module.InternalIncludePaths).Union(Module.PrivateIncludePaths).ToArray();
+					DirectoryReference[] ModuleIncludePaths = [.. Module.PublicIncludePaths.Union(Module.InternalIncludePaths).Union(Module.PrivateIncludePaths)];
 
-					UHTModuleInfo Info = new UHTModuleInfo
+					string VerseMountPoint = string.Empty;
+					string VersePackageName = string.Empty;
+					string VerseDirectoryPath = string.Empty;
+					if (Module.bHasVerse)
+					{
+						DirectoryReference SourceDirectory = DirectoryReference.Combine(Module.Rules.Directory, "Verse");
+						VerseMountPoint = GetVerseMountPointForModule(Module);
+						VersePackageName = $"{VerseMountPoint}/{Module.Name}";
+						VerseDirectoryPath = SourceDirectory.MakeRelativeTo(BaseDir).Replace("\\", "/", StringComparison.Ordinal);
+					}
+
+					UHTModuleInfo Info = new()
 					{
 						ModuleName = Module.Name,
 						ModuleRulesFile = Module.RulesFile,
@@ -519,7 +522,9 @@ namespace UnrealBuildTool
 						VersePath = Module.Rules.VersePath ?? "",
 						VerseScope = Module.Rules.VerseScope,
 						HasVerse = Module.bHasVerse,
-						VerseMountPoint = Module.bHasVerse ? GetVerseMountPointForModule(Module) : string.Empty,
+						VerseMountPoint = VerseMountPoint,
+						VersePackageName = VersePackageName,
+						VerseDirectoryPath = VerseDirectoryPath,
 						AlwaysExportStructs = Module.Rules.bAlwaysExportStructs,
 						AlwaysExportEnums = Module.Rules.bAlwaysExportEnums,
 						AllowUETypesInNamespaces = Module.Rules.bAllowUETypesInNamespaces,
@@ -531,6 +536,13 @@ namespace UnrealBuildTool
 				}
 			}
 
+			// Gather all the verse dependencies
+			BitArray ModuleHasBeenGathered = new(ModulesSortedByType.Count, false);
+			for (int Idx = 0; Idx < ModulesSortedByType.Count; Idx++)
+			{
+				GatherVerseDependencies(Idx, ModulesSortedByType, ModuleInfoArray, ModuleToIndex, ModuleHasBeenGathered);
+			}
+
 			// Filter out all the modules with reflection data
 			for (int Idx = 0; Idx < ModulesSortedByType.Count; Idx++)
 			{
@@ -548,7 +560,7 @@ namespace UnrealBuildTool
 					Info.GeneratedCPPFilenameBase = Path.Combine(GeneratedCodeDirectoryUHT.FullName, Info.ModuleName) + ".gen";
 					if (!Module.Rules.bUsePrecompiled)
 					{
-						Module.GeneratedCppDirectories ??= new List<string>();
+						Module.GeneratedCppDirectories ??= [];
 						Module.GeneratedCppDirectories.Add(GeneratedCodeDirectoryUHT.FullName);
 					}
 
@@ -556,11 +568,13 @@ namespace UnrealBuildTool
 
 					DirectoryItem ModuleDirectoryItem = DirectoryItem.GetItemByDirectoryReference(Module.ModuleDirectory);
 
-					List<FileItem> ReflectedHeaderFiles = new List<FileItem>();
-					ReflectedHeaderFiles.AddRange(Info.PublicUObjectClassesHeaders);
-					ReflectedHeaderFiles.AddRange(Info.PublicUObjectHeaders);
-					ReflectedHeaderFiles.AddRange(Info.InternalUObjectHeaders);
-					ReflectedHeaderFiles.AddRange(Info.PrivateUObjectHeaders);
+					List<FileItem> ReflectedHeaderFiles =
+					[
+						.. Info.PublicUObjectClassesHeaders,
+						.. Info.PublicUObjectHeaders,
+						.. Info.InternalUObjectHeaders,
+						.. Info.PrivateUObjectHeaders,
+					];
 					UObjectModuleHeaders.Add(new UHTModuleHeaderInfo(ModuleDirectoryItem, ReflectedHeaderFiles, Module.Rules.bUsePrecompiled));
 				}
 				else
@@ -584,30 +598,67 @@ namespace UnrealBuildTool
 			// Set Module.bHasUObjects for any IncludePathModules not already processed.
 			// This is necessary to keep include paths consistent between targets built with -AllModules and without
 			// Note that PublicIncludePathModules can be recursive so we need to traverse the entire chain
-			List<UEBuildModuleCPP> IncludePathModules = new();
+			List<UEBuildModuleCPP> IncludePathModules = [];
 			CollectModulesOnlyIncluded(ModulesToGenerateHeadersFor, IncludePathModules);
 			if (IncludePathModules.Count > 0)
 			{
-				Dictionary<UEBuildModuleCPP, UHTModuleInfo> IncludePathInfo = new();
-				using (ThreadPoolWorkQueue Queue = new ThreadPoolWorkQueue())
+				Dictionary<UEBuildModuleCPP, UHTModuleInfo> IncludePathInfo = [];
+
+				using ThreadPoolWorkQueue Queue = new();
+				IReadOnlySet<string> ExcludedFolders = UEBuildPlatform.GetBuildPlatform(Platform).GetExcludedFolderNames();
+				foreach (UEBuildModuleCPP Module in IncludePathModules)
 				{
-					IReadOnlySet<string> ExcludedFolders = UEBuildPlatform.GetBuildPlatform(Platform).GetExcludedFolderNames();
-					foreach (UEBuildModuleCPP Module in IncludePathModules)
+					Queue.Enqueue(() =>
 					{
-						Queue.Enqueue(() =>
+						foreach (DirectoryItem ModuleDirectoryItem in Module.ModuleDirectories.Select(x => DirectoryItem.GetItemByDirectoryReference(x)))
 						{
-							foreach (DirectoryItem ModuleDirectoryItem in Module.ModuleDirectories.Select(x => DirectoryItem.GetItemByDirectoryReference(x)))
+							List<FileItem> HeaderFiles = [];
+							FindHeaders(ModuleDirectoryItem, ExcludedFolders, HeaderFiles);
+
+							if (HeaderFiles.Any(x => MetadataCache.ContainsReflectionMarkup(x)))
 							{
-								List<FileItem> HeaderFiles = new();
-								FindHeaders(ModuleDirectoryItem, ExcludedFolders, HeaderFiles);
-
-								if (HeaderFiles.Any(x => MetadataCache.ContainsReflectionMarkup(x)))
-								{
-									Module.bHasUObjects = true;
-									break;
-								}
+								Module.bHasUObjects = true;
+								break;
 							}
-						});
+						}
+					});
+				}
+			}
+		}
+
+		static UHTModuleInfo GatherVerseDependencies(int Idx, List<UEBuildModuleCPP> ModulesSortedByType, UHTModuleInfo[] ModuleInfoArray, Dictionary<UEBuildModule, int> ModuleToIndex, BitArray ModuleHasBeenGathered)
+		{
+			UHTModuleInfo ModuleInfo = ModuleInfoArray[Idx];
+			if (ModuleHasBeenGathered[Idx])
+			{
+				return ModuleInfo;
+			}
+			ModuleHasBeenGathered[Idx] = true;
+
+			UEBuildModuleCPP BuildModule = ModulesSortedByType[Idx];
+			if (BuildModule.bHasVerse)
+			{
+				GatherVerseDependencies(ModulesSortedByType, ModuleInfoArray, ModuleToIndex, ModuleHasBeenGathered, BuildModule.PublicDependencyModules, ModuleInfo.VersePublicDependencies);
+				GatherVerseDependencies(ModulesSortedByType, ModuleInfoArray, ModuleToIndex, ModuleHasBeenGathered, BuildModule.PrivateDependencyModules, ModuleInfo.VersePrivateDependencies);
+			}
+			return ModuleInfo;
+		}
+
+		static void GatherVerseDependencies(List<UEBuildModuleCPP> ModulesSortedByType, UHTModuleInfo[] ModuleInfoArray, Dictionary<UEBuildModule, int> ModuleToIndex, BitArray ModuleHasBeenGathered, 
+			List<UEBuildModule>? InDependencies, HashSet<string> OutDependencies)
+		{
+			if (InDependencies != null)
+			{
+				foreach (UEBuildModule Dependency in InDependencies)
+				{
+					if (ModuleToIndex.TryGetValue(Dependency, out int GatherIdx))
+					{
+						UHTModuleInfo GatherModuleInfo = GatherVerseDependencies(GatherIdx, ModulesSortedByType, ModuleInfoArray, ModuleToIndex, ModuleHasBeenGathered);
+						if (GatherModuleInfo.HasVerse)
+						{
+							OutDependencies.Add(GatherModuleInfo.VersePackageName);
+						}
+						OutDependencies.UnionWith(GatherModuleInfo.VersePublicDependencies); // This is correct, private inherits public
 					}
 				}
 			}
@@ -641,7 +692,7 @@ namespace UnrealBuildTool
 
 		static void CollectModulesOnlyIncluded(IEnumerable<UEBuildModuleCPP> ModulesToGenerateHeadersFor, List<UEBuildModuleCPP> OutList)
 		{
-			HashSet<UEBuildModuleCPP> HandledModules = new();
+			HashSet<UEBuildModuleCPP> HandledModules = [];
 			foreach (UEBuildModuleCPP Module in ModulesToGenerateHeadersFor)
 			{
 				HandledModules.Add(Module);
@@ -666,7 +717,7 @@ namespace UnrealBuildTool
 			{
 				DirectoryItem ModuleDirectoryItem = DirectoryItem.GetItemByDirectoryReference(ModuleDirectory);
 
-				List<FileItem> HeaderFiles = new List<FileItem>();
+				List<FileItem> HeaderFiles = [];
 				FindHeaders(ModuleDirectoryItem, ExcludedFolders, HeaderFiles);
 
 				foreach (FileItem HeaderFile in HeaderFiles)
@@ -711,63 +762,6 @@ namespace UnrealBuildTool
 			}
 		}
 
-		/// <summary>
-		/// Gets UnrealHeaderTool.exe path. Does not care if UnrealheaderTool was build as a monolithic exe or not.
-		/// </summary>
-		static FileReference GetHeaderToolPath(FileReference ReceiptFile)
-		{
-			TargetReceipt Receipt = TargetReceipt.Read(ReceiptFile);
-			if (Receipt.Launch == null)
-			{
-				throw new BuildException("'Launch' property not set in UHT receipt.");
-			}
-			return Receipt.Launch;
-		}
-
-		/// <summary>
-		/// Gets the latest write time of any of the UnrealHeaderTool binaries (including DLLs and Plugins) or DateTime.MaxValue if UnrealHeaderTool does not exist
-		/// </summary>
-		/// <returns>Latest timestamp of UHT binaries or DateTime.MaxValue if UnrealHeaderTool is out of date and needs to be rebuilt.</returns>
-		static bool GetHeaderToolTimestampUtc(FileReference ReceiptPath, ILogger Logger, out DateTime Timestamp)
-		{
-			using (ScopedTimer TimestampTimer = new ScopedTimer("GetHeaderToolTimestamp", Logger))
-			{
-				// Try to read the receipt for UHT.
-				FileItem ReceiptFile = FileItem.GetItemByFileReference(ReceiptPath);
-				if (!ReceiptFile.Exists)
-				{
-					Timestamp = DateTime.MaxValue;
-					return false;
-				}
-
-				// Don't check timestamps for individual binaries if we're using the installed version of UHT. It will always be up to date.
-				if (!Unreal.IsFileInstalled(ReceiptFile.Location))
-				{
-					TargetReceipt? Receipt;
-					if (!TargetReceipt.TryRead(ReceiptPath, out Receipt))
-					{
-						Timestamp = DateTime.MaxValue;
-						return false;
-					}
-
-					// Make sure all the build products exist, and that the receipt is newer
-					foreach (BuildProduct BuildProduct in Receipt.BuildProducts)
-					{
-						FileItem BuildProductItem = FileItem.GetItemByFileReference(BuildProduct.Path);
-						if (!BuildProductItem.Exists || BuildProductItem.LastWriteTimeUtc > ReceiptFile.LastWriteTimeUtc)
-						{
-							Timestamp = DateTime.MaxValue;
-							return false;
-						}
-					}
-				}
-
-				// Return the timestamp for all the binaries
-				Timestamp = ReceiptFile.LastWriteTimeUtc;
-				return true;
-			}
-		}
-
 		/// <summary>
 		/// Gets the timestamp of CoreUObject.gen.cpp file.
 		/// </summary>
@@ -781,7 +775,7 @@ namespace UnrealBuildTool
 			}
 
 			// Otherwise look for CoreUObject.init.gen.cpp
-			FileInfo CoreGeneratedFileInfo = new FileInfo(Path.Combine(ModuleGeneratedCodeDirectory, ModuleName + ".init.gen.cpp"));
+			FileInfo CoreGeneratedFileInfo = new(Path.Combine(ModuleGeneratedCodeDirectory, ModuleName + ".init.gen.cpp"));
 			if (CoreGeneratedFileInfo.Exists)
 			{
 				return CoreGeneratedFileInfo.LastWriteTimeUtc;
@@ -817,12 +811,11 @@ namespace UnrealBuildTool
 		/// <summary>
 		/// Checks the class header files and determines if generated UObject code files are out of date in comparison.
 		/// </summary>
-		/// <param name="BuildConfiguration">Build configuration</param>
 		/// <param name="UObjectModules">Modules that we generate headers for</param>
 		/// <param name="HeaderToolTimestampUtc">Timestamp for UHT</param>
 		/// <param name="Logger">Logger for output</param>
 		/// <returns>True if the code files are out of date</returns>
-		private static bool AreGeneratedCodeFilesOutOfDate(BuildConfiguration BuildConfiguration, List<UHTModuleInfo> UObjectModules, DateTime HeaderToolTimestampUtc, ILogger Logger)
+		private static bool AreGeneratedCodeFilesOutOfDate(List<UHTModuleInfo> UObjectModules, DateTime HeaderToolTimestampUtc, ILogger Logger)
 		{
 			// Get CoreUObject.init.gen.cpp timestamp.  If the source files are older than the CoreUObject generated code, we'll
 			// need to regenerate code for the module
@@ -900,7 +893,7 @@ namespace UnrealBuildTool
 				}
 
 				// Has the .build.cs file changed since we last generated headers successfully?
-				FileInfo ModuleRulesFile = new FileInfo(Module.ModuleRulesFile.FullName);
+				FileInfo ModuleRulesFile = new(Module.ModuleRulesFile.FullName);
 				if (!ModuleRulesFile.Exists || ModuleRulesFile.LastWriteTimeUtc > SavedTimestampUtc)
 				{
 					Logger.LogDebug("UnrealHeaderTool needs to run because SavedTimestamp is older than the rules file ({ModuleModuleRulesFile}) for module {ModuleModuleName}", Module.ModuleRulesFile, Module.ModuleName);
@@ -908,11 +901,13 @@ namespace UnrealBuildTool
 				}
 
 				// Iterate over our UObjects headers and figure out if any of them have changed
-				List<FileItem> AllUObjectHeaders = new List<FileItem>();
-				AllUObjectHeaders.AddRange(Module.PublicUObjectClassesHeaders);
-				AllUObjectHeaders.AddRange(Module.PublicUObjectHeaders);
-				AllUObjectHeaders.AddRange(Module.InternalUObjectHeaders);
-				AllUObjectHeaders.AddRange(Module.PrivateUObjectHeaders);
+				List<FileItem> AllUObjectHeaders =
+				[
+					.. Module.PublicUObjectClassesHeaders,
+					.. Module.PublicUObjectHeaders,
+					.. Module.InternalUObjectHeaders,
+					.. Module.PrivateUObjectHeaders,
+				];
 
 				// Load up the old timestamp file and check to see if anything has changed
 				{
@@ -924,7 +919,7 @@ namespace UnrealBuildTool
 					}
 
 					// Iterate over our UObjects headers and figure out if any of them have changed
-					HashSet<string> ObjectHeadersSet = new HashSet<string>(AllUObjectHeaders.Select(x => x.AbsolutePath), FileReference.Comparer);
+					HashSet<string> ObjectHeadersSet = new(AllUObjectHeaders.Select(x => x.AbsolutePath), FileReference.Comparer);
 					foreach (string FileName in UObjectFilesFromPreviousRun)
 					{
 						if (!ObjectHeadersSet.Contains(FileName))
@@ -991,7 +986,7 @@ namespace UnrealBuildTool
 				if (!Module.bIsReadOnly)
 				{
 					string GeneratedCodeDirectory = Path.GetDirectoryName(Module.GeneratedCPPFilenameBase)!;
-					DirectoryInfo GeneratedCodeDirectoryInfo = new DirectoryInfo(GeneratedCodeDirectory);
+					DirectoryInfo GeneratedCodeDirectoryInfo = new(GeneratedCodeDirectory);
 
 					try
 					{
@@ -1003,11 +998,13 @@ namespace UnrealBuildTool
 							// files with UObject classes were deleted, so that we'll know to run UHT even if the timestamps of all
 							// of the other source files were unchanged
 							{
-								List<string> AllUObjectFiles = new List<string>();
-								AllUObjectFiles.AddRange(Module.PublicUObjectClassesHeaders.ConvertAll(Item => Item.AbsolutePath));
-								AllUObjectFiles.AddRange(Module.PublicUObjectHeaders.ConvertAll(Item => Item.AbsolutePath));
-								AllUObjectFiles.AddRange(Module.InternalUObjectHeaders.ConvertAll(Item => Item.AbsolutePath));
-								AllUObjectFiles.AddRange(Module.PrivateUObjectHeaders.ConvertAll(Item => Item.AbsolutePath));
+								List<string> AllUObjectFiles =
+								[
+									.. Module.PublicUObjectClassesHeaders.ConvertAll(Item => Item.AbsolutePath),
+									.. Module.PublicUObjectHeaders.ConvertAll(Item => Item.AbsolutePath),
+									.. Module.InternalUObjectHeaders.ConvertAll(Item => Item.AbsolutePath),
+									.. Module.PrivateUObjectHeaders.ConvertAll(Item => Item.AbsolutePath),
+								];
 								await FileReference.WriteAllLinesAsync(TimestampFile, AllUObjectFiles);
 							}
 
@@ -1029,29 +1026,27 @@ namespace UnrealBuildTool
 		public static int RunExternalNativeExecutable(FileReference ExePath, string Commandline, ILogger Logger)
 		{
 			Logger.LogDebug("RunExternalExecutable {ExePathFullName} {Commandline}", ExePath.FullName, Commandline);
-			using (LogEventParser Parser = new LogEventParser(Logger))
-			{
-				Parser.AddMatchersFromAssembly(Assembly.GetExecutingAssembly());
-				using (Process GameProcess = new Process())
-				{
-					GameProcess.StartInfo.FileName = ExePath.FullName;
-					GameProcess.StartInfo.Arguments = Commandline;
-					GameProcess.StartInfo.UseShellExecute = false;
-					GameProcess.StartInfo.RedirectStandardOutput = true;
-					GameProcess.OutputDataReceived += (s, e) => PrintProcessOutputAsync(s, e, Parser);
-					GameProcess.Start();
-					GameProcess.BeginOutputReadLine();
-					GameProcess.WaitForExit();
-
-					return GameProcess.ExitCode;
-				}
-			}
+
+			using LogEventParser Parser = new(Logger);
+			Parser.AddMatchersFromAssembly(Assembly.GetExecutingAssembly());
+
+			using Process GameProcess = new();
+			GameProcess.StartInfo.FileName = ExePath.FullName;
+			GameProcess.StartInfo.Arguments = Commandline;
+			GameProcess.StartInfo.UseShellExecute = false;
+			GameProcess.StartInfo.RedirectStandardOutput = true;
+			GameProcess.OutputDataReceived += (s, e) => PrintProcessOutputAsync(e, Parser);
+			GameProcess.Start();
+			GameProcess.BeginOutputReadLine();
+			GameProcess.WaitForExit();
+
+			return GameProcess.ExitCode;
 		}
 
 		/// <summary>
 		/// Simple function to pipe output asynchronously
 		/// </summary>
-		private static void PrintProcessOutputAsync(object Sender, DataReceivedEventArgs Event, LogEventParser Parser)
+		private static void PrintProcessOutputAsync(DataReceivedEventArgs Event, LogEventParser Parser)
 		{
 			// DataReceivedEventHandler is fired with a null string when the output stream is closed.  We don't want to
 			// print anything for that event.
@@ -1065,7 +1060,7 @@ namespace UnrealBuildTool
 		/// Builds and runs the header tool and touches the header directories.
 		/// Performs any early outs if headers need no changes, given the UObject modules, tool path, game name, and configuration
 		/// </summary>
-		public static async Task ExecuteHeaderToolIfNecessaryAsync(BuildConfiguration BuildConfiguration, FileReference? ProjectFile, TargetMakefile Makefile, string TargetName, ISourceFileWorkingSet WorkingSet, ILogger Logger)
+		public static async Task ExecuteHeaderToolIfNecessaryAsync(BuildConfiguration BuildConfiguration, FileReference? ProjectFile, TargetMakefile Makefile, string TargetName, ILogger Logger)
 		{
 
 			// No need to run UHT on itself
@@ -1075,7 +1070,7 @@ namespace UnrealBuildTool
 				return;
 			}
 
-			await ExecuteHeaderToolIfNecessaryInternalAsync(BuildConfiguration, ProjectFile, Makefile, TargetName, WorkingSet, Logger);
+			await ExecuteHeaderToolIfNecessaryInternalAsync(BuildConfiguration, ProjectFile, Makefile, TargetName, Logger);
 		}
 
 		/// <summary>
@@ -1108,9 +1103,12 @@ namespace UnrealBuildTool
 		/// <returns>Output UHT manifest</returns>
 		public static UHTManifest CreateUHTManifest(TargetMakefile Makefile, string TargetName, FileReference DepsFileName)
 		{
-			List<UHTManifest.Module> Modules = new List<UHTManifest.Module>();
+			List<UHTManifest.Module> Modules = [];
 			foreach (UHTModuleInfo UObjectModule in Makefile.UObjectModules)
 			{
+				List<string> VerseDependencies = [.. UObjectModule.VersePublicDependencies.Union(UObjectModule.VersePrivateDependencies)];
+				VerseDependencies.Sort(StringComparer.Ordinal);
+
 				Modules.Add(
 					new UHTManifest.Module
 					{
@@ -1118,12 +1116,12 @@ namespace UnrealBuildTool
 						ModuleType = (UHTModuleType)Enum.Parse(typeof(UHTModuleType), UObjectModule.ModuleType),
 						OverrideModuleType = (EPackageOverrideType)Enum.Parse(typeof(EPackageOverrideType), UObjectModule.OverrideModuleType),
 						BaseDirectory = UObjectModule.ModuleDirectories[0].FullName,
-						IncludePaths = UObjectModule.ModuleIncludePaths.Select(IncludePath => IncludePath.FullName).ToList(),
+						IncludePaths = [.. UObjectModule.ModuleIncludePaths.Select(IncludePath => IncludePath.FullName)],
 						OutputDirectory = Path.GetDirectoryName(UObjectModule.GeneratedCPPFilenameBase)!,
-						ClassesHeaders = UObjectModule.PublicUObjectClassesHeaders.Select((Header) => Header.AbsolutePath).ToList(),
-						PublicHeaders = UObjectModule.PublicUObjectHeaders.Select((Header) => Header.AbsolutePath).ToList(),
-						InternalHeaders = UObjectModule.InternalUObjectHeaders.Select((Header) => Header.AbsolutePath).ToList(),
-						PrivateHeaders = UObjectModule.PrivateUObjectHeaders.Select((Header) => Header.AbsolutePath).ToList(),
+						ClassesHeaders = [.. UObjectModule.PublicUObjectClassesHeaders.Select((Header) => Header.AbsolutePath)],
+						PublicHeaders = [.. UObjectModule.PublicUObjectHeaders.Select((Header) => Header.AbsolutePath)],
+						InternalHeaders = [.. UObjectModule.InternalUObjectHeaders.Select((Header) => Header.AbsolutePath)],
+						PrivateHeaders = [.. UObjectModule.PrivateUObjectHeaders.Select((Header) => Header.AbsolutePath)],
 						PublicDefines = UObjectModule.PublicDefines,
 						GeneratedCPPFilenameBase = UObjectModule.GeneratedCPPFilenameBase,
 						SaveExportedHeaders = !UObjectModule.bIsReadOnly,
@@ -1132,14 +1130,17 @@ namespace UnrealBuildTool
 						VerseScope = (UHTVerseScope)Enum.Parse(typeof(UHTVerseScope), UObjectModule.VerseScope.ToString()),
 						HasVerse = UObjectModule.HasVerse,
 						VerseMountPoint = UObjectModule.VerseMountPoint,
+						VersePackageName = UObjectModule.VersePackageName,
+						VerseDirectoryPath = UObjectModule.VerseDirectoryPath,
+						VerseDependencies = VerseDependencies,
 						AlwaysExportStructs = UObjectModule.AlwaysExportStructs,
 						AlwaysExportEnums = UObjectModule.AlwaysExportEnums,
 						AllowUETypesInNamespaces = UObjectModule.AllowUETypesInNamespaces,
 						MinimizeGeneratedIncludes = UObjectModule.MinimizeGeneratedIncludes,
-					}); ;
+					});
 			}
 
-			List<string> UhtPlugins = new List<string>();
+			List<string> UhtPlugins = [];
 			if (Makefile.EnabledUhtPlugins != null)
 			{
 				foreach (FileReference UhtPlugin in Makefile.EnabledUhtPlugins)
@@ -1148,7 +1149,7 @@ namespace UnrealBuildTool
 				}
 			}
 
-			UHTManifest Manifest = new UHTManifest
+			UHTManifest Manifest = new()
 			{
 				TargetName = TargetName,
 				IsGameTarget = Makefile.TargetType != TargetType.Program,
@@ -1191,21 +1192,20 @@ namespace UnrealBuildTool
 		/// Builds and runs the header tool and touches the header directories.
 		/// Performs any early outs if headers need no changes, given the UObject modules, tool path, game name, and configuration
 		/// </summary>
-		private static async Task ExecuteHeaderToolIfNecessaryInternalAsync(BuildConfiguration BuildConfiguration, FileReference? ProjectFile,
-			TargetMakefile Makefile, string TargetName, ISourceFileWorkingSet WorkingSet, ILogger Logger)
+		private static async Task ExecuteHeaderToolIfNecessaryInternalAsync(BuildConfiguration BuildConfiguration, FileReference? ProjectFile, TargetMakefile Makefile, string TargetName, ILogger Logger)
 		{
 			if (ProgressWriter.bWriteMarkup)
 			{
 				Logger.LogInformation("@progress push 5%");
 			}
-			using (ProgressWriter Progress = new ProgressWriter("Generating code...", false, Logger))
+			using (ProgressWriter Progress = new("Generating code...", false, Logger))
 			{
 				string RootLocalPath = Unreal.RootDirectory.FullName;
 
 				// Get the path to the assemblies we care about
 				string UhtAssemblyPath = typeof(UhtSession).Assembly.Location;
 
-				// Get a composite date/time for cirtical assemblies being used
+				// Get a composite date/time for critical assemblies being used
 				DateTime CompositeTimestamp = LatestDateTime(new FileInfo(Assembly.GetExecutingAssembly().Location).LastWriteTimeUtc, new FileInfo(UhtAssemblyPath).LastWriteTimeUtc);
 				if (Makefile.EnabledUhtPlugins != null)
 				{
@@ -1221,7 +1221,7 @@ namespace UnrealBuildTool
 				{
 					bUHTNeedsToRun = true;
 				}
-				else if (AreGeneratedCodeFilesOutOfDate(BuildConfiguration, Makefile.UObjectModules, CompositeTimestamp, Logger))
+				else if (AreGeneratedCodeFilesOutOfDate(Makefile.UObjectModules, CompositeTimestamp, Logger))
 				{
 					bUHTNeedsToRun = true;
 				}
@@ -1258,10 +1258,11 @@ namespace UnrealBuildTool
 					Logger.LogInformation("Parsing headers for {ActualTargetName}", ActualTargetName);
 
 					// Generate the command line
-					List<string> CmdArgs = new List<string>();
-					CmdArgs.Add((ProjectFile != null) ? ProjectFile.FullName : TargetName);
-					CmdArgs.Add(ModuleInfoFileName.FullName);
-					CmdArgs.Add("-WarningsAsErrors");
+					List<string> CmdArgs = [
+						(ProjectFile != null) ? ProjectFile.FullName : TargetName,
+						ModuleInfoFileName.FullName,
+						"-WarningsAsErrors",
+					];
 
 					if (Unreal.IsEngineInstalled())
 					{
@@ -1279,7 +1280,7 @@ namespace UnrealBuildTool
 						{
 							if (Arg[0] == '"' && Arg[^1] == '"')
 							{
-								CmdArgs.Add(Arg.Substring(1, Arg.Length - 2));
+								CmdArgs.Add(Arg[1..^1]);
 							}
 							else
 							{
@@ -1287,9 +1288,9 @@ namespace UnrealBuildTool
 							}
 						}
 					}
-					CommandLineArguments Arguments = new CommandLineArguments(CmdArgs.ToArray());
+					CommandLineArguments Arguments = new([.. CmdArgs]);
 
-					StringBuilder CmdLine = new StringBuilder();
+					StringBuilder CmdLine = new();
 					foreach (string Arg in CmdArgs)
 					{
 						CmdLine.AppendArgument(Arg);
@@ -1299,7 +1300,7 @@ namespace UnrealBuildTool
 
 					// Run UHT
 					IScope Timer = GlobalTracer.Instance.BuildSpan("Executing UnrealHeaderTool").StartActive();
-					UnrealHeaderToolMode UHTTool = new UnrealHeaderToolMode();
+					UnrealHeaderToolMode UHTTool = new();
 					CompilationResult UHTResult = (CompilationResult)await UHTTool.ExecuteAsync(Arguments, Logger);
 					Timer.Span.Finish();
 
