UnrealBuildTool: Disable uba for linking msvc pgo, there are still some issues that need to be investigated

#rnx

diff --git a/Engine/Source/Programs/UnrealBuildTool/Platform/Windows/VCToolChain.cs b/Engine/Source/Programs/UnrealBuildTool/Platform/Windows/VCToolChain.cs
--- a/Engine/Source/Programs/UnrealBuildTool/Platform/Windows/VCToolChain.cs
+++ b/Engine/Source/Programs/UnrealBuildTool/Platform/Windows/VCToolChain.cs
@@ -3316,6 +3316,11 @@ namespace UnrealBuildTool
 				LinkAction.bCanExecuteRemotely = false;
 			}
 
+			if (LinkEnvironment.bPGOOptimize || LinkEnvironment.bPGOProfile)
+			{
+				LinkAction.bCanExecuteInUBA = false; // Disabled for now. Should revisit to see why it is not working
+			}
+
 			if (EnvVars.Compiler.IsClang() && Target.WindowsPlatform.bSampleBasedPGO && LinkEnvironment.bPGOOptimize)
 			{
 				LinkAction.PrerequisiteItems.Add(FileItem.GetItemByFileReference(GetClangLinkProfDataFilename(LinkEnvironment)));
