UnrealBuildTool: Fix PerModuleInline.gen.cpp check when not linking

#rnx

diff --git a/Engine/Source/Programs/UnrealBuildTool/Configuration/UEBuildModuleCPP.cs b/Engine/Source/Programs/UnrealBuildTool/Configuration/UEBuildModuleCPP.cs
--- a/Engine/Source/Programs/UnrealBuildTool/Configuration/UEBuildModuleCPP.cs
+++ b/Engine/Source/Programs/UnrealBuildTool/Configuration/UEBuildModuleCPP.cs
@@ -784,7 +784,7 @@ namespace UnrealBuildTool
 
 			// The operator new/delete overrides need to exist in all modules.
 			// We put these overrides in a separate file in order to be able to ensure only one file is added per binary for the primary module
-			if (Target.LinkType == TargetLinkType.Modular && bCreatePerModuleFile && (!Target.bDisableLinking || bMergeUnityFiles))
+			if (Target.LinkType == TargetLinkType.Modular && bCreatePerModuleFile)
 			{
 				FileReference NewDeleteOverrides = FileReference.Combine(IntermediateDirectory, "PerModuleInline.gen.cpp");
 				List<string> Content =
