Fix UBT-generated VC projects when creating a modular editor build with a unique build environment.

These need a new `-Project=[OriginalProjectName].uproject` command line setting, or they will fail at startup.

diff --git a/Engine/Source/Programs/UnrealBuildTool/ProjectFiles/VisualStudio/VCProject.cs b/Engine/Source/Programs/UnrealBuildTool/ProjectFiles/VisualStudio/VCProject.cs
--- a/Engine/Source/Programs/UnrealBuildTool/ProjectFiles/VisualStudio/VCProject.cs
+++ b/Engine/Source/Programs/UnrealBuildTool/ProjectFiles/VisualStudio/VCProject.cs
@@ -2501,6 +2501,10 @@ namespace UnrealBuildTool
 						// Allow platforms to add any special properties they require
 						PlatformProjectGenerators.GenerateGamePlatformSpecificProperties(Platform, Configuration, TargetRulesObject.Type, VCProjectFileContent, RootOutputDirectory, TargetFilePath);
 					}
+					else if (TargetRulesObject.Type == TargetType.Editor && TargetRulesObject.ProjectFile != null && TargetRulesObject.LinkType == TargetLinkType.Modular && TargetRulesObject.BuildEnvironment == TargetBuildEnvironment.Unique)
+					{
+						VCProjectFileContent.AppendLine("    <LocalDebuggerCommandArguments>{0} -Project=\"{1}\"</LocalDebuggerCommandArguments>", TargetRulesObject.Name, NormalizeProjectPath(TargetRulesObject.ProjectFile));
+					}
 
 					VCProjectFileContent.AppendLine("  </PropertyGroup>");
 
