[RHI/UBT]
* Added support in UBT to have "FilesToGenerate" per module. This can be used by .build.cs files or platform specific code to inject files based on need.
* Moved only function using DDPI_SHADER_PLATFORM_NAME_MAP to cpp file
* Changed so DDPI_SHADER_PLATFORM_NAME_MAP is private definition

diff --git a/Engine/Source/Programs/ShaderCompileWorker/ShaderCompileWorker.Build.cs b/Engine/Source/Programs/ShaderCompileWorker/ShaderCompileWorker.Build.cs
--- a/Engine/Source/Programs/ShaderCompileWorker/ShaderCompileWorker.Build.cs
+++ b/Engine/Source/Programs/ShaderCompileWorker/ShaderCompileWorker.Build.cs
@@ -14,6 +14,7 @@ public class ShaderCompileWorker : ModuleRules
 				"Core",
 				"Projects",
 				"RenderCore",
+				"RHI",
 				"SandboxFile",
 				"TargetPlatform",
 				"ApplicationCore",
diff --git a/Engine/Source/Programs/UnrealBuildTool/Configuration/ModuleRules.cs b/Engine/Source/Programs/UnrealBuildTool/Configuration/ModuleRules.cs
--- a/Engine/Source/Programs/UnrealBuildTool/Configuration/ModuleRules.cs
+++ b/Engine/Source/Programs/UnrealBuildTool/Configuration/ModuleRules.cs
@@ -845,6 +845,11 @@ namespace UnrealBuildTool
 		/// </summary>
 		public bool bTreatAsEngineModule { get; set; }
 
+		/// <summary>
+		/// List of files that will be generated for this module. Entries in list are (FileName, FileContent)
+		/// </summary>
+		public Dictionary<string, IEnumerable<string>> FilesToGenerate = [];
+
 		/// <summary>
 		/// If this contains a platform and the project has had it's SDK version overriden from default, this module will be compiled as a project module, instead of shared engine module
 		/// </summary>
diff --git a/Engine/Source/Programs/UnrealBuildTool/Configuration/UEBuildModuleCPP.cs b/Engine/Source/Programs/UnrealBuildTool/Configuration/UEBuildModuleCPP.cs
--- a/Engine/Source/Programs/UnrealBuildTool/Configuration/UEBuildModuleCPP.cs
+++ b/Engine/Source/Programs/UnrealBuildTool/Configuration/UEBuildModuleCPP.cs
@@ -812,6 +812,18 @@ namespace UnrealBuildTool
 				CPPFiles.Add(PerModuleFile);
 			}
 
+			// Generate files provided through module rules.
+			foreach (KeyValuePair<string, IEnumerable<string>> kv in Rules.FilesToGenerate)
+			{
+				FileReference generatedFileRef = FileReference.Combine(IntermediateDirectory, "Gen", kv.Key);
+				FileItem generatedFile = Graph.CreateIntermediateTextFile(generatedFileRef, kv.Value, false);
+				if (generatedFileRef.HasExtension(".cpp"))
+				{
+					CPPFiles.Add(generatedFile);
+				}
+			}
+
+
 			// Engine modules will always use unity build mode unless MinSourceFilesForUnityBuildOverride is specified in
 			// the module rules file.  By default, game modules only use unity of they have enough source files for that
 			// to be worthwhile.  If you have a lot of small game modules, consider specifying MinSourceFilesForUnityBuildOverride=0
diff --git a/Engine/Source/Programs/UnrealBuildTool/Configuration/UEBuildPlatform.cs b/Engine/Source/Programs/UnrealBuildTool/Configuration/UEBuildPlatform.cs
--- a/Engine/Source/Programs/UnrealBuildTool/Configuration/UEBuildPlatform.cs
+++ b/Engine/Source/Programs/UnrealBuildTool/Configuration/UEBuildPlatform.cs
@@ -1394,5 +1394,71 @@ namespace UnrealBuildTool
 		/// </summary>
 		/// <param name="Receipt">Receipt for the target being deployed</param>
 		public abstract void Deploy(TargetReceipt Receipt);
+
+
+		/// <summary>
+		/// Struct used to add platform static shaders
+		/// </summary>
+		public struct StaticShader
+		{
+			/// <summary>
+			/// Name of shader
+			/// </summary>
+			public string Name { set; get; }
+			/// <summary>
+			/// Index of shader
+			/// </summary>
+			public int Index { set; get; }
+		}
+
+		/// <summary>
+		/// Called by platforms to add platform specific static shaders
+		/// </summary>
+		/// <param name="Target">Target rules</param>
+		/// <param name="Rules">Module rules</param>
+		/// <param name="StaticShaders">Static shaders supported by current platform</param>
+		public void AddPlatformStaticShaders(ReadOnlyTargetRules Target, ModuleRules Rules, IEnumerable<StaticShader> StaticShaders)
+		{
+			if (Rules.Name == "BuildSettings") // Put list of platforms in BuildSettings since this array will vary depending on what platforms machine has installed
+			{
+				if (!Rules.FilesToGenerate.ContainsKey("RHIStaticShaderPlatformNames.gen.h"))
+				{
+					Rules.FilesToGenerate.TryAdd("RHIStaticShaderPlatformNames.gen.h", 
+						[
+							"#pragma once",
+							"#include \"HAL/Platform.h\"",
+							"struct FStaticNameMapEntry { const TCHAR* Name; const TCHAR* Platform; const int Enum; };",
+							"extern BUILDSETTINGS_API FStaticNameMapEntry GStaticShaderNames[];"
+						]);
+					Rules.FilesToGenerate.Add("RHIStaticShaderPlatformNames.gen.cpp", new List<string>([
+						"#include \"RHIStaticShaderPlatformNames.gen.h\"",
+						"BUILDSETTINGS_API FStaticNameMapEntry GStaticShaderNames[] = ",
+						"{",
+						"	{ nullptr, nullptr, 0 },",
+						"};"
+						]));
+				}
+
+				if (Target.bNeedsExtraShaderFormats || Target.Platform == Platform)
+				{
+					List<string> content2 = (List<string>)Rules.FilesToGenerate["RHIStaticShaderPlatformNames.gen.cpp"];
+					content2.InsertRange(content2.Count - 2, StaticShaders.Select(shader =>
+						$"	{{ TEXT(\"{shader.Name}\"), TEXT(\"{Platform}\"), {shader.Index} }},"));
+				}
+			}
+			else if (Rules.Name == "RHI") // Add one file per platform with EShaderPlatform entries
+			{
+				if (Target.bNeedsExtraShaderFormats || Target.Platform == Platform)
+				{
+					List<string> content = ["#pragma once", "#include \"HAL/Platform.h\"", "enum EShaderPlatform : uint16;"];
+					foreach (StaticShader shader in StaticShaders)
+					{
+						content.Add($"inline constexpr EShaderPlatform SP_{shader.Name} = EShaderPlatform({shader.Index});");
+					}
+					Rules.FilesToGenerate.Add($"{Platform}RHIShaderPlatform.h", content);
+				}
+			}
+		}
+
 	}
 }
diff --git a/Engine/Source/Programs/UnrealBuildTool/Configuration/UEBuildTarget.cs b/Engine/Source/Programs/UnrealBuildTool/Configuration/UEBuildTarget.cs
--- a/Engine/Source/Programs/UnrealBuildTool/Configuration/UEBuildTarget.cs
+++ b/Engine/Source/Programs/UnrealBuildTool/Configuration/UEBuildTarget.cs
@@ -6315,6 +6315,14 @@ namespace UnrealBuildTool
 				// locations.
 				UEBuildPlatform.PlatformModifyHostModuleRules(ModuleName, RulesObject, Rules);
 
+				// If there are files to generate for module, then we have to create Gen directory and add it as public include
+				if (RulesObject.FilesToGenerate.Count != 0)
+				{
+					DirectoryReference GenDir = DirectoryReference.Combine(GetModuleIntermediateDirectory(RulesObject, Architectures), "Gen");
+					DirectoryReference.CreateDirectory(GenDir);
+					RulesObject.PublicIncludePaths.Add(NormalizeIncludePath(GenDir));
+				}
+
 				// Now, go ahead and create the module builder instance
 				Module = InstantiateModule(RulesObject, GeneratedCodeDirectory, Logger);
 
diff --git a/Engine/Source/Programs/UnrealBuildTool/Platform/IOS/UEBuildIOS.cs b/Engine/Source/Programs/UnrealBuildTool/Platform/IOS/UEBuildIOS.cs
--- a/Engine/Source/Programs/UnrealBuildTool/Platform/IOS/UEBuildIOS.cs
+++ b/Engine/Source/Programs/UnrealBuildTool/Platform/IOS/UEBuildIOS.cs
@@ -1001,6 +1001,8 @@ namespace UnrealBuildTool
 				return;
 			}		
 
+			AddPlatformStaticShaders(Target, Rules, []);
+
 			if ((Target.Platform == UnrealTargetPlatform.Win64) || (Target.Platform == UnrealTargetPlatform.Mac))
 			{
 				// For Windows we use the MetalDeveloperTools and we will check if the toolchain is present at runtime
diff --git a/Engine/Source/Runtime/RHI/Private/RHIStaticShaderPlatformNames.cpp b/Engine/Source/Runtime/RHI/Private/RHIStaticShaderPlatformNames.cpp
new file mode 100644
--- /dev/null
+++ b/Engine/Source/Runtime/RHI/Private/RHIStaticShaderPlatformNames.cpp
@@ -0,0 +1,26 @@
+// Copyright Epic Games, Inc. All Rights Reserved.
+
+#include "RHIStaticShaderPlatformNames.h"
+#include "RHIStaticShaderPlatformNames.gen.h"
+
+FStaticShaderPlatformNames::FStaticShaderPlatformNames()
+{
+	for (const FStaticNameMapEntry* Entry = GStaticShaderNames; Entry->Name; ++Entry)
+	{
+		check(IsStaticPlatform(EShaderPlatform(Entry->Enum)));
+		uint32 PlatformIndex = Entry->Enum - SP_StaticPlatform_First;
+
+		FPlatform& Platform = Platforms[PlatformIndex];
+		check(Platform.Name == NAME_None); // Check we've not already seen this platform
+
+		Platform.Name = Entry->Platform;
+		Platform.ShaderPlatform = FName(*FString::Printf(TEXT("SP_%s"), Entry->Name), FNAME_Add);
+		Platform.ShaderFormat = FName(*FString::Printf(TEXT("SF_%s"), Entry->Name), FNAME_Add);
+	}
+}
+
+const FStaticShaderPlatformNames& FStaticShaderPlatformNames::Get()
+{
+	static FStaticShaderPlatformNames Names;
+	return Names;
+}
diff --git a/Engine/Source/Runtime/RHI/Public/RHIShaderFormatDefinitions.inl b/Engine/Source/Runtime/RHI/Public/RHIShaderFormatDefinitions.inl
--- a/Engine/Source/Runtime/RHI/Public/RHIShaderFormatDefinitions.inl
+++ b/Engine/Source/Runtime/RHI/Public/RHIShaderFormatDefinitions.inl
@@ -66,9 +66,11 @@ static EShaderPlatform ShaderFormatNameToShaderPlatform(FName ShaderFormat)
 	if (ShaderFormat == NAME_VULKAN_SM6)						return SP_VULKAN_SM6;
 	if (ShaderFormat == NAME_VULKAN_SM5_ANDROID)				return SP_VULKAN_SM5_ANDROID;
 
+	const FStaticShaderPlatformNames& StaticNames = FStaticShaderPlatformNames::Get();
+
 	for (int32 StaticPlatform = SP_StaticPlatform_First; StaticPlatform <= SP_StaticPlatform_Last; ++StaticPlatform)
 	{
-		if (ShaderFormat == FStaticShaderPlatformNames::Get().GetShaderFormat(EShaderPlatform(StaticPlatform)))
+		if (ShaderFormat == StaticNames.GetShaderFormat(EShaderPlatform(StaticPlatform)))
 		{
 			return EShaderPlatform(StaticPlatform);
 		}
diff --git a/Engine/Source/Runtime/RHI/Public/RHIShaderPlatform.h b/Engine/Source/Runtime/RHI/Public/RHIShaderPlatform.h
--- a/Engine/Source/Runtime/RHI/Public/RHIShaderPlatform.h
+++ b/Engine/Source/Runtime/RHI/Public/RHIShaderPlatform.h
@@ -36,11 +36,8 @@ enum EShaderPlatform : uint16
 #define DDPI_NUM_STATIC_SHADER_PLATFORMS 16
 	SP_StaticPlatform_First = 32,
 
-	// Pull in the extra shader platform definitions from platform extensions.
-	// @todo - when we remove EShaderPlatform, fix up the shader platforms defined in UEBuild[Platform].cs files.
-#ifdef DDPI_EXTRA_SHADERPLATFORMS
-	DDPI_EXTRA_SHADERPLATFORMS
-#endif
+	// Platform specific shaders are now stored in platform specific files called <platform>RHIShaderPlatform.h
+	// These files are generated by UBT and can be included the same way as this file
 
 	SP_StaticPlatform_Last = (SP_StaticPlatform_First + DDPI_NUM_STATIC_SHADER_PLATFORMS - 1),
 
@@ -60,9 +57,13 @@ enum EShaderPlatform : uint16
 };
 static_assert(SP_NumPlatforms <= (1 << SP_NumBits), "SP_NumPlatforms will not fit on SP_NumBits");
 
-struct FGenericStaticShaderPlatform final
+#if USE_STATIC_SHADER_PLATFORM_ENUMS
+#include PREPROCESSOR_TO_STRING(PREPROCESSOR_JOIN(PLATFORM_HEADER_NAME, RHIShaderPlatform.h))
+#include COMPILED_PLATFORM_HEADER(StaticShaderPlatform.inl)
+#else
+struct FStaticShaderPlatform final
 {
-	inline FGenericStaticShaderPlatform(const EShaderPlatform InPlatform) : Platform(InPlatform) {}
+	inline FStaticShaderPlatform(const EShaderPlatform InPlatform) : Platform(InPlatform) {}
 	inline operator EShaderPlatform() const
 	{
 		return Platform;
@@ -79,11 +80,6 @@ struct FGenericStaticShaderPlatform final
 private:
 	const EShaderPlatform Platform;
 };
-
-#if USE_STATIC_SHADER_PLATFORM_ENUMS
-#include COMPILED_PLATFORM_HEADER(StaticShaderPlatform.inl)
-#else
-using FStaticShaderPlatform = FGenericStaticShaderPlatform;
 #endif
 
 // The maximum shader platform available on this system
diff --git a/Engine/Source/Runtime/RHI/Public/RHIStaticShaderPlatformNames.h b/Engine/Source/Runtime/RHI/Public/RHIStaticShaderPlatformNames.h
--- a/Engine/Source/Runtime/RHI/Public/RHIStaticShaderPlatformNames.h
+++ b/Engine/Source/Runtime/RHI/Public/RHIStaticShaderPlatformNames.h
@@ -17,41 +17,10 @@ private:
 		FName ShaderFormat;
 	} Platforms[NumPlatforms];
 
-	FStaticShaderPlatformNames()
-	{
-#ifdef DDPI_SHADER_PLATFORM_NAME_MAP
-		struct FStaticNameMapEntry
-		{
-			FName Name;
-			FName PlatformName;
-			int32 Index;
-		} NameMap[] =
-		{
-			DDPI_SHADER_PLATFORM_NAME_MAP
-		};
-
-		for (int32 MapIndex = 0; MapIndex < UE_ARRAY_COUNT(NameMap); ++MapIndex)
-		{
-			FStaticNameMapEntry const& Entry = NameMap[MapIndex];
-			check(IsStaticPlatform(EShaderPlatform(Entry.Index)));
-			uint32 PlatformIndex = Entry.Index - SP_StaticPlatform_First;
-
-			FPlatform& Platform = Platforms[PlatformIndex];
-			check(Platform.Name == NAME_None); // Check we've not already seen this platform
-
-			Platform.Name = Entry.PlatformName;
-			Platform.ShaderPlatform = FName(*FString::Printf(TEXT("SP_%s"), *Entry.Name.ToString()), FNAME_Add);
-			Platform.ShaderFormat = FName(*FString::Printf(TEXT("SF_%s"), *Entry.Name.ToString()), FNAME_Add);
-		}
-#endif
-	}
+	FStaticShaderPlatformNames();
 
 public:
-	static inline FStaticShaderPlatformNames const& Get()
-	{
-		static FStaticShaderPlatformNames Names;
-		return Names;
-	}
+	static RHI_API const FStaticShaderPlatformNames& Get();
 
 	static inline bool IsStaticPlatform(EShaderPlatform Platform)
 	{
diff --git a/Engine/Source/Runtime/RHI/RHI.Build.cs b/Engine/Source/Runtime/RHI/RHI.Build.cs
--- a/Engine/Source/Runtime/RHI/RHI.Build.cs
+++ b/Engine/Source/Runtime/RHI/RHI.Build.cs
@@ -12,6 +12,7 @@ public class RHI : ModuleRules
 		PrivateDependencyModuleNames.Add("TraceLog");
 		PrivateDependencyModuleNames.Add("ApplicationCore");
 		PrivateDependencyModuleNames.Add("Cbor");
+		PrivateDependencyModuleNames.Add("BuildSettings");
 
 		PublicDefinitions.AddDefinition("WITH_MGPU", Target.Platform.IsInGroup(UnrealPlatformGroup.Windows) && Target.Platform.IsInGroup(UnrealPlatformGroup.Desktop));
 
