Migrate OnlineServicesMcp to use new Proto UBT workflow. Disable the Protoc process on platforms outside of Win64, Linux and Mac until gRPC support is added for those platforms.

diff --git a/Engine/Source/Programs/UnrealBuildTool/Configuration/UEBuildModuleCPP.cs b/Engine/Source/Programs/UnrealBuildTool/Configuration/UEBuildModuleCPP.cs
--- a/Engine/Source/Programs/UnrealBuildTool/Configuration/UEBuildModuleCPP.cs
+++ b/Engine/Source/Programs/UnrealBuildTool/Configuration/UEBuildModuleCPP.cs
@@ -610,7 +610,13 @@ namespace UnrealBuildTool
 				ModuleCompileEnvironment.UserIncludePaths.Add(IntermediateDirectory);
 			}
 
-			if (bHasProto)
+			// We currently don't link gRPC on platforms other than Mac, Win64 and Unix and since
+			// all of the use cases we have also require gRPC we will compile only on those platforms.
+			// Once we add support for all platforms for gRPC this check can be removed.
+			bool bPlatformSupportsProtobuf = Target.IsInPlatformGroup(UnrealPlatformGroup.Unix)
+				|| Target.Platform == UnrealTargetPlatform.Mac
+				|| Target.Platform == UnrealTargetPlatform.Win64;
+			if (bHasProto && bPlatformSupportsProtobuf)
 			{
 				DirectoryReference ProtoDirectory = DirectoryReference.Combine(Rules.File.Directory, "Proto");
 				ProtocOutput OutputFiles = ProtoExecution.CompileProtoFiles(InputFiles.ProtoFiles, ProtoDirectory, IntermediateDirectory, Graph);
@@ -632,6 +638,9 @@ namespace UnrealBuildTool
 				ProtoCompileEnvironment.CppCompileWarnings.ShortenSizeTToIntWarningLevel = WarningLevel.Off;
 				ProtoCompileEnvironment.CppCompileWarnings.UnusedParameterWarningLevel = WarningLevel.Off;
 				ProtoCompileEnvironment.CppCompileWarnings.DeprecationWarningLevel = WarningLevel.Off;
+				// Some definitions that get passed cause issues with the proto compilations, and are not needed
+				// so we clear them for these.
+				ProtoCompileEnvironment.Definitions.Clear();
 
 				CPPOutput Output = ToolChain.CompileAllCPPFiles(ProtoCompileEnvironment, OutputFiles.CppFiles, IntermediateDirectory, Name, Graph);
 
