[UBA]
* Fixed so shared pch is in same bucket as the other actions using it
* Fixed so link step is in same bucket as compile steps.. since compile steps are inputs to link steps this will be an optimization to not have to download extra buckets

diff --git a/Engine/Source/Programs/UnrealBuildTool/Configuration/UEBuildModuleCPP.cs b/Engine/Source/Programs/UnrealBuildTool/Configuration/UEBuildModuleCPP.cs
--- a/Engine/Source/Programs/UnrealBuildTool/Configuration/UEBuildModuleCPP.cs
+++ b/Engine/Source/Programs/UnrealBuildTool/Configuration/UEBuildModuleCPP.cs
@@ -1862,6 +1862,7 @@ namespace UnrealBuildTool
 			CompileEnvironment.bValidateFormatStrings = ModuleCompileEnvironment.bValidateFormatStrings;
 			CompileEnvironment.bValidateInternalApi = ModuleCompileEnvironment.bValidateInternalApi;
 			CompileEnvironment.bValidateExperimentalApi = ModuleCompileEnvironment.bValidateExperimentalApi;
+			CompileEnvironment.bUseUnity = ModuleCompileEnvironment.bUseUnity; // Used by cache bucket. We want shared pch in same bucket
 		}
 
 		/// <summary>
diff --git a/Engine/Source/Programs/UnrealBuildTool/ToolChain/ISPCToolChain.cs b/Engine/Source/Programs/UnrealBuildTool/ToolChain/ISPCToolChain.cs
--- a/Engine/Source/Programs/UnrealBuildTool/ToolChain/ISPCToolChain.cs
+++ b/Engine/Source/Programs/UnrealBuildTool/ToolChain/ISPCToolChain.cs
@@ -376,13 +376,8 @@ namespace UnrealBuildTool
 		/// </summary>
 		/// <param name="target">Target rules</param>
 		/// <returns>bucket hash</returns>
-		static uint GetCacheBucket(ReadOnlyTargetRules? target)
+		static uint GetCacheBucket(ReadOnlyTargetRules target)
 		{
-			if (target == null)
-			{
-				return 0;
-			}
-
 			IEnumerable<string> strings = [
 				target.CacheSalt ?? String.Empty,			// Optional salt
 				target.BuildEnvironment.ToString(),			// Unique or Share environment
@@ -399,6 +394,11 @@ namespace UnrealBuildTool
 		/// <returns>bucket hash</returns>
 		static uint GetISPCCacheBucket(ReadOnlyTargetRules? target, CppCompileEnvironment compileEnvironment)
 		{
+			if (target == null)
+			{
+				return 0;
+			}
+
 			IEnumerable<string> strings = [
 				"ISCP",										// Create a separate bucket for ispc
 				compileEnvironment.Platform.ToString(),		// Platform
@@ -415,11 +415,16 @@ namespace UnrealBuildTool
 		/// <returns>bucket hash</returns>
 		public static uint GetCacheBucket(ReadOnlyTargetRules? target, CppCompileEnvironment compileEnvironment)
 		{
+			if (target == null)
+			{
+				return 0;
+			}
 			IEnumerable<string> strings = [
 				compileEnvironment.Platform.ToString(), // Platform
 				compileEnvironment.Configuration.ToString(), // Configuration
+				target.LinkType.ToString(), // Modular or monolithic (we want this for compiles too in order to end up in same bucket
 				!compileEnvironment.Architectures.bIsMultiArch ? compileEnvironment.Architecture.ToString() : compileEnvironment.Architectures.ToString(), // arm64 or x64
-				target != null && target.StaticAnalyzer != StaticAnalyzer.None && !compileEnvironment.bDisableStaticAnalysis ? target.StaticAnalyzer.ToString() : String.Empty, // If static analyzer
+				target.StaticAnalyzer != StaticAnalyzer.None && !compileEnvironment.bDisableStaticAnalysis ? target.StaticAnalyzer.ToString() : String.Empty, // If static analyzer
 				!compileEnvironment.bUseUnity ? "NU" : String.Empty, // Non-unity
 			];
 
@@ -453,12 +458,17 @@ namespace UnrealBuildTool
 		/// <returns>bucket hash</returns>
 		public static uint GetCacheBucket(ReadOnlyTargetRules? target, LinkEnvironment linkEnvironment)
 		{
+			if (target == null)
+			{
+				return 0;
+			}
+
 			IEnumerable<string> strings = [
 				linkEnvironment.Platform.ToString(), // Platform
 				linkEnvironment.Configuration.ToString(), // Configuration
-				linkEnvironment.LinkType.ToString(), // Modular or monolithic
+				target.LinkType.ToString(), // Modular or monolithic
 				!linkEnvironment.Architectures.bIsMultiArch ? linkEnvironment.Architecture.ToString() : linkEnvironment.Architectures.ToString(), // arm64 or x64
-				target?.bUseUnityBuild == false ? "NU" : String.Empty // Non-unity
+				target.bUseUnityBuild == false ? "NU" : String.Empty // Non-unity
 			];
 
 			// Optimize settings
