UnrealBuildTool: Prevent reprocessing when trimming ModuleRules definitions

#rnx

diff --git a/Engine/Source/Programs/UnrealBuildTool/Configuration/UEBuildModule.cs b/Engine/Source/Programs/UnrealBuildTool/Configuration/UEBuildModule.cs
--- a/Engine/Source/Programs/UnrealBuildTool/Configuration/UEBuildModule.cs
+++ b/Engine/Source/Programs/UnrealBuildTool/Configuration/UEBuildModule.cs
@@ -226,7 +226,7 @@ namespace UnrealBuildTool
 			ModuleApiDefine = Name.ToUpperInvariant() + "_API";
 
 			HashSet<string> PublicPreBuildLibraries = HashSetFromOptionalEnumerableStringParameter(Rules.PublicPreBuildLibraries);
-			PublicDefinitions = HashSetFromOptionalEnumerableStringParameter(Rules.PublicDefinitions);
+			PublicDefinitions = [.. HashSetFromOptionalEnumerableStringParameter(Rules.PublicDefinitions).Select(TrimDefinition)];
 			PublicIncludePaths = CreateDirectoryHashSet(Rules.PublicIncludePaths);
 			InternalIncludePaths = CreateDirectoryHashSet(Rules.InternalIncludePaths);
 			PublicSystemIncludePaths = CreateDirectoryHashSet(Rules.PublicSystemIncludePaths);
@@ -235,7 +235,7 @@ namespace UnrealBuildTool
 			PublicSystemLibraries = HashSetFromOptionalEnumerableStringParameter(Rules.PublicSystemLibraries);
 			PublicFrameworks = HashSetFromOptionalEnumerableStringParameter(Rules.PublicFrameworks);
 			PublicWeakFrameworks = HashSetFromOptionalEnumerableStringParameter(Rules.PublicWeakFrameworks);
-			PublicAutoRTFMExternalMappingFiles = new HashSet<FileItem>(Rules.AutoRTFMExternalMappingFiles.Select(x => FileItem.GetItemByPath(x)));
+			PublicAutoRTFMExternalMappingFiles = [.. Rules.AutoRTFMExternalMappingFiles.Select(FileItem.GetItemByPath)];
 
 			foreach (string LibraryName in PublicAdditionalLibraries)
 			{
@@ -433,7 +433,7 @@ namespace UnrealBuildTool
 		/// </summary>
 		protected HashSet<string> HashSetFromOptionalEnumerableStringParameter(IEnumerable<string> InEnumerableStrings)
 		{
-			return InEnumerableStrings == null ? new HashSet<string>() : new HashSet<string>(InEnumerableStrings.Select(x => ExpandPathVariables(x, null, null)));
+			return InEnumerableStrings == null ? [] : [.. InEnumerableStrings.Select(x => ExpandPathVariables(x, null, null))];
 		}
 
 		/// <summary>
@@ -604,33 +604,33 @@ namespace UnrealBuildTool
 		}
 
 		/// <summary>
-		/// Add definitions from source to target. This code also reformats defines to match strict rules which does not allow space before and after '='
+		/// Trims whitespace from a definition to match strict rules which does not allow space before and after '='.
 		/// </summary>
-		protected void AddDefinitions(List<string> Target, HashSet<string> Source)
+		/// <param name="Def"></param>
+		/// <returns>Trimmed string, or the current instance if unchanged</returns>
+		private static string TrimDefinition(string Def)
 		{
-			StringBuilder Builder = new();
-			foreach (string Def in Source)
-			{
-				string FixedDef = Def;
-				int IndexOfAssign = Def.IndexOf("=");
-				if (IndexOfAssign != -1)
-				{
-					ReadOnlySpan<char> DefSpan = Def.AsSpan();
-					ReadOnlySpan<char> Name = DefSpan.Slice(0, IndexOfAssign);
-					ReadOnlySpan<char> NameTrim = Name.Trim();
-					ReadOnlySpan<char> Value = DefSpan.Slice(IndexOfAssign + 1);
-					ReadOnlySpan<char> ValueTrim = Value.Trim();
-					if (Name.Length != NameTrim.Length || Value.Length != ValueTrim.Length)
-					{
-						Builder.Clear();
-						Builder.Append(NameTrim);
-						Builder.Append('=');
-						Builder.Append(ValueTrim);
-						FixedDef = Builder.ToString();
-					}
-				}
-				Target.Add(FixedDef);
-			}
+			int IndexOfAssign = Def.IndexOf('=');
+			if (IndexOfAssign != -1)
+			{
+				ReadOnlySpan<char> DefSpan = Def.AsSpan();
+				ReadOnlySpan<char> Name = DefSpan[..IndexOfAssign];
+				ReadOnlySpan<char> NameTrim = Name.Trim();
+				ReadOnlySpan<char> Value = DefSpan[(IndexOfAssign + 1)..];
+				ReadOnlySpan<char> ValueTrim = Value.Trim();
+				return Name.Length != NameTrim.Length || Value.Length != ValueTrim.Length
+					? $"{NameTrim}={ValueTrim}"
+					: Def;
+			}
+			return Def.Trim();
+		}
+
+		/// <summary>
+		/// Add definitions from source to target. TrimDefinition is called in the constructor to match strict rules which does not allow space before and after '='.
+		/// </summary>
+		protected static void AddDefinitions(List<string> Target, HashSet<string> Source)
+		{
+			Target.AddRange(Source);
 		}
 
 		/// <summary>
