UnrealBuildTool: Use list for linux link arguments instead of string

#rnx

diff --git a/Engine/Source/Programs/UnrealBuildTool/Platform/Linux/LinuxToolChain.cs b/Engine/Source/Programs/UnrealBuildTool/Platform/Linux/LinuxToolChain.cs
--- a/Engine/Source/Programs/UnrealBuildTool/Platform/Linux/LinuxToolChain.cs
+++ b/Engine/Source/Programs/UnrealBuildTool/Platform/Linux/LinuxToolChain.cs
@@ -991,13 +991,11 @@ namespace UnrealBuildTool
 			LinkAction.RootPaths = LinkEnvironment.RootPaths;
 			LinkAction.WorkingDirectory = Unreal.EngineSourceDirectory;
 
-			string LinkCommandString;
-			LinkCommandString = "\"" + Info.Clang + "\"";
+			List<string> LinkArguments = [];
+			LinkArguments.Add($"\"{Info.Clang}\"");
 
 			// Get link arguments.
-			List<string> LinkArguments = [];
 			GetLinkArguments_Global(LinkEnvironment, LinkArguments);
-			LinkCommandString += " " + String.Join(' ', LinkArguments);
 
 			// Tell the action that we're building an import library here and it should conditionally be
 			// ignored as a prerequisite for other actions
@@ -1051,14 +1049,14 @@ namespace UnrealBuildTool
 			}
 
 			// because the logic choosing between lld and ld is somewhat messy atm (lld fails to link .DSO due to bugs), make the name of the linker clear
-			LinkAction.CommandDescription += (LinkCommandString.Contains("-fuse-ld=lld")) ? " (lld)" : " (ld)";
+			LinkAction.CommandDescription += " (lld)";
 			LinkAction.CommandVersion = Info.ClangVersionString;
 			LinkAction.StatusDescription = Path.GetFileName(OutputFile.AbsolutePath);
 			LinkAction.CacheBucket = GetCacheBucket(TargetRules, LinkEnvironment);
 			LinkAction.ArtifactMode = ArtifactMode.Enabled;
 
 			// Add the output file to the command-line.
-			LinkCommandString += $" -o \"{OutputFile.AbsolutePath}\"";
+			LinkArguments.Add($"-o \"{OutputFile.AbsolutePath}\"");
 
 			// Add the input files to a response file, and pass the response file on the command-line.
 			List<string> ResponseLines = [];
@@ -1189,7 +1187,7 @@ namespace UnrealBuildTool
 			// when linking with --as-needed.
 			// Move external libraries to a separate --start-group/--end-group to fix it (and also make groups smaller and faster to link).
 			// See https://github.com/EpicGames/UnrealEngine/pull/2778 and https://github.com/EpicGames/UnrealEngine/pull/2793 for discussion
-			string ExternalLibraries = "";
+			List<string> ExternalLibraries = [];
 
 			// add libraries in a library group
 			ResponseLines.Add(String.Format(" --start-group"));
@@ -1199,7 +1197,7 @@ namespace UnrealBuildTool
 				if (String.IsNullOrEmpty(Path.GetDirectoryName(AdditionalLibrary)))
 				{
 					// library was passed just like "jemalloc", turn it into -ljemalloc
-					ExternalLibraries += $" -l{AdditionalLibrary}";
+					ExternalLibraries.Add($"-l{AdditionalLibrary}");
 				}
 				else if (Path.GetExtension(AdditionalLibrary) == ".a")
 				{
@@ -1214,11 +1212,11 @@ namespace UnrealBuildTool
 					// libcrypto/libssl contain number of functions that are being used in different DSOs. FIXME: generalize?
 					if (LinkEnvironment.bIsBuildingDLL && (AbsoluteAdditionalLibrary.Contains("libcrypto") || AbsoluteAdditionalLibrary.Contains("libssl")))
 					{
-						ResponseLines.Add(" --whole-archive " + AbsoluteAdditionalLibrary + " --no-whole-archive");
+						ResponseLines.Add($" --whole-archive {AbsoluteAdditionalLibrary} --no-whole-archive");
 					}
 					else
 					{
-						ResponseLines.Add(" " + AbsoluteAdditionalLibrary);
+						ResponseLines.Add($" {AbsoluteAdditionalLibrary}");
 					}
 
 					LinkAction.PrerequisiteItems.Add(FileItem.GetItemByPath(AdditionalLibrary));
@@ -1232,7 +1230,7 @@ namespace UnrealBuildTool
 					if (LinkToPrelinkFileList.TryGetValue(LibraryDependency, out FileItem? ToRelinkItem))
 					{
 						LibraryDependency = ToRelinkItem;
-						ExternalLibraries += $" -L\"{ToRelinkItem.Directory.FullName.Replace("\\", "/")}\"";
+						ExternalLibraries.Add($"-L\"{ToRelinkItem.Directory.FullName.Replace("\\", "/")}\"");
 					}
 
 					string LibName = Path.GetFileNameWithoutExtension(AdditionalLibrary);
@@ -1248,7 +1246,7 @@ namespace UnrealBuildTool
 						// the file name.
 						LibName = $":{LibraryDependency.Name}";
 					}
-					string LibLinkFlag = $" -l{LibName}";
+					string LibLinkFlag = $"-l{LibName}";
 
 					if (needsRelink)
 					{
@@ -1258,15 +1256,15 @@ namespace UnrealBuildTool
 						EngineAndGameLibrariesLinkFlags.Add(LibLinkFlag);
 						EngineAndGameLibrariesFiles.Add(LibraryDependency);
 						// it is important to add this exactly to the same place where the missing libraries would have been, it will be replaced later
-						if (!ExternalLibraries.Contains("--allow-shlib-undefined"))
+						if (!ExternalLibraries.Contains("-Wl,--allow-shlib-undefined"))
 						{
-							ExternalLibraries += String.Format(" -Wl,--allow-shlib-undefined");
+							ExternalLibraries.Add("-Wl,--allow-shlib-undefined");
 						}
 					}
 					else
 					{
 						LinkAction.PrerequisiteItems.Add(LibraryDependency);
-						ExternalLibraries += LibLinkFlag;
+						ExternalLibraries.Add(LibLinkFlag);
 					}
 				}
 			}
@@ -1275,54 +1273,54 @@ namespace UnrealBuildTool
 			FileReference ResponseFileName = GetResponseFileName(LinkEnvironment, OutputFile);
 			FileItem ResponseFileItem = Graph.CreateIntermediateTextFile(ResponseFileName, ResponseLines);
 
-			LinkCommandString += $" -Wl,@\"{ResponseFileName}\"";
+			LinkArguments.Add($"-Wl,@\"{ResponseFileName}\"");
 			LinkAction.PrerequisiteItems.Add(ResponseFileItem);
 
-			LinkCommandString += " -Wl,--start-group";
-			LinkCommandString += ExternalLibraries;
+			LinkArguments.Add("-Wl,--start-group");
+			LinkArguments.AddRange(ExternalLibraries);
 
 			// make unresolved symbols an error, unless a) building a cross-referenced DSO  b) we opted out
 			if ((!LinkEnvironment.bIsBuildingDLL || !LinkEnvironment.bIsCrossReferenced) && !LinkEnvironment.bIgnoreUnresolvedSymbols)
 			{
 				// This will make the linker report undefined symbols the current module, but ignore in the dependent DSOs.
 				// It is tempting, but may not be possible to change that report-all - due to circular dependencies between our libs.
-				LinkCommandString += " -Wl,--unresolved-symbols=ignore-in-shared-libs";
+				LinkArguments.Add("-Wl,--unresolved-symbols=ignore-in-shared-libs");
 			}
-			LinkCommandString += " -Wl,--end-group";
+			LinkArguments.Add("-Wl,--end-group");
 
-			LinkCommandString += " -lrt"; // needed for clock_gettime()
-			LinkCommandString += " -lm"; // math
+			LinkArguments.Add("-lrt"); // needed for clock_gettime()
+			LinkArguments.Add("-lm"); // math
 
 			if (ShouldUseLibcxx())
 			{
 				// libc++ and its abi lib
-				LinkCommandString += " -nodefaultlibs";
+				LinkArguments.Add("-nodefaultlibs");
 				string LibCxxLibDir = LinuxInfo.BaseLinuxPath + "/lib64/";
 				if(PlatformSDK.ForceUseLegacyLibCxx())
 				{
 					LibCxxLibDir = "ThirdParty/Unix/LibCxx/lib/Unix/" + LinkEnvironment.Architecture.LinuxName + "/";
 				}
-				LinkCommandString += " " +  LibCxxLibDir + "libc++.a";
-				LinkCommandString += " " +  LibCxxLibDir + "libc++abi.a";
-				LinkCommandString += " -lm";
-				LinkCommandString += " -lc";
-				LinkCommandString += " -lpthread"; // pthread_mutex_trylock is missing from libc stubs
-				LinkCommandString += " -lgcc_s";
-				LinkCommandString += " -lgcc";
+				LinkArguments.Add($"{LibCxxLibDir}libc++.a");
+				LinkArguments.Add($"{LibCxxLibDir}libc++abi.a");
+				LinkArguments.Add("-lm");
+				LinkArguments.Add("-lc");
+				LinkArguments.Add("-lpthread"); // pthread_mutex_trylock is missing from libc stubs
+				LinkArguments.Add("-lgcc_s");
+				LinkArguments.Add("-lgcc");
 			}
 
 			// these can be helpful for understanding the order of libraries or library search directories
 			if (PlatformSDK.bVerboseLinker)
 			{
-				LinkCommandString += " -Wl,--verbose";
-				LinkCommandString += " -Wl,--trace";
-				LinkCommandString += " -v";
+				LinkArguments.Add("-Wl,--verbose");
+				LinkArguments.Add("-Wl,--trace");
+				LinkArguments.Add("-v");
 			}
 
 			// Add the additional arguments specified by the environment.
-			LinkCommandString += LinkEnvironment.AdditionalArguments;
-			LinkCommandString = LinkCommandString.Replace("\\\\", "/");
-			LinkCommandString = LinkCommandString.Replace("\\", "/");
+			LinkArguments.Add(LinkEnvironment.AdditionalArguments);
+			LinkArguments = [.. LinkArguments.Select(x => x.Replace("\\\\", "/"))];
+			LinkArguments = [.. LinkArguments.Select(x => x.Replace("\\", "/"))];
 
 			bool bUseCmdExe = BuildHostPlatform.Current.ShellType == ShellType.Cmd;
 			FileReference ShellBinary = BuildHostPlatform.Current.Shell;
@@ -1331,9 +1329,9 @@ namespace UnrealBuildTool
 			// Linux has issues with scripts and parameter expansion from curely brakets
 			if (!bUseCmdExe)
 			{
-				LinkCommandString = LinkCommandString.Replace("{", "'{");
-				LinkCommandString = LinkCommandString.Replace("}", "}'");
-				LinkCommandString = LinkCommandString.Replace("$'{", "'${");    // fixing $'{ORIGIN}' to be '${ORIGIN}'
+				LinkArguments = [.. LinkArguments.Select(x => x.Replace("{", "'{"))];
+				LinkArguments = [.. LinkArguments.Select(x => x.Replace("}", "}'"))];
+				LinkArguments = [.. LinkArguments.Select(x => x.Replace("$'{", "'${"))]; // fixing $'{ORIGIN}' to be '${ORIGIN}'
 			}
 
 			string LinkScriptName = String.Format((bUseCmdExe ? "Link-{0}.bat" : "Link-{0}.sh"), OutputFile.Location.GetFileName());
@@ -1353,7 +1351,7 @@ namespace UnrealBuildTool
 					LinkWriter.WriteLine("set Retries=0");
 					LinkWriter.WriteLine(":linkloop");
 					LinkWriter.WriteLine("if %Retries% GEQ 1 goto failedtorelink");
-					LinkWriter.WriteLine(LinkCommandString);
+					LinkWriter.WriteLine(String.Join(' ', LinkArguments));
 					LinkWriter.WriteLine("if %errorlevel% neq 0 goto sleepandretry");
 					if (!needsRelink)
 					{
@@ -1376,7 +1374,7 @@ namespace UnrealBuildTool
 					LinkWriter.WriteLine("# *DO NOT EDIT*");
 					LinkWriter.WriteLine();
 					LinkWriter.WriteLine("set -o errexit");
-					LinkWriter.WriteLine(LinkCommandString);
+					LinkWriter.WriteLine(String.Join(' ', LinkArguments));
 					if (!needsRelink)
 					{
 						LinkWriter.WriteLine(GetDumpEncodeDebugCommand(LinkEnvironment, OutputFile, LinkAction));
@@ -1418,11 +1416,7 @@ namespace UnrealBuildTool
 				string LinkOutputFileForwardSlashes = OutputFile.AbsolutePath.Replace("\\", "/");
 				string RelinkedFileForwardSlashes = RealOutputFile.AbsolutePath.Replace("\\", "/");
 
-				string EngineAndGameLibrariesString = "";
-				foreach (string Library in EngineAndGameLibrariesLinkFlags)
-				{
-					EngineAndGameLibrariesString += Library;
-				}
+				string EngineAndGameLibrariesString = String.Join(' ', EngineAndGameLibrariesLinkFlags);
 
 				// create the relinking step
 				string RelinkScriptName = String.Format((bUseCmdExe ? "Relink-{0}.bat" : "Relink-{0}.sh"), RealOutputFile.Location.GetFileName());
@@ -1434,7 +1428,7 @@ namespace UnrealBuildTool
 				
 				{
 					StringWriter RelinkWriter = new();
-					string RelinkInvocation = LinkCommandString;
+					string RelinkInvocation = String.Join(' ', LinkArguments);
 					string Replace = "-Wl,--allow-shlib-undefined";
 					RelinkInvocation = RelinkInvocation.Replace(Replace, EngineAndGameLibrariesString);
 
