UnrealBuildTool: Update Linux to use response files for linking and relinking

#rnx

diff --git a/Engine/Source/Programs/UnrealBuildTool/Platform/Linux/LinuxToolChain.cs b/Engine/Source/Programs/UnrealBuildTool/Platform/Linux/LinuxToolChain.cs
--- a/Engine/Source/Programs/UnrealBuildTool/Platform/Linux/LinuxToolChain.cs
+++ b/Engine/Source/Programs/UnrealBuildTool/Platform/Linux/LinuxToolChain.cs
@@ -992,7 +992,6 @@ namespace UnrealBuildTool
 			LinkAction.WorkingDirectory = Unreal.EngineSourceDirectory;
 
 			List<string> LinkArguments = [];
-			LinkArguments.Add($"\"{Info.Clang}\"");
 
 			// Get link arguments.
 			GetLinkArguments_Global(LinkEnvironment, LinkArguments);
@@ -1124,7 +1123,7 @@ namespace UnrealBuildTool
 					if (!RPaths.Contains(RelativePath))
 					{
 						RPaths.Add(RelativePath);
-						ResponseLines.Add($" -rpath=\"${{ORIGIN}}/{RelativePath.Replace('\\', '/')}\"");
+						ResponseLines.Add($"-rpath=\"${{ORIGIN}}/{RelativePath.Replace('\\', '/')}\"");
 					}
 				}
 			}
@@ -1167,17 +1166,17 @@ namespace UnrealBuildTool
 				if (!RPaths.Contains(RelativePath))
 				{
 					RPaths.Add(RelativePath);
-					ResponseLines.Add($" -rpath=\"${{ORIGIN}}/{RelativePath.Replace('\\', '/')}\"");
+					ResponseLines.Add($"-rpath=\"${{ORIGIN}}/{RelativePath.Replace('\\', '/')}\"");
 				}
 			}
 
-			ResponseLines.Add($" -rpath-link=\"{Path.GetDirectoryName(RealOutputFile.AbsolutePath)}\"");
+			ResponseLines.Add($"-rpath-link=\"{Path.GetDirectoryName(RealOutputFile.AbsolutePath)}\"");
 
 			// Add the library paths to the argument list.
 			foreach (DirectoryReference LibraryPath in AllLibraryPaths)
 			{
 				// use absolute paths because of FixDependencies script again
-				ResponseLines.Add($" -L\"{LibraryPath.FullName.Replace('\\', '/')}\"");
+				ResponseLines.Add($"-L\"{LibraryPath.FullName.Replace('\\', '/')}\"");
 			}
 
 			List<string> EngineAndGameLibrariesLinkFlags = [];
@@ -1190,7 +1189,7 @@ namespace UnrealBuildTool
 			List<string> ExternalLibraries = [];
 
 			// add libraries in a library group
-			ResponseLines.Add(String.Format(" --start-group"));
+			ResponseLines.Add(String.Format("--start-group"));
 
 			foreach (string AdditionalLibrary in AdditionalLibraries)
 			{
@@ -1212,11 +1211,11 @@ namespace UnrealBuildTool
 					// libcrypto/libssl contain number of functions that are being used in different DSOs. FIXME: generalize?
 					if (LinkEnvironment.bIsBuildingDLL && (AbsoluteAdditionalLibrary.Contains("libcrypto") || AbsoluteAdditionalLibrary.Contains("libssl")))
 					{
-						ResponseLines.Add($" --whole-archive {AbsoluteAdditionalLibrary} --no-whole-archive");
+						ResponseLines.Add($"--whole-archive {AbsoluteAdditionalLibrary} --no-whole-archive");
 					}
 					else
 					{
-						ResponseLines.Add($" {AbsoluteAdditionalLibrary}");
+						ResponseLines.Add(AbsoluteAdditionalLibrary);
 					}
 
 					LinkAction.PrerequisiteItems.Add(FileItem.GetItemByPath(AdditionalLibrary));
@@ -1268,7 +1267,7 @@ namespace UnrealBuildTool
 					}
 				}
 			}
-			ResponseLines.Add(" --end-group");
+			ResponseLines.Add("--end-group");
 
 			FileReference ResponseFileName = GetResponseFileName(LinkEnvironment, OutputFile);
 			FileItem ResponseFileItem = Graph.CreateIntermediateTextFile(ResponseFileName, ResponseLines);
@@ -1340,6 +1339,10 @@ namespace UnrealBuildTool
 			LinkAction.PrerequisiteItems.Add(FileItem.GetItemByFileReference(LinkScriptFullPath));
 
 			{
+				FileReference LinkResponseFileName = GetResponseFileName(LinkEnvironment, FileItem.GetItemByFileReference(FileReference.Combine(OutputFile.Location.Directory, $"{OutputFile.Name}.link")));
+				FileItem LinkResponseFileItem = Graph.CreateIntermediateTextFile(LinkResponseFileName, LinkArguments);
+				LinkAction.PrerequisiteItems.Add(LinkResponseFileItem);
+
 				StringWriter LinkWriter = new();
 				if (bUseCmdExe)
 				{
@@ -1351,7 +1354,7 @@ namespace UnrealBuildTool
 					LinkWriter.WriteLine("set Retries=0");
 					LinkWriter.WriteLine(":linkloop");
 					LinkWriter.WriteLine("if %Retries% GEQ 1 goto failedtorelink");
-					LinkWriter.WriteLine(String.Join(' ', LinkArguments));
+					LinkWriter.WriteLine($"\"{Info.Clang}\" @\"{LinkResponseFileItem}\"");
 					LinkWriter.WriteLine("if %errorlevel% neq 0 goto sleepandretry");
 					if (!needsRelink)
 					{
@@ -1374,7 +1377,7 @@ namespace UnrealBuildTool
 					LinkWriter.WriteLine("# *DO NOT EDIT*");
 					LinkWriter.WriteLine();
 					LinkWriter.WriteLine("set -o errexit");
-					LinkWriter.WriteLine(String.Join(' ', LinkArguments));
+					LinkWriter.WriteLine($"\"{Info.Clang}\" @\"{LinkResponseFileItem}\"");
 					if (!needsRelink)
 					{
 						LinkWriter.WriteLine(GetDumpEncodeDebugCommand(LinkEnvironment, OutputFile, LinkAction));
@@ -1401,7 +1404,7 @@ namespace UnrealBuildTool
 				RelinkAction.RootPaths = LinkEnvironment.RootPaths;
 				RelinkAction.WorkingDirectory = LinkAction.WorkingDirectory;
 				RelinkAction.StatusDescription = LinkAction.StatusDescription;
-				RelinkAction.CommandDescription = "Relink";
+				RelinkAction.CommandDescription = "Relink (lld)";
 				RelinkAction.bCanExecuteRemotely = false;
 				RelinkAction.CacheBucket = GetCacheBucket(TargetRules, LinkEnvironment);
 				RelinkAction.ArtifactMode = ArtifactMode.Enabled;
@@ -1425,16 +1428,17 @@ namespace UnrealBuildTool
 				RelinkAction.PrerequisiteItems.UnionWith(LinkAction.PrerequisiteItems);
 				RelinkAction.PrerequisiteItems.Add(FileItem.GetItemByFileReference(RelinkScriptFullPath));
 
-				
 				{
 					StringWriter RelinkWriter = new();
-					string RelinkInvocation = String.Join(' ', LinkArguments);
-					string Replace = "-Wl,--allow-shlib-undefined";
-					RelinkInvocation = RelinkInvocation.Replace(Replace, EngineAndGameLibrariesString);
+					List<string> RelinkArguments = [.. LinkArguments];
+					RelinkArguments[RelinkArguments.IndexOf("-Wl,--allow-shlib-undefined")] = EngineAndGameLibrariesString;
+
+                    // should be the same as RelinkedFileRef
+                    RelinkArguments = [.. RelinkArguments.Select(x => x.Replace(LinkOutputFileForwardSlashes, RelinkedFileForwardSlashes))];
 
-					// should be the same as RelinkedFileRef
-					RelinkInvocation = RelinkInvocation.Replace(LinkOutputFileForwardSlashes, RelinkedFileForwardSlashes);
-					RelinkInvocation = RelinkInvocation.Replace("$", "\\$");
+					FileReference RelinkResponseFileName = GetResponseFileName(LinkEnvironment, FileItem.GetItemByFileReference(FileReference.Combine(OutputFile.Location.Directory, $"{OutputFile.Name}.relink")));
+					FileItem RelinkResponseFileItem = Graph.CreateIntermediateTextFile(RelinkResponseFileName, RelinkArguments);
+					RelinkAction.PrerequisiteItems.Add(RelinkResponseFileItem);
 
 					if (bUseCmdExe)
 					{
@@ -1445,7 +1449,7 @@ namespace UnrealBuildTool
 						RelinkWriter.WriteLine("set Retries=0");
 						RelinkWriter.WriteLine(":relinkloop");
 						RelinkWriter.WriteLine("if %Retries% GEQ 10 goto failedtorelink");
-						RelinkWriter.WriteLine(RelinkInvocation);
+						RelinkWriter.WriteLine($"\"{Info.Clang}\" @\"{RelinkResponseFileItem}\"");
 						RelinkWriter.WriteLine("if %errorlevel% neq 0 goto sleepandretry");
 						RelinkWriter.WriteLine(GetDumpEncodeDebugCommand(LinkEnvironment, RealOutputFile, RelinkAction));
 						RelinkWriter.WriteLine("exit 0");
@@ -1465,7 +1469,7 @@ namespace UnrealBuildTool
 						RelinkWriter.WriteLine("# *DO NOT EDIT*");
 						RelinkWriter.WriteLine();
 						RelinkWriter.WriteLine("set -o errexit");
-						RelinkWriter.WriteLine(RelinkInvocation);
+						RelinkWriter.WriteLine($"\"{Info.Clang}\" @\"{RelinkResponseFileItem}\"");
 						RelinkWriter.WriteLine(GetDumpEncodeDebugCommand(LinkEnvironment, RealOutputFile, RelinkAction));
 					}
 
