UnrealBuildTool: Prevent possible array out of bounds in linux toolchain, fix unquoted path in link args

#rnx

diff --git a/Engine/Source/Programs/UnrealBuildTool/Platform/Linux/LinuxToolChain.cs b/Engine/Source/Programs/UnrealBuildTool/Platform/Linux/LinuxToolChain.cs
--- a/Engine/Source/Programs/UnrealBuildTool/Platform/Linux/LinuxToolChain.cs
+++ b/Engine/Source/Programs/UnrealBuildTool/Platform/Linux/LinuxToolChain.cs
@@ -1299,8 +1299,8 @@ namespace UnrealBuildTool
 				{
 					LibCxxLibDir = "ThirdParty/Unix/LibCxx/lib/Unix/" + LinkEnvironment.Architecture.LinuxName + "/";
 				}
-				LinkArguments.Add($"{LibCxxLibDir}libc++.a");
-				LinkArguments.Add($"{LibCxxLibDir}libc++abi.a");
+				LinkArguments.Add($"\"{LibCxxLibDir}libc++.a\"");
+				LinkArguments.Add($"\"{LibCxxLibDir}libc++abi.a\"");
 				LinkArguments.Add("-lm");
 				LinkArguments.Add("-lc");
 				LinkArguments.Add("-lpthread"); // pthread_mutex_trylock is missing from libc stubs
@@ -1431,10 +1431,14 @@ namespace UnrealBuildTool
 				{
 					StringWriter RelinkWriter = new();
 					List<string> RelinkArguments = [.. LinkArguments];
-					RelinkArguments[RelinkArguments.IndexOf("-Wl,--allow-shlib-undefined")] = EngineAndGameLibrariesString;
+					int undefinedIndex = RelinkArguments.IndexOf("-Wl,--allow-shlib-undefined");
+					if (undefinedIndex != -1)
+					{
+						RelinkArguments[undefinedIndex] = EngineAndGameLibrariesString;
+					}
 
-                    // should be the same as RelinkedFileRef
-                    RelinkArguments = [.. RelinkArguments.Select(x => x.Replace(LinkOutputFileForwardSlashes, RelinkedFileForwardSlashes))];
+					// should be the same as RelinkedFileRef
+					RelinkArguments = [.. RelinkArguments.Select(x => x.Replace(LinkOutputFileForwardSlashes, RelinkedFileForwardSlashes))];
 
 					FileReference RelinkResponseFileName = GetResponseFileName(LinkEnvironment, FileItem.GetItemByFileReference(FileReference.Combine(OutputFile.Location.Directory, $"{OutputFile.Name}.relink")));
 					FileItem RelinkResponseFileItem = Graph.CreateIntermediateTextFile(RelinkResponseFileName, RelinkArguments);
