UnrealBuildTool: VFS linux link arguments

#rnx

diff --git a/Engine/Source/Programs/UnrealBuildTool/Platform/Linux/LinuxToolChain.cs b/Engine/Source/Programs/UnrealBuildTool/Platform/Linux/LinuxToolChain.cs
--- a/Engine/Source/Programs/UnrealBuildTool/Platform/Linux/LinuxToolChain.cs
+++ b/Engine/Source/Programs/UnrealBuildTool/Platform/Linux/LinuxToolChain.cs
@@ -676,7 +676,7 @@ namespace UnrealBuildTool
 
 			if (Options.HasFlag(ClangToolChainOptions.CompressDebugFile))
 			{
-				Arguments.Add("-Wl,--compress-debug-sections=zlib");
+				Arguments.Add("-gz=zlib");
 			}
 
 			// debugging symbols
@@ -726,7 +726,7 @@ namespace UnrealBuildTool
 			else
 			{
 				// x86_64 is now using updated ICU that doesn't need extra .so
-				Arguments.Add("-Wl,-rpath=${ORIGIN}/../../../Engine/Binaries/ThirdParty/ICU/icu4c-53_1/Unix/" + LinkEnvironment.Architecture.LinuxName);
+				Arguments.Add($"-Wl,-rpath=${{ORIGIN}}/../../../Engine/Binaries/ThirdParty/ICU/icu4c-53_1/Unix/{LinkEnvironment.Architecture.LinuxName}");
 			}
 
 			// @FIXME: Workaround for generating RPATHs for launching on devices UE-54136
@@ -1050,12 +1050,12 @@ namespace UnrealBuildTool
 			// because the logic choosing between lld and ld is somewhat messy atm (lld fails to link .DSO due to bugs), make the name of the linker clear
 			LinkAction.CommandDescription += " (lld)";
 			LinkAction.CommandVersion = Info.ClangVersionString;
-			LinkAction.StatusDescription = Path.GetFileName(OutputFile.AbsolutePath);
+			LinkAction.StatusDescription = OutputFile.Location.GetFileName();
 			LinkAction.CacheBucket = GetCacheBucket(TargetRules, LinkEnvironment);
 			LinkAction.ArtifactMode = ArtifactMode.Enabled;
 
 			// Add the output file to the command-line.
-			LinkArguments.Add($"-o \"{OutputFile.AbsolutePath}\"");
+			LinkArguments.Add($"-o \"{NormalizeCommandLinePath(OutputFile, LinkEnvironment.RootPaths)}\"");
 
 			// Add the input files to a response file, and pass the response file on the command-line.
 			List<string> ResponseLines = [];
@@ -1084,14 +1084,13 @@ namespace UnrealBuildTool
 			IEnumerable<string> AdditionalLibraries = Enumerable.Concat(LinkEnvironment.SystemLibraries, LinkEnvironment.Libraries.Select(x => x.FullName));
 			foreach (string AdditionalLibrary in AdditionalLibraries)
 			{
-				string PathToLib = Path.GetDirectoryName(AdditionalLibrary)!;
+				string? PathToLib = Path.GetDirectoryName(AdditionalLibrary);
 				if (!String.IsNullOrEmpty(PathToLib))
 				{
-					// make path absolute, because FixDependencies script may be executed in a different directory
-					DirectoryReference AbsolutePathToLib = new(PathToLib);
-					if (!AllLibraryPaths.Contains(AbsolutePathToLib))
+					DirectoryReference PathToLibDirectory = new(PathToLib);
+					if (!AllLibraryPaths.Contains(PathToLibDirectory))
 					{
-						AllLibraryPaths.Add(AbsolutePathToLib);
+						AllLibraryPaths.Add(PathToLibDirectory);
 					}
 				}
 
@@ -1170,13 +1169,12 @@ namespace UnrealBuildTool
 				}
 			}
 
-			ResponseLines.Add($"-rpath-link=\"{Path.GetDirectoryName(RealOutputFile.AbsolutePath)}\"");
+			ResponseLines.Add($"-rpath-link=\"{(LinkEnvironment.RootPaths.bUseVfs ? NormalizeCommandLinePath(RealOutputFile.Directory.Location, LinkEnvironment.RootPaths) :  RealOutputFile.Directory.FullName)}\"");
 
 			// Add the library paths to the argument list.
 			foreach (DirectoryReference LibraryPath in AllLibraryPaths)
 			{
-				// use absolute paths because of FixDependencies script again
-				ResponseLines.Add($"-L\"{LibraryPath.FullName.Replace('\\', '/')}\"");
+				ResponseLines.Add($"-L\"{NormalizeCommandLinePath(LibraryPath, LinkEnvironment.RootPaths)}\"");
 			}
 
 			List<string> EngineAndGameLibrariesLinkFlags = [];
@@ -1200,25 +1198,20 @@ namespace UnrealBuildTool
 				}
 				else if (Path.GetExtension(AdditionalLibrary) == ".a")
 				{
-					// static library passed in, pass it along but make path absolute, because FixDependencies script may be executed in a different directory
-					string AbsoluteAdditionalLibrary = Path.GetFullPath(AdditionalLibrary);
-					if (AbsoluteAdditionalLibrary.Contains(' '))
-					{
-						AbsoluteAdditionalLibrary = $"\"{AbsoluteAdditionalLibrary}\"";
-					}
-					AbsoluteAdditionalLibrary = AbsoluteAdditionalLibrary.Replace('\\', '/');
+					// static library passed in
+					FileReference AdditionalLibraryReference = new(AdditionalLibrary);
 
 					// libcrypto/libssl contain number of functions that are being used in different DSOs. FIXME: generalize?
-					if (LinkEnvironment.bIsBuildingDLL && (AbsoluteAdditionalLibrary.Contains("libcrypto") || AbsoluteAdditionalLibrary.Contains("libssl")))
+					if (LinkEnvironment.bIsBuildingDLL && (AdditionalLibraryReference.FullName.Contains("libcrypto") || AdditionalLibraryReference.FullName.Contains("libssl")))
 					{
-						ResponseLines.Add($"--whole-archive {AbsoluteAdditionalLibrary} --no-whole-archive");
+						ResponseLines.Add($"--whole-archive \"{NormalizeCommandLinePath(AdditionalLibraryReference, LinkEnvironment.RootPaths)}\" --no-whole-archive");
 					}
 					else
 					{
-						ResponseLines.Add(AbsoluteAdditionalLibrary);
+						ResponseLines.Add($"\"{NormalizeCommandLinePath(AdditionalLibraryReference, LinkEnvironment.RootPaths)}\"");
 					}
 
-					LinkAction.PrerequisiteItems.Add(FileItem.GetItemByPath(AdditionalLibrary));
+					LinkAction.PrerequisiteItems.Add(FileItem.GetItemByFileReference(AdditionalLibraryReference));
 				}
 				else
 				{
@@ -1229,7 +1222,7 @@ namespace UnrealBuildTool
 					if (LinkToPrelinkFileList.TryGetValue(LibraryDependency, out FileItem? ToRelinkItem))
 					{
 						LibraryDependency = ToRelinkItem;
-						ExternalLibraries.Add($"-L\"{ToRelinkItem.Directory.FullName.Replace("\\", "/")}\"");
+						ExternalLibraries.Add($"-L\"{NormalizeCommandLinePath(ToRelinkItem.Directory.Location, LinkEnvironment.RootPaths)}\"");
 					}
 
 					string LibName = Path.GetFileNameWithoutExtension(AdditionalLibrary);
@@ -1251,14 +1244,9 @@ namespace UnrealBuildTool
 					{
 						// We are building a cross referenced DLL so we can't actually include
 						// dependencies at this point. Instead we add it to the list of
-						// libraries to be used in the FixDependencies step.
+						// libraries to be used in the Relink step.
 						EngineAndGameLibrariesLinkFlags.Add(LibLinkFlag);
 						EngineAndGameLibrariesFiles.Add(LibraryDependency);
-						// it is important to add this exactly to the same place where the missing libraries would have been, it will be replaced later
-						if (!ExternalLibraries.Contains("-Wl,--allow-shlib-undefined"))
-						{
-							ExternalLibraries.Add("-Wl,--allow-shlib-undefined");
-						}
 					}
 					else
 					{
@@ -1267,12 +1255,17 @@ namespace UnrealBuildTool
 					}
 				}
 			}
+			if (needsRelink && EngineAndGameLibrariesLinkFlags.Count > 0)
+			{
+				// it is important to add this exactly to the same place where the missing libraries would have been, it will be replaced later
+				ExternalLibraries.Add("-Wl,--allow-shlib-undefined");
+			}
 			ResponseLines.Add("--end-group");
 
 			FileReference ResponseFileName = GetResponseFileName(LinkEnvironment, OutputFile);
 			FileItem ResponseFileItem = Graph.CreateIntermediateTextFile(ResponseFileName, ResponseLines);
 
-			LinkArguments.Add($"-Wl,@\"{ResponseFileName}\"");
+			LinkArguments.Add($"-Wl,@\"{NormalizeCommandLinePath(ResponseFileName, LinkEnvironment.RootPaths)}\"");
 			LinkAction.PrerequisiteItems.Add(ResponseFileItem);
 
 			LinkArguments.Add("-Wl,--start-group");
@@ -1294,13 +1287,13 @@ namespace UnrealBuildTool
 			{
 				// libc++ and its abi lib
 				LinkArguments.Add("-nodefaultlibs");
-				string LibCxxLibDir = LinuxInfo.BaseLinuxPath + "/lib64/";
+				DirectoryReference LibCxxLibDir = DirectoryReference.Combine(LinuxInfo.BaseLinuxPath!, "lib64");
 				if(PlatformSDK.ForceUseLegacyLibCxx())
 				{
-					LibCxxLibDir = "ThirdParty/Unix/LibCxx/lib/Unix/" + LinkEnvironment.Architecture.LinuxName + "/";
+					LibCxxLibDir = DirectoryReference.Combine(Unreal.EngineSourceDirectory, "ThirdParty", "Unix", "LibCxx", "lib", "Unix", LinkEnvironment.Architecture.LinuxName);
 				}
-				LinkArguments.Add($"\"{LibCxxLibDir}libc++.a\"");
-				LinkArguments.Add($"\"{LibCxxLibDir}libc++abi.a\"");
+				LinkArguments.Add($"\"{NormalizeCommandLinePath(FileReference.Combine(LibCxxLibDir, "libc++.a"), LinkEnvironment.RootPaths)}\"");
+				LinkArguments.Add($"\"{NormalizeCommandLinePath(FileReference.Combine(LibCxxLibDir, "libc++abi.a"), LinkEnvironment.RootPaths)}\"");
 				LinkArguments.Add("-lm");
 				LinkArguments.Add("-lc");
 				LinkArguments.Add("-lpthread"); // pthread_mutex_trylock is missing from libc stubs
@@ -1395,8 +1388,7 @@ namespace UnrealBuildTool
 			// Only execute linking on the local PC.
 			LinkAction.bCanExecuteRemotely = false;
 
-			// Prepare a script that will run later, once all shared libraries and the executable
-			// are created. This script will be called by action created in FixDependencies()
+			// Prepare a script that will run later, once all shared libraries and the executable are created.
 			if (needsRelink)
 			{
 				// Create the action to relink the library. This actions does not overwrite the source file so it can be executed in parallel
@@ -1404,7 +1396,7 @@ namespace UnrealBuildTool
 				RelinkAction.RootPaths = LinkEnvironment.RootPaths;
 				RelinkAction.WorkingDirectory = LinkAction.WorkingDirectory;
 				RelinkAction.StatusDescription = LinkAction.StatusDescription;
-				RelinkAction.CommandDescription = "Relink (lld)";
+				RelinkAction.CommandDescription = LinkAction.CommandDescription.Replace("Link", "Relink");
 				RelinkAction.bCanExecuteRemotely = false;
 				RelinkAction.CacheBucket = GetCacheBucket(TargetRules, LinkEnvironment);
 				RelinkAction.ArtifactMode = ArtifactMode.Enabled;
@@ -1416,10 +1408,8 @@ namespace UnrealBuildTool
 				}
 				RelinkAction.PrerequisiteItems.Add(OutputFile); // also depend on the first link action's output
 
-				string LinkOutputFileForwardSlashes = OutputFile.AbsolutePath.Replace("\\", "/");
-				string RelinkedFileForwardSlashes = RealOutputFile.AbsolutePath.Replace("\\", "/");
-
-				string EngineAndGameLibrariesString = String.Join(' ', EngineAndGameLibrariesLinkFlags);
+				string LinkOutputFileString = NormalizeCommandLinePath(OutputFile, LinkEnvironment.RootPaths);
+				string RelinkedFileString = NormalizeCommandLinePath(RealOutputFile, LinkEnvironment.RootPaths);
 
 				// create the relinking step
 				string RelinkScriptName = String.Format((bUseCmdExe ? "Relink-{0}.bat" : "Relink-{0}.sh"), RealOutputFile.Location.GetFileName());
@@ -1434,11 +1424,12 @@ namespace UnrealBuildTool
 					int undefinedIndex = RelinkArguments.IndexOf("-Wl,--allow-shlib-undefined");
 					if (undefinedIndex != -1)
 					{
-						RelinkArguments[undefinedIndex] = EngineAndGameLibrariesString;
+						RelinkArguments.RemoveAt(undefinedIndex);
+						RelinkArguments.InsertRange(undefinedIndex, EngineAndGameLibrariesLinkFlags);
 					}
 
 					// should be the same as RelinkedFileRef
-					RelinkArguments = [.. RelinkArguments.Select(x => x.Replace(LinkOutputFileForwardSlashes, RelinkedFileForwardSlashes))];
+					RelinkArguments = [.. RelinkArguments.Select(x => x.Replace(LinkOutputFileString, RelinkedFileString))];
 
 					FileReference RelinkResponseFileName = GetResponseFileName(LinkEnvironment, FileItem.GetItemByFileReference(FileReference.Combine(OutputFile.Location.Directory, $"{OutputFile.Name}.relink")));
 					FileItem RelinkResponseFileItem = Graph.CreateIntermediateTextFile(RelinkResponseFileName, RelinkArguments);
