UnrealBuildTool: VFS for linux link scripts except for linux shell script

#rnx

diff --git a/Engine/Source/Programs/UnrealBuildTool/Platform/Linux/LinuxToolChain.cs b/Engine/Source/Programs/UnrealBuildTool/Platform/Linux/LinuxToolChain.cs
--- a/Engine/Source/Programs/UnrealBuildTool/Platform/Linux/LinuxToolChain.cs
+++ b/Engine/Source/Programs/UnrealBuildTool/Platform/Linux/LinuxToolChain.cs
@@ -155,16 +155,19 @@ namespace UnrealBuildTool
 			bool bUseCmdExe = BuildHostPlatform.Current.Platform == UnrealTargetPlatform.Win64;
 			string DumpCommand = bUseCmdExe ? "\"{0}\" \"{1}\" \"{2}\" 2>NUL" : "\"{0}\" -c -o \"{2}\" \"{1}\"";
 
-			string DeleteCommand(string fileName)
+			string NormalizeFileRef(FileSystemReference file)
 			{
-				if (bUseCmdExe)
-				{
-					return $"if exist \"{fileName}\" del /F \"{fileName}\"";
-				}
-				else
-				{
-					return $"rm \"{fileName}\"";
-				}
+				// todo: binary paths on linux are currently having some issues
+				string normalized = bUseCmdExe ? NormalizeCommandLinePath(file, LinkEnvironment.RootPaths) : file.FullName;
+				return bUseCmdExe ? normalized.Replace('/', '\\') : normalized;
+			}
+
+			string NormalizeFile(FileItem file) => NormalizeFileRef(file.Location);
+
+			string DeleteCommand(FileItem file)
+			{
+				string fileName = NormalizeFile(file);
+				return bUseCmdExe ? $"if exist \"{fileName}\" del /F \"{fileName}\"" : $"rm \"{fileName}\"";
 			}
 
 			FileItem EncodedBinarySymbolsFile = FileItem.GetItemByPath(Path.Combine(LinkEnvironment.OutputDirectory!.FullName, OutputFile.Location.GetFileNameWithoutExtension() + ".sym"));
@@ -187,13 +190,13 @@ namespace UnrealBuildTool
 			{
 				// dump_syms
 				Out.WriteLine(DumpCommand,
-					LinuxInfo.DumpSyms,
-					OutputFile.AbsolutePath,
-					SymbolsFile.AbsolutePath
+					NormalizeFileRef(LinuxInfo.DumpSyms),
+					NormalizeFile(OutputFile),
+					NormalizeFile(SymbolsFile)
 				);
 
 				// encode breakpad symbols
-				Out.WriteLine($"\"{LinuxInfo.BreakpadEncoder}\" \"{SymbolsFile.AbsolutePath}\" \"{EncodedBinarySymbolsFile.AbsolutePath}\"");
+				Out.WriteLine($"\"{NormalizeFileRef(LinuxInfo.BreakpadEncoder)}\" \"{NormalizeFile(SymbolsFile)}\" \"{NormalizeFile(EncodedBinarySymbolsFile)}\"");
 
 				if (preservePsym)
 				{
@@ -201,13 +204,13 @@ namespace UnrealBuildTool
 				}
 				else
 				{
-					Out.WriteLine(DeleteCommand(SymbolsFile.AbsolutePath));
+					Out.WriteLine(DeleteCommand(SymbolsFile));
 				}
 			}
 			else
 			{
 				// we have to create dummy files to prevent packaging errors
-				Out.WriteLine($"echo DummySyms>> \"{EncodedBinarySymbolsFile.AbsolutePath}\"");
+				Out.WriteLine($"echo DummySyms>> \"{NormalizeFile(EncodedBinarySymbolsFile)}\"");
 			}
 
 			action.ProducedItems.Add(EncodedBinarySymbolsFile);
@@ -216,21 +219,21 @@ namespace UnrealBuildTool
 			{
 				if (MaxBinarySizeOverrideForObjcopy > 0 && bUseCmdExe)
 				{
-					Out.WriteLine($"for /F \"tokens=*\" %%F in (\"{OutputFile.AbsolutePath}\") DO set size=%%~zF");
+					Out.WriteLine($"for /F \"tokens=*\" %%F in (\"{NormalizeFile(OutputFile)}\") DO set size=%%~zF");
 
 					Out.WriteLine($"if %size% LSS {MaxBinarySizeOverrideForObjcopy} (");
 				}
 
 				// objcopy stripped file
-				Out.WriteLine($"\"{LinuxInfo.Objcopy}\" --strip-all \"{OutputFile.AbsolutePath}\" \"{StrippedFile.AbsolutePath}\"");
+				Out.WriteLine($"\"{NormalizeFileRef(LinuxInfo.Objcopy)}\" --strip-all \"{NormalizeFile(OutputFile)}\" \"{NormalizeFile(StrippedFile)}\"");
 
 				// objcopy debug file
-				Out.WriteLine($"\"{LinuxInfo.Objcopy}\" --only-keep-debug \"{OutputFile.AbsolutePath}\" \"{DebugFile.AbsolutePath}\"");
+				Out.WriteLine($"\"{NormalizeFileRef(LinuxInfo.Objcopy)}\" --only-keep-debug \"{NormalizeFile(OutputFile)}\" \"{NormalizeFile(DebugFile)}\"");
 
 				// objcopy link debug file to final so
-				Out.WriteLine($"\"{LinuxInfo.Objcopy}\" --add-gnu-debuglink=\"{DebugFile.AbsolutePath}\" \"{StrippedFile.AbsolutePath}\" \"{OutputFile.AbsolutePath}\"");
+				Out.WriteLine($"\"{NormalizeFileRef(LinuxInfo.Objcopy)}\" --add-gnu-debuglink=\"{NormalizeFile(DebugFile)}\" \"{NormalizeFile(StrippedFile)}\" \"{NormalizeFile(OutputFile)}\"");
 
-				Out.WriteLine(DeleteCommand(StrippedFile.AbsolutePath));
+				Out.WriteLine(DeleteCommand(StrippedFile));
 
 				if (bUseCmdExe)
 				{
@@ -238,20 +241,20 @@ namespace UnrealBuildTool
 					{
 						// If we have an override size, then we need to create a dummy file if that size is exceeded
 						Out.WriteLine(") ELSE (");
-						Out.WriteLine($"echo DummyDebug >> \"{DebugFile.AbsolutePath}\"");
+						Out.WriteLine($"echo DummyDebug >> \"{NormalizeFile(DebugFile)}\"");
 						Out.WriteLine(")");
 					}
 				}
 				else
 				{
 					// Change the debug file to normal permissions. It was taking on the +x rights from the output file
-					Out.WriteLine($"chmod 644 \"{DebugFile.AbsolutePath}\"");
+					Out.WriteLine($"chmod 644 \"{NormalizeFile(DebugFile)}\"");
 				}
 			}
 			else
 			{
 				// If we have disabled objcopy then we need to create a dummy debug file
-				Out.WriteLine($"echo DummyDebug >> \"{DebugFile.AbsolutePath}\"");
+				Out.WriteLine($"echo DummyDebug >> \"{NormalizeFile(DebugFile)}\"");
 			}
 
 			action.ProducedItems.Add(DebugFile);
@@ -1347,7 +1350,7 @@ namespace UnrealBuildTool
 					LinkWriter.WriteLine("set Retries=0");
 					LinkWriter.WriteLine(":linkloop");
 					LinkWriter.WriteLine("if %Retries% GEQ 1 goto failedtorelink");
-					LinkWriter.WriteLine($"\"{Info.Clang}\" @\"{LinkResponseFileItem}\"");
+					LinkWriter.WriteLine($"\"{NormalizeCommandLinePath(Info.Clang, LinkEnvironment.RootPaths).Replace('/', '\\')}\" @\"{NormalizeCommandLinePath(LinkResponseFileItem, LinkEnvironment.RootPaths).Replace('/', '\\')}\"");
 					LinkWriter.WriteLine("if %errorlevel% neq 0 goto sleepandretry");
 					if (!needsRelink)
 					{
@@ -1359,7 +1362,7 @@ namespace UnrealBuildTool
 					LinkWriter.WriteLine("set /a Retries+=1");
 					LinkWriter.WriteLine("goto linkloop");
 					LinkWriter.WriteLine(":failedtorelink");
-					LinkWriter.WriteLine($"echo Failed to link {OutputFile.AbsolutePath} after %Retries% retries");
+					LinkWriter.WriteLine($"echo Failed to link {NormalizeCommandLinePath(OutputFile, LinkEnvironment.RootPaths).Replace('/', '\\')} after %Retries% retries");
 					LinkWriter.WriteLine("exit 1");
 				}
 				else
@@ -1370,7 +1373,7 @@ namespace UnrealBuildTool
 					LinkWriter.WriteLine("# *DO NOT EDIT*");
 					LinkWriter.WriteLine();
 					LinkWriter.WriteLine("set -o errexit");
-					LinkWriter.WriteLine($"\"{Info.Clang}\" @\"{LinkResponseFileItem}\"");
+					LinkWriter.WriteLine($"\"{Info.Clang}\" @\"{NormalizeCommandLinePath(LinkResponseFileItem, LinkEnvironment.RootPaths)}\"");
 					if (!needsRelink)
 					{
 						LinkWriter.WriteLine(GetDumpEncodeDebugCommand(LinkEnvironment, OutputFile, LinkAction));
@@ -1412,7 +1415,7 @@ namespace UnrealBuildTool
 				string RelinkedFileString = NormalizeCommandLinePath(RealOutputFile, LinkEnvironment.RootPaths);
 
 				// create the relinking step
-				string RelinkScriptName = String.Format((bUseCmdExe ? "Relink-{0}.bat" : "Relink-{0}.sh"), RealOutputFile.Location.GetFileName());
+				string RelinkScriptName = String.Format(bUseCmdExe ? "Relink-{0}.bat" : "Relink-{0}.sh", RealOutputFile.Location.GetFileName());
 				FileReference RelinkScriptFullPath = FileReference.Combine(LinkEnvironment.IntermediateDirectory!, RelinkScriptName);
 
 				RelinkAction.PrerequisiteItems.UnionWith(LinkAction.PrerequisiteItems);
@@ -1444,7 +1447,7 @@ namespace UnrealBuildTool
 						RelinkWriter.WriteLine("set Retries=0");
 						RelinkWriter.WriteLine(":relinkloop");
 						RelinkWriter.WriteLine("if %Retries% GEQ 10 goto failedtorelink");
-						RelinkWriter.WriteLine($"\"{Info.Clang}\" @\"{RelinkResponseFileItem}\"");
+						RelinkWriter.WriteLine($"\"{NormalizeCommandLinePath(Info.Clang, LinkEnvironment.RootPaths).Replace('/', '\\')}\" @\"{NormalizeCommandLinePath(RelinkResponseFileItem, LinkEnvironment.RootPaths).Replace('/', '\\')}\"");
 						RelinkWriter.WriteLine("if %errorlevel% neq 0 goto sleepandretry");
 						RelinkWriter.WriteLine(GetDumpEncodeDebugCommand(LinkEnvironment, RealOutputFile, RelinkAction));
 						RelinkWriter.WriteLine("exit 0");
@@ -1453,7 +1456,7 @@ namespace UnrealBuildTool
 						RelinkWriter.WriteLine("set /a Retries+=1");
 						RelinkWriter.WriteLine("goto relinkloop");
 						RelinkWriter.WriteLine(":failedtorelink");
-						RelinkWriter.WriteLine($"echo Failed to relink {OutputFile.AbsolutePath} after %Retries% retries");
+						RelinkWriter.WriteLine($"echo Failed to relink {RelinkedFileString.Replace('/', '\\')} after %Retries% retries");
 						RelinkWriter.WriteLine("exit 1");
 					}
 					else
@@ -1464,7 +1467,7 @@ namespace UnrealBuildTool
 						RelinkWriter.WriteLine("# *DO NOT EDIT*");
 						RelinkWriter.WriteLine();
 						RelinkWriter.WriteLine("set -o errexit");
-						RelinkWriter.WriteLine($"\"{Info.Clang}\" @\"{RelinkResponseFileItem}\"");
+						RelinkWriter.WriteLine($"\"{Info.Clang}\" @\"{NormalizeCommandLinePath(RelinkResponseFileItem, LinkEnvironment.RootPaths)}\"");
 						RelinkWriter.WriteLine(GetDumpEncodeDebugCommand(LinkEnvironment, RealOutputFile, RelinkAction));
 					}
 
