[UBT/BuildSettings]
* Moved FilesToGenerate for RHIStaticPlatformNames.gen.h/cpp into BuildSettings.build.cs insteead of in UEBuildPlatform.cs. Fixes situation for machines having no special platforms synced

diff --git a/Engine/Source/Programs/UnrealBuildTool/Configuration/UEBuildPlatform.cs b/Engine/Source/Programs/UnrealBuildTool/Configuration/UEBuildPlatform.cs
--- a/Engine/Source/Programs/UnrealBuildTool/Configuration/UEBuildPlatform.cs
+++ b/Engine/Source/Programs/UnrealBuildTool/Configuration/UEBuildPlatform.cs
@@ -1421,24 +1421,6 @@ namespace UnrealBuildTool
 		{
 			if (Rules.Name == "BuildSettings") // Put list of platforms in BuildSettings since this array will vary depending on what platforms machine has installed
 			{
-				if (!Rules.FilesToGenerate.ContainsKey("RHIStaticShaderPlatformNames.gen.h"))
-				{
-					Rules.FilesToGenerate.TryAdd("RHIStaticShaderPlatformNames.gen.h", 
-						[
-							"#pragma once",
-							"#include \"HAL/Platform.h\"",
-							"struct FStaticNameMapEntry { const TCHAR* Name; const TCHAR* Platform; const int Enum; };",
-							"extern BUILDSETTINGS_API FStaticNameMapEntry GStaticShaderNames[];"
-						]);
-					Rules.FilesToGenerate.Add("RHIStaticShaderPlatformNames.gen.cpp", new List<string>([
-						"#include \"RHIStaticShaderPlatformNames.gen.h\"",
-						"BUILDSETTINGS_API FStaticNameMapEntry GStaticShaderNames[] = ",
-						"{",
-						"	{ nullptr, nullptr, 0 },",
-						"};"
-						]));
-				}
-
 				if (Target.bNeedsExtraShaderFormats || Target.Platform == Platform)
 				{
 					List<string> content2 = (List<string>)Rules.FilesToGenerate["RHIStaticShaderPlatformNames.gen.cpp"];
diff --git a/Engine/Source/Runtime/BuildSettings/BuildSettings.Build.cs b/Engine/Source/Runtime/BuildSettings/BuildSettings.Build.cs
--- a/Engine/Source/Runtime/BuildSettings/BuildSettings.Build.cs
+++ b/Engine/Source/Runtime/BuildSettings/BuildSettings.Build.cs
@@ -5,6 +5,7 @@ using UnrealBuildTool;
 using System;
 using System.Linq;
 using Microsoft.Extensions.Logging;
+using System.Collections.Generic;
 
 public class BuildSettings : ModuleRules
 {
@@ -64,6 +65,24 @@ public class BuildSettings : ModuleRules
 		{
 			bDisableAutoRTFMInstrumentation = true;
 		}
+
+		if (!FilesToGenerate.ContainsKey("RHIStaticShaderPlatformNames.gen.h"))
+		{
+			FilesToGenerate.TryAdd("RHIStaticShaderPlatformNames.gen.h",
+				[
+					"#pragma once",
+							"#include \"HAL/Platform.h\"",
+							"struct FStaticNameMapEntry { const TCHAR* Name; const TCHAR* Platform; const int Enum; };",
+							"extern BUILDSETTINGS_API FStaticNameMapEntry GStaticShaderNames[];"
+				]);
+			FilesToGenerate.Add("RHIStaticShaderPlatformNames.gen.cpp", new List<string>([
+				"#include \"RHIStaticShaderPlatformNames.gen.h\"",
+						"BUILDSETTINGS_API FStaticNameMapEntry GStaticShaderNames[] = ",
+						"{",
+						"	{ nullptr, nullptr, 0 },",
+						"};"
+				]));
+		}
 	}
 
 	private ulong GetPersistentAllocatorReserveSize()
