UnrealBuildTool: Disable 'C4702: unreachable code' when use bUseInlining is disabled

diff --git a/Engine/Source/Programs/UnrealBuildTool/Platform/Windows/VCToolChain.cs b/Engine/Source/Programs/UnrealBuildTool/Platform/Windows/VCToolChain.cs
--- a/Engine/Source/Programs/UnrealBuildTool/Platform/Windows/VCToolChain.cs
+++ b/Engine/Source/Programs/UnrealBuildTool/Platform/Windows/VCToolChain.cs
@@ -1266,6 +1266,11 @@ namespace UnrealBuildTool
 				Arguments.Add("/wd5321"); // nonstandard extension used: encoding '\xE5' as a multi-byte utf-8 character
 			}
 
+			// Disable C4702: unreachable code when use inlining is disabled
+			if (!CompileEnvironment.bUseInlining && EnvVars.Compiler.IsMSVC())
+			{
+				Arguments.Add("/wd4702"); // unreachable code
+			}
 
 			// Downgrade C4702: unreachable code to a warning when running LTCG or PGO
 			if (CompileEnvironment.bPGOOptimize || CompileEnvironment.bPGOProfile || CompileEnvironment.bAllowLTCG)
