[UBT] Fixes compilation using SN-DBS with clang-cl

CL 30684071 previously added code to include the corresponding .cpp file when a producing a PCH. That code only executed when bIsClangCompiler was true. That is set by ClangToolChain, but clang-cl uses VCToolChain so it skipped this code. The conditional has been removed so this path is now taken by clang-cl. This also means the path is taken by MSVC, but there are no issues caused by including the file.

Additionally, that code was searching for precompiled headers using the .gch extension. This extension is used by UBT when using ClangToolChain (see ClangToolChain::GetCompileArguments_FileType), but VCToolChain uses .pch (see VCToolChain::CompileCPPFiles). So when using clang-cl, we need to search for .pch instead and the code now searches for both.

Another issue impacting clang-cl is errors such as:

fatal error: malformed or corrupted AST file: 'could not find file '<ABS_PATH>\SharedPCH.CoreUObject.Cpp20.cpp' referenced by AST file '<ABS_PATH>/SharedPCH.CoreUObject.Cpp20.h.pch''

That indicates the AST stored in the .pch file has a reference back to the .cpp file used to create it, but when executing the job remotely the .cpp file was not being included. That file is now also included.

diff --git a/Engine/Source/Programs/UnrealBuildTool/Executors/SNDBS.cs b/Engine/Source/Programs/UnrealBuildTool/Executors/SNDBS.cs
--- a/Engine/Source/Programs/UnrealBuildTool/Executors/SNDBS.cs
+++ b/Engine/Source/Programs/UnrealBuildTool/Executors/SNDBS.cs
@@ -273,20 +273,34 @@ namespace UnrealBuildTool
 						}).ToList();
 					}
 
-					// Add PCH source file dependencies for clang or clang-cl
-					if (a.bIsClangCompiler)
+					// Look for any prerequisite actions that produce .pch files and add any .cpp files they depend on
+					List<Dictionary<string, object>>? ExplicitInputFiles = Job["explicit_input_files"] as List<Dictionary<string, object>>;
+
+					ExplicitInputFiles?.AddRange(a.PrerequisiteActions
+						.Where(Prereq => Prereq.ProducedItems.Any(Produced => Produced.AbsolutePath.EndsWith(".gch") || Produced.AbsolutePath.EndsWith(".pch")))
+						.SelectMany(Prereq => Prereq.PrerequisiteItems, (_, PrereqFile) => PrereqFile.AbsolutePath)
+						.Where(Path => Path.EndsWith(".cpp"))
+						.Select(Path => new Dictionary<string, object>()
+						{
+							["filename"] = Path
+						}));
+
+					// When using clang-cl, any time the PCH is included we need to include the corresponding .cpp file.
+					// This is because the AST in the PCH file references this file and without it we get a "malformed
+					// or corrupted AST file" error stating it could not find the .cpp file.
+					//
+					// The file is tiny and harmless to include when not using clang-cl.
+					List<Dictionary<string, object>>? PrecompiledDependencies = ExplicitInputFiles?.Select(File => File["filename"].ToString())
+						.Where(Path => Path != null && Path.EndsWith(".h.pch"))
+						.Select(Path => Path!.Replace(".h.pch", ".cpp"))
+						.Select(Path => new Dictionary<string, object>()
+						{
+							["filename"] = Path
+						}).ToList();
+
+					if (PrecompiledDependencies != null)
 					{
-						// Look for any prerequisite actions that produce .pch files and add any .cpp files they depend on
-						List<Dictionary<string, object>>? ExplicitInputFiles = Job["explicit_input_files"] as List<Dictionary<string, object>>;
-
-						ExplicitInputFiles?.AddRange(a.PrerequisiteActions
-								.Where(Prereq => Prereq.ProducedItems.Any(Produced => Produced.AbsolutePath.EndsWith(".gch")))
-								.SelectMany(Prereq => Prereq.PrerequisiteItems, (_, PrereqFile) => PrereqFile.AbsolutePath)
-								.Where(Path => Path.EndsWith(".cpp"))
-								.Select(Path => new Dictionary<string, object>()
-								{
-									["filename"] = Path
-								}));
+						ExplicitInputFiles?.AddRange(PrecompiledDependencies);
 					}
 
 					string CommandDescription = String.IsNullOrWhiteSpace(a.CommandDescription) ? a.ActionType.ToString() : a.CommandDescription;
