UnrealBuildTool: Add ini config for bUseVFS

[/Script/BuildSettings.BuildSettings]
bUseVFS=true

#rnx

diff --git a/Engine/Source/Programs/UnrealBuildTool/Configuration/TargetRules.cs b/Engine/Source/Programs/UnrealBuildTool/Configuration/TargetRules.cs
--- a/Engine/Source/Programs/UnrealBuildTool/Configuration/TargetRules.cs
+++ b/Engine/Source/Programs/UnrealBuildTool/Configuration/TargetRules.cs
@@ -2002,6 +2002,7 @@ namespace UnrealBuildTool
 		[CommandLine("-VFS", Value = "true")]
 		[CommandLine("-NoVFS", Value = "false")]
 		[XmlConfigFile(Category = "BuildConfiguration")]
+		[ConfigFile(ConfigHierarchyType.Engine, "/Script/BuildSettings.BuildSettings", "bUseVFS")]
 		[RequiresUniqueBuildEnvironment]
 		public bool bUseVFS { get; set; }
 
