UnrealBuildTool: Ensure .gen.cpp files are excluded from the adaptive working set

#rnx

diff --git a/Engine/Source/Programs/UnrealBuildTool/System/SourceFileWorkingSet.cs b/Engine/Source/Programs/UnrealBuildTool/System/SourceFileWorkingSet.cs
--- a/Engine/Source/Programs/UnrealBuildTool/System/SourceFileWorkingSet.cs
+++ b/Engine/Source/Programs/UnrealBuildTool/System/SourceFileWorkingSet.cs
@@ -204,6 +204,12 @@ namespace UnrealBuildTool
 		/// <returns>True if the file is part of the working set, false otherwise</returns>
 		public bool Contains(FileItem File)
 		{
+			// Generated .cpp files should never be treated as part of the working set
+			if (File.HasExtension(".gen.cpp"))
+			{
+				return false;
+			}
+
 			WaitForBackgroundProcess();
 			if (Files.Contains(File.Location) || Directories.Any(x => File.Location.IsUnderDirectory(x)))
 			{
