UnrealBuildTool: Handle if a plugin or module's rules are not found in the rules assembies

#rnx

diff --git a/Engine/Source/Programs/UnrealBuildTool/Configuration/UEBuildTarget.cs b/Engine/Source/Programs/UnrealBuildTool/Configuration/UEBuildTarget.cs
--- a/Engine/Source/Programs/UnrealBuildTool/Configuration/UEBuildTarget.cs
+++ b/Engine/Source/Programs/UnrealBuildTool/Configuration/UEBuildTarget.cs
@@ -4535,9 +4535,14 @@ namespace UnrealBuildTool
 						continue;
 					}
 
-					Type RulesType = RulesAssembly.GetModuleRulesType(ModuleName)!;
+					Type? RulesType = RulesAssembly.GetModuleRulesType(ModuleName);
 
-					if (!ModuleRules.IsValidForTarget(RulesType, Rules, out string? InvalidReason))
+					if (RulesType == null)
+					{
+						Logger.LogDebug("Excluding module {Module}: Unable to get module rules type from rules assemblies", ModuleName);
+						continue;
+					}
+					else if (!ModuleRules.IsValidForTarget(RulesType, Rules, out string? InvalidReason))
 					{
 						Logger.LogDebug("Excluding module {Module}: Does not support {InvalidReason}", ModuleName, InvalidReason);
 						continue;
@@ -4602,8 +4607,13 @@ namespace UnrealBuildTool
 						{
 							throw new BuildException("Unable to find module '{0}' referenced by {1}", ModuleDescriptor.Name, Plugin.File);
 						}
-						Type RulesType = RulesAssembly.GetModuleRulesType(ModuleDescriptor.Name)!;
-						if (!ModuleRules.IsValidForTarget(RulesType, Rules, out string? InvalidReason))
+						Type? RulesType = RulesAssembly.GetModuleRulesType(ModuleDescriptor.Name);
+						if (RulesType == null)
+						{
+							Logger.LogDebug("Excluding plugin {Plugin} module {Module}: Unable to get module rules type from rules assemblies", Plugin.Name, ModuleDescriptor.Name);
+							continue;
+						}
+						else if (!ModuleRules.IsValidForTarget(RulesType, Rules, out string? InvalidReason))
 						{
 							Logger.LogDebug("Excluding plugin {Plugin} module {Module}: Does not support {InvalidReason}", Plugin.Name, ModuleDescriptor.Name, InvalidReason);
 							continue;
