UnrealBuildTool: Fix how horde telemetry is loaded if running through a horde job

#rnx

diff --git a/Engine/Source/Programs/UnrealBuildTool/Configuration/TargetDescriptor.cs b/Engine/Source/Programs/UnrealBuildTool/Configuration/TargetDescriptor.cs
--- a/Engine/Source/Programs/UnrealBuildTool/Configuration/TargetDescriptor.cs
+++ b/Engine/Source/Programs/UnrealBuildTool/Configuration/TargetDescriptor.cs
@@ -610,7 +610,7 @@ namespace UnrealBuildTool
 
 			// Register any found descriptors for telemetry, using the first as the primary target
 			TelemetryService.Get().SetPrimaryTargetDetails(TargetDescriptors.FirstOrDefault());
-			TargetDescriptors.ForEach(x => TelemetryService.Get().AddEndpointsFromConfig(x.ProjectFile?.Directory));
+			TargetDescriptors.ForEach(x => TelemetryService.Get().AddEndpointsFromConfig(x.ProjectFile?.Directory, Logger));
 		}
 
 		/// <summary>
diff --git a/Engine/Source/Programs/UnrealBuildTool/System/Telemetry.cs b/Engine/Source/Programs/UnrealBuildTool/System/Telemetry.cs
--- a/Engine/Source/Programs/UnrealBuildTool/System/Telemetry.cs
+++ b/Engine/Source/Programs/UnrealBuildTool/System/Telemetry.cs
@@ -1,5 +1,8 @@
 // Copyright Epic Games, Inc. All Rights Reserved.
 
+using EpicGames.Core;
+using EpicGames.Core.Telemetry;
+using Microsoft.Extensions.Logging;
 using System;
 using System.Collections.Concurrent;
 using System.Collections.Generic;
@@ -7,8 +10,6 @@ using System.Linq;
 using System.Net.Http;
 using System.Text.Json.Serialization;
 using System.Threading.Tasks;
-using EpicGames.Core;
-using EpicGames.Core.Telemetry;
 using UnrealBuildBase;
 
 namespace UnrealBuildTool
@@ -319,21 +320,19 @@ namespace UnrealBuildTool
 			);
 		});
 
-		readonly Lazy<HttpClient> _httpClient;
-		readonly ConcurrentDictionary<Tuple<Uri, string>, TelemetryEndpoint> _endpoints;
-		readonly HashSet<string> _telemetryProviders = new();
+		readonly Lazy<HttpClient> _httpClient = new(HttpClientDefaults.GetClient);
+		readonly ConcurrentDictionary<Tuple<Uri, string>, TelemetryEndpoint> _endpoints = [];
+		readonly HashSet<string> _telemetryProviders = [];
 		bool _disposed = false;
 
 		TargetDescriptor? _descriptor = null;
 
 		#region Singleton
-		static readonly Lazy<TelemetryService> s_instance = new(() => new TelemetryService());
+		static readonly Lazy<TelemetryService> s_instance = new(() => new());
 		public static TelemetryService Get() => s_instance.Value;
 
 		private TelemetryService()
 		{
-			_httpClient = new(() => HttpClientDefaults.GetClient());
-			_endpoints = new();
 		}
 		#endregion
 
@@ -341,50 +340,54 @@ namespace UnrealBuildTool
 
 		public void AddTelemetryConfigProviders(IEnumerable<string> providers) => _telemetryProviders.UnionWith(providers);
 
-		public void AddEndpointsFromConfig(DirectoryReference? projectDir = null)
+		public void AddEndpointsFromConfig(DirectoryReference? projectDir, ILogger logger)
 		{
 			ConfigHierarchy engineIni = ConfigCache.ReadHierarchy(ConfigHierarchyType.Engine, projectDir, BuildHostPlatform.Current.Platform);
-			foreach (string category in _telemetryProviders)
+			foreach (string provider in _telemetryProviders.Where(x => !String.IsNullOrEmpty(x)))
 			{
-				AddEndpointFromConfigInternal(engineIni, category);
+				ConfigHierarchySection section = engineIni.FindSection(provider);
+				// If we are running via horde and UE_HORDE_URL is set, skip any horde providers, which will be added below
+				if (_hordeUrl != null && (!section.TryGetValue("APIKeyET", out string? apiKeyET) || apiKeyET.Contains("Horde", StringComparison.OrdinalIgnoreCase)))
+				{
+					continue;
+				}
+				AddEndpointFromConfigInternal(provider, section, logger);
 			}
 
-			// If we are running via horde, always the current horde url as a provider
-			if (!String.IsNullOrEmpty(_hordeUrl))
+			// Add matching horde provider if found
+			if (_hordeUrl != null)
 			{
-				AddEndpointFromConfigInternal(engineIni, "StudioTelemetry.Provider.Horde", _hordeUrl);
+				AddMatchingHordeEndpointFromConfig(projectDir, logger);
 			}
 		}
 
-		public void AddAllEndpointsFromConfig(DirectoryReference? projectDir = null)
+		public void AddEndpointsFromConfig(ILogger logger) => AddEndpointsFromConfig(null, logger);
+
+		private bool AddMatchingHordeEndpointFromConfig(DirectoryReference? projectDir, ILogger logger)
 		{
 			ConfigHierarchy engineIni = ConfigCache.ReadHierarchy(ConfigHierarchyType.Engine, projectDir, BuildHostPlatform.Current.Platform);
-			foreach (string section in engineIni.SectionNames.Where(x => x.StartsWith(ConfigCategoryPrefix, StringComparison.Ordinal)))
+			foreach (string provider in engineIni.SectionNames.Where(x => x.StartsWith(ConfigCategoryPrefix, StringComparison.Ordinal)))
 			{
-				AddEndpointFromConfigInternal(engineIni, section);
+				ConfigHierarchySection section = engineIni.FindSection(provider);
+				if (section.TryGetValue("APIServerET", out string? apiServerET) && apiServerET.Equals(_hordeUrl, StringComparison.OrdinalIgnoreCase))
+				{
+					return AddEndpointFromConfigInternal(provider, section, logger);
+				}
 			}
+			return false;
 		}
 
-		private bool AddEndpointFromConfigInternal(ConfigHierarchy engineIni, string provider, string? apiServerET = null)
+		private bool AddEndpointFromConfigInternal(string provider, ConfigHierarchySection section, ILogger logger)
 		{
-			if (String.IsNullOrEmpty(provider))
-			{
-				return false;
-			}
-
-			ConfigHierarchySection section = engineIni.FindSection(provider);
-
-			if (String.IsNullOrEmpty(apiServerET))
-			{
-				section.TryGetValue("APIServerET", out apiServerET);
-			}
-
-			if (!String.IsNullOrEmpty(apiServerET) && section.TryGetValue("APIKeyET", out string? apiKeyET) && section.TryGetValue("APIEndpointET", out string? apiEndpointET))
+			if (section.TryGetValue("APIServerET", out string? apiServerET) && section.TryGetValue("APIKeyET", out string? apiKeyET) && section.TryGetValue("APIEndpointET", out string? apiEndpointET))
 			{
 				Uri baseAddress = new($"{apiServerET}{apiEndpointET}");
 				Tuple<Uri, string> key = new(baseAddress, apiKeyET);
-				_endpoints.GetOrAdd(key, new TelemetryEndpoint(_httpClient.Value, baseAddress, apiKeyET, _metadata));
-				return true;
+				if (_endpoints.TryAdd(key, new(_httpClient.Value, baseAddress, apiKeyET, _metadata)))
+				{
+					logger.LogDebug("Loaded telemetry endpoint Name={Name} Key={Key} Url={Url}", provider, apiKeyET, baseAddress);
+					return true;
+				}
 			}
 
 			return false;
@@ -413,13 +416,8 @@ namespace UnrealBuildTool
 		/// <inheritdoc/>
 		public void FlushEvents()
 		{
-			List<Task> flushTasks = new List<Task>();
-
-			foreach (KeyValuePair<Tuple<Uri, string>, TelemetryEndpoint> endpoint in _endpoints)
-			{
-				flushTasks.Add(endpoint.Value.FlushEventsAsync());
-			}
-			Task.WaitAll(flushTasks.ToArray());
+			IEnumerable<Task> flushTasks = _endpoints.Select(x => x.Value.FlushEventsAsync());
+			Task.WaitAll([.. flushTasks]);
 		}
 		#endregion
 
diff --git a/Engine/Source/Programs/UnrealBuildTool/UnrealBuildTool.cs b/Engine/Source/Programs/UnrealBuildTool/UnrealBuildTool.cs
--- a/Engine/Source/Programs/UnrealBuildTool/UnrealBuildTool.cs
+++ b/Engine/Source/Programs/UnrealBuildTool/UnrealBuildTool.cs
@@ -545,7 +545,7 @@ namespace UnrealBuildTool
 				}
 
 				TelemetryService.Get().AddTelemetryConfigProviders(Options.TelemetryProviders.Concat(Options.CmdTelemetryProviders));
-				TelemetryService.Get().AddEndpointsFromConfig();
+				TelemetryService.Get().AddEndpointsFromConfig(Logger);
 
 				// Create a UbtRun file
 				try
