GameActivity.ConsoleKeyboardInput can only be used with a Theme.AppCompat theme.
Replace "[REDACTED_HANDLE]:style/Theme.Black.NoTitleBar.Fullscreen" with "[REDACTED_HANDLE]/UnrealBaseTheme".

diff --git a/Engine/Source/Programs/UnrealBuildTool/Platform/Android/UEDeployAndroid.cs b/Engine/Source/Programs/UnrealBuildTool/Platform/Android/UEDeployAndroid.cs
--- a/Engine/Source/Programs/UnrealBuildTool/Platform/Android/UEDeployAndroid.cs
+++ b/Engine/Source/Programs/UnrealBuildTool/Platform/Android/UEDeployAndroid.cs
@@ -2969,7 +2969,7 @@ namespace UnrealBuildTool
 				Text.AppendLine("\t\t<activity android:name=\"com.epicgames.unreal.GameActivity\"");
 				Text.AppendLine("\t\t          android:exported=\"true\"");
 				Text.AppendLine("\t\t          android:label=\"@string/app_name\"");
-				Text.AppendLine("\t\t          android:theme=\"@android:style/Theme.Black.NoTitleBar.Fullscreen\"");
+				Text.AppendLine("\t\t          android:theme=\"@style/UnrealBaseTheme\"");
 				Text.AppendLine(bAddDensity ? "\t\t          android:configChanges=\"mcc|mnc|uiMode|density|screenSize|smallestScreenSize|screenLayout|orientation|keyboardHidden|keyboard|navigation|touchscreen|locale|fontScale|layoutDirection\""
 											: "\t\t          android:configChanges=\"mcc|mnc|uiMode|screenSize|smallestScreenSize|screenLayout|orientation|keyboardHidden|keyboard|navigation|touchscreen|locale|fontScale|layoutDirection\"");
 
