Removed bashrc/bash_profile related code in Android SDK setup and Turnkey on Mac

diff --git a/Engine/Extras/Android/SetupAndroid.command b/Engine/Extras/Android/SetupAndroid.command
--- a/Engine/Extras/Android/SetupAndroid.command
+++ b/Engine/Extras/Android/SetupAndroid.command
@@ -11,9 +11,13 @@ pausefunc()
     read -rsp $'Press any key to continue...\n' -n1 key
 }
 
-# $1 parameter should be the profile file: .bash_profile as it is read by UE and .zprofile if default shell is zsh
-updateProfileFileFunc()
+# $1 parameter should be the profile file
+updateProfileFile()
 {
+	if [ ! -f "${1}" ]; then
+		touch "${1}"
+	fi
+	
 	if ! grep -q "export ANDROID_HOME=\"$STUDIO_SDK_PATH\"" "$1" 
 	then
 		echo >> "$1"
@@ -51,6 +55,24 @@ updateProfileFileFunc()
 	fi
 }
 
+# $1 parameter should be the profile file
+undoProfileFile()
+{
+	if [ ! -f "${1}" ]; then
+		return
+	fi
+	
+	echo "Undoing UE Android dev env variables from ${1}"
+	
+	if sed "s|\nexport ANDROID_HOME=\"$STUDIO_SDK_PATH\"\n\nexport JAVA_HOME=\"$JAVA_HOME\"\n\nexport PATH=\"\$PATH:$PLATFORMTOOLS\"\n\nexport NDKROOT=\"$NDKINSTALLPATH\"\nexport NDK_ROOT=\"$NDKINSTALLPATH\"||g" "$1"; then
+		sed -i "s|\nexport ANDROID_HOME=\"$STUDIO_SDK_PATH\"\n\nexport JAVA_HOME=\"$JAVA_HOME\"\n\nexport PATH=\"\$PATH:$PLATFORMTOOLS\"\n\nexport NDKROOT=\"$NDKINSTALLPATH\"\nexport NDK_ROOT=\"$NDKINSTALLPATH\"||g" "$1"
+	fi
+	
+	if [ ! -s "${1}" ] ; then
+		rm "${1}"
+    fi
+}
+
 if [[ ${ARG5} == "-noninteractive" ]]; then
     PAUSE=
 else
@@ -137,16 +159,10 @@ if [ $retVal -ne 0 ]; then
 	exit $retVal
 fi
 
-# on zsh, ~/.bash_profile may not exist
-if [ ! -f ~/.bash_profile ]; then
-	echo > ~/.bash_profile
-fi
- 
-updateProfileFileFunc ~/.bash_profile
-# on zsh, update .zprofile and create .bash_profile (as UE reads it)
-if [ $SHELL == "/bin/zsh" ]; then
-	updateProfileFileFunc ~/.zprofile
-	echo Adding UE Android dev env variables to .zprofile
+updateProfileFile ~/.zshrc
+
+if [ "${SHELL}" == "/bin/zsh" ]; then
+	undoProfileFile ~/.bash_profile
 fi
 
 echo Success!
diff --git a/Engine/Source/Developer/Android/AndroidDeviceDetection/Private/AndroidDeviceDetectionModule.cpp b/Engine/Source/Developer/Android/AndroidDeviceDetection/Private/AndroidDeviceDetectionModule.cpp
--- a/Engine/Source/Developer/Android/AndroidDeviceDetection/Private/AndroidDeviceDetectionModule.cpp
+++ b/Engine/Source/Developer/Android/AndroidDeviceDetection/Private/AndroidDeviceDetectionModule.cpp
@@ -959,8 +959,22 @@ public:
 			// didn't find ANDROID_HOME, so parse the .bashrc file on Linux
 			FArchive* FileReader = IFileManager::Get().CreateFileReader(*FString("~/.bashrc"));
 #else
-			// didn't find ANDROID_HOME, so parse the .bash_profile file on MAC
-			FArchive* FileReader = IFileManager::Get().CreateFileReader(*FString([@"~/.bash_profile" stringByExpandingTildeInPath]));
+			// didn't find ANDROID_HOME, so parse the .zshrc file on MAC
+			FArchive* FileReader = IFileManager::Get().CreateFileReader(*FString([@"~/.zshrc" stringByExpandingTildeInPath]));
+			if (!FileReader)
+			{
+				// Fallbacks
+				FileReader = IFileManager::Get().CreateFileReader(*FString([@"~/.bashrc" stringByExpandingTildeInPath]));
+				if (!FileReader)
+				{
+					FileReader = IFileManager::Get().CreateFileReader(*FString([@"~/.bash_profile" stringByExpandingTildeInPath]));
+					UE_LOG(LogCore, Warning, TEXT(".zshrc is missing. Falling back to .bash_profile."));
+				}
+				else
+				{
+					UE_LOG(LogCore, Warning, TEXT(".zshrc is missing. Falling back to .bashrc."));
+				}
+			}
 #endif
 			if (FileReader)
 			{
diff --git a/Engine/Source/Programs/AutomationTool/Android/AndroidPlatform.Automation.cs b/Engine/Source/Programs/AutomationTool/Android/AndroidPlatform.Automation.cs
--- a/Engine/Source/Programs/AutomationTool/Android/AndroidPlatform.Automation.cs
+++ b/Engine/Source/Programs/AutomationTool/Android/AndroidPlatform.Automation.cs
@@ -349,11 +349,21 @@ public class AndroidPlatform : Platform
 		}
 		else if (OperatingSystem.IsMacOS())
 		{
-			string BashProfilePath = Path.Combine(Environment.GetFolderPath(Environment.SpecialFolder.Personal), ".bash_profile");
+			string UserProfileDir = Environment.GetFolderPath(Environment.SpecialFolder.UserProfile);
+			string BashProfilePath = Path.Combine(UserProfileDir, ".zshrc");
 			if (!File.Exists(BashProfilePath))
 			{
-				// Try .bashrc if didn't fine .bash_profile
-				BashProfilePath = Path.Combine(Environment.GetFolderPath(Environment.SpecialFolder.Personal), ".bashrc");
+				// Fallbacks
+				BashProfilePath = Path.Combine(UserProfileDir, ".bashrc");
+				if (!File.Exists(BashProfilePath))
+				{
+					BashProfilePath = Path.Combine(UserProfileDir, ".bash_profile");
+					Logger.LogWarning(".zshrc is missing. Falling back to .bash_profile.");
+				}
+				else
+				{
+					Logger.LogWarning(".zshrc is missing. Falling back to .bashrc.");
+				}
 			}
 			if (File.Exists(BashProfilePath))
 			{
diff --git a/Engine/Source/Programs/UnrealBuildTool/Platform/Android/AndroidPlatformSDK.cs b/Engine/Source/Programs/UnrealBuildTool/Platform/Android/AndroidPlatformSDK.cs
--- a/Engine/Source/Programs/UnrealBuildTool/Platform/Android/AndroidPlatformSDK.cs
+++ b/Engine/Source/Programs/UnrealBuildTool/Platform/Android/AndroidPlatformSDK.cs
@@ -23,7 +23,7 @@ namespace UnrealBuildTool
 		{
 		}
 
-		private static void CheckEnvironmentVariables(bool bUseConfig)
+		private static void CheckEnvironmentVariables(bool bUseConfig, ILogger Logger)
 		{
 			Dictionary<string, string> EnvVarNames = new Dictionary<string, string> {
 													 {"ANDROID_HOME", "SDKPath"},
@@ -59,19 +59,31 @@ namespace UnrealBuildTool
 				}
 			}
 				
-			// If we are not running on Windows and we are still missing a key then go and find it from the .bash_profile
+			// If we are not running on Windows, and we are still missing a key then go and find it from alternate source
 			if (!OperatingSystem.IsWindows() && !EnvVarNames.All(s => AndroidEnv.ContainsKey(s.Key)))
 			{
-				string BashProfilePath = Path.Combine(Environment.GetFolderPath(Environment.SpecialFolder.UserProfile), ".bash_profile");
-				if (!File.Exists(BashProfilePath))
+				string UserProfileDir = Environment.GetFolderPath(Environment.SpecialFolder.UserProfile);
+				string BashProfilePath;
+				if (!OperatingSystem.IsMacOS())
 				{
-					// Try .bashrc if didn't fine .bash_profile
-					BashProfilePath = Path.Combine(Environment.GetFolderPath(Environment.SpecialFolder.UserProfile), ".bashrc");
-
-					// Lastly try .zprofile
+					BashProfilePath = Path.Combine(UserProfileDir, ".bashrc");
+				}
+				else
+				{
+					BashProfilePath = Path.Combine(UserProfileDir, ".zshrc");
 					if (!File.Exists(BashProfilePath))
 					{
-						BashProfilePath = Path.Combine(Environment.GetFolderPath(Environment.SpecialFolder.UserProfile), ".zprofile");
+						// Fallbacks
+						BashProfilePath = Path.Combine(UserProfileDir, ".bashrc");
+						if (!File.Exists(BashProfilePath))
+						{
+							BashProfilePath = Path.Combine(UserProfileDir, ".bash_profile");
+							Logger.LogWarning(".zshrc is missing. Falling back to .bash_profile.");
+						}
+						else
+						{
+							Logger.LogWarning(".zshrc is missing. Falling back to .bashrc.");
+						}
 					}
 				}
 				if (File.Exists(BashProfilePath))
@@ -105,7 +117,7 @@ namespace UnrealBuildTool
 			}
 		}
 
-		public static string? GetNDKRoot()
+		public static string? GetNDKRoot(ILogger Logger)
 		{
 			string? NDKPath = Environment.GetEnvironmentVariable("NDKROOT");
 
@@ -116,7 +128,7 @@ namespace UnrealBuildTool
 			}
 
 			// set environment variables (don't check config)
-			CheckEnvironmentVariables(false);
+			CheckEnvironmentVariables(false, Logger);
 
 			// try NDKROOT again
 			NDKPath = Environment.GetEnvironmentVariable("NDKROOT");
@@ -129,7 +141,7 @@ namespace UnrealBuildTool
 
 		protected override string? GetInstalledSDKVersion()
 		{
-			string? NDKPath = GetNDKRoot();
+			string? NDKPath = GetNDKRoot(Logger);
 
 			// don't register if we don't have an NDKROOT specified
 			if (String.IsNullOrEmpty(NDKPath))
@@ -396,7 +408,7 @@ namespace UnrealBuildTool
 			}
 
 			// Set environment variables as needed (allow config overrides)
-			CheckEnvironmentVariables(true);
+			CheckEnvironmentVariables(true, Logger);
 			
 			string? NDKPath = Environment.GetEnvironmentVariable("NDKROOT");
 			string? JavaPath = Environment.GetEnvironmentVariable("JAVA_HOME");
diff --git a/Engine/Source/Programs/UnrealBuildTool/Platform/Android/AndroidToolChain.cs b/Engine/Source/Programs/UnrealBuildTool/Platform/Android/AndroidToolChain.cs
--- a/Engine/Source/Programs/UnrealBuildTool/Platform/Android/AndroidToolChain.cs
+++ b/Engine/Source/Programs/UnrealBuildTool/Platform/Android/AndroidToolChain.cs
@@ -132,7 +132,7 @@ namespace UnrealBuildTool
 			Options = ToolchainOptions;
 			ProjectFile = InProjectFile;
 
-			string? NDKPath = AndroidPlatformSDK.GetNDKRoot();
+			string? NDKPath = AndroidPlatformSDK.GetNDKRoot(InLogger);
 
 			// don't register if we don't have an NDKROOT specified
 			if (String.IsNullOrEmpty(NDKPath))
@@ -257,7 +257,7 @@ namespace UnrealBuildTool
 
 		protected override ClangToolChainInfo GetToolChainInfo()
 		{
-			return new ClangToolChainInfo(DirectoryReference.FromString(AndroidPlatformSDK.GetNDKRoot()), FileReference.FromString(ClangPath)!, FileReference.FromString(ArPathArm64)!, Logger);
+			return new ClangToolChainInfo(DirectoryReference.FromString(AndroidPlatformSDK.GetNDKRoot(Logger)), FileReference.FromString(ClangPath)!, FileReference.FromString(ArPathArm64)!, Logger);
 		}
 
 		public static string GetGLESVersion(bool bBuildForES31)
