UnrealBuildTool: VFS for linux shell script

#rnx

diff --git a/Engine/Source/Programs/UnrealBuildTool/Platform/Linux/LinuxToolChain.cs b/Engine/Source/Programs/UnrealBuildTool/Platform/Linux/LinuxToolChain.cs
--- a/Engine/Source/Programs/UnrealBuildTool/Platform/Linux/LinuxToolChain.cs
+++ b/Engine/Source/Programs/UnrealBuildTool/Platform/Linux/LinuxToolChain.cs
@@ -157,8 +157,7 @@ namespace UnrealBuildTool
 
 			string NormalizeFileRef(FileSystemReference file)
 			{
-				// todo: binary paths on linux are currently having some issues
-				string normalized = bUseCmdExe ? NormalizeCommandLinePath(file, LinkEnvironment.RootPaths) : file.FullName;
+				string normalized = NormalizeCommandLinePath(file, LinkEnvironment.RootPaths);
 				return bUseCmdExe ? normalized.Replace('/', '\\') : normalized;
 			}
 
@@ -1373,7 +1372,7 @@ namespace UnrealBuildTool
 					LinkWriter.WriteLine("# *DO NOT EDIT*");
 					LinkWriter.WriteLine();
 					LinkWriter.WriteLine("set -o errexit");
-					LinkWriter.WriteLine($"\"{Info.Clang}\" @\"{NormalizeCommandLinePath(LinkResponseFileItem, LinkEnvironment.RootPaths)}\"");
+					LinkWriter.WriteLine($"\"{NormalizeCommandLinePath(Info.Clang, LinkEnvironment.RootPaths)}\" @\"{NormalizeCommandLinePath(LinkResponseFileItem, LinkEnvironment.RootPaths)}\"");
 					if (!needsRelink)
 					{
 						LinkWriter.WriteLine(GetDumpEncodeDebugCommand(LinkEnvironment, OutputFile, LinkAction));
@@ -1467,7 +1466,7 @@ namespace UnrealBuildTool
 						RelinkWriter.WriteLine("# *DO NOT EDIT*");
 						RelinkWriter.WriteLine();
 						RelinkWriter.WriteLine("set -o errexit");
-						RelinkWriter.WriteLine($"\"{Info.Clang}\" @\"{NormalizeCommandLinePath(RelinkResponseFileItem, LinkEnvironment.RootPaths)}\"");
+						RelinkWriter.WriteLine($"\"{NormalizeCommandLinePath(Info.Clang, LinkEnvironment.RootPaths)}\" @\"{NormalizeCommandLinePath(RelinkResponseFileItem, LinkEnvironment.RootPaths)}\"");
 						RelinkWriter.WriteLine(GetDumpEncodeDebugCommand(LinkEnvironment, RealOutputFile, RelinkAction));
 					}
 
