[UBT]
* Commented out Timeline.Print for now. Takes 1 minute to run when generating full visual studio solution

diff --git a/Engine/Source/Programs/UnrealBuildTool/UnrealBuildTool.cs b/Engine/Source/Programs/UnrealBuildTool/UnrealBuildTool.cs
--- a/Engine/Source/Programs/UnrealBuildTool/UnrealBuildTool.cs
+++ b/Engine/Source/Programs/UnrealBuildTool/UnrealBuildTool.cs
@@ -696,7 +696,10 @@ namespace UnrealBuildTool
 				bool bShowExecutionTime = ModeOptions.HasFlag(ToolModeOptions.ShowExecutionTime);
 				LogLevel ExecutionLogLevel = bShowExecutionTime ? LogLevel.Information : LogLevel.Debug;
 				LogLevel TimelineLogLevel = bShowExecutionTime && bShowTimeline && ApplicationResult == CompilationResult.Succeeded ?  LogLevel.Information : LogLevel.Debug;
-				Timeline.Print(TimeSpan.FromMilliseconds(100.0), TimeSpan.FromMilliseconds(200.0), TimelineLogLevel, Logger);
+				
+				// Commented out for now. Takes 1 minute to run when generating vs solution
+				//Timeline.Print(TimeSpan.FromMilliseconds(100.0), TimeSpan.FromMilliseconds(200.0), TimelineLogLevel, Logger);
+				
 				Logger.Log(ExecutionLogLevel, "");
 				if (ApplicationResult == CompilationResult.Succeeded)
 				{
