[AutoRTFM] Restructure verse-clang deploy directories to support multiple versions

Lets us have LLVM 18 and LLVM 20 side-by-side without aliasing executables.

diff --git a/Engine/Source/Programs/UnrealBuildTool/Platform/Linux/LinuxToolChain.cs b/Engine/Source/Programs/UnrealBuildTool/Platform/Linux/LinuxToolChain.cs
--- a/Engine/Source/Programs/UnrealBuildTool/Platform/Linux/LinuxToolChain.cs
+++ b/Engine/Source/Programs/UnrealBuildTool/Platform/Linux/LinuxToolChain.cs
@@ -120,7 +120,10 @@ namespace UnrealBuildTool
 				// this *shouldnt* cause issues as these compiler should ideally be the same version
 				if (Options.HasFlag(ClangToolChainOptions.UseAutoRTFMCompiler))
 				{
-					DirectoryReference AutoRTFMDir = DirectoryReference.Combine(Unreal.EngineDirectory, "Restricted", "NotForLicensees", "Binaries", BuildHostPlatform.Current.Platform.ToString(), "AutoRTFM", "bin");
+					// Use the same clang version as the regular non-instrumented compiler.
+					LinuxToolChainInfo RegularClang = new LinuxToolChainInfo(BaseLinuxPath, MultiArchRoot, ClangPath, LlvmArPath, ObjcopyPath, Logger);
+					string MajorVersion = RegularClang.ClangVersion.Major.ToString();
+					DirectoryReference AutoRTFMDir = DirectoryReference.Combine(Unreal.EngineDirectory, "Restricted", "NotForLicensees", "Binaries", BuildHostPlatform.Current.Platform.ToString(), "AutoRTFM", MajorVersion, "bin");
 					
 					// set up the path to our toolchain
 					if (BuildHostPlatform.Current.Platform == UnrealTargetPlatform.Linux)
diff --git a/Engine/Source/Programs/UnrealBuildTool/Platform/Windows/MicrosoftPlatformSDK.cs b/Engine/Source/Programs/UnrealBuildTool/Platform/Windows/MicrosoftPlatformSDK.cs
--- a/Engine/Source/Programs/UnrealBuildTool/Platform/Windows/MicrosoftPlatformSDK.cs
+++ b/Engine/Source/Programs/UnrealBuildTool/Platform/Windows/MicrosoftPlatformSDK.cs
@@ -769,8 +769,11 @@ namespace UnrealBuildTool
 					else if (Compiler == WindowsCompiler.ClangRTFM)
 					{
 						// the AutoRTFM Clang compiler is in-tree
-						DirectoryReference AutoRTFMDir = DirectoryReference.Combine(Unreal.EngineDirectory, "Restricted", "NotForLicensees", "Binaries", "Win64", "AutoRTFM");
-						AddClangToolChain(Compiler, AutoRTFMDir, ToolChains, IsAutoSdk: false, Logger);
+						DirectoryReference AutoRTFMBaseDir = DirectoryReference.Combine(Unreal.EngineDirectory, "Restricted", "NotForLicensees", "Binaries", "Win64", "AutoRTFM");
+						foreach (DirectoryReference AutoRTFMDir in DirectoryReference.EnumerateDirectories(AutoRTFMBaseDir))
+						{
+							AddClangToolChain(Compiler, AutoRTFMDir, ToolChains, IsAutoSdk: false, Logger);
+						}
 
 						// Check for AutoSDK paths
 						DirectoryReference? AutoSdkDir;
