UnrealBuildTool: If the preferred MSVC requires a newer clang than preferred, downgrade MSVC if possible

#rnx

diff --git a/Engine/Source/Programs/UnrealBuildTool/Platform/Windows/MicrosoftPlatformSDK.cs b/Engine/Source/Programs/UnrealBuildTool/Platform/Windows/MicrosoftPlatformSDK.cs
--- a/Engine/Source/Programs/UnrealBuildTool/Platform/Windows/MicrosoftPlatformSDK.cs
+++ b/Engine/Source/Programs/UnrealBuildTool/Platform/Windows/MicrosoftPlatformSDK.cs
@@ -564,6 +564,15 @@ namespace UnrealBuildTool
 					return false;
 				}
 			}
+			else if (CompilerVersion.StartsWith('<') && VersionNumber.TryParse(CompilerVersion[1..], out VersionNumber? LessThanToolChainVersion))
+			{
+				ToolChain = SelectToolChain(ToolChains.Where(x => x.Version < LessThanToolChainVersion && x.Family < LessThanToolChainVersion), x => x, Architecture);
+				if (ToolChain == null || !(ToolChain.Version < LessThanToolChainVersion || ToolChain.Family < LessThanToolChainVersion))
+				{
+					DumpToolChains(ToolChains, x => x.ThenByDescending(x => x.Version < LessThanToolChainVersion).ThenByDescending(x => x.Family < LessThanToolChainVersion), Architecture, Logger);
+					throw new BuildException("Unable to find valid <{0} C++ toolchain for {1} {2}", LessThanToolChainVersion, Compiler, Architecture.ToString());
+				}
+			}
 			else if (VersionNumber.TryParse(CompilerVersion, out VersionNumber? ToolChainVersion))
 			{
 				ToolChain = SelectToolChain(ToolChains, x => x.ThenByDescending(x => x.Version == ToolChainVersion).ThenByDescending(x => x.Family == ToolChainVersion), Architecture);
diff --git a/Engine/Source/Programs/UnrealBuildTool/Platform/Windows/VCEnvironment.cs b/Engine/Source/Programs/UnrealBuildTool/Platform/Windows/VCEnvironment.cs
--- a/Engine/Source/Programs/UnrealBuildTool/Platform/Windows/VCEnvironment.cs
+++ b/Engine/Source/Programs/UnrealBuildTool/Platform/Windows/VCEnvironment.cs
@@ -2,8 +2,6 @@
 
 using System;
 using System.Collections.Generic;
-using System.Diagnostics;
-using System.IO;
 using System.Linq;
 using System.Runtime.InteropServices;
 using System.Runtime.Versioning;
@@ -111,12 +109,12 @@ namespace UnrealBuildTool
 		/// <summary>
 		/// The default system include paths
 		/// </summary>
-		public readonly List<DirectoryReference> IncludePaths = new List<DirectoryReference>();
+		public readonly List<DirectoryReference> IncludePaths = [];
 
 		/// <summary>
 		/// The default system library paths
 		/// </summary>
-		public readonly List<DirectoryReference> LibraryPaths = new List<DirectoryReference>();
+		public readonly List<DirectoryReference> LibraryPaths = [];
 
 		/// <summary>
 		/// Constructor
@@ -143,7 +141,7 @@ namespace UnrealBuildTool
 			// Get the compiler and linker paths from the Toolchain directory
 			CompilerPath = GetCompilerToolPath(Compiler, Architecture, CompilerDir);
 			LinkerPath = GetLinkerToolPath(Compiler, Architecture, CompilerDir, ToolChainDir);
-			LibraryManagerPath = GetLibraryLinkerToolPath(Compiler, Architecture, CompilerDir, ToolChainDir);
+			LibraryManagerPath = GetLibraryLinkerToolPath(Compiler, Architecture, ToolChainDir);
 			ToolchainCompilerPath = GetCompilerToolPath(ToolChain, Architecture, ToolChainDir);
 
 			// Get the resource compiler path from the Windows SDK
@@ -250,7 +248,7 @@ namespace UnrealBuildTool
 		{
 			if (bAllowRadLinker)
 			{
-				return FileReference.Combine(UnrealBuildBase.Unreal.EngineDirectory, "Binaries", "Win64", "RadDebugger", Architecture.ToString(), "radlink.exe");
+				return FileReference.Combine(Unreal.EngineDirectory, "Binaries", "Win64", "RadDebugger", Architecture.ToString(), "radlink.exe");
 			}
 			else if ((Compiler == WindowsCompiler.Clang || Compiler == WindowsCompiler.ClangRTFM || Compiler == WindowsCompiler.ClangInstrument || Compiler == WindowsCompiler.ClangCustom) && bAllowClangLinker)
 			{
@@ -266,12 +264,12 @@ namespace UnrealBuildTool
 		/// <summary>
 		/// Gets the path to the library linker.
 		/// </summary>
-		FileReference GetLibraryLinkerToolPath(WindowsCompiler Compiler, UnrealArch Architecture, DirectoryReference CompilerDir, DirectoryReference ToochainDir)
+		FileReference GetLibraryLinkerToolPath(WindowsCompiler Compiler, UnrealArch Architecture, DirectoryReference ToochainDir)
 		{
 			if (bAllowClangLinker && Compiler.IsClang())
 			{
 				// Since obj files could be LLVM IR Stream file format (when building with ltcg) we can't use link.exe. UbaObjTool support all formats and produce identical lib files as link.exe when running non-ltcg
-				DirectoryReference Dir = DirectoryReference.Combine(UnrealBuildBase.Unreal.EngineDirectory, "Binaries", "Win64", "UnrealBuildAccelerator", RuntimeInformation.ProcessArchitecture.ToString().ToLowerInvariant());
+				DirectoryReference Dir = DirectoryReference.Combine(Unreal.EngineDirectory, "Binaries", "Win64", "UnrealBuildAccelerator", RuntimeInformation.ProcessArchitecture.ToString().ToLowerInvariant());
 				return FileReference.Combine(Dir, "UbaObjTool.exe");
 			}
 			return FileReference.Combine(GetVCToolPath(ToochainDir, Architecture), "link.exe"); // We add /LIB to cmd line so we can use link.exe directly instead of going via lib.exe
@@ -295,7 +293,7 @@ namespace UnrealBuildTool
 		/// <summary>
 		/// Registers all SDK binaries with UBA for cross-architecture support
 		/// </summary>
-		void RegisterCrossArchitectureToolPaths(DirectoryReference WindowsSdkDir, VersionNumber WindowsSdkVersion)
+		static void RegisterCrossArchitectureToolPaths(DirectoryReference WindowsSdkDir, VersionNumber WindowsSdkVersion)
 		{
 			IEnumerable<DirectoryReference> ToolPaths = [
 				 DirectoryReference.Combine(WindowsSdkDir, "bin", WindowsSdkVersion.ToString()),
@@ -307,7 +305,7 @@ namespace UnrealBuildTool
 			{
 				DirectoryReference HostBinDir = DirectoryReference.Combine(ToolPath, MicrosoftPlatformSDK.MSVCHostUnrealArch.bIsX64 ? "x64" : "arm64");
 				DirectoryReference CrossBinDir = DirectoryReference.Combine(ToolPath, !MicrosoftPlatformSDK.MSVCHostUnrealArch.bIsX64 ? "x64" : "arm64");
-				if (FileReference.Exists(HostBinDir) && FileReference.Exists(CrossBinDir))
+				if (DirectoryReference.Exists(HostBinDir) && DirectoryReference.Exists(CrossBinDir))
 				{
 					EpicGames.UBA.Utils.RegisterCrossArchitecturePath(HostBinDir.FullName, CrossBinDir.FullName);
 				}
@@ -501,10 +499,7 @@ namespace UnrealBuildTool
 		public VCEnvironmentParameters(WindowsCompiler Compiler, WindowsCompiler ToolChain, UnrealTargetPlatform Platform, UnrealArch Architecture, string? CompilerVersion, string? ToolchainVersion, string? WindowsSdkVersion, string? SuppliedSdkDirectoryForVersion, bool bUseCPPWinRT, bool bAllowClangLinker, bool bAllowRadLinker, ILogger Logger)
 		{
 			// Get the compiler version info
-			VersionNumber? SelectedCompilerVersion;
-			DirectoryReference? SelectedCompilerDir;
-			DirectoryReference? SelectedRedistDir;
-			if (!WindowsPlatform.TryGetToolChainDir(Compiler, CompilerVersion, Architecture, Logger, out SelectedCompilerVersion, out SelectedCompilerDir, out SelectedRedistDir))
+			if (!WindowsPlatform.TryGetToolChainDir(Compiler, CompilerVersion, Architecture, Logger, out VersionNumber? SelectedCompilerVersion, out DirectoryReference? SelectedCompilerDir, out DirectoryReference? SelectedRedistDir))
 			{
 				throw new BuildException("{0}{1} {2} must be installed in order to build this target.", WindowsPlatform.GetCompilerName(Compiler), String.IsNullOrEmpty(CompilerVersion) ? "" : String.Format(" ({0})", CompilerVersion), Architecture.ToString());
 			}
@@ -523,6 +518,30 @@ namespace UnrealBuildTool
 				{
 					throw new BuildException("{0} must be installed in order to build this target.", WindowsPlatform.GetCompilerName(WindowsCompiler.VisualStudio2022));
 				}
+
+				VersionNumber? PreferredSelectedToolChainVersion = SelectedToolChainVersion;
+				DirectoryReference? PreferredSelectedToolChainDir = SelectedToolChainDir;
+				DirectoryReference? PreferredSelectedRedistDir = SelectedRedistDir;
+
+				try
+				{
+					// If the preferred MSVC toolchain is not supported for the Clang compiler version, search for an older version that is supported
+					VersionNumber SelectedClangVersion = !Compiler.IsIntel() ? SelectedCompilerVersion : MicrosoftPlatformSDK.GetClangVersionForIntelCompiler(FileReference.Combine(SelectedCompilerDir, "bin", "icx.exe"));
+					while (MicrosoftPlatformSDK.GetMinimumClangVersionForVcVersion(SelectedToolChainVersion) > SelectedCompilerVersion)
+					{
+						if (!WindowsPlatform.TryGetToolChainDir(ToolChain, $"<{SelectedToolChainVersion.GetComponent(0)}.{SelectedToolChainVersion.GetComponent(1)}", Architecture, Logger, out SelectedToolChainVersion, out SelectedToolChainDir, out SelectedRedistDir))
+						{
+							throw new BuildException("{0} compiler that supports Clang {1} must be installed in order to build this target.", WindowsPlatform.GetCompilerName(WindowsCompiler.VisualStudio2022));
+						}
+					}
+				}
+				catch (BuildException)
+				{
+					// In the event a valid toolchain is not found that is supported with the requested clang version, use the original preferred choice and the error will be handed later
+					SelectedToolChainVersion = PreferredSelectedToolChainVersion;
+					SelectedToolChainDir = PreferredSelectedToolChainDir;
+					SelectedRedistDir = PreferredSelectedRedistDir;
+				}
 			}
 			else
 			{
