UnrealBuildTool: Default WindowsPlatform.bClangStandaloneDebug to true. This will slightly increase pdb size when compiling with clang, but will improve debugging experience

diff --git a/Engine/Source/Programs/UnrealBuildTool/Platform/Windows/UEBuildWindows.cs b/Engine/Source/Programs/UnrealBuildTool/Platform/Windows/UEBuildWindows.cs
--- a/Engine/Source/Programs/UnrealBuildTool/Platform/Windows/UEBuildWindows.cs
+++ b/Engine/Source/Programs/UnrealBuildTool/Platform/Windows/UEBuildWindows.cs
@@ -374,13 +374,14 @@ namespace UnrealBuildTool
 		public bool bVCExtendedWarningInfo = true;
 
 		/// <summary>
-		/// True if optimizations to reduce the size of debug information should be disabled
+		/// True if optimizations to reduce the size of debug information should be disabled. Disabling this option can decrease pdb size but may degrade the debugging experience.
 		/// See https://clang.llvm.org/docs/UsersManual.html#cmdoption-fstandalone-debug for more information
 		/// </summary>
 		[ConfigFile(ConfigHierarchyType.Engine, "/Script/WindowsTargetPlatform.WindowsTargetSettings", "bClangStandaloneDebug")]
 		[XmlConfigFile(Category = "WindowsPlatform")]
 		[CommandLine("-ClangStandaloneDebug")]
-		public bool bClangStandaloneDebug = false;
+		[CommandLine("-NoClangStandaloneDebug", Value = "false")]
+		public bool bClangStandaloneDebug = true;
 
 		/// <summary>
 		/// True if we should use the Clang linker (LLD) when we are compiling with Clang or Intel oneAPI, otherwise we use the MSVC linker.
