[Trivial] struct/class config auto-parsing does not fail on missing values

#rnx

diff --git a/Engine/Source/Programs/UnrealBuildTool/System/ConfigValueParser.cs b/Engine/Source/Programs/UnrealBuildTool/System/ConfigValueParser.cs
--- a/Engine/Source/Programs/UnrealBuildTool/System/ConfigValueParser.cs
+++ b/Engine/Source/Programs/UnrealBuildTool/System/ConfigValueParser.cs
@@ -241,7 +241,7 @@ namespace UnrealBuildTool
 						continue;
 					}
 
-					throw new Exception($"Cannot parse {Field.Name} for {T.Name} from {Line}");
+					// NB. We allow uninitialized values in structs/classes - it may be expected & ini saving omits empty strings etc.
 				}
 				return true;
 			}
