static analysis fixes:
- guid is not a primitive
- cannot rely on FieldAtttributes.HasDefault (condition is not needed anyway)

#rnx

diff --git a/Engine/Source/Programs/UnrealBuildTool/System/ConfigValueParser.cs b/Engine/Source/Programs/UnrealBuildTool/System/ConfigValueParser.cs
--- a/Engine/Source/Programs/UnrealBuildTool/System/ConfigValueParser.cs
+++ b/Engine/Source/Programs/UnrealBuildTool/System/ConfigValueParser.cs
@@ -54,11 +54,6 @@ namespace UnrealBuildTool
 						Result = DoubleValue;
 						return true;
 					}
-					else if (T == typeof(Guid) && TryParse(Line, out Guid GuidValue))
-					{
-						Result = GuidValue;
-						return true;
-					}
 
 					throw new Exception($"Cannot parse \"{Line}\" for primitive type {T}");
 				}
@@ -101,6 +96,16 @@ namespace UnrealBuildTool
 						return true;
 					}
 				}
+				else if (T == typeof(Guid))
+				{
+					if (TryParse(Line, out Guid GuidValue))
+					{
+						Result = GuidValue;
+						return true;
+					}
+
+					throw new Exception($"Cannot parse Guid from {Line}");
+				}
 				else if (T.IsGenericType)
 				{
 					if (T.GetGenericTypeDefinition() == typeof(Dictionary<,>) && TryParseAsMap(Line, out Dictionary<string, string>? DictValues))
@@ -235,11 +240,6 @@ namespace UnrealBuildTool
 							return false;
 						}
 					}
-					else if (Field.Attributes.HasFlag(System.Reflection.FieldAttributes.HasDefault))
-					{
-						// value has a default
-						continue;
-					}
 
 					// NB. We allow uninitialized values in structs/classes - it may be expected & ini saving omits empty strings etc.
 				}
