UnrealBuildTool: Add vfs paths to PGO_PATH_TRANSLATION

#rnx

diff --git a/Engine/Source/Programs/UnrealBuildTool/Platform/Windows/VCToolChain.cs b/Engine/Source/Programs/UnrealBuildTool/Platform/Windows/VCToolChain.cs
--- a/Engine/Source/Programs/UnrealBuildTool/Platform/Windows/VCToolChain.cs
+++ b/Engine/Source/Programs/UnrealBuildTool/Platform/Windows/VCToolChain.cs
@@ -123,12 +123,25 @@ namespace UnrealBuildTool
 			EnvVars.SetEnvironmentVariables();
 
 			// This allows PGD files generated from this build to be portable to other machines. It needs to be set in all configurations so we don't split the environment
-			List<string> pgoPathTranslations = new() { $"{Unreal.RootDirectory.FullName}=ROOT" };
+			HashSet<string> pgoPathTranslations = [];
+			string? pgoPathTranslationEnv = Environment.GetEnvironmentVariable("PGO_PATH_TRANSLATION");
+			if (!String.IsNullOrWhiteSpace(pgoPathTranslationEnv))
+			{
+				pgoPathTranslations= [.. pgoPathTranslationEnv.Split(';', StringSplitOptions.RemoveEmptyEntries)];
+			}
+
+			CppRootPaths RootPaths = new() { bUseVfs = true };
+			RootPaths[CppRootPathFolder.Root] = Unreal.RootDirectory;
 			if (Target.ProjectFile != null)
 			{
-				pgoPathTranslations.Add($"{Target.ProjectFile.Directory.FullName}=PROJ");
+				RootPaths[CppRootPathFolder.Project] = Target.ProjectFile.Directory;
+			}
+			foreach ((_, DirectoryReference vfs, DirectoryReference local) in RootPaths)
+			{
+				pgoPathTranslations.Add($"{local}={vfs.GetDirectoryName().ToUpperInvariant()}");
+				pgoPathTranslations.Add($"{vfs}={vfs.GetDirectoryName().ToUpperInvariant()}");
 			}
-			Environment.SetEnvironmentVariable("PGO_PATH_TRANSLATION", String.Join(";", pgoPathTranslations));
+			Environment.SetEnvironmentVariable("PGO_PATH_TRANSLATION", String.Join(';', pgoPathTranslations.Order()));
 
 			// Don't allow the INCLUDE environment variable to propagate. It's set by the IDE based on the IncludePath property in the project files which we
 			// add to improve Visual Studio memory usage, but we don't actually need it to set when invoking the compiler. Doing so results in it being converted
