UnrealBuildTool: Enable UBA for msvc pgo

#rnx

diff --git a/Engine/Source/Programs/UnrealBuildTool/Platform/Windows/VCToolChain.cs b/Engine/Source/Programs/UnrealBuildTool/Platform/Windows/VCToolChain.cs
--- a/Engine/Source/Programs/UnrealBuildTool/Platform/Windows/VCToolChain.cs
+++ b/Engine/Source/Programs/UnrealBuildTool/Platform/Windows/VCToolChain.cs
@@ -1587,23 +1587,15 @@ namespace UnrealBuildTool
 			}
 			else
 			{
-				// Remove when we fix so pgo linking run in uba
-				CppRootPaths RootPaths = LinkEnvironment.RootPaths;
-				if (LinkEnvironment.bPGOOptimize || LinkEnvironment.bPGOProfile)
-				{
-					RootPaths = new(RootPaths);
-					RootPaths.bUseVfs = false;
-				}
-
 				Arguments.Add("/MANIFEST:EMBED");
 				FileItem ManifestFile = FileItem.GetItemByPath(Target.WindowsPlatform.ManifestFile);
-				Arguments.Add($"/MANIFESTINPUT:\"{NormalizeCommandLinePath(ManifestFile, RootPaths)}\"");
+				Arguments.Add($"/MANIFESTINPUT:\"{NormalizeCommandLinePath(ManifestFile, LinkEnvironment.RootPaths)}\"");
 
 				// Embed a dependency on the SxS assembly manifest. this is for dependent runtime DLLs
 				if (!Target.Architecture.bIsX64)
 				{
 					string AssemblyName = Target.Architecture.ToString().ToLower();
-					Arguments.Add("/MANIFESTDEPENDENCY:\"name='" + AssemblyName + "' processorArchitecture='arm64' version='1.0.0.0' type='win32'\"");
+					Arguments.Add($"/MANIFESTDEPENDENCY:\"name='{AssemblyName}' processorArchitecture='arm64' version='1.0.0.0' type='win32'\"");
 				}
 			}
 
@@ -1664,7 +1656,7 @@ namespace UnrealBuildTool
 				if (LinkEnvironment.bIsBuildingConsoleApplication && !LinkEnvironment.bIsBuildingDLL && !String.IsNullOrEmpty(LinkEnvironment.WindowsEntryPointOverride))
 				{
 					// Use overridden entry point
-					Arguments.Add("/ENTRY:" + LinkEnvironment.WindowsEntryPointOverride);
+					Arguments.Add($"/ENTRY:{LinkEnvironment.WindowsEntryPointOverride}");
 				}
 
 				// Allow the OS to load the EXE at different base addresses than its preferred base address.
@@ -1681,11 +1673,11 @@ namespace UnrealBuildTool
 				{
 					if (LinkEnvironment.DefaultStackSizeCommit > 0)
 					{
-						Arguments.Add("/STACK:" + LinkEnvironment.DefaultStackSize + "," + LinkEnvironment.DefaultStackSizeCommit);
+						Arguments.Add($"/STACK:{LinkEnvironment.DefaultStackSize},{LinkEnvironment.DefaultStackSizeCommit}");
 					}
 					else
 					{
-						Arguments.Add("/STACK:" + LinkEnvironment.DefaultStackSize);
+						Arguments.Add($"/STACK:{LinkEnvironment.DefaultStackSize}");
 					}
 				}
 			}
@@ -3334,11 +3326,6 @@ namespace UnrealBuildTool
 				LinkAction.bCanExecuteRemotely = false;
 			}
 
-			if (LinkEnvironment.bPGOOptimize || LinkEnvironment.bPGOProfile)
-			{
-				LinkAction.bCanExecuteInUBA = false; // Disabled for now. Should revisit to see why it is not working
-			}
-
 			if (EnvVars.Compiler.IsClang() && Target.WindowsPlatform.bSampleBasedPGO && LinkEnvironment.bPGOOptimize)
 			{
 				LinkAction.PrerequisiteItems.Add(FileItem.GetItemByFileReference(GetClangLinkProfDataFilename(LinkEnvironment)));
