UnrealBuildTool: Query file encoding from .cs files that will be dynamically compiled to ensure hashes match for debugging

diff --git a/Engine/Source/Programs/UnrealBuildTool/System/DynamicCompilation.cs b/Engine/Source/Programs/UnrealBuildTool/System/DynamicCompilation.cs
--- a/Engine/Source/Programs/UnrealBuildTool/System/DynamicCompilation.cs
+++ b/Engine/Source/Programs/UnrealBuildTool/System/DynamicCompilation.cs
@@ -1,5 +1,12 @@
 // Copyright Epic Games, Inc. All Rights Reserved.
 
+using EpicGames.Core;
+using Microsoft.CodeAnalysis;
+using Microsoft.CodeAnalysis.CSharp;
+using Microsoft.CodeAnalysis.Emit;
+using Microsoft.CodeAnalysis.Text;
+using Microsoft.Extensions.Logging;
+using OpenTracing.Util;
 using System;
 using System.Collections.Generic;
 using System.IO;
@@ -8,13 +15,6 @@ using System.Reflection;
 using System.Runtime.InteropServices;
 using System.Threading;
 using System.Threading.Tasks;
-using EpicGames.Core;
-using Microsoft.CodeAnalysis;
-using Microsoft.CodeAnalysis.CSharp;
-using Microsoft.CodeAnalysis.Emit;
-using Microsoft.CodeAnalysis.Text;
-using Microsoft.Extensions.Logging;
-using OpenTracing.Util;
 using UnrealBuildBase;
 
 namespace UnrealBuildTool
@@ -170,8 +170,8 @@ namespace UnrealBuildTool
 
 		private static async Task<SyntaxTree> ParseSyntaxTreeAsync(FileReference sourceFileName, CSharpParseOptions parseOptions, CancellationToken cancellationToken)
 		{
-			string content = await FileReference.ReadAllTextAsync(sourceFileName, cancellationToken);
-			SourceText source = SourceText.From(content, System.Text.Encoding.UTF8);
+			byte[] bytes = await FileReference.ReadAllBytesAsync(sourceFileName, cancellationToken);
+			SourceText source = SourceText.From(bytes, bytes.Length, FileUtils.GetEncoding(bytes));
 			return CSharpSyntaxTree.ParseText(source, parseOptions, sourceFileName.FullName, cancellationToken);
 		}
 
