UnrealBuildTool: Remove empty items from UBA providers list before adding in default provider

#rnx

diff --git a/Engine/Source/Programs/UnrealBuildTool/Executors/UnrealBuildAccelerator/UBAAgentCoordinatorHorde.cs b/Engine/Source/Programs/UnrealBuildTool/Executors/UnrealBuildAccelerator/UBAAgentCoordinatorHorde.cs
--- a/Engine/Source/Programs/UnrealBuildTool/Executors/UnrealBuildAccelerator/UBAAgentCoordinatorHorde.cs
+++ b/Engine/Source/Programs/UnrealBuildTool/Executors/UnrealBuildAccelerator/UBAAgentCoordinatorHorde.cs
@@ -736,14 +736,14 @@ namespace UnrealBuildTool
 			IEnumerable<string> providers = Unreal.IsBuildMachine()
 				? [.. ubaConfig.BuildMachineProviders, .. ubaConfig.IniBuildMachineProviders]
 				: [.. ubaConfig.Providers, .. ubaConfig.IniProviders];
-
+			providers = providers.Distinct().Where(x => !String.IsNullOrEmpty(x) && x.StartsWith(ProviderPrefix));
 
 			if (!providers.Any())
 			{
 				providers = [ProviderPrefix];
 			}
 
-			IEnumerable<UBAAgentCoordinatorHorde> coordinators = providers.Distinct().Where(x => x.StartsWith(ProviderPrefix))
+			IEnumerable<UBAAgentCoordinatorHorde> coordinators = providers
 				.SelectMany(provider => targetDescriptors.Select(targetDescriptor => new UBAAgentCoordinatorHorde(logger, ubaConfig, provider, targetDescriptor.AdditionalArguments, targetDescriptor.ProjectFile?.Directory)))
 				.Where(x => !x._hordeConfig.bDisableHorde && !String.IsNullOrWhiteSpace(x._hordeConfig.HordeServer))
 				.DistinctBy(x => HashCode.Combine(x._hordeConfig.HordeServer, x._hordeConfig.HordeCondition, x._hordeConfig.HordeCluster, x._hordeConfig.HordePool, x._hordeConfig.HordeConnectionMode));
