UnrealBuildTool: Automated analyze message cleanup

#rnx

diff --git a/Engine/Source/Programs/UnrealBuildTool/Configuration/UEBuildModuleCPP.cs b/Engine/Source/Programs/UnrealBuildTool/Configuration/UEBuildModuleCPP.cs
--- a/Engine/Source/Programs/UnrealBuildTool/Configuration/UEBuildModuleCPP.cs
+++ b/Engine/Source/Programs/UnrealBuildTool/Configuration/UEBuildModuleCPP.cs
@@ -9,7 +9,6 @@ using System.Collections.Generic;
 using System.IO;
 using System.Linq;
 using System.Text;
-using System.Text.Json;
 using UnrealBuildBase;
 
 namespace UnrealBuildTool
@@ -24,20 +23,20 @@ namespace UnrealBuildTool
 		/// </summary>
 		public class InputFileCollection
 		{
-			public readonly List<FileItem> HeaderFiles = new List<FileItem>();
-			public readonly List<FileItem> ISPCHeaderFiles = new List<FileItem>();
-
-			public readonly List<FileItem> IXXFiles = new List<FileItem>();
-			public readonly List<FileItem> CPPFiles = new List<FileItem>();
-			public readonly List<FileItem> CFiles = new List<FileItem>();
-			public readonly List<FileItem> CCFiles = new List<FileItem>();
-			public readonly List<FileItem> CXXFiles = new List<FileItem>();
-			public readonly List<FileItem> MFiles = new List<FileItem>();
-			public readonly List<FileItem> MMFiles = new List<FileItem>();
-			public readonly List<FileItem> SwiftFiles = new List<FileItem>();
-			public readonly List<FileItem> RCFiles = new List<FileItem>();
-			public readonly List<FileItem> ISPCFiles = new List<FileItem>();
-			public readonly List<FileItem> ProtoFiles = new List<FileItem>();
+			public readonly List<FileItem> HeaderFiles = [];
+			public readonly List<FileItem> ISPCHeaderFiles = [];
+
+			public readonly List<FileItem> IXXFiles = [];
+			public readonly List<FileItem> CPPFiles = [];
+			public readonly List<FileItem> CFiles = [];
+			public readonly List<FileItem> CCFiles = [];
+			public readonly List<FileItem> CXXFiles = [];
+			public readonly List<FileItem> MFiles = [];
+			public readonly List<FileItem> MMFiles = [];
+			public readonly List<FileItem> SwiftFiles = [];
+			public readonly List<FileItem> RCFiles = [];
+			public readonly List<FileItem> ISPCFiles = [];
+			public readonly List<FileItem> ProtoFiles = [];
 
 			public IEnumerable<FileItem> AllHeaderFiles => HeaderFiles.Concat(ISPCHeaderFiles);
 			public IEnumerable<FileItem> AllCompiledFiles => IXXFiles.Concat(CPPFiles).Concat(CFiles).Concat(CCFiles).Concat(CXXFiles).Concat(MFiles).Concat(MMFiles).Concat(SwiftFiles).Concat(RCFiles).Concat(ISPCFiles);
@@ -46,16 +45,10 @@ namespace UnrealBuildTool
 		/// <summary>
 		/// Instance where a CPP file doesn't include the expected header as the first include.
 		/// </summary>
-		public readonly struct InvalidIncludeDirective
+		public readonly struct InvalidIncludeDirective(FileReference cppFile, FileReference headerFile)
 		{
-			public FileReference CppFile { get; }
-			public FileReference HeaderFile { get; }
-
-			public InvalidIncludeDirective(FileReference cppFile, FileReference headerFile)
-			{
-				CppFile = cppFile;
-				HeaderFile = headerFile;
-			}
+			public FileReference CppFile { get; } = cppFile;
+			public FileReference HeaderFile { get; } = headerFile;
 		}
 
 		/// <summary>
@@ -106,7 +99,7 @@ namespace UnrealBuildTool
 		/// <summary>
 		/// Verse source code directory associated with this module
 		/// </summary>
-		DirectoryReference? AssociatedVerseDirectory;
+		readonly DirectoryReference? AssociatedVerseDirectory;
 
 		/// <summary>
 		/// Create PerModuleInline source file. Defaults to PrimaryModule in a binary but could be another one if primary module does not depend on core
@@ -511,19 +504,21 @@ namespace UnrealBuildTool
 			}
 
 			// Find all the input files
-			Dictionary<DirectoryItem, FileItem[]> DirectoryToSourceFiles = new Dictionary<DirectoryItem, FileItem[]>();
+			Dictionary<DirectoryItem, FileItem[]> DirectoryToSourceFiles = [];
 			InputFileCollection InputFiles = FindInputFiles(Target.Platform, DirectoryToSourceFiles, Logger);
 
 			{
-				Dictionary<string, List<FileItem>> nameToPath = new();
+				Dictionary<string, List<FileItem>> nameToPath = [];
 				foreach (FileItem InputFile in InputFiles.AllCompiledFiles)
 				{
 					string name = InputFile.Name.ToUpperInvariant();
-					if (!nameToPath.ContainsKey(name))
+					if (!nameToPath.TryGetValue(name, out List<FileItem>? path))
 					{
-						nameToPath.Add(name, new List<FileItem>());
+						path = [];
+						nameToPath.Add(name, path);
 					}
-					nameToPath[name].Add(InputFile);
+
+					path.Add(InputFile);
 				}
 				IEnumerable<KeyValuePair<string, List<FileItem>>> fileConflicts = nameToPath.Where(item => item.Value.Count > 1);
 				if (fileConflicts.Any())
@@ -542,12 +537,12 @@ namespace UnrealBuildTool
 				Graph.AddSourceFiles(Pair.Key, Pair.Value);
 			}
 
-			Graph.AddHeaderFiles(InputFiles.AllHeaderFiles.ToArray());
+			Graph.AddHeaderFiles([.. InputFiles.AllHeaderFiles]);
 
-			// We are building with IWYU and thismodule does not support it, early out
+			// We are building with IWYU and this module does not support it, early out
 			if (Target.bIWYU && Rules.IWYUSupport == IWYUSupport.None)
 			{
-				return new List<FileItem>();
+				return [];
 			}
 
 			// Process all of the header file dependencies for this module
@@ -587,7 +582,7 @@ namespace UnrealBuildTool
 			}
 
 			// Configure the precompiled headers for this module
-			CppCompileEnvironment = SetupPrecompiledHeaders(Target, ToolChain, CppCompileEnvironment, LinkInputFiles, Graph, Logger);
+			CppCompileEnvironment = SetupPrecompiledHeaders(Target, ToolChain, CppCompileEnvironment, LinkInputFiles, Graph);
 			if (CppCompileEnvironment.PerArchPrecompiledHeaderFiles != null)
 			{
 				foreach (UnrealArch Arch in CppCompileEnvironment.PerArchPrecompiledHeaderFiles.Keys)
@@ -619,15 +614,17 @@ namespace UnrealBuildTool
 				CppCompileEnvironment.UserIncludePaths.Add(IntermediateDirectory);
 				ModuleCompileEnvironment.UserIncludePaths.Add(IntermediateDirectory);
 
-				CppCompileEnvironment ProtoCompileEnvironment = new CppCompileEnvironment(CppCompileEnvironment);
-				ProtoCompileEnvironment.PrecompiledHeaderAction = PrecompiledHeaderAction.None;
-				ProtoCompileEnvironment.bUseUnity = false;
-				ProtoCompileEnvironment.bWarningsAsErrors = false;
-				// Disable dynamic analysis on generated proto code as it causes issues when the proto/grpc libs are not
-				// compiled the same way. This can cause false positive in some situations, for example with address
-				// sanitize enabled (see: [REDACTED_PROFILE_URL]/protobuf/issues/13115).
-				ProtoCompileEnvironment.bDisableDynamicAnalysis = true;
-				// Disable warnings that are common in genereated proto cpp
+				CppCompileEnvironment ProtoCompileEnvironment = new(CppCompileEnvironment)
+				{
+					PrecompiledHeaderAction = PrecompiledHeaderAction.None,
+					bUseUnity = false,
+					bWarningsAsErrors = false,
+					// Disable dynamic analysis on generated proto code as it causes issues when the proto/grpc libs are not
+					// compiled the same way. This can cause false positive in some situations, for example with address
+					// sanitize enabled (see: [REDACTED_PROFILE_URL]/protobuf/issues/13115).
+					bDisableDynamicAnalysis = true
+				};
+				// Disable warnings that are common in generated proto cpp
 				ProtoCompileEnvironment.CppCompileWarnings.DeprecationWarningLevel = WarningLevel.Off;
 				ProtoCompileEnvironment.CppCompileWarnings.ShortenSizeTToIntWarningLevel = WarningLevel.Off;
 				ProtoCompileEnvironment.CppCompileWarnings.UnusedParameterWarningLevel = WarningLevel.Off;
@@ -646,17 +643,17 @@ namespace UnrealBuildTool
 			CppCompileEnvironment = ToolChain.CreateSharedResponseFile(CppCompileEnvironment, SharedResponseFile, Graph);
 
 			// Mapping of source file to unity file. We output this to intermediate directories for other tools (eg. live coding) to use.
-			Dictionary<FileItem, FileItem> SourceFileToUnityFile = new Dictionary<FileItem, FileItem>();
+			Dictionary<FileItem, FileItem> SourceFileToUnityFile = [];
 
-			List<FileItem> CPPFiles = new List<FileItem>(InputFiles.CPPFiles);
-			List<FileItem> GeneratedFileItems = new List<FileItem>();
+			List<FileItem> CPPFiles = [.. InputFiles.CPPFiles];
+			List<FileItem> GeneratedFileItems = [];
 			CppCompileEnvironment GeneratedCPPCompileEnvironment = CppCompileEnvironment;
 			bool bMergeUnityFiles = !(Target.DisableMergingModuleAndGeneratedFilesInUnityFiles?.Contains(Name) ?? false) && Rules.bMergeUnityFiles;
 
 			// Compile all the generated CPP files
 			if (GeneratedCppDirectories != null && !CppCompileEnvironment.bHackHeaderGenerator)
 			{
-				Dictionary<string, FileItem> GeneratedFiles = new Dictionary<string, FileItem>();
+				Dictionary<string, FileItem> GeneratedFiles = [];
 				foreach (string GeneratedDir in GeneratedCppDirectories)
 				{
 					if (!Directory.Exists(GeneratedDir))
@@ -672,7 +669,7 @@ namespace UnrealBuildTool
 						string FileName = File.Name;
 						if (FileName.EndsWith(".gen.cpp"))
 						{
-							string Key = Prefix + FileName.Substring(0, FileName.Length - ".gen.cpp".Length);
+							string Key = string.Concat(Prefix, FileName.AsSpan(0, FileName.Length - ".gen.cpp".Length));
 							GeneratedFiles.Add(Key, File);
 						}
 						else if (GeneratedCPPCompileEnvironment.FileMatchesExtraGeneratedCPPTypes(FileName))
@@ -694,11 +691,13 @@ namespace UnrealBuildTool
 							string Key = Prefix + ListOfInlinedGenCppsItem;
 							if (GeneratedFiles.Remove(Key, out FileItem? FoundGenCppFile))
 							{
-								if (!CppCompileEnvironment.FileInlineGenCPPMap.ContainsKey(CPPFileItem))
+								if (!CppCompileEnvironment.FileInlineGenCPPMap.TryGetValue(CPPFileItem, out List<FileItem>? genFiles))
 								{
-									CppCompileEnvironment.FileInlineGenCPPMap[CPPFileItem] = new List<FileItem>();
+									genFiles = [];
+									CppCompileEnvironment.FileInlineGenCPPMap[CPPFileItem] = genFiles;
 								}
-								CppCompileEnvironment.FileInlineGenCPPMap[CPPFileItem].Add(FoundGenCppFile);
+
+								genFiles.Add(FoundGenCppFile);
 							}
 							else
 							{
@@ -710,7 +709,7 @@ namespace UnrealBuildTool
 
 					if (Rules.CppCompileWarningSettings.NonInlinedGenCppWarningLevel >= WarningLevel.Warning)
 					{
-						Dictionary<string, FileItem> CPPFilesLookup = new Dictionary<string, FileItem>();
+						Dictionary<string, FileItem> CPPFilesLookup = [];
 						foreach (FileItem CPPFile in CPPFiles)
 						{
 							CPPFilesLookup.Add(CPPFile.Name, CPPFile);
@@ -723,11 +722,11 @@ namespace UnrealBuildTool
 								continue;
 							}
 
-							string NameWithoutPrefix = Name.Substring(4);
+							string NameWithoutPrefix = Name[4..];
 							string NameWithoutPrefixWithSuffix = NameWithoutPrefix + ".cpp";
 							if (CPPFilesLookup.TryGetValue(NameWithoutPrefixWithSuffix, out FileItem? Item))
 							{
-								Logger.LogWarning("{0}(1): .gen.cpp not inlined. Add '#include UE_INLINE_GENERATED_CPP_BY_NAME({1})' after header includes.", Item.AbsolutePath, NameWithoutPrefix);
+								Logger.LogWarning("{Path}(1): .gen.cpp not inlined. Add '#include UE_INLINE_GENERATED_CPP_BY_NAME({Name})' after header includes.", Item.AbsolutePath, NameWithoutPrefix);
 							}
 						}
 					}
@@ -735,15 +734,19 @@ namespace UnrealBuildTool
 					// Create a compile environment for the generated files. We can disable creating debug info here to improve link times.
 					if (GeneratedCPPCompileEnvironment.bCreateDebugInfo && Target.bDisableDebugInfoForGeneratedCode)
 					{
-						GeneratedCPPCompileEnvironment = new CppCompileEnvironment(GeneratedCPPCompileEnvironment);
-						GeneratedCPPCompileEnvironment.bCreateDebugInfo = false;
+						GeneratedCPPCompileEnvironment = new CppCompileEnvironment(GeneratedCPPCompileEnvironment)
+						{
+							bCreateDebugInfo = false
+						};
 						bMergeUnityFiles = false;
 					}
 
 					if (Target.StaticAnalyzer != StaticAnalyzer.None && !Target.bStaticAnalyzerIncludeGenerated)
 					{
-						GeneratedCPPCompileEnvironment = new CppCompileEnvironment(GeneratedCPPCompileEnvironment);
-						GeneratedCPPCompileEnvironment.bDisableStaticAnalysis = true;
+						GeneratedCPPCompileEnvironment = new CppCompileEnvironment(GeneratedCPPCompileEnvironment)
+						{
+							bDisableStaticAnalysis = true
+						};
 						bMergeUnityFiles = false;
 					}
 
@@ -886,11 +889,13 @@ namespace UnrealBuildTool
 			GeneratedCPPCompileEnvironment.bUseUnity = CppCompileEnvironment.bUseUnity;
 
 			// Create and register a special action that can be used to compile single files (even when unity is enabled) for generated files
-			if (GeneratedFileItems.Any())
+			if (GeneratedFileItems.Count != 0)
 			{
-				CppCompileEnvironment SingleGeneratedCPPCompileEnvironment = new CppCompileEnvironment(GeneratedCPPCompileEnvironment);
-				SingleGeneratedCPPCompileEnvironment.PrecompiledHeaderAction = PrecompiledHeaderAction.None;
-				SingleGeneratedCPPCompileEnvironment.bUseUnity = false;
+				CppCompileEnvironment SingleGeneratedCPPCompileEnvironment = new(GeneratedCPPCompileEnvironment)
+				{
+					PrecompiledHeaderAction = PrecompiledHeaderAction.None,
+					bUseUnity = false
+				};
 				foreach (DirectoryReference Directory in GeneratedFileItems.Select(x => x.Location.Directory).Distinct())
 				{
 					ToolChain.CreateSpecificFileAction(SingleGeneratedCPPCompileEnvironment, Directory, IntermediateDirectory, Graph);
@@ -898,11 +903,13 @@ namespace UnrealBuildTool
 			}
 
 			// Create and register a special action that can be used to compile single files (even when unity is enabled) for normal files
-			if (ModuleDirectories.Any())
+			if (ModuleDirectories.Length != 0)
 			{
-				CppCompileEnvironment SingleCompileEnvironment = new CppCompileEnvironment(CppCompileEnvironment);
-				SingleCompileEnvironment.PrecompiledHeaderAction = PrecompiledHeaderAction.None;
-				SingleCompileEnvironment.bUseUnity = false;
+				CppCompileEnvironment SingleCompileEnvironment = new(CppCompileEnvironment)
+				{
+					PrecompiledHeaderAction = PrecompiledHeaderAction.None,
+					bUseUnity = false
+				};
 				foreach (DirectoryReference Directory in ModuleDirectories)
 				{
 					ToolChain.CreateSpecificFileAction(SingleCompileEnvironment, Directory, IntermediateDirectory, Graph);
@@ -917,24 +924,24 @@ namespace UnrealBuildTool
 			else if (Rules.Target.bAlwaysUseUnityForGeneratedFiles && !Rules.Target.bEnableCppModules)
 			{
 				GeneratedCPPCompileEnvironment.bUseUnity = true;
-				Unity.GenerateUnitySource(Target, GeneratedFileItems, new List<FileItem>(), GeneratedCPPCompileEnvironment, WorkingSet, (Rules.ShortName ?? Name) + ".gen", IntermediateDirectory, Graph, SourceFileToUnityFile,
+				Unity.GenerateUnitySource(Target, GeneratedFileItems, [], GeneratedCPPCompileEnvironment, WorkingSet, (Rules.ShortName ?? Name) + ".gen", IntermediateDirectory, Graph, SourceFileToUnityFile,
 					out List<FileItem> NormalGeneratedFiles, out List<FileItem> AdaptiveGeneratedFiles, NumIncludedBytesPerUnityCPP);
 				ModuleCompileEnvironment.FileInlineGenCPPMap = new(GeneratedCPPCompileEnvironment.FileInlineGenCPPMap);
 				CppCompileEnvironment.FileInlineGenCPPMap = new(GeneratedCPPCompileEnvironment.FileInlineGenCPPMap);
-				LinkInputFiles.AddRange(CompileFilesWithToolChain(Target, ToolChain, GeneratedCPPCompileEnvironment, ModuleCompileEnvironment, NormalGeneratedFiles, AdaptiveGeneratedFiles, Graph, Logger).ObjectFiles);
+				LinkInputFiles.AddRange(CompileFilesWithToolChain(Target, ToolChain, GeneratedCPPCompileEnvironment, ModuleCompileEnvironment, NormalGeneratedFiles, AdaptiveGeneratedFiles, Graph).ObjectFiles);
 			}
 			else
 			{
 				Unity.GetAdaptiveFiles(Target, GeneratedFileItems, InputFiles.HeaderFiles, GeneratedCPPCompileEnvironment, WorkingSet, (Rules.ShortName ?? Name) + ".gen", IntermediateDirectory, Graph,
 					out List<FileItem> NormalGeneratedFiles, out List<FileItem> AdaptiveGeneratedFiles);
-				LinkInputFiles.AddRange(CompileFilesWithToolChain(Target, ToolChain, GeneratedCPPCompileEnvironment, ModuleCompileEnvironment, NormalGeneratedFiles, AdaptiveGeneratedFiles, Graph, Logger).ObjectFiles);
+				LinkInputFiles.AddRange(CompileFilesWithToolChain(Target, ToolChain, GeneratedCPPCompileEnvironment, ModuleCompileEnvironment, NormalGeneratedFiles, AdaptiveGeneratedFiles, Graph).ObjectFiles);
 			}
 
 
 			// Compile Swift - which may generate a header that the cpp files need below
 			if (InputFiles.SwiftFiles.Count > 0 && Rules.SwiftInteropHeader != null)
 			{
-				CppCompileEnvironment SwiftCompileEnvironment = new CppCompileEnvironment(ModuleCompileEnvironment);
+				CppCompileEnvironment SwiftCompileEnvironment = new(ModuleCompileEnvironment);
 				// currently nothing is using this header, but it could be helpful to have in swift
 				CreateHeaderForDefinitions(SwiftCompileEnvironment, IntermediateDirectory, "swift", Graph);
 				SwiftCompileEnvironment.ForceIncludeFiles.Clear();
@@ -955,7 +962,7 @@ namespace UnrealBuildTool
 				// both used when compiling files that may need the headers
 				ModuleCompileEnvironment.AdditionalPrerequisites.AddRange(HeaderFiles);
 				CppCompileEnvironment.AdditionalPrerequisites.AddRange(HeaderFiles);
-				if (HeaderFiles.Count() > 0)
+				if (HeaderFiles.Any())
 				{
 					// @todo: this could have different per-arch headers- should be rare but possible
 					ModuleCompileEnvironment.UserIncludePaths.Add(HeaderFiles.First().Location.Directory);
@@ -971,7 +978,7 @@ namespace UnrealBuildTool
 					Unity.GenerateUnitySource(Target, CPPFiles, InputFiles.HeaderFiles, CppCompileEnvironment, WorkingSet, Rules.ShortName ?? Name, IntermediateDirectory, Graph, SourceFileToUnityFile,
 						out List<FileItem> NormalFiles, out List<FileItem> AdaptiveFiles, NumIncludedBytesPerUnityCPP);
 					ModuleCompileEnvironment.FileInlineGenCPPMap = new(CppCompileEnvironment.FileInlineGenCPPMap);
-					LinkInputFiles.AddRange(CompileFilesWithToolChain(Target, ToolChain, CppCompileEnvironment, ModuleCompileEnvironment, NormalFiles, AdaptiveFiles, Graph, Logger).ObjectFiles);
+					LinkInputFiles.AddRange(CompileFilesWithToolChain(Target, ToolChain, CppCompileEnvironment, ModuleCompileEnvironment, NormalFiles, AdaptiveFiles, Graph).ObjectFiles);
 				}
 				else
 				{
@@ -982,7 +989,7 @@ namespace UnrealBuildTool
 						NormalFiles = CPPFiles;
 						AdaptiveFiles.RemoveAll(new HashSet<FileItem>(NormalFiles).Contains);
 					}
-					LinkInputFiles.AddRange(CompileFilesWithToolChain(Target, ToolChain, CppCompileEnvironment, ModuleCompileEnvironment, NormalFiles, AdaptiveFiles, Graph, Logger).ObjectFiles);
+					LinkInputFiles.AddRange(CompileFilesWithToolChain(Target, ToolChain, CppCompileEnvironment, ModuleCompileEnvironment, NormalFiles, AdaptiveFiles, Graph).ObjectFiles);
 				}
 			}
 
@@ -996,15 +1003,17 @@ namespace UnrealBuildTool
 			// Compile C files. Do not use a PCH here, because a C++ PCH is not compatible with C source files.
 			if (InputFiles.CFiles.Count > 0)
 			{
-				CppCompileEnvironment CCompileEnvironment = new CppCompileEnvironment(ModuleCompileEnvironment);
-				CCompileEnvironment.bUseUnity = bModuleUsesUnityBuild && Rules.UnityFileTypes.HasFlag(ModuleRules.UnityFileType.C);
+				CppCompileEnvironment CCompileEnvironment = new(ModuleCompileEnvironment)
+				{
+					bUseUnity = bModuleUsesUnityBuild && Rules.UnityFileTypes.HasFlag(ModuleRules.UnityFileType.C)
+				};
 				CreateHeaderForDefinitions(CCompileEnvironment, IntermediateDirectory, "c", Graph);
 
 				if (CCompileEnvironment.bUseUnity)
 				{
 					Unity.GenerateUnitySource(Target, InputFiles.CFiles, InputFiles.HeaderFiles, CCompileEnvironment, WorkingSet, Rules.ShortName ?? Name, IntermediateDirectory, Graph, SourceFileToUnityFile,
 						out List<FileItem> NormalFiles, out List<FileItem> AdaptiveFiles, NumIncludedBytesPerUnityCPP);
-					LinkInputFiles.AddRange(CompileFilesWithToolChain(Target, ToolChain, CCompileEnvironment, ModuleCompileEnvironment, NormalFiles, AdaptiveFiles, Graph, Logger).ObjectFiles);
+					LinkInputFiles.AddRange(CompileFilesWithToolChain(Target, ToolChain, CCompileEnvironment, ModuleCompileEnvironment, NormalFiles, AdaptiveFiles, Graph).ObjectFiles);
 				}
 				else
 				{
@@ -1015,15 +1024,17 @@ namespace UnrealBuildTool
 			// Compile CC files
 			if (InputFiles.CCFiles.Count > 0)
 			{
-				CppCompileEnvironment CCCompileEnvironment = new CppCompileEnvironment(CppCompileEnvironment);
-				CCCompileEnvironment.bUseUnity = bModuleUsesUnityBuild && Rules.UnityFileTypes.HasFlag(ModuleRules.UnityFileType.CC);
+				CppCompileEnvironment CCCompileEnvironment = new(CppCompileEnvironment)
+				{
+					bUseUnity = bModuleUsesUnityBuild && Rules.UnityFileTypes.HasFlag(ModuleRules.UnityFileType.CC)
+				};
 				CreateHeaderForDefinitions(CCCompileEnvironment, IntermediateDirectory, "cc", Graph);
 
 				if (CCCompileEnvironment.bUseUnity)
 				{
 					Unity.GenerateUnitySource(Target, InputFiles.CCFiles, InputFiles.HeaderFiles, CCCompileEnvironment, WorkingSet, Rules.ShortName ?? Name, IntermediateDirectory, Graph, SourceFileToUnityFile,
 						out List<FileItem> NormalFiles, out List<FileItem> AdaptiveFiles, NumIncludedBytesPerUnityCPP);
-					LinkInputFiles.AddRange(CompileFilesWithToolChain(Target, ToolChain, CCCompileEnvironment, CppCompileEnvironment, NormalFiles, AdaptiveFiles, Graph, Logger).ObjectFiles);
+					LinkInputFiles.AddRange(CompileFilesWithToolChain(Target, ToolChain, CCCompileEnvironment, CppCompileEnvironment, NormalFiles, AdaptiveFiles, Graph).ObjectFiles);
 				}
 				else
 				{
@@ -1034,15 +1045,17 @@ namespace UnrealBuildTool
 			// Compile CXX files
 			if (InputFiles.CXXFiles.Count > 0)
 			{
-				CppCompileEnvironment CXXCompileEnvironment = new CppCompileEnvironment(CppCompileEnvironment);
-				CXXCompileEnvironment.bUseUnity = bModuleUsesUnityBuild && Rules.UnityFileTypes.HasFlag(ModuleRules.UnityFileType.CXX);
+				CppCompileEnvironment CXXCompileEnvironment = new(CppCompileEnvironment)
+				{
+					bUseUnity = bModuleUsesUnityBuild && Rules.UnityFileTypes.HasFlag(ModuleRules.UnityFileType.CXX)
+				};
 				CreateHeaderForDefinitions(CXXCompileEnvironment, IntermediateDirectory, "cxx", Graph);
 
 				if (CXXCompileEnvironment.bUseUnity)
 				{
 					Unity.GenerateUnitySource(Target, InputFiles.CXXFiles, InputFiles.HeaderFiles, CXXCompileEnvironment, WorkingSet, Rules.ShortName ?? Name, IntermediateDirectory, Graph, SourceFileToUnityFile,
 						out List<FileItem> NormalFiles, out List<FileItem> AdaptiveFiles, NumIncludedBytesPerUnityCPP);
-					LinkInputFiles.AddRange(CompileFilesWithToolChain(Target, ToolChain, CXXCompileEnvironment, CppCompileEnvironment, NormalFiles, AdaptiveFiles, Graph, Logger).ObjectFiles);
+					LinkInputFiles.AddRange(CompileFilesWithToolChain(Target, ToolChain, CXXCompileEnvironment, CppCompileEnvironment, NormalFiles, AdaptiveFiles, Graph).ObjectFiles);
 				}
 				else
 				{
@@ -1053,15 +1066,17 @@ namespace UnrealBuildTool
 			// Compile M files. Do not use a PCH here, because a C++ PCH is not compatible with Objective-C source files.
 			if (InputFiles.MFiles.Count > 0)
 			{
-				CppCompileEnvironment MCompileEnvironment = new CppCompileEnvironment(ModuleCompileEnvironment);
-				MCompileEnvironment.bUseUnity = bModuleUsesUnityBuild && Rules.UnityFileTypes.HasFlag(ModuleRules.UnityFileType.M);
+				CppCompileEnvironment MCompileEnvironment = new(ModuleCompileEnvironment)
+				{
+					bUseUnity = bModuleUsesUnityBuild && Rules.UnityFileTypes.HasFlag(ModuleRules.UnityFileType.M)
+				};
 				CreateHeaderForDefinitions(MCompileEnvironment, IntermediateDirectory, "m", Graph);
 
 				if (MCompileEnvironment.bUseUnity)
 				{
 					Unity.GenerateUnitySource(Target, InputFiles.MFiles, InputFiles.HeaderFiles, MCompileEnvironment, WorkingSet, Rules.ShortName ?? Name, IntermediateDirectory, Graph, SourceFileToUnityFile,
 						out List<FileItem> NormalFiles, out List<FileItem> AdaptiveFiles, NumIncludedBytesPerUnityCPP);
-					LinkInputFiles.AddRange(CompileFilesWithToolChain(Target, ToolChain, MCompileEnvironment, ModuleCompileEnvironment, NormalFiles, AdaptiveFiles, Graph, Logger).ObjectFiles);
+					LinkInputFiles.AddRange(CompileFilesWithToolChain(Target, ToolChain, MCompileEnvironment, ModuleCompileEnvironment, NormalFiles, AdaptiveFiles, Graph).ObjectFiles);
 				}
 				else
 				{
@@ -1072,15 +1087,17 @@ namespace UnrealBuildTool
 			// Compile MM files
 			if (InputFiles.MMFiles.Count > 0)
 			{
-				CppCompileEnvironment MMCompileEnvironment = new CppCompileEnvironment(CppCompileEnvironment);
-				MMCompileEnvironment.bUseUnity = bModuleUsesUnityBuild && Rules.UnityFileTypes.HasFlag(ModuleRules.UnityFileType.MM);
+				CppCompileEnvironment MMCompileEnvironment = new(CppCompileEnvironment)
+				{
+					bUseUnity = bModuleUsesUnityBuild && Rules.UnityFileTypes.HasFlag(ModuleRules.UnityFileType.MM)
+				};
 				CreateHeaderForDefinitions(MMCompileEnvironment, IntermediateDirectory, "mm", Graph);
 
 				if (MMCompileEnvironment.bUseUnity)
 				{
 					Unity.GenerateUnitySource(Target, InputFiles.MMFiles, InputFiles.HeaderFiles, MMCompileEnvironment, WorkingSet, Rules.ShortName ?? Name, IntermediateDirectory, Graph, SourceFileToUnityFile,
 						out List<FileItem> NormalFiles, out List<FileItem> AdaptiveFiles, NumIncludedBytesPerUnityCPP);
-					LinkInputFiles.AddRange(CompileFilesWithToolChain(Target, ToolChain, MMCompileEnvironment, CppCompileEnvironment, NormalFiles, AdaptiveFiles, Graph, Logger).ObjectFiles);
+					LinkInputFiles.AddRange(CompileFilesWithToolChain(Target, ToolChain, MMCompileEnvironment, CppCompileEnvironment, NormalFiles, AdaptiveFiles, Graph).ObjectFiles);
 				}
 				else
 				{
@@ -1093,7 +1110,7 @@ namespace UnrealBuildTool
 			// must be specified relative to the resource file itself or Engine/Source.
 			if (InputFiles.RCFiles.Count > 0)
 			{
-				CppCompileEnvironment ResourceCompileEnvironment = new CppCompileEnvironment(BinaryCompileEnvironment);
+				CppCompileEnvironment ResourceCompileEnvironment = new(BinaryCompileEnvironment);
 				if (Binary != null)
 				{
 					// @todo: This should be in some Windows code somewhere...
@@ -1107,7 +1124,7 @@ namespace UnrealBuildTool
 			{
 				DirectoryReference.CreateDirectory(PrecompiledManifestLocation.Directory);
 
-				PrecompiledManifest Manifest = new PrecompiledManifest();
+				PrecompiledManifest Manifest = new();
 				Manifest.OutputFiles.AddRange(LinkInputFiles.Select(x => x.Location));
 				Manifest.WriteIfModified(PrecompiledManifestLocation);
 			}
@@ -1116,7 +1133,7 @@ namespace UnrealBuildTool
 			if (Rules.Target.bWithLiveCoding)
 			{
 				StringWriter StringWriter = new();
-				using (JsonWriter Writer = new JsonWriter(StringWriter))
+				using (JsonWriter Writer = new(StringWriter))
 				{
 					Writer.WriteObjectStart();
 					Writer.WriteObjectStart("RemapUnityFiles");
@@ -1158,8 +1175,10 @@ namespace UnrealBuildTool
 					}
 					else
 					{
-						HeaderCompileEnvironment = new CppCompileEnvironment(ModuleCompileEnvironment);
-						HeaderCompileEnvironment.bUseUnity = false;
+						HeaderCompileEnvironment = new CppCompileEnvironment(ModuleCompileEnvironment)
+						{
+							bUseUnity = false
+						};
 					}
 
 					DirectoryReference HeaderIntermediateDirectory = DirectoryReference.Combine(IntermediateDirectory, "H");
@@ -1170,7 +1189,7 @@ namespace UnrealBuildTool
 					foreach (DirectoryItem ParentDir in HeaderFileItems.Select(x => x.Directory!).Distinct().OrderBy(x => x.FullName))
 					{
 						DirectoryReference HeaderSubIntermediateDirectory = DirectoryReference.Combine(HeaderIntermediateDirectory, (DirectoryIndex++).ToString());
-						List<FileItem> HeadersToCompile = new();
+						List<FileItem> HeadersToCompile = [];
 						foreach (FileItem Header in HeaderFileItems.Where(x => x.Directory == ParentDir))
 						{
 							if (Target.bIWYU)
@@ -1210,7 +1229,7 @@ namespace UnrealBuildTool
 		{
 			if (Rules.IWYUSupport == IWYUSupport.None)
 			{
-				return new List<FileItem>();
+				return [];
 			}
 
 			// Find FileItems for module's pch files
@@ -1226,12 +1245,12 @@ namespace UnrealBuildTool
 				SharedPchFileItem = FileItem.GetItemByFileReference(FileReference.Combine(ModuleDirectory, Rules.SharedPCHHeaderFile));
 			}
 
-			Dictionary<string, FileItem> NameToFileItem = new();
+			Dictionary<string, FileItem> NameToFileItem = [];
 
-			HashSet<FileItem> CollidingFiles = new();
+			HashSet<FileItem> CollidingFiles = [];
 
 			// Collect the headers that should be built
-			List<FileItem> HeaderFileItems = new();
+			List<FileItem> HeaderFileItems = [];
 			foreach (FileItem HeaderFileItem in InputFiles.HeaderFiles)
 			{
 				// We don't want to build pch files in iwyu, skip those.
@@ -1264,7 +1283,7 @@ namespace UnrealBuildTool
 		}
 
 		// Cache of files found while searching the includes
-		static ConcurrentDictionary<DirectoryReference, HashSet<FileReference>?> KnownIncludeFilesDict = new();
+		static readonly ConcurrentDictionary<DirectoryReference, HashSet<FileReference>?> KnownIncludeFilesDict = new();
 		static HashSet<FileReference>? GetIncludeFiles(DirectoryReference Directory)
 		{
 			if (Directory == null)
@@ -1277,7 +1296,7 @@ namespace UnrealBuildTool
 				{
 					try
 					{
-						return DirectoryLookupCache.EnumerateFiles(Directory).ToHashSet();
+						return [.. DirectoryLookupCache.EnumerateFiles(Directory)];
 					}
 					catch (Exception)
 					{
@@ -1293,6 +1312,17 @@ namespace UnrealBuildTool
 
 		private void FindIncludedHeaders(HashSet<string> VisitedIncludes, HashSet<FileItem> FoundIncludeFileItems, HashSet<string> IncludesNotFound, CppCompileEnvironment CompileEnvironment, FileItem FileToSearch, ILogger Logger)
 		{
+			static FileItem? SearchForFileItem(DirectoryReference dir, string tranformedHeaderInclude)
+			{
+				FileReference FileRef = FileReference.Combine(dir, tranformedHeaderInclude);
+				HashSet<FileReference>? Files = GetIncludeFiles(FileRef.Directory);
+				if (Files != null && Files.Contains(FileRef))
+				{
+					return FileItem.GetItemByFileReference(FileRef);
+				}
+				return null;
+			}
+
 			foreach (string HeaderInclude in CompileEnvironment.MetadataCache.GetHeaderIncludes(FileToSearch))
 			{
 				string TranformedHeaderInclude = HeaderInclude;
@@ -1304,26 +1334,15 @@ namespace UnrealBuildTool
 
 				if (VisitedIncludes.Add(TranformedHeaderInclude))
 				{
-					Func<DirectoryReference, FileItem?> SearchForFileItem = (DirectoryReference dir) =>
-					{
-						FileReference FileRef = FileReference.Combine(dir, TranformedHeaderInclude);
-						HashSet<FileReference>? Files = GetIncludeFiles(FileRef.Directory);
-						if (Files != null && Files.Contains(FileRef))
-						{
-							return FileItem.GetItemByFileReference(FileRef);
-						}
-						return null;
-					};
-
 					// check to see if the file is a relative path first
-					FileItem? IncludeFileItem = SearchForFileItem(FileToSearch.Directory.Location);
+					FileItem? IncludeFileItem = SearchForFileItem(FileToSearch.Directory.Location, TranformedHeaderInclude);
 
 					// search through the include paths if the file isn't relative
 					if (IncludeFileItem == null)
 					{
 						foreach (DirectoryReference IncludePath in CompileEnvironment.UserIncludePaths)
 						{
-							IncludeFileItem = SearchForFileItem(IncludePath);
+							IncludeFileItem = SearchForFileItem(IncludePath, TranformedHeaderInclude);
 							if (IncludeFileItem != null)
 							{
 								break;
@@ -1343,19 +1362,19 @@ namespace UnrealBuildTool
 						if (!TranformedHeaderInclude.Contains('.'))
 						{
 #if DEBUG
-							Logger.LogDebug("{0} SharedPCH - Skipping '{1}' found in '{2}' because it doesn't appear to be a module header.", Name, TranformedHeaderInclude, FileToSearch.Location);
+							Logger.LogDebug("{Name} SharedPCH - Skipping '{Header}' found in '{File}' because it doesn't appear to be a module header.", Name, TranformedHeaderInclude, FileToSearch.Location);
 #endif
 						}
 						else if (TranformedHeaderInclude.EndsWith(".generated.h"))
 						{
 #if DEBUG
-							Logger.LogDebug("{0} SharedPCH - Skipping '{1}' found in '{2}' because it appears to be a generated UHT header.", Name, TranformedHeaderInclude, FileToSearch.Location);
+							Logger.LogDebug("{Name} SharedPCH - Skipping '{Header}' found in '{File}' because it appears to be a generated UHT header.", Name, TranformedHeaderInclude, FileToSearch.Location);
 #endif
 						}
 						else
 						{
 #if DEBUG
-							Logger.LogDebug("{0} SharedPCH - Could not find include directory for '{1}' found in '{2}'.", Name, TranformedHeaderInclude, FileToSearch.Location);
+							Logger.LogDebug("{Name} SharedPCH - Could not find include directory for '{Header}' found in '{File}'.", Name, TranformedHeaderInclude, FileToSearch.Location);
 #endif
 							IncludesNotFound.Add(TranformedHeaderInclude);
 						}
@@ -1373,15 +1392,15 @@ namespace UnrealBuildTool
 		/// <returns>Returns an optimized list of modules this module's shared PCH uses</returns>
 		private HashSet<UEBuildModule> GetSharedPCHModuleDependencies(CppCompileEnvironment CompileEnvironment, FileItem PCHHeaderFile, ILogger Logger)
 		{
-			Dictionary<UEBuildModule, bool> ModuleToIncludePathsOnlyFlag = new Dictionary<UEBuildModule, bool>();
+			Dictionary<UEBuildModule, bool> ModuleToIncludePathsOnlyFlag = [];
 			FindModulesInPublicCompileEnvironment(ModuleToIncludePathsOnlyFlag);
 
-			HashSet<UEBuildModule> AllModuleDeps = ModuleToIncludePathsOnlyFlag.Keys.ToHashSet();
-			HashSet<UEBuildModule> ModuleDeps = new HashSet<UEBuildModule>(AllModuleDeps) { this };
+			HashSet<UEBuildModule> AllModuleDeps = [.. ModuleToIncludePathsOnlyFlag.Keys];
+			HashSet<UEBuildModule> ModuleDeps = [.. AllModuleDeps, this];
 
-			HashSet<string> VisitedIncludes = new();
-			HashSet<FileItem> FoundIncludeFileItems = new();
-			HashSet<string> IncludesNotFound = new();
+			HashSet<string> VisitedIncludes = [];
+			HashSet<FileItem> FoundIncludeFileItems = [];
+			HashSet<string> IncludesNotFound = [];
 
 			// pre-populate all the known file locations for this compile environment
 			foreach (DirectoryReference dir in CompileEnvironment.UserIncludePaths)
@@ -1390,14 +1409,14 @@ namespace UnrealBuildTool
 			}
 
 			FindIncludedHeaders(VisitedIncludes, FoundIncludeFileItems, IncludesNotFound, CompileEnvironment, PCHHeaderFile, Logger);
-			if (IncludesNotFound.Any())
+			if (IncludesNotFound.Count != 0)
 			{
 				// We are assuming that the includes not found aren't important. If this is found to be not true then we should return all the module dependencies
 				//return AllModuleDeps;
 			}
 
 			bool FoundAllModules = true;
-			HashSet<UEBuildModule> OptModules = new();
+			HashSet<UEBuildModule> OptModules = [];
 			foreach (FileItem IncludeFile in FoundIncludeFileItems)
 			{
 				if (CompileEnvironment.MetadataCache.UsesAPIDefine(IncludeFile) || CompileEnvironment.MetadataCache.ContainsReflectionMarkup(IncludeFile))
@@ -1408,7 +1427,7 @@ namespace UnrealBuildTool
 #if DEBUG
 						if (ModuleToIncludePathsOnlyFlag[FoundModule])
 						{
-							Logger.LogDebug("{0} SharedPCH - '{1}' is exporting types but '{2}' isn't declared as a public dependency.", Name, IncludeFile.Location, FoundModule);
+							Logger.LogDebug("{Name} SharedPCH - '{Include}' is exporting types but '{Module}' isn't declared as a public dependency.", Name, IncludeFile.Location, FoundModule);
 						}
 #endif
 
@@ -1417,7 +1436,7 @@ namespace UnrealBuildTool
 					else
 					{
 #if DEBUG
-						Logger.LogDebug("{0} SharedPCH - '{1}' is exporting types but the module this file belongs to isn't declared as a public dependency or include.", Name, IncludeFile.Location);
+						Logger.LogDebug("{Name} SharedPCH - '{Include}' is exporting types but the module this file belongs to isn't declared as a public dependency or include.", Name, IncludeFile.Location);
 #endif
 						FoundAllModules = false;
 					}
@@ -1425,7 +1444,7 @@ namespace UnrealBuildTool
 				else
 				{
 #if DEBUG
-					Logger.LogDebug("{0} SharedPCH - '{1}' is not exporting types so we are ignoring the dependency", Name, IncludeFile.Location);
+					Logger.LogDebug("{Name} SharedPCH - '{Include}' is not exporting types so we are ignoring the dependency", Name, IncludeFile.Location);
 #endif
 				}
 			}
@@ -1433,7 +1452,7 @@ namespace UnrealBuildTool
 			if (!FoundAllModules)
 			{
 #if DEBUG
-				Logger.LogDebug("{0} SharedPCH - Is missing public dependencies. To be safe, this shared PCH will fall back to use all the module dependencies. Note that this could affect compile times.", Name);
+				Logger.LogDebug("{Name} SharedPCH - Is missing public dependencies. To be safe, this shared PCH will fall back to use all the module dependencies. Note that this could affect compile times.", Name);
 #endif
 				return AllModuleDeps;
 			}
@@ -1470,7 +1489,7 @@ namespace UnrealBuildTool
 
 		static HashSet<string> GetImmutableDefinitions(List<string> Definitions)
 		{
-			HashSet<string> ImmutableDefinitions = new();
+			HashSet<string> ImmutableDefinitions = [];
 			foreach (string Definition in Definitions)
 			{
 				if (Definition.Contains("UE_IS_ENGINE_MODULE", StringComparison.Ordinal) ||
@@ -1504,7 +1523,7 @@ namespace UnrealBuildTool
 			FileItem WrapperFile = CreatePCHWrapperFile(WrapperLocation, DefinitionsFileItem!, HeaderFile, Graph);
 
 			// Create a new C++ environment that is used to create the PCH.
-			CppCompileEnvironment CompileEnvironment = new CppCompileEnvironment(ModuleCompileEnvironment);
+			CppCompileEnvironment CompileEnvironment = new(ModuleCompileEnvironment);
 			CompileEnvironment.Definitions.Clear();
 			CompileEnvironment.PrecompiledHeaderAction = PrecompiledHeaderAction.Create;
 			CompileEnvironment.PrecompiledHeaderIncludeFilename = WrapperFile.Location;
@@ -1519,7 +1538,7 @@ namespace UnrealBuildTool
 			}
 			else
 			{
-				Output = ToolChain.CompileAllCPPFiles(CompileEnvironment, new List<FileItem>() { WrapperFile }, IntermediateDirectory, Name, Graph);
+				Output = ToolChain.CompileAllCPPFiles(CompileEnvironment, [WrapperFile], IntermediateDirectory, Name, Graph);
 			}
 
 			return new PrecompiledHeaderInstance(WrapperFile, DefinitionsFileItem, CompileEnvironment, Output, GetImmutableDefinitions(ModuleCompileEnvironment.Definitions));
@@ -1543,7 +1562,7 @@ namespace UnrealBuildTool
 				// Modify definitions if we need to create a new shared pch for the include order
 				if (ModuleCompileEnvironment.IncludeOrderVersion != Template.BaseCompileEnvironment.IncludeOrderVersion)
 				{
-					Definitions = new List<string>(Definitions);
+					Definitions = [.. Definitions];
 					foreach (string OldDefine in EngineIncludeOrderHelper.GetDeprecationDefines(Template.BaseCompileEnvironment.IncludeOrderVersion))
 					{
 						Definitions.Remove(OldDefine);
@@ -1554,7 +1573,7 @@ namespace UnrealBuildTool
 				// Modify definitions if we need to create a new shared pch for project modules
 				if (ModuleCompileEnvironment.bTreatAsEngineModule != Template.BaseCompileEnvironment.bTreatAsEngineModule)
 				{
-					Definitions = new List<string>(Definitions);
+					Definitions = [.. Definitions];
 					Definitions.RemoveAll(x => x.Contains("UE_IS_ENGINE_MODULE", StringComparison.Ordinal));
 					Definitions.RemoveAll(x => x.Contains("UE_DEPRECATED_FORGAME", StringComparison.Ordinal));
 					Definitions.RemoveAll(x => x.Contains("UE_DEPRECATED_FORENGINE", StringComparison.Ordinal));
@@ -1567,7 +1586,7 @@ namespace UnrealBuildTool
 				// Modify definitions if we need to create a new shared pch for validating format strings
 				if (ModuleCompileEnvironment.bValidateFormatStrings != Template.BaseCompileEnvironment.bValidateFormatStrings)
 				{
-					Definitions = new List<string>(Definitions);
+					Definitions = [.. Definitions];
 					Definitions.RemoveAll(x => x.Contains("UE_VALIDATE_FORMAT_STRINGS", StringComparison.Ordinal));
 					Definitions.Add($"UE_VALIDATE_FORMAT_STRINGS={(ModuleCompileEnvironment.bValidateFormatStrings ? "1" : "0")}");
 				}
@@ -1575,7 +1594,7 @@ namespace UnrealBuildTool
 				// Modify definitions if we need to create a new shared pch for validating internal api usage
 				if (ModuleCompileEnvironment.bValidateInternalApi != Template.BaseCompileEnvironment.bValidateInternalApi)
 				{
-					Definitions = new List<string>(Definitions);
+					Definitions = [.. Definitions];
 					Definitions.RemoveAll(x => x.Contains("UE_VALIDATE_INTERNAL_API", StringComparison.Ordinal));
 					Definitions.Add($"UE_VALIDATE_INTERNAL_API={(ModuleCompileEnvironment.bValidateInternalApi ? "1" : "0")}");
 				}
@@ -1583,7 +1602,7 @@ namespace UnrealBuildTool
 				// Modify definitions if we need to create a new shared pch for validating experimental api usage
 				if (ModuleCompileEnvironment.bValidateExperimentalApi != Template.BaseCompileEnvironment.bValidateExperimentalApi)
 				{
-					Definitions = new List<string>(Definitions);
+					Definitions = [.. Definitions];
 					Definitions.RemoveAll(x => x.Contains("UE_VALIDATE_EXPERIMENTAL_API", StringComparison.Ordinal));
 					Definitions.Add($"UE_VALIDATE_EXPERIMENTAL_API={(ModuleCompileEnvironment.bValidateExperimentalApi ? "1" : "0")}");
 				}
@@ -1592,8 +1611,8 @@ namespace UnrealBuildTool
 				string Variant = GetSuffixForSharedPCH(ModuleCompileEnvironment, Template.BaseCompileEnvironment);
 
 				FileReference SharedDefinitionsLocation = FileReference.Combine(Template.OutputDir, $"SharedDefinitions.{Template.Module.Name}{Variant}.h");
-				List<string> NewDefinitions = new();
-				StringBuilder Writer = new StringBuilder();
+				List<string> NewDefinitions = [];
+				StringBuilder Writer = new();
 				WriteDefinitions($"Shared Definitions for {Template.Module.Name}{Variant}", [], Definitions, [], Writer);
 				FileItem SharedDefinitionsFileItem = Graph.CreateIntermediateTextFile(SharedDefinitionsLocation, Writer.ToString(), AllowAsync: false);
 
@@ -1602,7 +1621,7 @@ namespace UnrealBuildTool
 				FileItem WrapperFile = CreatePCHWrapperFile(WrapperLocation, SharedDefinitionsFileItem, Template.HeaderFile, Graph);
 
 				// Create the compile environment for this PCH
-				CppCompileEnvironment CompileEnvironment = new CppCompileEnvironment(Template.BaseCompileEnvironment);
+				CppCompileEnvironment CompileEnvironment = new(Template.BaseCompileEnvironment);
 				CompileEnvironment.Definitions.Clear();
 				CompileEnvironment.PrecompiledHeaderAction = PrecompiledHeaderAction.Create;
 				CompileEnvironment.PrecompiledHeaderIncludeFilename = WrapperFile.Location;
@@ -1613,8 +1632,8 @@ namespace UnrealBuildTool
 				if (Rules.Target.bChainPCHs)
 				{
 					// Find all the dependencies of this module
-					HashSet<UEBuildModule> ReferencedModules = new HashSet<UEBuildModule>();
-					Template.Module.GetAllDependencyModules(new List<UEBuildModule>(), ReferencedModules, bIncludeDynamicallyLoaded: false, bForceCircular: false, bOnlyDirectDependencies: false);
+					HashSet<UEBuildModule> ReferencedModules = [];
+					Template.Module.GetAllDependencyModules([], ReferencedModules, bIncludeDynamicallyLoaded: false, bForceCircular: false, bOnlyDirectDependencies: false);
 
 					PrecompiledHeaderTemplate? ParentTemplate = CompileEnvironment.SharedPCHs
 						.Skip(CompileEnvironment.SharedPCHs.IndexOf(Template) + 1)
@@ -1634,7 +1653,7 @@ namespace UnrealBuildTool
 				}
 				else
 				{
-					Output = ToolChain.CompileAllCPPFiles(CompileEnvironment, new List<FileItem>() { WrapperFile }, Template.OutputDir, "Shared", Graph);
+					Output = ToolChain.CompileAllCPPFiles(CompileEnvironment, [WrapperFile], Template.OutputDir, "Shared", Graph);
 				}
 				Instance = new PrecompiledHeaderInstance(WrapperFile, SharedDefinitionsFileItem, CompileEnvironment, Output, GetImmutableDefinitions(Template.BaseCompileEnvironment.Definitions))
 				{
@@ -1875,14 +1894,13 @@ namespace UnrealBuildTool
 			CppCompileEnvironment ModuleCompileEnvironment,
 			List<FileItem> NormalFiles,
 			List<FileItem> AdaptiveFiles,
-			IActionGraphBuilder Graph,
-			ILogger Logger)
+			IActionGraphBuilder Graph)
 		{
 			bool isEngine = Rules.File.IsUnderDirectory(Unreal.EngineDirectory) || Rules.bTreatAsEngineModule || Rules.PrivatePCHHeaderFile == null;
 			bool bAdaptiveUnityDisablesPCH = isEngine ? Target.bAdaptiveUnityDisablesPCH : Target.bAdaptiveUnityDisablesPCHForProject;
 			string adaptiveUnityDisablesPCHName = isEngine ? nameof(Target.bAdaptiveUnityDisablesPCH) : nameof(Target.bAdaptiveUnityDisablesPCHForProject);
 
-			CPPOutput OutputFiles = new CPPOutput();
+			CPPOutput OutputFiles = new();
 
 			if (NormalFiles.Count > 0)
 			{
@@ -1912,11 +1930,11 @@ namespace UnrealBuildTool
 					Graph.AddDiagnostic("[Adaptive Build] Enabling Edit & Continue for excluded files. Set bAdaptiveUnityEnablesEditAndContinue to false in BuildConfiguration.xml to change this behavior.");
 				}
 
-				Graph.AddDiagnostic($"[Adaptive Build] Excluded from {Name} unity file: " + String.Join(", ", AdaptiveFiles.Select(File => Path.GetFileName(File.AbsolutePath))));
+				Graph.AddDiagnostic($"[Adaptive Build] Excluded from {Name} unity file: {String.Join(", ", AdaptiveFiles.Select(File => Path.GetFileName(File.AbsolutePath)))}");
 
 				// Create the new compile environment. Always turn off PCH due to different compiler settings.
-				CppCompileEnvironment AdaptiveUnityEnvironment = new CppCompileEnvironment(ModuleCompileEnvironment);
-				AdaptiveUnityEnvironment.AdditionalPrerequisites = new(AdaptiveUnityEnvironment.AdditionalPrerequisites.Concat(CompileEnvironment.AdditionalPrerequisites).Distinct());
+				CppCompileEnvironment AdaptiveUnityEnvironment = new(ModuleCompileEnvironment);
+				AdaptiveUnityEnvironment.AdditionalPrerequisites = [.. AdaptiveUnityEnvironment.AdditionalPrerequisites.Concat(CompileEnvironment.AdditionalPrerequisites).Distinct()];
 				AdaptiveUnityEnvironment.UserIncludePaths.UnionWith(CompileEnvironment.SharedUserIncludePaths);
 				AdaptiveUnityEnvironment.bDisableStaticAnalysis = CompileEnvironment.bDisableStaticAnalysis;
 				if (Target.bAdaptiveUnityDisablesOptimizations)
@@ -1978,11 +1996,11 @@ namespace UnrealBuildTool
 
 		static CPPOutput CompileAdaptiveNonUnityFilesWithDedicatedPCH(UEToolChain ToolChain, CppCompileEnvironment CompileEnvironment, List<FileItem> Files, DirectoryReference IntermediateDirectory, string ModuleName, IActionGraphBuilder Graph)
 		{
-			CPPOutput Output = new CPPOutput();
+			CPPOutput Output = new();
 			foreach (FileItem File in Files)
 			{
 				// Build the contents of the wrapper file
-				StringBuilder WrapperContents = new StringBuilder();
+				StringBuilder WrapperContents = new();
 				{
 					string FileString = File.AbsolutePath;
 					if (File.Location.IsUnderDirectory(Unreal.RootDirectory))
@@ -1990,12 +2008,10 @@ namespace UnrealBuildTool
 						FileString = File.Location.MakeRelativeTo(Unreal.EngineSourceDirectory);
 					}
 					FileString = FileString.Replace('\\', '/');
-					WriteDefinitions($"Dedicated PCH for {FileString}", Array.Empty<string>(), CompileEnvironment.Definitions, Array.Empty<string>(), WrapperContents);
+					WriteDefinitions($"Dedicated PCH for {FileString}", [], CompileEnvironment.Definitions, [], WrapperContents);
 					WrapperContents.AppendLine();
-					using (StreamReader Reader = new StreamReader(File.Location.FullName))
-					{
-						CppIncludeParser.CopyIncludeDirectives(Reader, WrapperContents);
-					}
+					using StreamReader Reader = new(File.Location.FullName);
+					CppIncludeParser.CopyIncludeDirectives(Reader, WrapperContents);
 				}
 
 				// Write the PCH header
@@ -2003,25 +2019,25 @@ namespace UnrealBuildTool
 				FileItem DedicatedPchFile = Graph.CreateIntermediateTextFile(DedicatedPchLocation, WrapperContents.ToString());
 
 				// Create a new C++ environment to compile the PCH
-				CppCompileEnvironment PchEnvironment = new CppCompileEnvironment(CompileEnvironment);
+				CppCompileEnvironment PchEnvironment = new(CompileEnvironment);
 				PchEnvironment.Definitions.Clear();
 				PchEnvironment.UserIncludePaths.Add(File.Location.Directory); // Need to be able to include headers in the same directory as the source file
 				PchEnvironment.PrecompiledHeaderAction = PrecompiledHeaderAction.Create;
 				PchEnvironment.PrecompiledHeaderIncludeFilename = DedicatedPchFile.Location;
 
 				// Create the action to compile the PCH file.
-				CPPOutput PchOutput = ToolChain.CompileAllCPPFiles(PchEnvironment, new List<FileItem>() { DedicatedPchFile }, IntermediateDirectory, ModuleName, Graph);
+				CPPOutput PchOutput = ToolChain.CompileAllCPPFiles(PchEnvironment, [DedicatedPchFile], IntermediateDirectory, ModuleName, Graph);
 				Output.ObjectFiles.AddRange(PchOutput.ObjectFiles);
 
 				// Create a new C++ environment to compile the original file
-				CppCompileEnvironment FileEnvironment = new CppCompileEnvironment(CompileEnvironment);
+				CppCompileEnvironment FileEnvironment = new(CompileEnvironment);
 				FileEnvironment.Definitions.Clear();
 				FileEnvironment.PrecompiledHeaderAction = PrecompiledHeaderAction.Include;
 				FileEnvironment.PrecompiledHeaderIncludeFilename = DedicatedPchFile.Location;
 				FileEnvironment.PCHInstance = new PrecompiledHeaderInstance(DedicatedPchFile, DedicatedPchFile, PchEnvironment, PchOutput, GetImmutableDefinitions(CompileEnvironment.Definitions));
 
 				// Create the action to compile the PCH file.
-				CPPOutput FileOutput = ToolChain.CompileAllCPPFiles(FileEnvironment, new List<FileItem>() { File }, IntermediateDirectory, ModuleName, Graph);
+				CPPOutput FileOutput = ToolChain.CompileAllCPPFiles(FileEnvironment, [File], IntermediateDirectory, ModuleName, Graph);
 				Output.ObjectFiles.AddRange(FileOutput.ObjectFiles);
 			}
 			return Output;
@@ -2035,8 +2051,8 @@ namespace UnrealBuildTool
 		/// <param name="CompileEnvironment">The current compile environment</param>
 		/// <param name="LinkInputFiles">List of files that will be linked for the target</param>
 		/// <param name="Graph">List of build actions</param>
-		/// <param name="Logger">Logger for output</param>
-		CppCompileEnvironment SetupPrecompiledHeaders(ReadOnlyTargetRules Target, UEToolChain? ToolChain, CppCompileEnvironment CompileEnvironment, List<FileItem> LinkInputFiles, IActionGraphBuilder Graph, ILogger Logger)
+		/// 
+		CppCompileEnvironment SetupPrecompiledHeaders(ReadOnlyTargetRules Target, UEToolChain? ToolChain, CppCompileEnvironment CompileEnvironment, List<FileItem> LinkInputFiles, IActionGraphBuilder Graph)
 		{
 			if (Target.bUsePCHFiles && Rules.PCHUsage != ModuleRules.PCHUsageMode.NoPCHs)
 			{
@@ -2070,8 +2086,8 @@ namespace UnrealBuildTool
 						FileItem PrivateDefinitionsFileItem;
 						{
 							// Only add new definitions that are not already existing in the shared pch
-							List<string> NewDefinitions = new();
-							List<string> Undefintions = new();
+							List<string> NewDefinitions = [];
+							List<string> Undefinitions = [];
 							bool ModuleApiUndef = false;
 							foreach (string Definition in CompileEnvironment.Definitions)
 							{
@@ -2086,12 +2102,12 @@ namespace UnrealBuildTool
 								if (!ModuleApiUndef && Definition.StartsWith(ModuleApiDefine, StringComparison.Ordinal))
 								{
 									ModuleApiUndef = true;
-									Undefintions.Add(ModuleApiDefine);
+									Undefinitions.Add(ModuleApiDefine);
 								}
 							}
 
 							StringBuilder Writer = new();
-							WriteDefinitions($"Shared PCH Definitions for {Name}", new string[] { Instance.DefinitionsFile.Name }, NewDefinitions, Undefintions, Writer);
+							WriteDefinitions($"Shared PCH Definitions for {Name}", [Instance.DefinitionsFile.Name], NewDefinitions, Undefinitions, Writer);
 							PrivateDefinitionsFileItem = Graph.CreateIntermediateTextFile(PrivateDefinitionsFile, Writer.ToString());
 						}
 
@@ -2133,7 +2149,7 @@ namespace UnrealBuildTool
 				FileReference PrivateDefinitionsFile = FileReference.Combine(IntermediateDirectory, PrivateDefinitionsName);
 				{
 					StringBuilder Writer = new();
-					WriteDefinitions("Definitions", Array.Empty<string>(), CompileEnvironment.Definitions, Array.Empty<string>(), Writer);
+					WriteDefinitions("Definitions", [], CompileEnvironment.Definitions, [], Writer);
 					CompileEnvironment.Definitions.Clear();
 
 					FileItem PrivateDefinitionsFileItem = Graph.CreateIntermediateTextFile(PrivateDefinitionsFile, Writer.ToString());
@@ -2157,7 +2173,7 @@ namespace UnrealBuildTool
 		static FileItem CreatePCHWrapperFile(FileReference OutputFile, FileItem DefinitionsFile, FileItem IncludedFile, IActionGraphBuilder Graph)
 		{
 			// Build the contents of the wrapper file
-			StringBuilder WrapperContents = new StringBuilder();
+			StringBuilder WrapperContents = new();
 			{
 				WrapperContents.AppendLine("// PCH for {0}", IncludedFile.Name);
 				WrapperContents.AppendLine("#include \"{0}\"", DefinitionsFile.Name);
@@ -2200,7 +2216,7 @@ namespace UnrealBuildTool
 				Writer.Append("#undef ").Append(Undefinition).AppendLine();
 			}
 
-			HashSet<string> Processed = new();
+			HashSet<string> Processed = [];
 			foreach (string Definition in Definitions.Where(x => Processed.Add(x)))
 			{
 				int EqualsIdx = Definition.IndexOf('=', StringComparison.Ordinal);
@@ -2229,14 +2245,14 @@ namespace UnrealBuildTool
 				if (InvalidIncludeDirectives == null)
 				{
 					// Find headers used by the source file.
-					Dictionary<string, FileReference> NameToHeaderFile = new Dictionary<string, FileReference>();
+					Dictionary<string, FileReference> NameToHeaderFile = [];
 					foreach (FileItem HeaderFile in HeaderFiles)
 					{
 						NameToHeaderFile[HeaderFile.Location.GetFileNameWithoutExtension()] = HeaderFile.Location;
 					}
 
 					// Find the directly included files for each source file, and make sure it includes the matching header if possible
-					InvalidIncludeDirectives = new();
+					InvalidIncludeDirectives = [];
 					if (Rules != null && Rules.IWYUSupport != IWYUSupport.None && Target.bEnforceIWYU)
 					{
 						foreach (FileItem CppFile in CppFiles)
@@ -2248,8 +2264,7 @@ namespace UnrealBuildTool
 								string ExpectedName = CppFile.Location.GetFileNameWithoutExtension();
 								if (!String.Equals(IncludeName, ExpectedName, StringComparison.OrdinalIgnoreCase))
 								{
-									FileReference? HeaderFile;
-									if (NameToHeaderFile.TryGetValue(ExpectedName, out HeaderFile))
+									if (NameToHeaderFile.TryGetValue(ExpectedName, out FileReference? HeaderFile))
 									{
 										InvalidIncludeDirectives.Add(new InvalidIncludeDirective(CppFile.Location, HeaderFile));
 									}
@@ -2419,24 +2434,19 @@ namespace UnrealBuildTool
 			{
 				return false;
 			}
-			switch (Setting)
+			return Setting switch
 			{
-				case ModuleRules.CodeOptimization.Never:
-					return false;
-				case ModuleRules.CodeOptimization.Default:
-				case ModuleRules.CodeOptimization.InNonDebugBuilds:
-					return Configuration != UnrealTargetConfiguration.Debug && (Configuration != UnrealTargetConfiguration.DebugGame || bIsEngineModule);
-				case ModuleRules.CodeOptimization.InShippingBuildsOnly:
-					return (Configuration == UnrealTargetConfiguration.Shipping);
-				default:
-					return true;
-			}
+				ModuleRules.CodeOptimization.Never => false,
+				ModuleRules.CodeOptimization.Default or ModuleRules.CodeOptimization.InNonDebugBuilds => Configuration != UnrealTargetConfiguration.Debug && (Configuration != UnrealTargetConfiguration.DebugGame || bIsEngineModule),
+				ModuleRules.CodeOptimization.InShippingBuildsOnly => (Configuration == UnrealTargetConfiguration.Shipping),
+				_ => true,
+			};
 		}
 
 		public CppCompileEnvironment CreateCompileEnvironmentForIntellisense(ReadOnlyTargetRules Target, CppCompileEnvironment BaseCompileEnvironment, ILogger Logger)
 		{
 			CppCompileEnvironment CompileEnvironment = CreateModuleCompileEnvironment(Target, BaseCompileEnvironment, Logger);
-			CompileEnvironment = SetupPrecompiledHeaders(Target, null, CompileEnvironment, new List<FileItem>(), new NullActionGraphBuilder(Logger), Logger);
+			CompileEnvironment = SetupPrecompiledHeaders(Target, null, CompileEnvironment, [], new NullActionGraphBuilder(Logger));
 			CreateHeaderForDefinitions(CompileEnvironment, IntermediateDirectory, null, new NullActionGraphBuilder(Logger));
 			return CompileEnvironment;
 		}
@@ -2450,11 +2460,12 @@ namespace UnrealBuildTool
 		/// <returns>The new module compile environment.</returns>
 		public CppCompileEnvironment CreateModuleCompileEnvironment(ReadOnlyTargetRules Target, CppCompileEnvironment BaseCompileEnvironment, ILogger Logger)
 		{
-			CppCompileEnvironment Result = new CppCompileEnvironment(BaseCompileEnvironment);
-
-			// Override compile environment
-			Result.bUseUnity = Rules.bUseUnity;
-			Result.bCodeCoverage = Target.bCodeCoverage;
+			CppCompileEnvironment Result = new(BaseCompileEnvironment)
+			{
+				// Override compile environment
+				bUseUnity = Rules.bUseUnity,
+				bCodeCoverage = Target.bCodeCoverage
+			};
 			Result.bOptimizeCode = ShouldEnableOptimization(Rules.OptimizeCode, Target.Configuration, Rules.bTreatAsEngineModule, Result.bCodeCoverage);
 			Result.bUseRTTI |= Rules.bUseRTTI;
 			Result.bVcRemoveUnreferencedComdat &= Rules.bVcRemoveUnreferencedComdat;
@@ -2679,19 +2690,20 @@ namespace UnrealBuildTool
 		/// <returns>The new shared PCH compile environment.</returns>
 		public CppCompileEnvironment CreateSharedPCHCompileEnvironment(UEBuildTarget Target, CppCompileEnvironment BaseCompileEnvironment)
 		{
-			CppCompileEnvironment CompileEnvironment = new CppCompileEnvironment(BaseCompileEnvironment);
-
-			// Use the default optimization setting for
-			CompileEnvironment.bOptimizeCode = ShouldEnableOptimization(ModuleRules.CodeOptimization.Default, Target.Configuration, Rules.bTreatAsEngineModule, Rules.bCodeCoverage);
-			CompileEnvironment.bCodeCoverage = Rules.bCodeCoverage;
-			CompileEnvironment.bTreatAsEngineModule = Rules.bTreatAsEngineModule;
-			CompileEnvironment.bValidateFormatStrings = Rules.bValidateFormatStrings;
-			CompileEnvironment.bValidateInternalApi = Rules.bValidateInternalApi;
-			CompileEnvironment.bValidateExperimentalApi = Rules.bValidateExperimentalApi;
+			CppCompileEnvironment CompileEnvironment = new(BaseCompileEnvironment)
+			{
+				// Use the default optimization setting for
+				bOptimizeCode = ShouldEnableOptimization(ModuleRules.CodeOptimization.Default, Target.Configuration, Rules.bTreatAsEngineModule, Rules.bCodeCoverage),
+				bCodeCoverage = Rules.bCodeCoverage,
+				bTreatAsEngineModule = Rules.bTreatAsEngineModule,
+				bValidateFormatStrings = Rules.bValidateFormatStrings,
+				bValidateInternalApi = Rules.bValidateInternalApi,
+				bValidateExperimentalApi = Rules.bValidateExperimentalApi,
 
-			// Override compile environment
-			CompileEnvironment.bIsBuildingDLL = !Target.ShouldCompileMonolithic();
-			CompileEnvironment.bIsBuildingLibrary = false;
+				// Override compile environment
+				bIsBuildingDLL = !Target.ShouldCompileMonolithic(),
+				bIsBuildingLibrary = false
+			};
 
 			// Add a macro for when we're compiling an engine module, to enable additional compiler diagnostics through code.
 			CompileEnvironment.Definitions.Add($"UE_IS_ENGINE_MODULE={(Rules.bTreatAsEngineModule || Rules.Target.bWarnAboutMonolithicHeadersIncluded ? "1" : "0")}");
@@ -2712,7 +2724,7 @@ namespace UnrealBuildTool
 			CompileEnvironment.Definitions.AddRange(PublicDefinitions);
 
 			// Find all the modules that are part of the public compile environment for this module.
-			Dictionary<UEBuildModule, bool> ModuleToIncludePathsOnlyFlag = new Dictionary<UEBuildModule, bool>();
+			Dictionary<UEBuildModule, bool> ModuleToIncludePathsOnlyFlag = [];
 			FindModulesInPublicCompileEnvironment(ModuleToIncludePathsOnlyFlag);
 
 			// Set up additional root paths
@@ -2742,9 +2754,9 @@ namespace UnrealBuildTool
 		{
 			IReadOnlySet<string> ExcludedNames = UEBuildPlatform.GetBuildPlatform(Platform).GetExcludedFolderNames();
 
-			InputFileCollection InputFiles = new InputFileCollection();
+			InputFileCollection InputFiles = new();
 
-			SourceDirectories = new HashSet<DirectoryReference>();
+			SourceDirectories = [];
 			foreach (DirectoryReference Dir in ModuleDirectories)
 			{
 				DirectoryItem ModuleDirectoryItem = DirectoryItem.GetItemByDirectoryReference(Dir);
@@ -2765,8 +2777,7 @@ namespace UnrealBuildTool
 		/// <param name="Logger">Logger for output</param>
 		static void FindInputFilesFromDirectoryRecursive(DirectoryItem BaseDirectory, IReadOnlySet<string> ExcludedNames, HashSet<DirectoryReference> SourceDirectories, Dictionary<DirectoryItem, FileItem[]> DirectoryToSourceFiles, InputFileCollection InputFiles, ILogger Logger)
 		{
-			bool bIgnoreFileFound;
-			FileItem[] SourceFiles = FindInputFilesFromDirectory(BaseDirectory, InputFiles, out bIgnoreFileFound, Logger);
+			FileItem[] SourceFiles = FindInputFilesFromDirectory(BaseDirectory, InputFiles, out bool bIgnoreFileFound, Logger);
 
 			if (bIgnoreFileFound)
 			{
@@ -2790,14 +2801,14 @@ namespace UnrealBuildTool
 
 		// List of known extensions that are checked in FindInputFilesFromDirectory, so we can log out only if a hidden file
 		// is skipped, only if it has an extension that would normally be compiled
-		static string[] KnownInputFileExtensions =
-		{
+		static readonly string[] KnownInputFileExtensions =
+		[
 			".h", ".ipsh",
 			".cpp", ".cxx", ".ixx", ".c", ".cc",
 			".m", ".mm", ".swift",
 			".rc", ".ispc",
 			".proto",
-		};
+		];
 
 		/// <summary>
 		/// Finds the input files that should be built for this module, from a given directory
@@ -2810,7 +2821,7 @@ namespace UnrealBuildTool
 		static FileItem[] FindInputFilesFromDirectory(DirectoryItem BaseDirectory, InputFileCollection InputFiles, out bool bIgnoreFileFound, ILogger Logger)
 		{
 			bIgnoreFileFound = false;
-			List<FileItem> SourceFiles = new List<FileItem>();
+			List<FileItem> SourceFiles = [];
 			foreach (FileItem InputFile in BaseDirectory.EnumerateFiles())
 			{
 				if (InputFile.Name == ".ubtignore")
@@ -2819,7 +2830,7 @@ namespace UnrealBuildTool
 				}
 				// some text editors will leave temp files that start with . next to a, say, .cpp file, and we do not want to compile them,
 				// but skip files like .DS_Store that don't have an extension (so, no more .'s after the first)
-				else if (InputFile.Name.StartsWith(".") && InputFile.Name.IndexOf('.', 1) != -1)
+				else if (InputFile.Name.StartsWith('.') && InputFile.Name.IndexOf('.', 1) != -1)
 				{
 					if (KnownInputFileExtensions.Contains(InputFile.Location.GetExtension()))
 					{
@@ -2890,19 +2901,19 @@ namespace UnrealBuildTool
 				}
 				// if you add any extension checks here, make sure to update KnownInputFileExtensions
 			}
-			return SourceFiles.ToArray();
+			return [.. SourceFiles];
 		}
 
 		public static List<string> GenerateHeaderCpp(string HeaderName, string IncludeFileString)
 		{
-			HeaderName = HeaderName.Substring(0, HeaderName.Length - 2).Replace(".", "_").Replace("-", "_");
-
-			List<string> GeneratedHeaderCppContents = new();
-			GeneratedHeaderCppContents.Add("// This file is automatically generated at compile-time to include a user created header file.");
-			GeneratedHeaderCppContents.Add("#define UE_DIRECT_HEADER_COMPILE 1");
-			GeneratedHeaderCppContents.Add($"#define __COMPILING_{HeaderName}");
-			GeneratedHeaderCppContents.Add($"#include \"{IncludeFileString.Replace('\\', '/')}\"");
-			return GeneratedHeaderCppContents;
+			HeaderName = HeaderName[..^2].Replace(".", "_").Replace("-", "_");
+
+			return [
+				"// This file is automatically generated at compile-time to include a user created header file.",
+				"#define UE_DIRECT_HEADER_COMPILE 1",
+				$"#define __COMPILING_{HeaderName}",
+				$"#include \"{IncludeFileString.Replace('\\', '/')}\"",
+			];
 		}
 
 		/// <summary>
@@ -2913,11 +2924,10 @@ namespace UnrealBuildTool
 		/// <returns>Array of source files</returns>
 		public static FileItem[] GetSourceFiles(DirectoryItem Directory, ILogger Logger)
 		{
-			bool bIgnoreFileFound;
-			FileItem[] Files = FindInputFilesFromDirectory(Directory, new InputFileCollection(), out bIgnoreFileFound, Logger);
+			FileItem[] Files = FindInputFilesFromDirectory(Directory, new InputFileCollection(), out bool bIgnoreFileFound, Logger);
 			if (bIgnoreFileFound)
 			{
-				return Array.Empty<FileItem>();
+				return [];
 			}
 			return Files;
 		}
@@ -2930,50 +2940,27 @@ namespace UnrealBuildTool
 		/// <returns>Array of source files and headers</returns>
 		public static FileItem[] GetSourceFilesAndHeaders(DirectoryItem Directory, ILogger Logger)
 		{
-			bool bIgnoreFileFound;
-			InputFileCollection InputFiles = new InputFileCollection();
-			List<FileItem> Files = new List<FileItem>(FindInputFilesFromDirectory(Directory, InputFiles, out bIgnoreFileFound, Logger));
+			InputFileCollection InputFiles = new();
+			List<FileItem> Files = [.. FindInputFilesFromDirectory(Directory, InputFiles, out bool bIgnoreFileFound, Logger)];
 			if (bIgnoreFileFound)
 			{
-				return Array.Empty<FileItem>();
+				return [];
 			}
 			Files.AddRange(InputFiles.HeaderFiles);
-			return Files.ToArray();
+			return [.. Files];
 		}
 
 		/// <summary>
 		/// Checks a given directory path whether it exists and if it contains any Verse source files
 		/// </summary>
-		public static bool IsValidVerseDirectory(DirectoryReference MaybeVerseDirectory)
-		{
-			if (!DirectoryReference.Exists(MaybeVerseDirectory))
-			{
-				return false;
-			}
-
-			foreach (string FilePath in Directory.EnumerateFiles(MaybeVerseDirectory.FullName, "*.v*", SearchOption.AllDirectories))
-			{
-				if (FilePath.EndsWith(".verse") || FilePath.EndsWith(".vmodule"))
-				{
-					return true;
-				}
-			}
+		public static bool IsValidVerseDirectory(DirectoryReference MaybeVerseDirectory) => DirectoryReference.Exists(MaybeVerseDirectory)
+			&& Directory.EnumerateFiles(MaybeVerseDirectory.FullName, "*.v*", SearchOption.AllDirectories).Any(x => x.EndsWith(".verse") || x.EndsWith(".vmodule"));
 
-			return false;
-		}
-		
 
 		/// <summary>
 		/// Checks a given directory path whether it exists and if it contains any Proto source files
 		/// </summary>
-		public static bool IsValidProtoDirectory(DirectoryReference MaybeProtoDirectory)
-		{
-			if (!DirectoryReference.Exists(MaybeProtoDirectory))
-			{
-				return false;
-			}
-
-			return Directory.EnumerateFiles(MaybeProtoDirectory.FullName, "*.proto", SearchOption.AllDirectories).Any();
-		}
+		public static bool IsValidProtoDirectory(DirectoryReference MaybeProtoDirectory) => DirectoryReference.Exists(MaybeProtoDirectory)
+			&& Directory.EnumerateFiles(MaybeProtoDirectory.FullName, "*.proto", SearchOption.AllDirectories).Any();
 	}
 }
diff --git a/Engine/Source/Programs/UnrealBuildTool/Modes/AnalyzeMode.cs b/Engine/Source/Programs/UnrealBuildTool/Modes/AnalyzeMode.cs
--- a/Engine/Source/Programs/UnrealBuildTool/Modes/AnalyzeMode.cs
+++ b/Engine/Source/Programs/UnrealBuildTool/Modes/AnalyzeMode.cs
@@ -27,7 +27,7 @@ namespace UnrealBuildTool.Modes
 			Arguments.ApplyTo(this);
 
 			// Create the build configuration object, and read the settings
-			BuildConfiguration BuildConfiguration = new BuildConfiguration();
+			BuildConfiguration BuildConfiguration = new();
 			XmlConfig.ApplyTo(BuildConfiguration);
 			Arguments.ApplyTo(BuildConfiguration);
 
@@ -38,7 +38,7 @@ namespace UnrealBuildTool.Modes
 			using (ISourceFileWorkingSet WorkingSet = new EmptySourceFileWorkingSet())
 			{
 				// Find the compile commands for each file in the target
-				Dictionary<FileReference, string> FileToCommand = new Dictionary<FileReference, string>();
+				Dictionary<FileReference, string> FileToCommand = [];
 				foreach (TargetDescriptor TargetDescriptor in TargetDescriptors)
 				{
 					AnalyzeTarget(TargetDescriptor, BuildConfiguration, Logger);
@@ -51,15 +51,15 @@ namespace UnrealBuildTool.Modes
 		class ModuleInfo
 		{
 			public UEBuildModule Module;
-			public IReadOnlyList<string> IncludeChain;
-			public HashSet<UEBuildModule> InwardRefs = new HashSet<UEBuildModule>();
-			public HashSet<UEBuildModule> UniqueInwardRefs = new HashSet<UEBuildModule>();
-			public HashSet<UEBuildModule> OutwardRefs = new HashSet<UEBuildModule>();
-			public HashSet<UEBuildModule> UniqueOutwardRefs = new HashSet<UEBuildModule>();
+			public string[] IncludeChain;
+			public HashSet<UEBuildModule> InwardRefs = [];
+			public HashSet<UEBuildModule> UniqueInwardRefs = [];
+			public HashSet<UEBuildModule> OutwardRefs = [];
+			public HashSet<UEBuildModule> UniqueOutwardRefs = [];
 
-			public List<FileReference> ObjectFiles = new List<FileReference>();
+			public List<FileReference> ObjectFiles = [];
 			public long ObjSize = 0;
-			public List<FileReference> BinaryFiles = new List<FileReference>();
+			public List<FileReference> BinaryFiles = [];
 			public long BinSize = 0;
 
 			public ModuleInfo(UEBuildModule Module, params string[] IncludeChain)
@@ -81,22 +81,22 @@ namespace UnrealBuildTool.Modes
 
 			UEBuildModule Module = Target.GetModuleByName(ModuleName);
 
-			List<string> CurrentChain = new(ParentChain);
-			CurrentChain.Add(Module.Name);
+			List<string> CurrentChain = [.. ParentChain, Module.Name];
 
-			if (!ModuleToInfo.ContainsKey(Module))
+			if (!ModuleToInfo.TryGetValue(Module, out ModuleInfo? moduleInfo))
 			{
-				ModuleToInfo[Module] = new ModuleInfo(Module, CurrentChain.ToArray());
+				moduleInfo = new ModuleInfo(Module, [.. CurrentChain]);
+				ModuleToInfo[Module] = moduleInfo;
 			}
 
-			if (ModuleToInfo[Module].IncludeChain.Count > CurrentChain.Count)
+			if (moduleInfo.IncludeChain.Length > CurrentChain.Count)
 			{
 				// Now we need to recheck all downstream dependencies again because the chain may be shorter
-				List<UEBuildModule> RecheckModules = new();
-				Module.GetAllDependencyModules(RecheckModules, new(), true, false, true);
+				List<UEBuildModule> RecheckModules = [];
+				Module.GetAllDependencyModules(RecheckModules, [], true, false, true);
 				Visited.ExceptWith(RecheckModules);
-				Logger.LogDebug("Found shorter chain for {Module} {Prev} -> {New}, rechecking {Count} already visited dependencies", Module.Name, ModuleToInfo[Module].IncludeChain.Count, CurrentChain.Count, RecheckModules.Count);
-				ModuleToInfo[Module].IncludeChain = CurrentChain.ToArray();
+				Logger.LogDebug("Found shorter chain for {Module} {Prev} -> {New}, rechecking {Count} already visited dependencies", Module.Name, moduleInfo.IncludeChain.Length, CurrentChain.Count, RecheckModules.Count);
+				moduleInfo.IncludeChain = [.. CurrentChain];
 			}
 			else if (Visited.Contains(Module))
 			{
@@ -105,8 +105,8 @@ namespace UnrealBuildTool.Modes
 
 			Visited.Add(Module);
 
-			List<UEBuildModule> TargetModules = new();
-			Module.GetAllDependencyModules(TargetModules, new(), true, false, true);
+			List<UEBuildModule> TargetModules = [];
+			Module.GetAllDependencyModules(TargetModules, [], true, false, true);
 			TargetModules.ForEach(x => AnalyzeModuleChains(x.Name, CurrentChain, Target, ModuleToInfo, Visited, Logger));
 		}
 
@@ -117,17 +117,17 @@ namespace UnrealBuildTool.Modes
 			DirectoryReference.CreateDirectory(Target.ReceiptFileName.Directory);
 
 			// Find the shortest path from the target to each module
-			HashSet<UEBuildModule> Visited = new();
-			Dictionary<UEBuildModule, ModuleInfo> ModuleToInfo = new Dictionary<UEBuildModule, ModuleInfo>();
+			HashSet<UEBuildModule> Visited = [];
+			Dictionary<UEBuildModule, ModuleInfo> ModuleToInfo = [];
 
 			if (Target.Rules.LaunchModuleName != null)
 			{
-				AnalyzeModuleChains(Target.Rules.LaunchModuleName, new List<string>() { "target" }, Target, ModuleToInfo, Visited, Logger);
+				AnalyzeModuleChains(Target.Rules.LaunchModuleName, ["target"], Target, ModuleToInfo, Visited, Logger);
 			}
 
 			foreach (string RootModuleName in Target.Rules.ExtraModuleNames)
 			{
-				AnalyzeModuleChains(RootModuleName, new List<string>() { "target" }, Target, ModuleToInfo, Visited, Logger);
+				AnalyzeModuleChains(RootModuleName, ["target"], Target, ModuleToInfo, Visited, Logger);
 			}
 
 			// Also enable all the plugin modules
@@ -137,7 +137,7 @@ namespace UnrealBuildTool.Modes
 				{
 					if (!ModuleToInfo.ContainsKey(Module))
 					{
-						AnalyzeModuleChains(Module.Name, new List<string>() { "target", Plugin.ReferenceChain, Plugin.File.GetFileName() }, Target, ModuleToInfo, Visited, Logger);
+						AnalyzeModuleChains(Module.Name, ["target", Plugin.ReferenceChain, Plugin.File.GetFileName()], Target, ModuleToInfo, Visited, Logger);
 					}
 				}
 			}
@@ -151,7 +151,7 @@ namespace UnrealBuildTool.Modes
 						if (!ModuleToInfo.ContainsKey(Module))
 						{
 							// quick hack to make allmodules always worse (empty entries are ignored when writing)
-							List<string> IncludeChain = Enumerable.Repeat(String.Empty, 1000).ToList();
+							List<string> IncludeChain = [.. Enumerable.Repeat(String.Empty, 1000)];
 							IncludeChain.Add("allmodules option");
 							if (Module.Rules.Plugin != null)
 							{
@@ -171,21 +171,20 @@ namespace UnrealBuildTool.Modes
 			foreach ((UEBuildModule SourceModule, ModuleInfo SourceModuleInfo) in ModuleToInfo)
 			{
 				SourceModuleInfo.OutwardRefs.Add(SourceModule);
-				SourceModule.GetAllDependencyModules(new List<UEBuildModule>(), SourceModuleInfo.OutwardRefs, false, false, false);
+				SourceModule.GetAllDependencyModules([], SourceModuleInfo.OutwardRefs, false, false, false);
 				SourceModuleInfo.OutwardRefs.Remove(SourceModule);
 			}
 
 			// Find the direct output dependencies of each module
 			foreach ((UEBuildModule SourceModule, ModuleInfo SourceModuleInfo) in ModuleToInfo)
 			{
-				SourceModuleInfo.UniqueOutwardRefs = new HashSet<UEBuildModule>(SourceModuleInfo.OutwardRefs);
+				SourceModuleInfo.UniqueOutwardRefs = [.. SourceModuleInfo.OutwardRefs];
 
 				foreach (UEBuildModule TargetModule in SourceModuleInfo.OutwardRefs)
 				{
-					HashSet<UEBuildModule> VisitedTargetModules = new HashSet<UEBuildModule>();
-					VisitedTargetModules.Add(SourceModule);
+					HashSet<UEBuildModule> VisitedTargetModules = [SourceModule];
 
-					List<UEBuildModule> DependencyModules = new List<UEBuildModule>();
+					List<UEBuildModule> DependencyModules = [];
 					TargetModule.GetAllDependencyModules(DependencyModules, VisitedTargetModules, false, false, false);
 					DependencyModules.Remove(TargetModule);
 
@@ -222,7 +221,7 @@ namespace UnrealBuildTool.Modes
 				}
 			}
 
-			HashSet<UEBuildModule> MissingModules = new HashSet<UEBuildModule>();
+			HashSet<UEBuildModule> MissingModules = [];
 			foreach (UEBuildBinary Binary in Target.Binaries)
 			{
 				long BinSize = 0;
@@ -236,8 +235,7 @@ namespace UnrealBuildTool.Modes
 				}
 				foreach (UEBuildModule Module in Binary.Modules)
 				{
-					ModuleInfo? ModuleInfo;
-					if (!ModuleToInfo.TryGetValue(Module, out ModuleInfo))
+					if (!ModuleToInfo.TryGetValue(Module, out ModuleInfo? ModuleInfo))
 					{
 						MissingModules.Add(Module);
 						continue;
@@ -254,7 +252,7 @@ namespace UnrealBuildTool.Modes
 				Logger.LogInformation("Missing module '{MissingModuleName}'", MissingModule.Name);
 			}
 
-			List<KeyValuePair<FileReference, BuildProductType>> AnalyzeProducts = new();
+			List<KeyValuePair<FileReference, BuildProductType>> AnalyzeProducts = [];
 
 			// Generate the dependency graph between modules
 			FileReference DependencyGraphFile = Target.ReceiptFileName.ChangeExtension(".Dependencies.gexf");
@@ -271,7 +269,7 @@ namespace UnrealBuildTool.Modes
 			// Write all the target stats as a text file
 			FileReference TextFile = Target.ReceiptFileName.ChangeExtension(".txt");
 			Logger.LogInformation("Writing module information to {TextFile}", TextFile);
-			using (StreamWriter Writer = new StreamWriter(TextFile.FullName))
+			using (StreamWriter Writer = new(TextFile.FullName))
 			{
 				Writer.WriteLine("All modules in {0}, ordered by number of indirect references", Target.TargetName);
 
@@ -295,23 +293,25 @@ namespace UnrealBuildTool.Modes
 			// Write all the target stats as a CSV file
 			FileReference CsvFile = Target.ReceiptFileName.ChangeExtension(".csv");
 			Logger.LogInformation("Writing module information to {CsvFile}", CsvFile);
-			using (StreamWriter Writer = new StreamWriter(CsvFile.FullName))
+			using (StreamWriter Writer = new(CsvFile.FullName))
 			{
-				List<string> Columns = new List<string>();
-				Columns.Add("Module");
-				Columns.Add("ShortestPath");
-				Columns.Add("NumUniqueInwardRefs");
-				Columns.Add("UniqueInwardRefs");
-				Columns.Add("NumRecursiveInwardRefs");
-				Columns.Add("RecursiveInwardRefs");
-				Columns.Add("NumUniqueOutwardRefs");
-				Columns.Add("UniqueOutwardRefs");
-				Columns.Add("NumRecursiveOutwardRefs");
-				Columns.Add("RecursiveOutwardRefs");
-				Columns.Add("ObjSize");
-				Columns.Add("ObjFiles");
-				Columns.Add("BinSize");
-				Columns.Add("BinFiles");
+				List<string> Columns =
+				[
+					"Module",
+					"ShortestPath",
+					"NumUniqueInwardRefs",
+					"UniqueInwardRefs",
+					"NumRecursiveInwardRefs",
+					"RecursiveInwardRefs",
+					"NumUniqueOutwardRefs",
+					"UniqueOutwardRefs",
+					"NumRecursiveOutwardRefs",
+					"RecursiveOutwardRefs",
+					"ObjSize",
+					"ObjFiles",
+					"BinSize",
+					"BinFiles",
+				];
 				Writer.WriteLine(String.Join(",", Columns));
 
 				foreach (ModuleInfo ModuleInfo in ModuleToInfo.Values.OrderByDescending(x => x.InwardRefs.Count).ThenBy(x => x.BinSize))
@@ -356,29 +356,20 @@ namespace UnrealBuildTool.Modes
 
 		private static void WriteDependencyGraph(UEBuildTarget Target, Dictionary<UEBuildModule, ModuleInfo> ModuleToInfo, FileReference FileName)
 		{
-			List<GraphNode> Nodes = new List<GraphNode>();
+			List<GraphNode> Nodes = [];
 
-			Dictionary<UEBuildModule, GraphNode> ModuleToNode = new Dictionary<UEBuildModule, GraphNode>();
+			Dictionary<UEBuildModule, GraphNode> ModuleToNode = [];
 			foreach (ModuleInfo ModuleInfo in ModuleToInfo.Values)
 			{
-				GraphNode Node = new GraphNode(ModuleInfo.Module.Name);
-
-				long Size;
-				if (Target.ShouldCompileMonolithic())
-				{
-					Size = ModuleInfo.ObjSize;
-				}
-				else
-				{
-					Size = ModuleInfo.BinSize;
-				}
+				GraphNode Node = new(ModuleInfo.Module.Name);
 
+				long Size = Target.ShouldCompileMonolithic() ? ModuleInfo.ObjSize : ModuleInfo.BinSize;
 				Node.Size = 1.0f + (Size / (50.0f * 1024.0f * 1024.0f));
 				Nodes.Add(Node);
 				ModuleToNode[ModuleInfo.Module] = Node;
 			}
 
-			List<GraphEdge> Edges = new List<GraphEdge>();
+			List<GraphEdge> Edges = [];
 			foreach ((UEBuildModule SourceModule, ModuleInfo SourceModuleInfo) in ModuleToInfo)
 			{
 				GraphNode SourceNode = ModuleToNode[SourceModule];
@@ -386,12 +377,12 @@ namespace UnrealBuildTool.Modes
 				{
 					ModuleInfo TargetModuleInfo = ModuleToInfo[TargetModule];
 
-					GraphNode? TargetNode;
-					if (ModuleToNode.TryGetValue(TargetModule, out TargetNode))
+					if (ModuleToNode.TryGetValue(TargetModule, out GraphNode? TargetNode))
 					{
-						GraphEdge Edge = new GraphEdge(SourceNode, TargetNode);
-						Edge.Thickness = TargetModuleInfo.InwardRefs.Count;
-						Edges.Add(Edge);
+						Edges.Add(new(SourceNode, TargetNode)
+						{
+							Thickness = TargetModuleInfo.InwardRefs.Count
+						});
 					}
 				}
 			}
@@ -401,10 +392,10 @@ namespace UnrealBuildTool.Modes
 
 		private static void WriteShortestPathGraph(UEBuildTarget Target, Dictionary<UEBuildModule, ModuleInfo> ModuleToInfo, FileReference FileName)
 		{
-			Dictionary<string, GraphNode> NameToNode = new Dictionary<string, GraphNode>(StringComparer.Ordinal);
+			Dictionary<string, GraphNode> NameToNode = new(StringComparer.Ordinal);
 
-			HashSet<(GraphNode, GraphNode)> EdgesSet = new HashSet<(GraphNode, GraphNode)>();
-			List<GraphEdge> Edges = new List<GraphEdge>();
+			HashSet<(GraphNode, GraphNode)> EdgesSet = [];
+			List<GraphEdge> Edges = [];
 
 			foreach ((UEBuildModule Module, ModuleInfo ModuleInfo) in ModuleToInfo)
 			{
@@ -413,30 +404,28 @@ namespace UnrealBuildTool.Modes
 				GraphNode? PrevNode = null;
 				foreach (string Part in Parts)
 				{
-					GraphNode? NextNode;
-					if (!NameToNode.TryGetValue(Part, out NextNode))
+					if (!NameToNode.TryGetValue(Part, out GraphNode? NextNode))
 					{
 						NextNode = new GraphNode(Part);
 						NameToNode[Part] = NextNode;
 					}
 					if (PrevNode != null && EdgesSet.Add((PrevNode, NextNode)))
 					{
-						GraphEdge Edge = new GraphEdge(PrevNode, NextNode);
-						Edges.Add(Edge);
+						Edges.Add(new(PrevNode, NextNode));
 					}
 					PrevNode = NextNode;
 				}
 			}
 
-			GraphVisualization.WriteGraphFile(FileName, $"Module dependency graph for {Target.TargetName}", NameToNode.Values.ToList(), Edges);
+			GraphVisualization.WriteGraphFile(FileName, $"Module dependency graph for {Target.TargetName}", [.. NameToNode.Values], Edges);
 		}
 
 		private static HashSet<UEBuildModule> GetDirectDependencyModules(UEBuildModule Module)
 		{
-			HashSet<UEBuildModule> ReferencedModules = new HashSet<UEBuildModule>();
-			Module.GetAllDependencyModules(new List<UEBuildModule>(), ReferencedModules, true, false, false);
+			HashSet<UEBuildModule> ReferencedModules = [];
+			Module.GetAllDependencyModules([], ReferencedModules, true, false, false);
 
-			HashSet<UEBuildModule> Modules = new HashSet<UEBuildModule>(Module.GetDirectDependencyModules());
+			HashSet<UEBuildModule> Modules = [.. Module.GetDirectDependencyModules()];
 			Modules.ExceptWith(ReferencedModules);
 			return Modules;
 		}
