UnrealBuildTool: Downgrade some logging messages to Trace (enabled to log via -Verbose)

#rnx

diff --git a/Engine/Source/Programs/Shared/EpicGames.Build/System/CompileScriptModules.cs b/Engine/Source/Programs/Shared/EpicGames.Build/System/CompileScriptModules.cs
--- a/Engine/Source/Programs/Shared/EpicGames.Build/System/CompileScriptModules.cs
+++ b/Engine/Source/Programs/Shared/EpicGames.Build/System/CompileScriptModules.cs
@@ -568,7 +568,7 @@ namespace UnrealBuildBase
 			// Was a build record loaded for this project path? (relevant when considering referenced projects)
 			if (!BuildRecords.TryGetValue(ProjectPath, out CsProjBuildRecordEntry? Entry))
 			{
-				Logger.LogDebug("Found project {ProjectPath} with no existing build record", ProjectPath);
+				Logger.LogTrace("Found project {ProjectPath} with no existing build record", ProjectPath);
 				return;
 			}
 
@@ -593,7 +593,7 @@ namespace UnrealBuildBase
 			if (!ValidateBuildRecord(Entry.BuildRecord, ProjectPath.Directory, out string ValidationFailureMessage, Hook))
 			{
 				string ProjectRelativePath = Path.GetRelativePath(Unreal.EngineDirectory.FullName, ProjectPath.FullName);
-				Logger.LogDebug("[{ProjectRelativePath}] {ValidationFailureMessage}", ProjectRelativePath, ValidationFailureMessage);
+				Logger.LogTrace("[{ProjectRelativePath}] {ValidationFailureMessage}", ProjectRelativePath, ValidationFailureMessage);
 				Entry.Status = CsProjBuildRecordStatus.Invalid;
 				return;
 			}
@@ -608,7 +608,7 @@ namespace UnrealBuildBase
 				{
 					string ProjectRelativePath = Path.GetRelativePath(Unreal.EngineDirectory.FullName, ProjectPath.FullName);
 					string DependencyRelativePath = Path.GetRelativePath(Unreal.EngineDirectory.FullName, FullProjectPath.FullName);
-					Logger.LogDebug("[{ProjectRelativePath}] Existing output is not valid because dependency {DependencyRelativePath} could not be found", ProjectRelativePath, DependencyRelativePath);
+					Logger.LogTrace("[{ProjectRelativePath}] Existing output is not valid because dependency {DependencyRelativePath} could not be found", ProjectRelativePath, DependencyRelativePath);
 					Entry.Status = CsProjBuildRecordStatus.Invalid;
 					return;
 				}
@@ -617,7 +617,7 @@ namespace UnrealBuildBase
 				{
 					string ProjectRelativePath = Path.GetRelativePath(Unreal.EngineDirectory.FullName, ProjectPath.FullName);
 					string DependencyRelativePath = Path.GetRelativePath(Unreal.EngineDirectory.FullName, FullProjectPath.FullName);
-					Logger.LogDebug("[{ProjectRelativePath}] Existing output is not valid because dependency {DependencyRelativePath} is not valid", ProjectRelativePath, DependencyRelativePath);
+					Logger.LogTrace("[{ProjectRelativePath}] Existing output is not valid because dependency {DependencyRelativePath} is not valid", ProjectRelativePath, DependencyRelativePath);
 					Entry.Status = CsProjBuildRecordStatus.Invalid;
 					return;
 				}
@@ -629,7 +629,7 @@ namespace UnrealBuildBase
 					{
 						string ProjectRelativePath = Path.GetRelativePath(Unreal.EngineDirectory.FullName, ProjectPath.FullName);
 						string DependencyRelativePath = Path.GetRelativePath(Unreal.EngineDirectory.FullName, FullProjectPath.FullName);
-						Logger.LogDebug("[{ProjectRelativePath}] Existing output is not valid because dependency {DependencyRelativePath} is newer", ProjectRelativePath, DependencyRelativePath);
+						Logger.LogTrace("[{ProjectRelativePath}] Existing output is not valid because dependency {DependencyRelativePath} is newer", ProjectRelativePath, DependencyRelativePath);
 						Entry.Status = CsProjBuildRecordStatus.Invalid;
 						return;
 					}
@@ -640,7 +640,7 @@ namespace UnrealBuildBase
 					{
 						string ProjectRelativePath = Path.GetRelativePath(Unreal.EngineDirectory.FullName, ProjectPath.FullName);
 						string DependencyRelativePath = Path.GetRelativePath(Unreal.EngineDirectory.FullName, FullProjectPath.FullName);
-						Logger.LogDebug("[{ProjectRelativePath}] Existing dependency output time stamp for {DependencyRelativePath} does not match expected value", ProjectRelativePath, DependencyRelativePath);
+						Logger.LogTrace("[{ProjectRelativePath}] Existing dependency output time stamp for {DependencyRelativePath} does not match expected value", ProjectRelativePath, DependencyRelativePath);
 						Entry.Status = CsProjBuildRecordStatus.Invalid;
 						return;
 					}
diff --git a/Engine/Source/Programs/Shared/EpicGames.MsBuild/CsProjBuilder.cs b/Engine/Source/Programs/Shared/EpicGames.MsBuild/CsProjBuilder.cs
--- a/Engine/Source/Programs/Shared/EpicGames.MsBuild/CsProjBuilder.cs
+++ b/Engine/Source/Programs/Shared/EpicGames.MsBuild/CsProjBuilder.cs
@@ -340,7 +340,7 @@ namespace EpicGames.MsBuild
 				// write all build records
 				if (FileReference.WriteAllTextIfDifferent(entry.BuildRecordFile, JsonSerializer.Serialize(entry.BuildRecord, jsonOptions)))
 				{
-					logger.LogDebug("Wrote script module build record to {BuildRecordPath}", entry.BuildRecordFile);
+					logger.LogTrace("Wrote script module build record to {BuildRecordPath}", entry.BuildRecordFile);
 				}
 			});
 
diff --git a/Engine/Source/Programs/UnrealBuildTool/Configuration/UEBuildModule.cs b/Engine/Source/Programs/UnrealBuildTool/Configuration/UEBuildModule.cs
--- a/Engine/Source/Programs/UnrealBuildTool/Configuration/UEBuildModule.cs
+++ b/Engine/Source/Programs/UnrealBuildTool/Configuration/UEBuildModule.cs
@@ -1317,10 +1317,10 @@ namespace UnrealBuildTool
 
 			if (allowed.Any())
 			{
-				logger.LogDebug(KnownLogEvents.Validation, "Allowed circular dependency for '{RulesFile}' detected:", RulesFile.MakeRelativeTo(Unreal.EngineDirectory));
+				logger.LogTrace(KnownLogEvents.Validation, "Allowed circular dependency for '{RulesFile}' detected:", RulesFile.MakeRelativeTo(Unreal.EngineDirectory));
 				foreach (CircularReference circularReference in allowed)
 				{
-					logger.LogDebug(KnownLogEvents.Validation, "* {ReferenceStack} (allowed via {GuiltyModules})", String.Join(" -> ", circularReference.ReferenceStack), String.Join(", ", circularReference.GuiltyModules.Select(x => $"{x.Key} -> {x.Value}")));
+					logger.LogTrace(KnownLogEvents.Validation, "* {ReferenceStack} (allowed via {GuiltyModules})", String.Join(" -> ", circularReference.ReferenceStack), String.Join(", ", circularReference.GuiltyModules.Select(x => $"{x.Key} -> {x.Value}")));
 				}
 			}
 			return errors.Any();
diff --git a/Engine/Source/Programs/UnrealBuildTool/System/ActionGraph.cs b/Engine/Source/Programs/UnrealBuildTool/System/ActionGraph.cs
--- a/Engine/Source/Programs/UnrealBuildTool/System/ActionGraph.cs
+++ b/Engine/Source/Programs/UnrealBuildTool/System/ActionGraph.cs
@@ -738,11 +738,12 @@ namespace UnrealBuildTool
 				{
 					if (ProducedItem.Exists)
 					{
-						Logger.LogDebug("{StatusDescription}: Produced item \"{ProducedItem}\" was produced by outdated attributes.", RootAction.StatusDescription, ProducedItem.Location);
-						Logger.LogDebug("  New attributes: {NewProducingAttributes}", NewProducingAttributes);
+						Logger.LogTrace("{StatusDescription}: Produced item \"{ProducedItem}\" was produced by outdated attributes.", RootAction.StatusDescription, ProducedItem.Location);
+						Logger.LogTrace("  New attributes: {NewProducingAttributes}", NewProducingAttributes);
 					}
 
 					bIsOutdated = true;
+					break;
 				}
 
 				// If the produced file doesn't exist or has zero size, consider it outdated.  The zero size check is to detect cases
@@ -762,8 +763,9 @@ namespace UnrealBuildTool
 				else
 				{
 					// If any of the produced items doesn't exist, the action is outdated.
-					Logger.LogDebug("{StatusDescription}: Produced item \"{ProducedItem}\" doesn't exist.", RootAction.StatusDescription, ProducedItem.Location);
+					Logger.LogTrace("{StatusDescription}: Produced item \"{ProducedItem}\" doesn't exist.", RootAction.StatusDescription, ProducedItem.Location);
 					bIsOutdated = true;
+					break;
 				}
 			}
 
@@ -785,7 +787,7 @@ namespace UnrealBuildTool
 						bool bPrerequisiteItemIsNewerThanLastExecution = TimeDifference.TotalSeconds > 1;
 						if (bPrerequisiteItemIsNewerThanLastExecution)
 						{
-							Logger.LogDebug("{StatusDescription}: Prerequisite {PrerequisiteItem} is newer than the last execution of the action: {Message}", RootAction.StatusDescription, PrerequisiteItem.Location,
+							Logger.LogTrace("{StatusDescription}: Prerequisite {PrerequisiteItem} is newer than the last execution of the action: {Message}", RootAction.StatusDescription, PrerequisiteItem.Location,
 									$"{PrerequisiteItem.LastWriteTimeUtc.ToLocalTime().ToString(CultureInfo.CurrentCulture)} vs {LastExecutionTimeUtc.LocalDateTime.ToString(CultureInfo.CurrentCulture)}");
 							bIsOutdated = true;
 							break;
@@ -795,7 +797,7 @@ namespace UnrealBuildTool
 					{
 						// If the PrerequisiteItem doesn't exist, it means it's being produced right now and we need to update the ProducedItem that
 						// depends on it
-						Logger.LogDebug("{StatusDescription}: Prerequisite {PrerequisiteItem} doesn't exist.", RootAction.StatusDescription, PrerequisiteItem.Location);
+						Logger.LogTrace("{StatusDescription}: Prerequisite {PrerequisiteItem} doesn't exist.", RootAction.StatusDescription, PrerequisiteItem.Location);
 						bIsOutdated = true;
 						break;
 					}
@@ -807,22 +809,22 @@ namespace UnrealBuildTool
 			{
 				if (!CppDependencies.TryGetDependencies(RootAction.DependencyListFile, Logger, out List<FileItem>? DependencyFiles))
 				{
-					Logger.LogDebug("{StatusDescription}: Missing dependency list file \"{DependencyListFile}\"", RootAction.StatusDescription, RootAction.DependencyListFile);
+					Logger.LogTrace("{StatusDescription}: Missing dependency list file \"{DependencyListFile}\"", RootAction.StatusDescription, RootAction.DependencyListFile);
 					bIsOutdated = true;
 				}
 				else
 				{
-					foreach (FileItem DependencyFile in DependencyFiles)
+					foreach (FileItem DependencyFile in DependencyFiles.Except(RootAction.PrerequisiteItems))
 					{
 						if (!DependencyFile.Exists)
 						{
-							Logger.LogDebug("{RootAction.StatusDescription}: Dependency {DependencyFile} doesn't exist", RootAction.StatusDescription, DependencyFile.AbsolutePath);
+							Logger.LogTrace("{StatusDescription}: Dependency {DependencyFile} doesn't exist", RootAction.StatusDescription, DependencyFile.AbsolutePath);
 							bIsOutdated = true;
 							break;
 						}
 						else if (DependencyFile.LastWriteTimeUtc > LastExecutionTimeUtc)
 						{
-							Logger.LogDebug("{RootAction.StatusDescription}: Dependency {DependencyFile} is newer than the last execution of the action: {Message}", RootAction.StatusDescription, DependencyFile.AbsolutePath,
+							Logger.LogTrace("{StatusDescription}: Dependency {DependencyFile} is newer than the last execution of the action: {Message}", RootAction.StatusDescription, DependencyFile.AbsolutePath,
 								$"{DependencyFile.LastWriteTimeUtc.ToLocalTime().ToString(CultureInfo.CurrentCulture)} vs {LastExecutionTimeUtc.LocalDateTime.ToString(CultureInfo.CurrentCulture)}");
 							bIsOutdated = true;
 							break;
@@ -867,7 +869,7 @@ namespace UnrealBuildTool
 					// won't have to wait for dependent targets to be relinked after each change.
 					if (!bIgnoreOutdatedImportLibraries || !IsImportLibraryDependency(RootAction, PrerequisiteAction))
 					{
-						Logger.LogDebug("{StatusDescription}: Prerequisite {PrereqStatusDescription} is produced by outdated action.", RootAction.StatusDescription, PrerequisiteAction.StatusDescription);
+						Logger.LogTrace("{StatusDescription}: Prerequisite {PrereqStatusDescription} is produced by outdated action.", RootAction.StatusDescription, PrerequisiteAction.StatusDescription);
 						bIsOutdated = true;
 						break;
 					}
