C# analysis cleanup, mostly replacing var with types

#rnx

diff --git a/Engine/Source/Programs/Shared/EpicGames.MsBuild/CsProjBuilder.cs b/Engine/Source/Programs/Shared/EpicGames.MsBuild/CsProjBuilder.cs
--- a/Engine/Source/Programs/Shared/EpicGames.MsBuild/CsProjBuilder.cs
+++ b/Engine/Source/Programs/Shared/EpicGames.MsBuild/CsProjBuilder.cs
@@ -16,7 +16,6 @@ using EpicGames.Core;
 using Microsoft.Build.Evaluation;
 using Microsoft.Build.Locator;
 using Microsoft.Extensions.Logging;
-using Microsoft.Extensions.Logging.Abstractions;
 
 namespace EpicGames.MsBuild
 {
diff --git a/Engine/Source/Programs/Shared/EpicGames.OIDC/OidcTokenManager.cs b/Engine/Source/Programs/Shared/EpicGames.OIDC/OidcTokenManager.cs
--- a/Engine/Source/Programs/Shared/EpicGames.OIDC/OidcTokenManager.cs
+++ b/Engine/Source/Programs/Shared/EpicGames.OIDC/OidcTokenManager.cs
@@ -712,18 +712,18 @@ namespace EpicGames.OIDC
 			}
 			else if (RuntimeInformation.IsOSPlatform(OSPlatform.Linux))
 			{
-				List<string> Browsers = ["xdg-open"];
-				bool IsWsl = IsRunningWsl();
-				if (IsWsl)
+				List<string> browsers = ["xdg-open"];
+				bool isWsl = IsRunningWsl();
+				if (isWsl)
 				{
-					Browsers.Add("wslview");
+					browsers.Add("wslview");
 				}
 
-				foreach (string Browser in Browsers)
+				foreach (string browser in browsers)
 				{
 					try
 					{
-						return Process.Start(Browser, url);
+						return Process.Start(browser, url);
 					}
 					catch (Exception e)
 					{
@@ -731,7 +731,7 @@ namespace EpicGames.OIDC
 					}
 				}
 
-				if (IsWsl)
+				if (isWsl)
 				{
 					_logger?.LogInformation("Did not find a browser, use 'sudo apt install xdg-utils' or 'sudo apt install wslu'.");
 				}
diff --git a/Engine/Source/Programs/Shared/EpicGames.UBA/Impl/SchedulerImpl.cs b/Engine/Source/Programs/Shared/EpicGames.UBA/Impl/SchedulerImpl.cs
--- a/Engine/Source/Programs/Shared/EpicGames.UBA/Impl/SchedulerImpl.cs
+++ b/Engine/Source/Programs/Shared/EpicGames.UBA/Impl/SchedulerImpl.cs
@@ -30,9 +30,13 @@ namespace EpicGames.UBA.Impl
 		static extern ulong ProcessHandle_GetTotalWallTime(nint process);
 		#endregion
 
-		internal ProcessFinishedInfoImpl(object userData, nint processHandle) { _userData = userData;  _processHandle = processHandle; }
-		object _userData;
-		nint _processHandle;
+		internal ProcessFinishedInfoImpl(object userData, nint processHandle)
+		{
+			_userData = userData;
+			_processHandle = processHandle;
+		}
+		readonly object _userData;
+		readonly nint _processHandle;
 
 		public ProcessExecutionType ExecutionType => (ProcessExecutionType)ProcessHandle_GetExecutionType(_processHandle);
 		public int ExitCode => (int)ProcessHandle_GetExitCode(_processHandle);
@@ -68,13 +72,17 @@ namespace EpicGames.UBA.Impl
 		{
 			_scheduler = scheduler;
 			_userData = userData;
-			_exitCallbackDelegate = RaiseExited;
-			_startInfoHandle = ProcessStartInfo_Create2(info.Application, info.Arguments, info.WorkingDirectory, info.Description, (uint)info.Priority, info.RootsHandle, (byte)(info.TrackInputs ? 1 : 0), info.LogFile ?? String.Empty, _exitCallbackDelegate);
+			ExitCallbackDelegate = RaiseExited;
+			_startInfoHandle = ProcessStartInfo_Create2(info.Application, info.Arguments, info.WorkingDirectory, info.Description, (uint)info.Priority, info.RootsHandle, (byte)(info.TrackInputs ? 1 : 0), info.LogFile ?? String.Empty, ExitCallbackDelegate);
 		}
 
 		void RaiseExited(nint userData, nint processHandle, ref byte exitedResponse)
 		{
-			while (_index == UInt32.MaxValue) { Thread.Sleep(1); } // This to prevent a (most likely impossible) race condition (Index is set after process is enqueued)
+			// This to prevent a (most likely impossible) race condition (Index is set after process is enqueued)
+			while (_index == UInt32.MaxValue)
+			{
+				Thread.Sleep(1);
+			}
 
 			//ProcessImpl process = new(handle);
 			ProcessFinishedInfoImpl info = new(_userData, processHandle);
@@ -83,9 +91,9 @@ namespace EpicGames.UBA.Impl
 		}
 
 		public delegate void ExitCallback(nint userData, nint processHandle, ref byte exitedResponse);
-		SchedulerImpl _scheduler;
-		object _userData;
-		event ExitCallback _exitCallbackDelegate;
+		readonly SchedulerImpl _scheduler;
+		readonly object _userData;
+		event ExitCallback ExitCallbackDelegate;
 		public nint _startInfoHandle;
 		public uint _index = UInt32.MaxValue;
 	}
@@ -120,7 +128,7 @@ namespace EpicGames.UBA.Impl
 
 		nint _schedulerHandle = IntPtr.Zero;
 		internal Func<IProcessFinishedInfo, ProcessFinishedResponse>? _processFinishedCallback;
-		List<SchedulerProcess> _processes = new();
+		readonly List<SchedulerProcess> _processes = new();
 
 		internal SchedulerImpl(ISessionServer server, IEnumerable<ICacheClient> cacheClients, int maxLocalProcessors, bool forceRemote)
 		{
diff --git a/Engine/Source/Programs/Shared/EpicGames.UHT/Exporters/CodeGen/EditorDataStorageUhtExtension.cs b/Engine/Source/Programs/Shared/EpicGames.UHT/Exporters/CodeGen/EditorDataStorageUhtExtension.cs
--- a/Engine/Source/Programs/Shared/EpicGames.UHT/Exporters/CodeGen/EditorDataStorageUhtExtension.cs
+++ b/Engine/Source/Programs/Shared/EpicGames.UHT/Exporters/CodeGen/EditorDataStorageUhtExtension.cs
@@ -1,10 +1,8 @@
 ﻿// Copyright Epic Games, Inc. All Rights Reserved.
 
-using System;
 using System.Text;
 using EpicGames.UHT.Tables;
 using EpicGames.UHT.Utils;
-using EpicGames.UHT.Exporters.CodeGen;
 using EpicGames.UHT.Types;
 
 namespace EditorDataStorageUhtExtension
diff --git a/Engine/Source/Programs/Shared/EpicGames.UHT/Exporters/CodeGen/UhtUsedDefineScopes.cs b/Engine/Source/Programs/Shared/EpicGames.UHT/Exporters/CodeGen/UhtUsedDefineScopes.cs
--- a/Engine/Source/Programs/Shared/EpicGames.UHT/Exporters/CodeGen/UhtUsedDefineScopes.cs
+++ b/Engine/Source/Programs/Shared/EpicGames.UHT/Exporters/CodeGen/UhtUsedDefineScopes.cs
@@ -4,7 +4,6 @@ using System;
 using System.Collections.Generic;
 using System.Linq;
 using System.Text;
-using System.Xml.Linq;
 using EpicGames.Core;
 using EpicGames.UHT.Parsers;
 using EpicGames.UHT.Types;
diff --git a/Engine/Source/Programs/Shared/EpicGames.UHT/Parsers/UhtNamespaceParser.cs b/Engine/Source/Programs/Shared/EpicGames.UHT/Parsers/UhtNamespaceParser.cs
--- a/Engine/Source/Programs/Shared/EpicGames.UHT/Parsers/UhtNamespaceParser.cs
+++ b/Engine/Source/Programs/Shared/EpicGames.UHT/Parsers/UhtNamespaceParser.cs
@@ -83,7 +83,6 @@ namespace EpicGames.UHT.Parsers
 				topScope.TokenReader.Optional(';');
 			}
 
-
 			// Remove the namespaces from the active namespace
 			topScope.HeaderParser.PopNamespaces(namespaceCount);
 			return UhtParseResult.Handled;
@@ -95,12 +94,12 @@ namespace EpicGames.UHT.Parsers
 			classObj.ClassType = UhtClassType.VModule;
 
 			using UhtParsingScope topScope = new(parentScope, classObj, parentScope.Session.GetKeywordTable(UhtTableNames.VModule), UhtAccessSpecifier.Public);
-			const string scopeName = "VMODULE";
-			using UhtMessageContext tokenContext = new(scopeName);
+			const string ScopeName = "VMODULE";
+			using UhtMessageContext tokenContext = new(ScopeName);
 
 			// Parse the specifiers
 			UhtSpecifierContext specifierContext = new(topScope, topScope.TokenReader, classObj.MetaData);
-			UhtSpecifierParser specifiers = UhtSpecifierParser.GetThreadInstance(specifierContext, scopeName, parentScope.Session.GetSpecifierTable(UhtTableNames.Class));
+			UhtSpecifierParser specifiers = UhtSpecifierParser.GetThreadInstance(specifierContext, ScopeName, parentScope.Session.GetSpecifierTable(UhtTableNames.Class));
 			specifiers.ParseSpecifiers();
 			classObj.GeneratedBodyLineNumber = classObj.LineNumber;
 			classObj.GeneratedBodyAccessSpecifier = UhtAccessSpecifier.Public;
diff --git a/Engine/Source/Programs/Shared/EpicGames.UHT/Tables/UhtCodeGeneratorInjectorTable.cs b/Engine/Source/Programs/Shared/EpicGames.UHT/Tables/UhtCodeGeneratorInjectorTable.cs
--- a/Engine/Source/Programs/Shared/EpicGames.UHT/Tables/UhtCodeGeneratorInjectorTable.cs
+++ b/Engine/Source/Programs/Shared/EpicGames.UHT/Tables/UhtCodeGeneratorInjectorTable.cs
@@ -1,14 +1,10 @@
-﻿// Copyright Epic Games, Inc. All Rights Reserved.
+// Copyright Epic Games, Inc. All Rights Reserved.
 
 using System;
 using System.Collections.Generic;
-using System.Linq;
 using System.Reflection;
 using System.Text;
-using EpicGames.Core;
-using EpicGames.UHT.Exporters.CodeGen;
 using EpicGames.UHT.Types;
-using EpicGames.UHT.Utils;
 
 namespace EpicGames.UHT.Tables
 {
@@ -124,11 +120,14 @@ namespace EpicGames.UHT.Tables
 			// Sort the delegates to ensure determinism across multiple runs
 			foreach (List<UhtCodeGeneratorInjectorDelegate>? injectors in _injectorsTable.Values)
 			{
-				if (injectors == null) { continue; }
-				injectors.Sort((Lhs, Rhs) =>
+				if (injectors == null)
 				{
-					MethodInfo methodLhs = Lhs.Method;
-					MethodInfo methodRhs = Rhs.Method;
+					continue;
+				}
+				injectors.Sort((lhs, rhs) =>
+				{
+					MethodInfo methodLhs = lhs.Method;
+					MethodInfo methodRhs = rhs.Method;
 					
 					// Sort by module
 					int moduleComparison = String.Compare(methodLhs.Module.Name, methodLhs.Module.Name, StringComparison.Ordinal);
diff --git a/Engine/Source/Programs/Shared/EpicGames.UHT/Tables/UhtPropertyTypeTable.cs b/Engine/Source/Programs/Shared/EpicGames.UHT/Tables/UhtPropertyTypeTable.cs
--- a/Engine/Source/Programs/Shared/EpicGames.UHT/Tables/UhtPropertyTypeTable.cs
+++ b/Engine/Source/Programs/Shared/EpicGames.UHT/Tables/UhtPropertyTypeTable.cs
@@ -4,7 +4,6 @@ using System;
 using System.Collections.Generic;
 using System.Reflection;
 using EpicGames.Core;
-using EpicGames.UHT.Tokenizer;
 using EpicGames.UHT.Types;
 using EpicGames.UHT.Utils;
 
diff --git a/Engine/Source/Programs/Shared/EpicGames.UHT/Types/Properties/UhtSoftClassProperty.cs b/Engine/Source/Programs/Shared/EpicGames.UHT/Types/Properties/UhtSoftClassProperty.cs
--- a/Engine/Source/Programs/Shared/EpicGames.UHT/Types/Properties/UhtSoftClassProperty.cs
+++ b/Engine/Source/Programs/Shared/EpicGames.UHT/Types/Properties/UhtSoftClassProperty.cs
@@ -3,7 +3,6 @@
 using System.Diagnostics.CodeAnalysis;
 using System.Text;
 using EpicGames.UHT.Tables;
-using EpicGames.UHT.Tokenizer;
 using EpicGames.UHT.Utils;
 
 namespace EpicGames.UHT.Types
diff --git a/Engine/Source/Programs/Shared/EpicGames.UHT/Types/UhtField.cs b/Engine/Source/Programs/Shared/EpicGames.UHT/Types/UhtField.cs
--- a/Engine/Source/Programs/Shared/EpicGames.UHT/Types/UhtField.cs
+++ b/Engine/Source/Programs/Shared/EpicGames.UHT/Types/UhtField.cs
@@ -1,7 +1,6 @@
 // Copyright Epic Games, Inc. All Rights Reserved.
 
 using System;
-using System.Globalization;
 using System.Text;
 using System.Text.Json.Serialization;
 using EpicGames.UHT.Utils;
diff --git a/Engine/Source/Programs/Shared/EpicGames.UHT/Types/UhtNamespace.cs b/Engine/Source/Programs/Shared/EpicGames.UHT/Types/UhtNamespace.cs
--- a/Engine/Source/Programs/Shared/EpicGames.UHT/Types/UhtNamespace.cs
+++ b/Engine/Source/Programs/Shared/EpicGames.UHT/Types/UhtNamespace.cs
@@ -2,7 +2,6 @@
 
 using System;
 using System.Text;
-using System.Web;
 using EpicGames.Core;
 
 namespace EpicGames.UHT.Types
diff --git a/Engine/Source/Programs/Shared/EpicGames.UHT/Types/UhtProperty.cs b/Engine/Source/Programs/Shared/EpicGames.UHT/Types/UhtProperty.cs
--- a/Engine/Source/Programs/Shared/EpicGames.UHT/Types/UhtProperty.cs
+++ b/Engine/Source/Programs/Shared/EpicGames.UHT/Types/UhtProperty.cs
@@ -5,7 +5,6 @@ using System.Collections.Generic;
 using System.Diagnostics;
 using System.Diagnostics.CodeAnalysis;
 using System.Globalization;
-using System.Linq;
 using System.Text;
 using System.Text.Json.Serialization;
 using EpicGames.Core;
diff --git a/Engine/Source/Programs/Shared/EpicGames.UHT/Utils/UhtPropertyResolveArgs.cs b/Engine/Source/Programs/Shared/EpicGames.UHT/Utils/UhtPropertyResolveArgs.cs
--- a/Engine/Source/Programs/Shared/EpicGames.UHT/Utils/UhtPropertyResolveArgs.cs
+++ b/Engine/Source/Programs/Shared/EpicGames.UHT/Utils/UhtPropertyResolveArgs.cs
@@ -1,7 +1,6 @@
 // Copyright Epic Games, Inc. All Rights Reserved.
 
 using System;
-using System.Diagnostics;
 using EpicGames.Core;
 using EpicGames.UHT.Parsers;
 using EpicGames.UHT.Tables;
@@ -465,7 +464,6 @@ namespace EpicGames.UHT.Utils
 			return Session.TryGetPropertyType(fullType, out propertyType);
 		}
 
-
 		/// <summary>
 		/// Given the current position of the token reader, return the segment being parsed
 		/// </summary>
diff --git a/Engine/Source/Programs/Shared/EpicGames.UHT/Utils/UhtSession.cs b/Engine/Source/Programs/Shared/EpicGames.UHT/Utils/UhtSession.cs
--- a/Engine/Source/Programs/Shared/EpicGames.UHT/Utils/UhtSession.cs
+++ b/Engine/Source/Programs/Shared/EpicGames.UHT/Utils/UhtSession.cs
@@ -1230,15 +1230,15 @@ namespace EpicGames.UHT.Utils
 
 			if (arguments != null)
 			{
-				foreach (UhtExporter Exporter in ExporterTable)
+				foreach (UhtExporter exporter in ExporterTable)
 				{
-					if (arguments.HasOption($"-{Exporter.Name}"))
+					if (arguments.HasOption($"-{exporter.Name}"))
 					{
-						SetExporterStatus(Exporter.Name, true);
+						SetExporterStatus(exporter.Name, true);
 					}
-					else if (arguments.HasOption($"-no{Exporter.Name}"))
+					else if (arguments.HasOption($"-no{exporter.Name}"))
 					{
-						SetExporterStatus(Exporter.Name, false);
+						SetExporterStatus(exporter.Name, false);
 					}
 				}
 			}
diff --git a/Engine/Source/Programs/Shared/EpicGames.UHT/Utils/UhtUtilities.cs b/Engine/Source/Programs/Shared/EpicGames.UHT/Utils/UhtUtilities.cs
--- a/Engine/Source/Programs/Shared/EpicGames.UHT/Utils/UhtUtilities.cs
+++ b/Engine/Source/Programs/Shared/EpicGames.UHT/Utils/UhtUtilities.cs
@@ -7,7 +7,6 @@ using System.Collections.Generic;
 using System.Linq;
 using System.Text;
 using EpicGames.Core;
-using EpicGames.UHT.Types;
 
 namespace EpicGames.UHT.Utils
 {
diff --git a/Engine/Source/Programs/Shared/UnrealBuildTool.Tests/TestUtilities/StaticTargetRulesInitializer.cs b/Engine/Source/Programs/Shared/UnrealBuildTool.Tests/TestUtilities/StaticTargetRulesInitializer.cs
--- a/Engine/Source/Programs/Shared/UnrealBuildTool.Tests/TestUtilities/StaticTargetRulesInitializer.cs
+++ b/Engine/Source/Programs/Shared/UnrealBuildTool.Tests/TestUtilities/StaticTargetRulesInitializer.cs
@@ -14,6 +14,7 @@ namespace UnrealBuildTool.Tests.TestUtilities
 		internal static readonly string UBTTestFolderRoot = Path.GetTempPath();
 
 		[AssemblyInitialize]
+		[System.Diagnostics.CodeAnalysis.SuppressMessage("Style", "IDE0060:Remove unused parameter", Justification = "AssemblyInitialize requirement")]
 		public static void InitializeBuildEnvironment(TestContext context)
 		{
 			// Clean any previous runs linger temp folder
@@ -32,10 +33,10 @@ namespace UnrealBuildTool.Tests.TestUtilities
 				}
 			}
 
-			// Example, on how to register a fully stubbed test environment.
-			/*
+            // Example, on how to register a fully stubbed test environment.
+            /*
 			 	UEBuildPlatform.RegisterBuildPlatform(
-					new TestBuildPlatfrom(
+					new TestBuildPlatform(
 							UnrealTargetPlatform.Win64,
 							new TestSDK(logger),
 							new UnrealArchitectureConfig(UnrealArch.X64),
@@ -43,19 +44,18 @@ namespace UnrealBuildTool.Tests.TestUtilities
 					logger
 			);*/
 
-			ILogger logger = new TestLogger();
+            ILogger logger = new TestLogger();
 
-			UEBuildPlatform? _;
-			if (!UEBuildPlatform.TryGetBuildPlatform(UnrealTargetPlatform.Win64, out _))
+			if (!UEBuildPlatform.TryGetBuildPlatform(UnrealTargetPlatform.Win64, out UEBuildPlatform? _))
 			{
-				MicrosoftPlatformSDK SDK = new MicrosoftPlatformSDK(logger);
-				UEBuildPlatform.RegisterBuildPlatform(new WindowsPlatform(UnrealTargetPlatform.Win64, SDK, logger), logger);
+				MicrosoftPlatformSDK sdk = new MicrosoftPlatformSDK(logger);
+				UEBuildPlatform.RegisterBuildPlatform(new WindowsPlatform(UnrealTargetPlatform.Win64, sdk, logger), logger);
 			}
 
-			if (!UEBuildPlatform.TryGetBuildPlatform(UnrealTargetPlatform.Android, out _))
+			if (!UEBuildPlatform.TryGetBuildPlatform(UnrealTargetPlatform.Android, out UEBuildPlatform? _))
 			{
-				AndroidPlatformSDK ASDK = new AndroidPlatformSDK(logger);
-				UEBuildPlatform.RegisterBuildPlatform(new AndroidPlatform(ASDK, logger), logger);
+				AndroidPlatformSDK sdk = new AndroidPlatformSDK(logger);
+				UEBuildPlatform.RegisterBuildPlatform(new AndroidPlatform(sdk, logger), logger);
 			}
 
 			XmlConfig.ReadConfigFiles(null, null, logger);
diff --git a/Engine/Source/Programs/Shared/UnrealBuildTool.Tests/TestUtilities/TestBuildPlatform.cs b/Engine/Source/Programs/Shared/UnrealBuildTool.Tests/TestUtilities/TestBuildPlatform.cs
--- a/Engine/Source/Programs/Shared/UnrealBuildTool.Tests/TestUtilities/TestBuildPlatform.cs
+++ b/Engine/Source/Programs/Shared/UnrealBuildTool.Tests/TestUtilities/TestBuildPlatform.cs
@@ -5,33 +5,33 @@ using Microsoft.Extensions.Logging;
 
 namespace UnrealBuildTool.Tests.TestUtilities
 {
-	internal class TestBuildPlatfrom : UEBuildPlatform
+	internal class TestBuildPlatform : UEBuildPlatform
 	{
-		public TestBuildPlatfrom(UnrealTargetPlatform InPlatform, UEBuildPlatformSDK SDK, UnrealArchitectureConfig ArchitectureConfig, ILogger InLogger) : base(InPlatform, SDK, ArchitectureConfig, InLogger)
+		public TestBuildPlatform(UnrealTargetPlatform inPlatform, UEBuildPlatformSDK sdk, UnrealArchitectureConfig architectureConfig, ILogger inLogger) : base(inPlatform, sdk, architectureConfig, inLogger)
 		{
 
 		}
-		public override UEToolChain CreateToolChain(ReadOnlyTargetRules Target)
+		public override UEToolChain CreateToolChain(ReadOnlyTargetRules target)
 		{
 			throw new System.NotImplementedException();
 		}
 
-		public override void Deploy(TargetReceipt Receipt)
+		public override void Deploy(TargetReceipt receipt)
 		{
 
 		}
 
-		public override bool IsBuildProduct(string FileName, string[] NamePrefixes, string[] NameSuffixes)
+		public override bool IsBuildProduct(string fileName, string[] namePrefixes, string[] nameSuffixes)
 		{
 			return false;
 		}
 
-		public override void SetUpEnvironment(ReadOnlyTargetRules Target, CppCompileEnvironment CompileEnvironment, LinkEnvironment LinkEnvironment)
+		public override void SetUpEnvironment(ReadOnlyTargetRules target, CppCompileEnvironment compileEnvironment, LinkEnvironment linkEnvironment)
 		{
 
 		}
 
-		public override bool ShouldCreateDebugInfo(ReadOnlyTargetRules Target)
+		public override bool ShouldCreateDebugInfo(ReadOnlyTargetRules target)
 		{
 			return false;
 		}
diff --git a/Engine/Source/Programs/Shared/UnrealBuildTool.Tests/TestUtilities/TestBuildPlatformFactory.cs b/Engine/Source/Programs/Shared/UnrealBuildTool.Tests/TestUtilities/TestBuildPlatformFactory.cs
--- a/Engine/Source/Programs/Shared/UnrealBuildTool.Tests/TestUtilities/TestBuildPlatformFactory.cs
+++ b/Engine/Source/Programs/Shared/UnrealBuildTool.Tests/TestUtilities/TestBuildPlatformFactory.cs
@@ -8,7 +8,7 @@ namespace UnrealBuildTool.Tests.TestUtilities
 	{
 		public override UnrealTargetPlatform TargetPlatform => UnrealTargetPlatform.Win64;
 
-		public override void RegisterBuildPlatforms(ILogger Logger)
+		public override void RegisterBuildPlatforms(ILogger logger)
 		{
 		}
 	}
diff --git a/Engine/Source/Programs/Shared/UnrealBuildTool.Tests/TestUtilities/TestBuildPlatformSDK.cs b/Engine/Source/Programs/Shared/UnrealBuildTool.Tests/TestUtilities/TestBuildPlatformSDK.cs
--- a/Engine/Source/Programs/Shared/UnrealBuildTool.Tests/TestUtilities/TestBuildPlatformSDK.cs
+++ b/Engine/Source/Programs/Shared/UnrealBuildTool.Tests/TestUtilities/TestBuildPlatformSDK.cs
@@ -7,13 +7,13 @@ namespace UnrealBuildTool.Tests.TestUtilities
 {
 	internal class TestSDK : UEBuildPlatformSDK
 	{
-		public TestSDK(ILogger InLogger) : base(InLogger)
+		public TestSDK(ILogger inLogger) : base(inLogger)
 		{
 		}
 
-		public override bool TryConvertVersionToInt(string? StringValue, out ulong OutValue, string? Hint = null)
+		public override bool TryConvertVersionToInt(string? stringValue, out ulong outValue, string? hint = null)
 		{
-			OutValue = 1;
+			outValue = 1;
 			return true;
 		}
 
diff --git a/Engine/Source/Programs/Shared/UnrealBuildTool.Tests/TestUtilities/TestToolChain.cs b/Engine/Source/Programs/Shared/UnrealBuildTool.Tests/TestUtilities/TestToolChain.cs
--- a/Engine/Source/Programs/Shared/UnrealBuildTool.Tests/TestUtilities/TestToolChain.cs
+++ b/Engine/Source/Programs/Shared/UnrealBuildTool.Tests/TestUtilities/TestToolChain.cs
@@ -13,22 +13,22 @@ namespace UnrealBuildTool.Tests.TestUtilities
 		{
 		}
 
-		public override FileItem? LinkFiles(LinkEnvironment LinkEnvironment, bool bBuildImportLibraryOnly, IActionGraphBuilder Graph)
+		public override FileItem? LinkFiles(LinkEnvironment linkEnvironment, bool bBuildImportLibraryOnly, IActionGraphBuilder graph)
 		{
 			return null;
 		}
 
-		protected override CPPOutput CompileCPPFiles(CppCompileEnvironment CompileEnvironment, IEnumerable<FileItem> InputFiles, DirectoryReference OutputDir, string ModuleName, IActionGraphBuilder Graph)
+		protected override CPPOutput CompileCPPFiles(CppCompileEnvironment compileEnvironment, IEnumerable<FileItem> inputFiles, DirectoryReference outputDir, string moduleName, IActionGraphBuilder graph)
 		{
 			return null!;
 		}
 
-		protected override CPPOutput CompileISPCFiles(CppCompileEnvironment CompileEnvironment, IEnumerable<FileItem> InputFiles, DirectoryReference OutputDir, IActionGraphBuilder Graph)
+		protected override CPPOutput CompileISPCFiles(CppCompileEnvironment compileEnvironment, IEnumerable<FileItem> inputFiles, DirectoryReference outputDir, IActionGraphBuilder graph)
 		{
 			return null!;
 		}
 
-		protected override CPPOutput GenerateISPCHeaders(CppCompileEnvironment CompileEnvironment, IEnumerable<FileItem> InputFiles, DirectoryReference OutputDir, IActionGraphBuilder Graph)
+		protected override CPPOutput GenerateISPCHeaders(CppCompileEnvironment compileEnvironment, IEnumerable<FileItem> inputFiles, DirectoryReference outputDir, IActionGraphBuilder graph)
 		{
 			return null!;
 		}
diff --git a/Engine/Source/Programs/UnrealBuildTool/Configuration/UEBuildModuleCPP.cs b/Engine/Source/Programs/UnrealBuildTool/Configuration/UEBuildModuleCPP.cs
--- a/Engine/Source/Programs/UnrealBuildTool/Configuration/UEBuildModuleCPP.cs
+++ b/Engine/Source/Programs/UnrealBuildTool/Configuration/UEBuildModuleCPP.cs
@@ -2627,7 +2627,7 @@ namespace UnrealBuildTool
 				Result.Definitions.Add(String.Format("UE_WITH_DEBUG_INFO={0}", BaseCompileEnvironment.bCreateDebugInfo ? "1" : "0"));
 
 				StringBuilder sb = new();
-				foreach (var (id, vfs, local) in BaseCompileEnvironment.RootPaths)
+				foreach ((uint id, DirectoryReference vfs, DirectoryReference local) in BaseCompileEnvironment.RootPaths)
 				{
 					sb.Append(vfs.FullName.Replace("\\", "\\\\")).Append(';');
 					sb.Append(local.FullName.Replace("\\", "\\\\")).Append(';');
diff --git a/Engine/Source/Programs/UnrealBuildTool/Configuration/UEBuildTarget.cs b/Engine/Source/Programs/UnrealBuildTool/Configuration/UEBuildTarget.cs
--- a/Engine/Source/Programs/UnrealBuildTool/Configuration/UEBuildTarget.cs
+++ b/Engine/Source/Programs/UnrealBuildTool/Configuration/UEBuildTarget.cs
@@ -3538,12 +3538,12 @@ namespace UnrealBuildTool
 			{
 				binary.PrimaryModule.bCreatePerModuleFile = false;
 
-				var queue = new Queue<(UEBuildModule Module, UEBuildModule Source)>();
-				var visited = new HashSet<UEBuildModule>();
+				Queue<(UEBuildModule Module, UEBuildModule Source)> queue = new Queue<(UEBuildModule Module, UEBuildModule Source)>();
+				HashSet<UEBuildModule> visited = new HashSet<UEBuildModule>();
 
 				// Queue up primary module and all direct dependencies
 				queue.Enqueue((binary.PrimaryModule, binary.PrimaryModule));
-				foreach (var m in binary.Modules)
+				foreach (UEBuildModule m in binary.Modules)
 				{
 					if (visited.Add(m))
 					{
@@ -3554,7 +3554,7 @@ namespace UnrealBuildTool
 				// Search for non-engine PerModuleInlineFile.. 
 				while (queue.Count > 0)
 				{
-					var current = queue.Dequeue();
+					(UEBuildModule Module, UEBuildModule Source) current = queue.Dequeue();
 					if (current.Module.RulesFile.IsUnderDirectory(Unreal.EngineDirectory))
 					{
 						continue;
@@ -3565,7 +3565,7 @@ namespace UnrealBuildTool
 						return;
 					}
 
-					foreach (var dep in current.Module.PublicDependencyModules!.Union(current.Module.PrivateDependencyModules!))
+					foreach (UEBuildModule? dep in current.Module.PublicDependencyModules!.Union(current.Module.PrivateDependencyModules!))
 					{
 						if (visited.Add(dep))
 						{
diff --git a/Engine/Source/Programs/UnrealBuildTool/Executors/UnrealBuildAccelerator/UBAAgentCoordinatorHorde.cs b/Engine/Source/Programs/UnrealBuildTool/Executors/UnrealBuildAccelerator/UBAAgentCoordinatorHorde.cs
--- a/Engine/Source/Programs/UnrealBuildTool/Executors/UnrealBuildAccelerator/UBAAgentCoordinatorHorde.cs
+++ b/Engine/Source/Programs/UnrealBuildTool/Executors/UnrealBuildAccelerator/UBAAgentCoordinatorHorde.cs
@@ -186,7 +186,7 @@ namespace UnrealBuildTool
 				Thread.Sleep(10);
 			}
 
-			var clusterId = new ClusterId(UnrealBuildAcceleratorHordeConfig.ClusterDefault);
+			ClusterId clusterId = new ClusterId(UnrealBuildAcceleratorHordeConfig.ClusterDefault);
 			UbaConfig ubaConfig = _client.AllocateUbaCacheServerAsync(clusterId, cancellationToken).Result;
 			if (String.IsNullOrEmpty(ubaConfig.CacheEndpoint))
 			{
@@ -889,7 +889,7 @@ namespace UnrealBuildTool
 					return;
 				}
 
-				var (hordeSession, reason) = await _hordeSessionTask;
+				(UBAHordeSession? hordeSession, string reason) = await _hordeSessionTask;
 
 				if (hordeSession == null)
 				{
@@ -982,7 +982,7 @@ namespace UnrealBuildTool
 				return;
 			}
 
-			var (hordeSession, reason) = await _hordeSessionTask;
+			(UBAHordeSession? hordeSession, string reason) = await _hordeSessionTask;
 			_hordeSessionTask = null;
 			if (hordeSession != null)
 			{
diff --git a/Engine/Source/Programs/UnrealBuildTool/Executors/UnrealBuildAccelerator/UBAExecutor.cs b/Engine/Source/Programs/UnrealBuildTool/Executors/UnrealBuildAccelerator/UBAExecutor.cs
--- a/Engine/Source/Programs/UnrealBuildTool/Executors/UnrealBuildAccelerator/UBAExecutor.cs
+++ b/Engine/Source/Programs/UnrealBuildTool/Executors/UnrealBuildAccelerator/UBAExecutor.cs
@@ -624,7 +624,7 @@ namespace UnrealBuildTool
 
 			using IScheduler scheduler = IScheduler.CreateScheduler(session, cache, NumParallelProcesses, UBAConfig.bForceBuildAllRemote);
 
-			var cancelBuild = () =>
+			System.Action cancelBuild = () =>
 			{
 				if (_bIsCancelled)
 					return;
@@ -692,7 +692,7 @@ namespace UnrealBuildTool
 
 							string executingHost = info.ExecutingHost ?? "Unknown";
 
-							var rerun = (string error) =>
+							Func<string, ProcessFinishedResponse> rerun = (string error) =>
 							{
 								_threadedLogger.LogInformation("{Description} {StatusDescription} [RemoteExecutor: {ExecutingHost}]: Exited with error code {ExitCode} ({Error}). This action will retry locally", action.CommandDescription, action.StatusDescription, executingHost, exitCode, error);
 								foreach (string line in logLines)
@@ -850,11 +850,11 @@ namespace UnrealBuildTool
 				ProcessStartInfo startInfo = GetActionStartInfo(action);
 				startInfo.RootsHandle = GetActionRootsHandle(action);
 
-				var prereqActionsSortIndex = action.PrerequisiteActions.Select(a => a.SortIndex).ToArray();
+				int[] prereqActionsSortIndex = action.PrerequisiteActions.Select(a => a.SortIndex).ToArray();
 				bool enableDetour = !ForceLocalNoDetour(action) && action.bCanExecuteInUBA;
 				bool canRunRemotely = CanRunRemotely(action);
 
-				var (knownInputs, knownInputsCount) = BuildKnownInputs(action);
+				(byte[]? knownInputs, uint knownInputsCount) = BuildKnownInputs(action);
 
 				uint bucket = GetCacheBucket(action);
 				uint memoryGroup = GetMemoryGroup(action);
diff --git a/Engine/Source/Programs/UnrealBuildTool/Modes/IWYUMode.cs b/Engine/Source/Programs/UnrealBuildTool/Modes/IWYUMode.cs
--- a/Engine/Source/Programs/UnrealBuildTool/Modes/IWYUMode.cs
+++ b/Engine/Source/Programs/UnrealBuildTool/Modes/IWYUMode.cs
@@ -978,7 +978,7 @@ namespace UnrealBuildTool
 							if (fixedPath.StartsWith(Dir.FullName))
 							{
 								string printable = value.File.Substring(Dir.FullName.Length + 1);
-								var entry = new IWYUIncludeEntry() { Full = MissingInclude.Full, Printable = printable, System = false, Resolved = value, };
+								IWYUIncludeEntry entry = new IWYUIncludeEntry() { Full = MissingInclude.Full, Printable = printable, System = false, Resolved = value, };
 								Info.Includes.Add(entry);
 								break;
 							}
diff --git a/Engine/Source/Programs/UnrealBuildTool/Platform/Apple/UEBuildApple.cs b/Engine/Source/Programs/UnrealBuildTool/Platform/Apple/UEBuildApple.cs
--- a/Engine/Source/Programs/UnrealBuildTool/Platform/Apple/UEBuildApple.cs
+++ b/Engine/Source/Programs/UnrealBuildTool/Platform/Apple/UEBuildApple.cs
@@ -218,7 +218,7 @@ namespace UnrealBuildTool
 			
 			CompileEnvironment.Definitions.Add("PLATFORM_APPLE=1");
 
-			AppleExports.GetSwiftIntegrationSettings(Target.ProjectFile, Target.Type, Platform, out var bUseSwiftUIMain, out _);
+			AppleExports.GetSwiftIntegrationSettings(Target.ProjectFile, Target.Type, Platform, out bool bUseSwiftUIMain, out _);
 			CompileEnvironment.Definitions.Add("UE_USE_SWIFT_UI_MAIN=" + (bUseSwiftUIMain ? "1" : "0"));
 
 			CompileEnvironment.Definitions.Add("WITH_TTS=0");
diff --git a/Engine/Source/Programs/UnrealBuildTool/Platform/Mac/MacToolChain.cs b/Engine/Source/Programs/UnrealBuildTool/Platform/Mac/MacToolChain.cs
--- a/Engine/Source/Programs/UnrealBuildTool/Platform/Mac/MacToolChain.cs
+++ b/Engine/Source/Programs/UnrealBuildTool/Platform/Mac/MacToolChain.cs
@@ -749,7 +749,7 @@ namespace UnrealBuildTool
 
 					HashSet<DirectoryReference> UsedDestDirs = new();
 					HashSet<FileReference> UsedSourceLibs = new();
-					var ProcessLib = (FileReference SourceLib, FileReference DestLib, FileReference ReferencingFile) => 
+					Action<FileReference, FileReference, FileReference> ProcessLib = (FileReference SourceLib, FileReference DestLib, FileReference ReferencingFile) => 
 					{
 						if (string.Equals(SourceLib.GetExtension(), ".dylib", StringComparison.InvariantCultureIgnoreCase))
 						{
diff --git a/Engine/Source/Programs/UnrealBuildTool/ProjectFiles/VisualStudioWorkspace/VSWorkspaceProjectFile.cs b/Engine/Source/Programs/UnrealBuildTool/ProjectFiles/VisualStudioWorkspace/VSWorkspaceProjectFile.cs
--- a/Engine/Source/Programs/UnrealBuildTool/ProjectFiles/VisualStudioWorkspace/VSWorkspaceProjectFile.cs
+++ b/Engine/Source/Programs/UnrealBuildTool/ProjectFiles/VisualStudioWorkspace/VSWorkspaceProjectFile.cs
@@ -123,14 +123,14 @@ namespace UnrealBuildTool
 
 					try
 					{
-						foreach (var Architecture in TargetDesc.Architectures.Architectures)
+						foreach (UnrealArch Architecture in TargetDesc.Architectures.Architectures)
 						{
 							UEBuildTarget BuildTarget = UEBuildTarget.Create(TargetDesc, false, false, false, UnrealIntermediateEnvironment.GenerateProjectFiles, Logger);
 							BuildTarget.PreBuildSetup(Logger);
 
 							ExportedTargetInfo TargetInfo = ExportTarget(BuildTarget, Architecture, bBuildByDefault, PlatformProjectGenerators, Logger, VCFormat);
 
-							var sb = new StringBuilder();
+							StringBuilder sb = new StringBuilder();
 							sb.Append(ProjectTarget.TargetFilePath.GetFileNameWithoutAnyExtensions());
 							sb.Append("_");
 							sb.Append(TargetInfo.Configuration);
@@ -468,7 +468,7 @@ namespace UnrealBuildTool
 
 				if (Architecture != null)
 				{
-					var result = Architecture.ToString();
+					string? result = Architecture.ToString();
 					return result != null ? result : string.Empty;
 				}
 			}
diff --git a/Engine/Source/Programs/UnrealBuildTool/ProjectFiles/VisualStudioWorkspace/VSWorkspaceProjectFileGenerator.cs b/Engine/Source/Programs/UnrealBuildTool/ProjectFiles/VisualStudioWorkspace/VSWorkspaceProjectFileGenerator.cs
--- a/Engine/Source/Programs/UnrealBuildTool/ProjectFiles/VisualStudioWorkspace/VSWorkspaceProjectFileGenerator.cs
+++ b/Engine/Source/Programs/UnrealBuildTool/ProjectFiles/VisualStudioWorkspace/VSWorkspaceProjectFileGenerator.cs
@@ -105,7 +105,7 @@ namespace UnrealBuildTool
 
 		protected override bool WriteProjectFiles(PlatformProjectGeneratorCollection PlatformProjectGenerators, ILogger Logger)
 		{
-			var ProjectGenerationDescriptionData = ReadProjectGenerateDescriptionFilePath()?.Targets;
+			Dictionary<string, AvailableConfiguration>? ProjectGenerationDescriptionData = ReadProjectGenerateDescriptionFilePath()?.Targets;
 			using ProgressWriter Progress = new("Writing project files...", true, Logger);
 			List<ProjectFile> ProjectsToGenerate = new(GeneratedProjectFiles);
 			if (ProjectNames.Any())
@@ -136,7 +136,7 @@ namespace UnrealBuildTool
 
 			List<UnrealTargetConfiguration> Configurations = ConfigurationsToGenerate.ToList();
 			VCProjectFileFormat VCFormat = GetVCProjectFileFormat(Logger);
-			var wroteVCProjectFiles = false;
+			bool wroteVCProjectFiles = false;
 
 			for (int ProjectFileIndex = 0; ProjectFileIndex < ProjectsToGenerate.Count; ++ProjectFileIndex)
 			{
@@ -150,7 +150,7 @@ namespace UnrealBuildTool
 					return false;
 				}
 
-				foreach (var Target in CurrentProject.ProjectTargets.OfType<ProjectTarget>())
+				foreach (ProjectTarget Target in CurrentProject.ProjectTargets.OfType<ProjectTarget>())
 				{
 					if (!TargetTypes.Any() || (Target.TargetRules != null && TargetTypes.Contains(Target.TargetRules.Type)))
 					{
@@ -507,12 +507,12 @@ namespace UnrealBuildTool
 					FilteredPlatforms = new List<UnrealTargetPlatform>();
 					FilteredConfigurations = new List<UnrealTargetConfiguration>();
 
-					foreach (var Platform in Platforms)
+					foreach (UnrealTargetPlatform Platform in Platforms)
 					{
 						IncludePlatform(AvailableConfiguration, Platform, FilteredPlatforms);
 					}
 
-					foreach (var Configuration in Configurations)
+					foreach (UnrealTargetConfiguration Configuration in Configurations)
 					{
 						IncludeConfiguration(AvailableConfiguration, Configuration, FilteredConfigurations);
 					}
@@ -523,10 +523,10 @@ namespace UnrealBuildTool
 					return true;
 				}
 
-				var Settings = new VCProjectFileSettings();
+				VCProjectFileSettings Settings = new VCProjectFileSettings();
 				Settings.MaxSharedIncludePaths = 0;
 
-				var VCProject = new VCProjectFile(Output, RootPath, Format, false, true, null, Settings);
+				VCProjectFile VCProject = new VCProjectFile(Output, RootPath, Format, false, true, null, Settings);
 				VCProject.ProjectTargets.Add(Target);
 				if (VCProject.WriteProjectFile(FilteredPlatforms, FilteredConfigurations, PlatformProjectGenerators, Logger))
 				{
@@ -541,7 +541,7 @@ namespace UnrealBuildTool
 					while (Nodes.Count != 0)
 					{
 						XmlNode Node = Nodes.Dequeue();
-						using (var NodeList = Node.ChildNodes)
+						using (XmlNodeList NodeList = Node.ChildNodes)
 						{
 							if (NodeList != null)
 							{
@@ -561,7 +561,7 @@ namespace UnrealBuildTool
 										}
 										else
 										{
-											foreach (var RegEx in BlockedNodesListRegEx)
+											foreach (Regex RegEx in BlockedNodesListRegEx)
 											{
 												if (RegEx.Match(ChildNode.Name).Success)
 												{
@@ -585,7 +585,7 @@ namespace UnrealBuildTool
 							}
 						}
 
-						foreach (var ChildNode in RemovedChildNodes)
+						foreach (XmlNode ChildNode in RemovedChildNodes)
 						{
 							Node.RemoveChild(ChildNode);
 						}
diff --git a/Engine/Source/Programs/UnrealBuildTool/System/HotReload.cs b/Engine/Source/Programs/UnrealBuildTool/System/HotReload.cs
--- a/Engine/Source/Programs/UnrealBuildTool/System/HotReload.cs
+++ b/Engine/Source/Programs/UnrealBuildTool/System/HotReload.cs
@@ -1432,7 +1432,7 @@ namespace UnrealBuildTool
 
 							if (Action.RootPaths.bUseVfs)
 							{
-								foreach (var RootPath in Action.RootPaths)
+								foreach ((uint id, DirectoryReference vfs, DirectoryReference local) RootPath in Action.RootPaths)
 								{
 									if (vfsPathsAdded.Add(RootPath.vfs))
 									{
@@ -1446,7 +1446,7 @@ namespace UnrealBuildTool
 				Writer.WriteArrayEnd();
 
 				Writer.WriteArrayStart("Vfs");
-				foreach (var kv in vfsPaths)
+				foreach ((DirectoryReference virtualPath, DirectoryReference localPath) kv in vfsPaths)
 				{
 					Writer.WriteObjectStart();
 					Writer.WriteValue("Virtual", kv.virtualPath.FullName);
