Add `bInjectAutoRTFMAttributeDisables` and let us use it to plumb the `autortfm(disable)` attribute into `*_API` macros in modules with AutoRTFM explicitly disabled.

This allows us to automagically get the benefit of static compile-time checking of transactional calls to functions that are explicitly disabled for AutoRTFM. I've had to introduce a `*_NON_ATTRIBUTED_API` define for a few select cases - namely in this CL `StaticEnum` is template specialized in each UHT generated file and that *cannot* have AutoRTFM disabled on it because the original template definition is used under transactions (and thus *all* template specializations of it must be compiled for transactions for safety reasons).

diff --git a/Engine/Source/Programs/Shared/EpicGames.UHT/Exporters/CodeGen/UhtHeaderCodeGeneratorCppFile.cs b/Engine/Source/Programs/Shared/EpicGames.UHT/Exporters/CodeGen/UhtHeaderCodeGeneratorCppFile.cs
--- a/Engine/Source/Programs/Shared/EpicGames.UHT/Exporters/CodeGen/UhtHeaderCodeGeneratorCppFile.cs
+++ b/Engine/Source/Programs/Shared/EpicGames.UHT/Exporters/CodeGen/UhtHeaderCodeGeneratorCppFile.cs
@@ -458,7 +458,7 @@ namespace EpicGames.UHT.Exporters.CodeGen
 
 			if (Module.Module.AlwaysExportEnums)
 			{
-				builder.Append("template<> ").Append(Module.Api).Append("UEnum* StaticEnum<").Append(enumObj.CppType).Append(">()\r\n");
+				builder.Append("template<> ").Append(Module.NonAttributedApi).Append("UEnum* StaticEnum<").Append(enumObj.CppType).Append(">()\r\n");
 			}
 			else
 			{
diff --git a/Engine/Source/Programs/Shared/EpicGames.UHT/Exporters/CodeGen/UhtHeaderCodeGeneratorHFile.cs b/Engine/Source/Programs/Shared/EpicGames.UHT/Exporters/CodeGen/UhtHeaderCodeGeneratorHFile.cs
--- a/Engine/Source/Programs/Shared/EpicGames.UHT/Exporters/CodeGen/UhtHeaderCodeGeneratorHFile.cs
+++ b/Engine/Source/Programs/Shared/EpicGames.UHT/Exporters/CodeGen/UhtHeaderCodeGeneratorHFile.cs
@@ -774,7 +774,7 @@ namespace EpicGames.UHT.Exporters.CodeGen
 				// Forward declare the StaticEnum<> specialization for enum classes
 				if (Module.Module.AlwaysExportEnums || enumObj.EnumExportFlags.HasAnyFlags(UhtEnumExportFlags.MinimalAPI))
 				{
-					builder.Append("template<> ").Append(Module.Api).Append("UEnum* StaticEnum<").Append(enumObj.CppType).Append(">();\r\n");
+					builder.Append("template<> ").Append(Module.NonAttributedApi).Append("UEnum* StaticEnum<").Append(enumObj.CppType).Append(">();\r\n");
 				}
 				else
 				{
diff --git a/Engine/Source/Programs/Shared/EpicGames.UHT/Types/UhtModule.cs b/Engine/Source/Programs/Shared/EpicGames.UHT/Types/UhtModule.cs
--- a/Engine/Source/Programs/Shared/EpicGames.UHT/Types/UhtModule.cs
+++ b/Engine/Source/Programs/Shared/EpicGames.UHT/Types/UhtModule.cs
@@ -47,6 +47,12 @@ namespace EpicGames.UHT.Types
 		[JsonIgnore]
 		public string Api { get; }
 
+		/// <summary>
+		/// C++ API non-attributed name for the module with a trailing space
+		/// </summary>
+		[JsonIgnore]
+		public string NonAttributedApi { get; }
+
 		/// <summary>
 		/// Enumeration of all the packages in the module.
 		/// </summary>
@@ -170,6 +176,7 @@ namespace EpicGames.UHT.Types
 			}
 
 			Api = $"{ShortName.ToUpper()}_API ";
+			NonAttributedApi = $"{ShortName.ToUpper()}_NON_ATTRIBUTED_API ";
 
 			ScriptPackage = new(this, packageName, packageFlags);
 		}
diff --git a/Engine/Source/Programs/UnrealBuildTool/Configuration/ModuleRules.cs b/Engine/Source/Programs/UnrealBuildTool/Configuration/ModuleRules.cs
--- a/Engine/Source/Programs/UnrealBuildTool/Configuration/ModuleRules.cs
+++ b/Engine/Source/Programs/UnrealBuildTool/Configuration/ModuleRules.cs
@@ -1146,6 +1146,12 @@ namespace UnrealBuildTool
 		/// Disables AutoRTFM instrumentation to this module only when AutoRTFMCompiler is enabled
 		/// </summary>
 		public bool bDisableAutoRTFMInstrumentation { get; set; }
+
+		/// <summary>
+		/// If `bDisableAutoRTFMInstrumentation` is true, will make the modules `*_API` define contain `AUTORTFM_DISABLE` to let us enable errors for code
+		/// that uses transactionally-disabled code in places that are transactionally-enabled.
+		/// </summary>
+		public bool bInjectAutoRTFMAttributeDisables { get; set; }
 		
 		/// <summary>
 		/// If true, autogenerated functions of USTRUCTS are always exported
diff --git a/Engine/Source/Programs/UnrealBuildTool/Configuration/UEBuildModule.cs b/Engine/Source/Programs/UnrealBuildTool/Configuration/UEBuildModule.cs
--- a/Engine/Source/Programs/UnrealBuildTool/Configuration/UEBuildModule.cs
+++ b/Engine/Source/Programs/UnrealBuildTool/Configuration/UEBuildModule.cs
@@ -61,10 +61,15 @@ namespace UnrealBuildTool
 		public UEBuildBinary Binary = null!;
 
 		/// <summary>
-		/// The name of the _API define for this module
+		/// The name of the _API define for this module (for functions)
 		/// </summary>
 		protected readonly string ModuleApiDefine;
 
+		/// <summary>
+		/// The name of the _NON_ATTRIBUTED_API define for this module (for variables)
+		/// </summary>
+		protected readonly string ModuleNonAttributedDefine;
+
 		/// <summary>
 		/// Set of all the public definitions
 		/// </summary>
@@ -224,6 +229,7 @@ namespace UnrealBuildTool
 			this.IntermediateDirectoryNoArch = IntermediateDirectoryNoArch;
 
 			ModuleApiDefine = Name.ToUpperInvariant() + "_API";
+			ModuleNonAttributedDefine = Name.ToUpperInvariant() + "_NON_ATTRIBUTED_API";
 
 			HashSet<string> PublicPreBuildLibraries = HashSetFromOptionalEnumerableStringParameter(Rules.PublicPreBuildLibraries);
 			PublicDefinitions = [.. HashSetFromOptionalEnumerableStringParameter(Rules.PublicDefinitions).Select(TrimDefinition)];
@@ -678,32 +684,42 @@ namespace UnrealBuildTool
 			// Add this module's public definitions
 			AddDefinitions(Definitions, PublicDefinitions);
 
-			// Add the import or export declaration for the module
 			if (Rules.Type == ModuleRules.ModuleType.CPlusPlus)
 			{
+				String ApiDefinition = String.Empty;
+
 				if (Rules.Target.LinkType == TargetLinkType.Monolithic)
 				{
 					if (Rules.Target.bShouldCompileAsDLL && (Rules.Target.bHasExports || Rules.ModuleSymbolVisibility == ModuleRules.SymbolVisibility.VisibileForDll))
 					{
-						Definitions.Add(ModuleApiDefine + "=DLLEXPORT");
-					}
-					else
-					{
-						Definitions.Add(ModuleApiDefine + "=");
+						ApiDefinition = "DLLEXPORT";
 					}
 				}
 				else if (Binary == null || (Rules.Target.bMergeModules ? (SourceBinary == null || SourceModule != this) : SourceBinary != Binary))
 				{
-					Definitions.Add(ModuleApiDefine + "=DLLIMPORT");
+					ApiDefinition = "DLLIMPORT";
 				}
-				else if (!Binary.bAllowExports || Rules.Target.bMergeModules)
+				else if (!(!Binary.bAllowExports || Rules.Target.bMergeModules))
 				{
-					Definitions.Add(ModuleApiDefine + "=");
-				}
-				else
-				{
-					Definitions.Add(ModuleApiDefine + "=DLLEXPORT");
+					ApiDefinition = "DLLEXPORT";
 				}
+
+				// We use a C-style attribute here because of where these macros are defined in UE. C++ attributes must come at the beginning of the function
+				// definition, like `[[my_attr]] void Foo();`, but these macros are used like `void FOO_API Foo();`, and if we used the C++ attribute name to
+				// disable AutoRTFM on this function the compiler would spit nails. But it'll happily consume the C-style attribute at that location and do
+				// the right thing, so we use that instead.
+				String AutoRTFMDisable = (Rules.Target.bUseAutoRTFMCompiler && Rules.bDisableAutoRTFMInstrumentation && Rules.bInjectAutoRTFMAttributeDisables) ?
+					"__attribute__((autortfm(disable)))" : String.Empty;
+
+				// Add a spacer if and only if we've actually used import/export above.
+				String Spacer = ApiDefinition.Length == 0 ? String.Empty : " ";
+
+				// Add the import or export declaration for the module
+				String ModuleApiDefinitionString = $"{ModuleApiDefine}={ApiDefinition}{Spacer}{AutoRTFMDisable}";
+				Definitions.Add(ModuleApiDefinitionString);
+
+				String ModuleNonAttributedApiDefinitionString = $"{ModuleNonAttributedDefine}={ApiDefinition}";
+				Definitions.Add(ModuleNonAttributedApiDefinitionString);
 			}
 
 			// Add the additional frameworks so that the compiler can know about their #include paths
diff --git a/Engine/Source/Programs/UnrealBuildTool/Configuration/UEBuildModuleCPP.cs b/Engine/Source/Programs/UnrealBuildTool/Configuration/UEBuildModuleCPP.cs
--- a/Engine/Source/Programs/UnrealBuildTool/Configuration/UEBuildModuleCPP.cs
+++ b/Engine/Source/Programs/UnrealBuildTool/Configuration/UEBuildModuleCPP.cs
@@ -9,6 +9,7 @@ using System.Collections.Generic;
 using System.IO;
 using System.Linq;
 using System.Text;
+using System.Text.Json;
 using UnrealBuildBase;
 
 namespace UnrealBuildTool
@@ -23,20 +24,20 @@ namespace UnrealBuildTool
 		/// </summary>
 		public class InputFileCollection
 		{
-			public readonly List<FileItem> HeaderFiles = [];
-			public readonly List<FileItem> ISPCHeaderFiles = [];
-
-			public readonly List<FileItem> IXXFiles = [];
-			public readonly List<FileItem> CPPFiles = [];
-			public readonly List<FileItem> CFiles = [];
-			public readonly List<FileItem> CCFiles = [];
-			public readonly List<FileItem> CXXFiles = [];
-			public readonly List<FileItem> MFiles = [];
-			public readonly List<FileItem> MMFiles = [];
-			public readonly List<FileItem> SwiftFiles = [];
-			public readonly List<FileItem> RCFiles = [];
-			public readonly List<FileItem> ISPCFiles = [];
-			public readonly List<FileItem> ProtoFiles = [];
+			public readonly List<FileItem> HeaderFiles = new List<FileItem>();
+			public readonly List<FileItem> ISPCHeaderFiles = new List<FileItem>();
+
+			public readonly List<FileItem> IXXFiles = new List<FileItem>();
+			public readonly List<FileItem> CPPFiles = new List<FileItem>();
+			public readonly List<FileItem> CFiles = new List<FileItem>();
+			public readonly List<FileItem> CCFiles = new List<FileItem>();
+			public readonly List<FileItem> CXXFiles = new List<FileItem>();
+			public readonly List<FileItem> MFiles = new List<FileItem>();
+			public readonly List<FileItem> MMFiles = new List<FileItem>();
+			public readonly List<FileItem> SwiftFiles = new List<FileItem>();
+			public readonly List<FileItem> RCFiles = new List<FileItem>();
+			public readonly List<FileItem> ISPCFiles = new List<FileItem>();
+			public readonly List<FileItem> ProtoFiles = new List<FileItem>();
 
 			public IEnumerable<FileItem> AllHeaderFiles => HeaderFiles.Concat(ISPCHeaderFiles);
 			public IEnumerable<FileItem> AllCompiledFiles => IXXFiles.Concat(CPPFiles).Concat(CFiles).Concat(CCFiles).Concat(CXXFiles).Concat(MFiles).Concat(MMFiles).Concat(SwiftFiles).Concat(RCFiles).Concat(ISPCFiles);
@@ -45,10 +46,16 @@ namespace UnrealBuildTool
 		/// <summary>
 		/// Instance where a CPP file doesn't include the expected header as the first include.
 		/// </summary>
-		public readonly struct InvalidIncludeDirective(FileReference cppFile, FileReference headerFile)
+		public readonly struct InvalidIncludeDirective
 		{
-			public FileReference CppFile { get; } = cppFile;
-			public FileReference HeaderFile { get; } = headerFile;
+			public FileReference CppFile { get; }
+			public FileReference HeaderFile { get; }
+
+			public InvalidIncludeDirective(FileReference cppFile, FileReference headerFile)
+			{
+				CppFile = cppFile;
+				HeaderFile = headerFile;
+			}
 		}
 
 		/// <summary>
@@ -99,7 +106,7 @@ namespace UnrealBuildTool
 		/// <summary>
 		/// Verse source code directory associated with this module
 		/// </summary>
-		readonly DirectoryReference? AssociatedVerseDirectory;
+		DirectoryReference? AssociatedVerseDirectory;
 
 		/// <summary>
 		/// Create PerModuleInline source file. Defaults to PrimaryModule in a binary but could be another one if primary module does not depend on core
@@ -504,21 +511,19 @@ namespace UnrealBuildTool
 			}
 
 			// Find all the input files
-			Dictionary<DirectoryItem, FileItem[]> DirectoryToSourceFiles = [];
+			Dictionary<DirectoryItem, FileItem[]> DirectoryToSourceFiles = new Dictionary<DirectoryItem, FileItem[]>();
 			InputFileCollection InputFiles = FindInputFiles(Target.Platform, DirectoryToSourceFiles, Logger);
 
 			{
-				Dictionary<string, List<FileItem>> nameToPath = [];
+				Dictionary<string, List<FileItem>> nameToPath = new();
 				foreach (FileItem InputFile in InputFiles.AllCompiledFiles)
 				{
 					string name = InputFile.Name.ToUpperInvariant();
-					if (!nameToPath.TryGetValue(name, out List<FileItem>? path))
+					if (!nameToPath.ContainsKey(name))
 					{
-						path = [];
-						nameToPath.Add(name, path);
+						nameToPath.Add(name, new List<FileItem>());
 					}
-
-					path.Add(InputFile);
+					nameToPath[name].Add(InputFile);
 				}
 				IEnumerable<KeyValuePair<string, List<FileItem>>> fileConflicts = nameToPath.Where(item => item.Value.Count > 1);
 				if (fileConflicts.Any())
@@ -537,12 +542,12 @@ namespace UnrealBuildTool
 				Graph.AddSourceFiles(Pair.Key, Pair.Value);
 			}
 
-			Graph.AddHeaderFiles([.. InputFiles.AllHeaderFiles]);
+			Graph.AddHeaderFiles(InputFiles.AllHeaderFiles.ToArray());
 
-			// We are building with IWYU and this module does not support it, early out
+			// We are building with IWYU and thismodule does not support it, early out
 			if (Target.bIWYU && Rules.IWYUSupport == IWYUSupport.None)
 			{
-				return [];
+				return new List<FileItem>();
 			}
 
 			// Process all of the header file dependencies for this module
@@ -582,7 +587,7 @@ namespace UnrealBuildTool
 			}
 
 			// Configure the precompiled headers for this module
-			CppCompileEnvironment = SetupPrecompiledHeaders(Target, ToolChain, CppCompileEnvironment, LinkInputFiles, Graph);
+			CppCompileEnvironment = SetupPrecompiledHeaders(Target, ToolChain, CppCompileEnvironment, LinkInputFiles, Graph, Logger);
 			if (CppCompileEnvironment.PerArchPrecompiledHeaderFiles != null)
 			{
 				foreach (UnrealArch Arch in CppCompileEnvironment.PerArchPrecompiledHeaderFiles.Keys)
@@ -614,17 +619,15 @@ namespace UnrealBuildTool
 				CppCompileEnvironment.UserIncludePaths.Add(IntermediateDirectory);
 				ModuleCompileEnvironment.UserIncludePaths.Add(IntermediateDirectory);
 
-				CppCompileEnvironment ProtoCompileEnvironment = new(CppCompileEnvironment)
-				{
-					PrecompiledHeaderAction = PrecompiledHeaderAction.None,
-					bUseUnity = false,
-					bWarningsAsErrors = false,
-					// Disable dynamic analysis on generated proto code as it causes issues when the proto/grpc libs are not
-					// compiled the same way. This can cause false positive in some situations, for example with address
-					// sanitize enabled (see: [REDACTED_PROFILE_URL]/protobuf/issues/13115).
-					bDisableDynamicAnalysis = true
-				};
-				// Disable warnings that are common in generated proto cpp
+				CppCompileEnvironment ProtoCompileEnvironment = new CppCompileEnvironment(CppCompileEnvironment);
+				ProtoCompileEnvironment.PrecompiledHeaderAction = PrecompiledHeaderAction.None;
+				ProtoCompileEnvironment.bUseUnity = false;
+				ProtoCompileEnvironment.bWarningsAsErrors = false;
+				// Disable dynamic analysis on generated proto code as it causes issues when the proto/grpc libs are not
+				// compiled the same way. This can cause false positive in some situations, for example with address
+				// sanitize enabled (see: [REDACTED_PROFILE_URL]/protobuf/issues/13115).
+				ProtoCompileEnvironment.bDisableDynamicAnalysis = true;
+				// Disable warnings that are common in genereated proto cpp
 				ProtoCompileEnvironment.CppCompileWarnings.DeprecationWarningLevel = WarningLevel.Off;
 				ProtoCompileEnvironment.CppCompileWarnings.ShortenSizeTToIntWarningLevel = WarningLevel.Off;
 				ProtoCompileEnvironment.CppCompileWarnings.UnusedParameterWarningLevel = WarningLevel.Off;
@@ -643,17 +646,17 @@ namespace UnrealBuildTool
 			CppCompileEnvironment = ToolChain.CreateSharedResponseFile(CppCompileEnvironment, SharedResponseFile, Graph);
 
 			// Mapping of source file to unity file. We output this to intermediate directories for other tools (eg. live coding) to use.
-			Dictionary<FileItem, FileItem> SourceFileToUnityFile = [];
+			Dictionary<FileItem, FileItem> SourceFileToUnityFile = new Dictionary<FileItem, FileItem>();
 
-			List<FileItem> CPPFiles = [.. InputFiles.CPPFiles];
-			List<FileItem> GeneratedFileItems = [];
+			List<FileItem> CPPFiles = new List<FileItem>(InputFiles.CPPFiles);
+			List<FileItem> GeneratedFileItems = new List<FileItem>();
 			CppCompileEnvironment GeneratedCPPCompileEnvironment = CppCompileEnvironment;
 			bool bMergeUnityFiles = !(Target.DisableMergingModuleAndGeneratedFilesInUnityFiles?.Contains(Name) ?? false) && Rules.bMergeUnityFiles;
 
 			// Compile all the generated CPP files
 			if (GeneratedCppDirectories != null && !CppCompileEnvironment.bHackHeaderGenerator)
 			{
-				Dictionary<string, FileItem> GeneratedFiles = [];
+				Dictionary<string, FileItem> GeneratedFiles = new Dictionary<string, FileItem>();
 				foreach (string GeneratedDir in GeneratedCppDirectories)
 				{
 					if (!Directory.Exists(GeneratedDir))
@@ -669,7 +672,7 @@ namespace UnrealBuildTool
 						string FileName = File.Name;
 						if (FileName.EndsWith(".gen.cpp"))
 						{
-							string Key = string.Concat(Prefix, FileName.AsSpan(0, FileName.Length - ".gen.cpp".Length));
+							string Key = Prefix + FileName.Substring(0, FileName.Length - ".gen.cpp".Length);
 							GeneratedFiles.Add(Key, File);
 						}
 						else if (GeneratedCPPCompileEnvironment.FileMatchesExtraGeneratedCPPTypes(FileName))
@@ -691,13 +694,11 @@ namespace UnrealBuildTool
 							string Key = Prefix + ListOfInlinedGenCppsItem;
 							if (GeneratedFiles.Remove(Key, out FileItem? FoundGenCppFile))
 							{
-								if (!CppCompileEnvironment.FileInlineGenCPPMap.TryGetValue(CPPFileItem, out List<FileItem>? genFiles))
+								if (!CppCompileEnvironment.FileInlineGenCPPMap.ContainsKey(CPPFileItem))
 								{
-									genFiles = [];
-									CppCompileEnvironment.FileInlineGenCPPMap[CPPFileItem] = genFiles;
+									CppCompileEnvironment.FileInlineGenCPPMap[CPPFileItem] = new List<FileItem>();
 								}
-
-								genFiles.Add(FoundGenCppFile);
+								CppCompileEnvironment.FileInlineGenCPPMap[CPPFileItem].Add(FoundGenCppFile);
 							}
 							else
 							{
@@ -709,7 +710,7 @@ namespace UnrealBuildTool
 
 					if (Rules.CppCompileWarningSettings.NonInlinedGenCppWarningLevel >= WarningLevel.Warning)
 					{
-						Dictionary<string, FileItem> CPPFilesLookup = [];
+						Dictionary<string, FileItem> CPPFilesLookup = new Dictionary<string, FileItem>();
 						foreach (FileItem CPPFile in CPPFiles)
 						{
 							CPPFilesLookup.Add(CPPFile.Name, CPPFile);
@@ -722,11 +723,11 @@ namespace UnrealBuildTool
 								continue;
 							}
 
-							string NameWithoutPrefix = Name[4..];
+							string NameWithoutPrefix = Name.Substring(4);
 							string NameWithoutPrefixWithSuffix = NameWithoutPrefix + ".cpp";
 							if (CPPFilesLookup.TryGetValue(NameWithoutPrefixWithSuffix, out FileItem? Item))
 							{
-								Logger.LogWarning("{Path}(1): .gen.cpp not inlined. Add '#include UE_INLINE_GENERATED_CPP_BY_NAME({Name})' after header includes.", Item.AbsolutePath, NameWithoutPrefix);
+								Logger.LogWarning("{0}(1): .gen.cpp not inlined. Add '#include UE_INLINE_GENERATED_CPP_BY_NAME({1})' after header includes.", Item.AbsolutePath, NameWithoutPrefix);
 							}
 						}
 					}
@@ -734,19 +735,15 @@ namespace UnrealBuildTool
 					// Create a compile environment for the generated files. We can disable creating debug info here to improve link times.
 					if (GeneratedCPPCompileEnvironment.bCreateDebugInfo && Target.bDisableDebugInfoForGeneratedCode)
 					{
-						GeneratedCPPCompileEnvironment = new CppCompileEnvironment(GeneratedCPPCompileEnvironment)
-						{
-							bCreateDebugInfo = false
-						};
+						GeneratedCPPCompileEnvironment = new CppCompileEnvironment(GeneratedCPPCompileEnvironment);
+						GeneratedCPPCompileEnvironment.bCreateDebugInfo = false;
 						bMergeUnityFiles = false;
 					}
 
 					if (Target.StaticAnalyzer != StaticAnalyzer.None && !Target.bStaticAnalyzerIncludeGenerated)
 					{
-						GeneratedCPPCompileEnvironment = new CppCompileEnvironment(GeneratedCPPCompileEnvironment)
-						{
-							bDisableStaticAnalysis = true
-						};
+						GeneratedCPPCompileEnvironment = new CppCompileEnvironment(GeneratedCPPCompileEnvironment);
+						GeneratedCPPCompileEnvironment.bDisableStaticAnalysis = true;
 						bMergeUnityFiles = false;
 					}
 
@@ -889,13 +886,11 @@ namespace UnrealBuildTool
 			GeneratedCPPCompileEnvironment.bUseUnity = CppCompileEnvironment.bUseUnity;
 
 			// Create and register a special action that can be used to compile single files (even when unity is enabled) for generated files
-			if (GeneratedFileItems.Count != 0)
+			if (GeneratedFileItems.Any())
 			{
-				CppCompileEnvironment SingleGeneratedCPPCompileEnvironment = new(GeneratedCPPCompileEnvironment)
-				{
-					PrecompiledHeaderAction = PrecompiledHeaderAction.None,
-					bUseUnity = false
-				};
+				CppCompileEnvironment SingleGeneratedCPPCompileEnvironment = new CppCompileEnvironment(GeneratedCPPCompileEnvironment);
+				SingleGeneratedCPPCompileEnvironment.PrecompiledHeaderAction = PrecompiledHeaderAction.None;
+				SingleGeneratedCPPCompileEnvironment.bUseUnity = false;
 				foreach (DirectoryReference Directory in GeneratedFileItems.Select(x => x.Location.Directory).Distinct())
 				{
 					ToolChain.CreateSpecificFileAction(SingleGeneratedCPPCompileEnvironment, Directory, IntermediateDirectory, Graph);
@@ -903,13 +898,11 @@ namespace UnrealBuildTool
 			}
 
 			// Create and register a special action that can be used to compile single files (even when unity is enabled) for normal files
-			if (ModuleDirectories.Length != 0)
+			if (ModuleDirectories.Any())
 			{
-				CppCompileEnvironment SingleCompileEnvironment = new(CppCompileEnvironment)
-				{
-					PrecompiledHeaderAction = PrecompiledHeaderAction.None,
-					bUseUnity = false
-				};
+				CppCompileEnvironment SingleCompileEnvironment = new CppCompileEnvironment(CppCompileEnvironment);
+				SingleCompileEnvironment.PrecompiledHeaderAction = PrecompiledHeaderAction.None;
+				SingleCompileEnvironment.bUseUnity = false;
 				foreach (DirectoryReference Directory in ModuleDirectories)
 				{
 					ToolChain.CreateSpecificFileAction(SingleCompileEnvironment, Directory, IntermediateDirectory, Graph);
@@ -924,24 +917,24 @@ namespace UnrealBuildTool
 			else if (Rules.Target.bAlwaysUseUnityForGeneratedFiles && !Rules.Target.bEnableCppModules)
 			{
 				GeneratedCPPCompileEnvironment.bUseUnity = true;
-				Unity.GenerateUnitySource(Target, GeneratedFileItems, [], GeneratedCPPCompileEnvironment, WorkingSet, (Rules.ShortName ?? Name) + ".gen", IntermediateDirectory, Graph, SourceFileToUnityFile,
+				Unity.GenerateUnitySource(Target, GeneratedFileItems, new List<FileItem>(), GeneratedCPPCompileEnvironment, WorkingSet, (Rules.ShortName ?? Name) + ".gen", IntermediateDirectory, Graph, SourceFileToUnityFile,
 					out List<FileItem> NormalGeneratedFiles, out List<FileItem> AdaptiveGeneratedFiles, NumIncludedBytesPerUnityCPP);
 				ModuleCompileEnvironment.FileInlineGenCPPMap = new(GeneratedCPPCompileEnvironment.FileInlineGenCPPMap);
 				CppCompileEnvironment.FileInlineGenCPPMap = new(GeneratedCPPCompileEnvironment.FileInlineGenCPPMap);
-				LinkInputFiles.AddRange(CompileFilesWithToolChain(Target, ToolChain, GeneratedCPPCompileEnvironment, ModuleCompileEnvironment, NormalGeneratedFiles, AdaptiveGeneratedFiles, Graph).ObjectFiles);
+				LinkInputFiles.AddRange(CompileFilesWithToolChain(Target, ToolChain, GeneratedCPPCompileEnvironment, ModuleCompileEnvironment, NormalGeneratedFiles, AdaptiveGeneratedFiles, Graph, Logger).ObjectFiles);
 			}
 			else
 			{
 				Unity.GetAdaptiveFiles(Target, GeneratedFileItems, InputFiles.HeaderFiles, GeneratedCPPCompileEnvironment, WorkingSet, (Rules.ShortName ?? Name) + ".gen", IntermediateDirectory, Graph,
 					out List<FileItem> NormalGeneratedFiles, out List<FileItem> AdaptiveGeneratedFiles);
-				LinkInputFiles.AddRange(CompileFilesWithToolChain(Target, ToolChain, GeneratedCPPCompileEnvironment, ModuleCompileEnvironment, NormalGeneratedFiles, AdaptiveGeneratedFiles, Graph).ObjectFiles);
+				LinkInputFiles.AddRange(CompileFilesWithToolChain(Target, ToolChain, GeneratedCPPCompileEnvironment, ModuleCompileEnvironment, NormalGeneratedFiles, AdaptiveGeneratedFiles, Graph, Logger).ObjectFiles);
 			}
 
 
 			// Compile Swift - which may generate a header that the cpp files need below
 			if (InputFiles.SwiftFiles.Count > 0 && Rules.SwiftInteropHeader != null)
 			{
-				CppCompileEnvironment SwiftCompileEnvironment = new(ModuleCompileEnvironment);
+				CppCompileEnvironment SwiftCompileEnvironment = new CppCompileEnvironment(ModuleCompileEnvironment);
 				// currently nothing is using this header, but it could be helpful to have in swift
 				CreateHeaderForDefinitions(SwiftCompileEnvironment, IntermediateDirectory, "swift", Graph);
 				SwiftCompileEnvironment.ForceIncludeFiles.Clear();
@@ -962,7 +955,7 @@ namespace UnrealBuildTool
 				// both used when compiling files that may need the headers
 				ModuleCompileEnvironment.AdditionalPrerequisites.AddRange(HeaderFiles);
 				CppCompileEnvironment.AdditionalPrerequisites.AddRange(HeaderFiles);
-				if (HeaderFiles.Any())
+				if (HeaderFiles.Count() > 0)
 				{
 					// @todo: this could have different per-arch headers- should be rare but possible
 					ModuleCompileEnvironment.UserIncludePaths.Add(HeaderFiles.First().Location.Directory);
@@ -978,7 +971,7 @@ namespace UnrealBuildTool
 					Unity.GenerateUnitySource(Target, CPPFiles, InputFiles.HeaderFiles, CppCompileEnvironment, WorkingSet, Rules.ShortName ?? Name, IntermediateDirectory, Graph, SourceFileToUnityFile,
 						out List<FileItem> NormalFiles, out List<FileItem> AdaptiveFiles, NumIncludedBytesPerUnityCPP);
 					ModuleCompileEnvironment.FileInlineGenCPPMap = new(CppCompileEnvironment.FileInlineGenCPPMap);
-					LinkInputFiles.AddRange(CompileFilesWithToolChain(Target, ToolChain, CppCompileEnvironment, ModuleCompileEnvironment, NormalFiles, AdaptiveFiles, Graph).ObjectFiles);
+					LinkInputFiles.AddRange(CompileFilesWithToolChain(Target, ToolChain, CppCompileEnvironment, ModuleCompileEnvironment, NormalFiles, AdaptiveFiles, Graph, Logger).ObjectFiles);
 				}
 				else
 				{
@@ -989,7 +982,7 @@ namespace UnrealBuildTool
 						NormalFiles = CPPFiles;
 						AdaptiveFiles.RemoveAll(new HashSet<FileItem>(NormalFiles).Contains);
 					}
-					LinkInputFiles.AddRange(CompileFilesWithToolChain(Target, ToolChain, CppCompileEnvironment, ModuleCompileEnvironment, NormalFiles, AdaptiveFiles, Graph).ObjectFiles);
+					LinkInputFiles.AddRange(CompileFilesWithToolChain(Target, ToolChain, CppCompileEnvironment, ModuleCompileEnvironment, NormalFiles, AdaptiveFiles, Graph, Logger).ObjectFiles);
 				}
 			}
 
@@ -1003,17 +996,15 @@ namespace UnrealBuildTool
 			// Compile C files. Do not use a PCH here, because a C++ PCH is not compatible with C source files.
 			if (InputFiles.CFiles.Count > 0)
 			{
-				CppCompileEnvironment CCompileEnvironment = new(ModuleCompileEnvironment)
-				{
-					bUseUnity = bModuleUsesUnityBuild && Rules.UnityFileTypes.HasFlag(ModuleRules.UnityFileType.C)
-				};
+				CppCompileEnvironment CCompileEnvironment = new CppCompileEnvironment(ModuleCompileEnvironment);
+				CCompileEnvironment.bUseUnity = bModuleUsesUnityBuild && Rules.UnityFileTypes.HasFlag(ModuleRules.UnityFileType.C);
 				CreateHeaderForDefinitions(CCompileEnvironment, IntermediateDirectory, "c", Graph);
 
 				if (CCompileEnvironment.bUseUnity)
 				{
 					Unity.GenerateUnitySource(Target, InputFiles.CFiles, InputFiles.HeaderFiles, CCompileEnvironment, WorkingSet, Rules.ShortName ?? Name, IntermediateDirectory, Graph, SourceFileToUnityFile,
 						out List<FileItem> NormalFiles, out List<FileItem> AdaptiveFiles, NumIncludedBytesPerUnityCPP);
-					LinkInputFiles.AddRange(CompileFilesWithToolChain(Target, ToolChain, CCompileEnvironment, ModuleCompileEnvironment, NormalFiles, AdaptiveFiles, Graph).ObjectFiles);
+					LinkInputFiles.AddRange(CompileFilesWithToolChain(Target, ToolChain, CCompileEnvironment, ModuleCompileEnvironment, NormalFiles, AdaptiveFiles, Graph, Logger).ObjectFiles);
 				}
 				else
 				{
@@ -1024,17 +1015,15 @@ namespace UnrealBuildTool
 			// Compile CC files
 			if (InputFiles.CCFiles.Count > 0)
 			{
-				CppCompileEnvironment CCCompileEnvironment = new(CppCompileEnvironment)
-				{
-					bUseUnity = bModuleUsesUnityBuild && Rules.UnityFileTypes.HasFlag(ModuleRules.UnityFileType.CC)
-				};
+				CppCompileEnvironment CCCompileEnvironment = new CppCompileEnvironment(CppCompileEnvironment);
+				CCCompileEnvironment.bUseUnity = bModuleUsesUnityBuild && Rules.UnityFileTypes.HasFlag(ModuleRules.UnityFileType.CC);
 				CreateHeaderForDefinitions(CCCompileEnvironment, IntermediateDirectory, "cc", Graph);
 
 				if (CCCompileEnvironment.bUseUnity)
 				{
 					Unity.GenerateUnitySource(Target, InputFiles.CCFiles, InputFiles.HeaderFiles, CCCompileEnvironment, WorkingSet, Rules.ShortName ?? Name, IntermediateDirectory, Graph, SourceFileToUnityFile,
 						out List<FileItem> NormalFiles, out List<FileItem> AdaptiveFiles, NumIncludedBytesPerUnityCPP);
-					LinkInputFiles.AddRange(CompileFilesWithToolChain(Target, ToolChain, CCCompileEnvironment, CppCompileEnvironment, NormalFiles, AdaptiveFiles, Graph).ObjectFiles);
+					LinkInputFiles.AddRange(CompileFilesWithToolChain(Target, ToolChain, CCCompileEnvironment, CppCompileEnvironment, NormalFiles, AdaptiveFiles, Graph, Logger).ObjectFiles);
 				}
 				else
 				{
@@ -1045,17 +1034,15 @@ namespace UnrealBuildTool
 			// Compile CXX files
 			if (InputFiles.CXXFiles.Count > 0)
 			{
-				CppCompileEnvironment CXXCompileEnvironment = new(CppCompileEnvironment)
-				{
-					bUseUnity = bModuleUsesUnityBuild && Rules.UnityFileTypes.HasFlag(ModuleRules.UnityFileType.CXX)
-				};
+				CppCompileEnvironment CXXCompileEnvironment = new CppCompileEnvironment(CppCompileEnvironment);
+				CXXCompileEnvironment.bUseUnity = bModuleUsesUnityBuild && Rules.UnityFileTypes.HasFlag(ModuleRules.UnityFileType.CXX);
 				CreateHeaderForDefinitions(CXXCompileEnvironment, IntermediateDirectory, "cxx", Graph);
 
 				if (CXXCompileEnvironment.bUseUnity)
 				{
 					Unity.GenerateUnitySource(Target, InputFiles.CXXFiles, InputFiles.HeaderFiles, CXXCompileEnvironment, WorkingSet, Rules.ShortName ?? Name, IntermediateDirectory, Graph, SourceFileToUnityFile,
 						out List<FileItem> NormalFiles, out List<FileItem> AdaptiveFiles, NumIncludedBytesPerUnityCPP);
-					LinkInputFiles.AddRange(CompileFilesWithToolChain(Target, ToolChain, CXXCompileEnvironment, CppCompileEnvironment, NormalFiles, AdaptiveFiles, Graph).ObjectFiles);
+					LinkInputFiles.AddRange(CompileFilesWithToolChain(Target, ToolChain, CXXCompileEnvironment, CppCompileEnvironment, NormalFiles, AdaptiveFiles, Graph, Logger).ObjectFiles);
 				}
 				else
 				{
@@ -1066,17 +1053,15 @@ namespace UnrealBuildTool
 			// Compile M files. Do not use a PCH here, because a C++ PCH is not compatible with Objective-C source files.
 			if (InputFiles.MFiles.Count > 0)
 			{
-				CppCompileEnvironment MCompileEnvironment = new(ModuleCompileEnvironment)
-				{
-					bUseUnity = bModuleUsesUnityBuild && Rules.UnityFileTypes.HasFlag(ModuleRules.UnityFileType.M)
-				};
+				CppCompileEnvironment MCompileEnvironment = new CppCompileEnvironment(ModuleCompileEnvironment);
+				MCompileEnvironment.bUseUnity = bModuleUsesUnityBuild && Rules.UnityFileTypes.HasFlag(ModuleRules.UnityFileType.M);
 				CreateHeaderForDefinitions(MCompileEnvironment, IntermediateDirectory, "m", Graph);
 
 				if (MCompileEnvironment.bUseUnity)
 				{
 					Unity.GenerateUnitySource(Target, InputFiles.MFiles, InputFiles.HeaderFiles, MCompileEnvironment, WorkingSet, Rules.ShortName ?? Name, IntermediateDirectory, Graph, SourceFileToUnityFile,
 						out List<FileItem> NormalFiles, out List<FileItem> AdaptiveFiles, NumIncludedBytesPerUnityCPP);
-					LinkInputFiles.AddRange(CompileFilesWithToolChain(Target, ToolChain, MCompileEnvironment, ModuleCompileEnvironment, NormalFiles, AdaptiveFiles, Graph).ObjectFiles);
+					LinkInputFiles.AddRange(CompileFilesWithToolChain(Target, ToolChain, MCompileEnvironment, ModuleCompileEnvironment, NormalFiles, AdaptiveFiles, Graph, Logger).ObjectFiles);
 				}
 				else
 				{
@@ -1087,17 +1072,15 @@ namespace UnrealBuildTool
 			// Compile MM files
 			if (InputFiles.MMFiles.Count > 0)
 			{
-				CppCompileEnvironment MMCompileEnvironment = new(CppCompileEnvironment)
-				{
-					bUseUnity = bModuleUsesUnityBuild && Rules.UnityFileTypes.HasFlag(ModuleRules.UnityFileType.MM)
-				};
+				CppCompileEnvironment MMCompileEnvironment = new CppCompileEnvironment(CppCompileEnvironment);
+				MMCompileEnvironment.bUseUnity = bModuleUsesUnityBuild && Rules.UnityFileTypes.HasFlag(ModuleRules.UnityFileType.MM);
 				CreateHeaderForDefinitions(MMCompileEnvironment, IntermediateDirectory, "mm", Graph);
 
 				if (MMCompileEnvironment.bUseUnity)
 				{
 					Unity.GenerateUnitySource(Target, InputFiles.MMFiles, InputFiles.HeaderFiles, MMCompileEnvironment, WorkingSet, Rules.ShortName ?? Name, IntermediateDirectory, Graph, SourceFileToUnityFile,
 						out List<FileItem> NormalFiles, out List<FileItem> AdaptiveFiles, NumIncludedBytesPerUnityCPP);
-					LinkInputFiles.AddRange(CompileFilesWithToolChain(Target, ToolChain, MMCompileEnvironment, CppCompileEnvironment, NormalFiles, AdaptiveFiles, Graph).ObjectFiles);
+					LinkInputFiles.AddRange(CompileFilesWithToolChain(Target, ToolChain, MMCompileEnvironment, CppCompileEnvironment, NormalFiles, AdaptiveFiles, Graph, Logger).ObjectFiles);
 				}
 				else
 				{
@@ -1110,7 +1093,7 @@ namespace UnrealBuildTool
 			// must be specified relative to the resource file itself or Engine/Source.
 			if (InputFiles.RCFiles.Count > 0)
 			{
-				CppCompileEnvironment ResourceCompileEnvironment = new(BinaryCompileEnvironment);
+				CppCompileEnvironment ResourceCompileEnvironment = new CppCompileEnvironment(BinaryCompileEnvironment);
 				if (Binary != null)
 				{
 					// @todo: This should be in some Windows code somewhere...
@@ -1124,7 +1107,7 @@ namespace UnrealBuildTool
 			{
 				DirectoryReference.CreateDirectory(PrecompiledManifestLocation.Directory);
 
-				PrecompiledManifest Manifest = new();
+				PrecompiledManifest Manifest = new PrecompiledManifest();
 				Manifest.OutputFiles.AddRange(LinkInputFiles.Select(x => x.Location));
 				Manifest.WriteIfModified(PrecompiledManifestLocation);
 			}
@@ -1133,7 +1116,7 @@ namespace UnrealBuildTool
 			if (Rules.Target.bWithLiveCoding)
 			{
 				StringWriter StringWriter = new();
-				using (JsonWriter Writer = new(StringWriter))
+				using (JsonWriter Writer = new JsonWriter(StringWriter))
 				{
 					Writer.WriteObjectStart();
 					Writer.WriteObjectStart("RemapUnityFiles");
@@ -1175,10 +1158,8 @@ namespace UnrealBuildTool
 					}
 					else
 					{
-						HeaderCompileEnvironment = new CppCompileEnvironment(ModuleCompileEnvironment)
-						{
-							bUseUnity = false
-						};
+						HeaderCompileEnvironment = new CppCompileEnvironment(ModuleCompileEnvironment);
+						HeaderCompileEnvironment.bUseUnity = false;
 					}
 
 					DirectoryReference HeaderIntermediateDirectory = DirectoryReference.Combine(IntermediateDirectory, "H");
@@ -1189,7 +1170,7 @@ namespace UnrealBuildTool
 					foreach (DirectoryItem ParentDir in HeaderFileItems.Select(x => x.Directory!).Distinct().OrderBy(x => x.FullName))
 					{
 						DirectoryReference HeaderSubIntermediateDirectory = DirectoryReference.Combine(HeaderIntermediateDirectory, (DirectoryIndex++).ToString());
-						List<FileItem> HeadersToCompile = [];
+						List<FileItem> HeadersToCompile = new();
 						foreach (FileItem Header in HeaderFileItems.Where(x => x.Directory == ParentDir))
 						{
 							if (Target.bIWYU)
@@ -1229,7 +1210,7 @@ namespace UnrealBuildTool
 		{
 			if (Rules.IWYUSupport == IWYUSupport.None)
 			{
-				return [];
+				return new List<FileItem>();
 			}
 
 			// Find FileItems for module's pch files
@@ -1245,12 +1226,12 @@ namespace UnrealBuildTool
 				SharedPchFileItem = FileItem.GetItemByFileReference(FileReference.Combine(ModuleDirectory, Rules.SharedPCHHeaderFile));
 			}
 
-			Dictionary<string, FileItem> NameToFileItem = [];
+			Dictionary<string, FileItem> NameToFileItem = new();
 
-			HashSet<FileItem> CollidingFiles = [];
+			HashSet<FileItem> CollidingFiles = new();
 
 			// Collect the headers that should be built
-			List<FileItem> HeaderFileItems = [];
+			List<FileItem> HeaderFileItems = new();
 			foreach (FileItem HeaderFileItem in InputFiles.HeaderFiles)
 			{
 				// We don't want to build pch files in iwyu, skip those.
@@ -1283,7 +1264,7 @@ namespace UnrealBuildTool
 		}
 
 		// Cache of files found while searching the includes
-		static readonly ConcurrentDictionary<DirectoryReference, HashSet<FileReference>?> KnownIncludeFilesDict = new();
+		static ConcurrentDictionary<DirectoryReference, HashSet<FileReference>?> KnownIncludeFilesDict = new();
 		static HashSet<FileReference>? GetIncludeFiles(DirectoryReference Directory)
 		{
 			if (Directory == null)
@@ -1296,7 +1277,7 @@ namespace UnrealBuildTool
 				{
 					try
 					{
-						return [.. DirectoryLookupCache.EnumerateFiles(Directory)];
+						return DirectoryLookupCache.EnumerateFiles(Directory).ToHashSet();
 					}
 					catch (Exception)
 					{
@@ -1312,17 +1293,6 @@ namespace UnrealBuildTool
 
 		private void FindIncludedHeaders(HashSet<string> VisitedIncludes, HashSet<FileItem> FoundIncludeFileItems, HashSet<string> IncludesNotFound, CppCompileEnvironment CompileEnvironment, FileItem FileToSearch, ILogger Logger)
 		{
-			static FileItem? SearchForFileItem(DirectoryReference dir, string tranformedHeaderInclude)
-			{
-				FileReference FileRef = FileReference.Combine(dir, tranformedHeaderInclude);
-				HashSet<FileReference>? Files = GetIncludeFiles(FileRef.Directory);
-				if (Files != null && Files.Contains(FileRef))
-				{
-					return FileItem.GetItemByFileReference(FileRef);
-				}
-				return null;
-			}
-
 			foreach (string HeaderInclude in CompileEnvironment.MetadataCache.GetHeaderIncludes(FileToSearch))
 			{
 				string TranformedHeaderInclude = HeaderInclude;
@@ -1334,15 +1304,26 @@ namespace UnrealBuildTool
 
 				if (VisitedIncludes.Add(TranformedHeaderInclude))
 				{
+					Func<DirectoryReference, FileItem?> SearchForFileItem = (DirectoryReference dir) =>
+					{
+						FileReference FileRef = FileReference.Combine(dir, TranformedHeaderInclude);
+						HashSet<FileReference>? Files = GetIncludeFiles(FileRef.Directory);
+						if (Files != null && Files.Contains(FileRef))
+						{
+							return FileItem.GetItemByFileReference(FileRef);
+						}
+						return null;
+					};
+
 					// check to see if the file is a relative path first
-					FileItem? IncludeFileItem = SearchForFileItem(FileToSearch.Directory.Location, TranformedHeaderInclude);
+					FileItem? IncludeFileItem = SearchForFileItem(FileToSearch.Directory.Location);
 
 					// search through the include paths if the file isn't relative
 					if (IncludeFileItem == null)
 					{
 						foreach (DirectoryReference IncludePath in CompileEnvironment.UserIncludePaths)
 						{
-							IncludeFileItem = SearchForFileItem(IncludePath, TranformedHeaderInclude);
+							IncludeFileItem = SearchForFileItem(IncludePath);
 							if (IncludeFileItem != null)
 							{
 								break;
@@ -1362,19 +1343,19 @@ namespace UnrealBuildTool
 						if (!TranformedHeaderInclude.Contains('.'))
 						{
 #if DEBUG
-							Logger.LogDebug("{Name} SharedPCH - Skipping '{Header}' found in '{File}' because it doesn't appear to be a module header.", Name, TranformedHeaderInclude, FileToSearch.Location);
+							Logger.LogDebug("{0} SharedPCH - Skipping '{1}' found in '{2}' because it doesn't appear to be a module header.", Name, TranformedHeaderInclude, FileToSearch.Location);
 #endif
 						}
 						else if (TranformedHeaderInclude.EndsWith(".generated.h"))
 						{
 #if DEBUG
-							Logger.LogDebug("{Name} SharedPCH - Skipping '{Header}' found in '{File}' because it appears to be a generated UHT header.", Name, TranformedHeaderInclude, FileToSearch.Location);
+							Logger.LogDebug("{0} SharedPCH - Skipping '{1}' found in '{2}' because it appears to be a generated UHT header.", Name, TranformedHeaderInclude, FileToSearch.Location);
 #endif
 						}
 						else
 						{
 #if DEBUG
-							Logger.LogDebug("{Name} SharedPCH - Could not find include directory for '{Header}' found in '{File}'.", Name, TranformedHeaderInclude, FileToSearch.Location);
+							Logger.LogDebug("{0} SharedPCH - Could not find include directory for '{1}' found in '{2}'.", Name, TranformedHeaderInclude, FileToSearch.Location);
 #endif
 							IncludesNotFound.Add(TranformedHeaderInclude);
 						}
@@ -1392,15 +1373,15 @@ namespace UnrealBuildTool
 		/// <returns>Returns an optimized list of modules this module's shared PCH uses</returns>
 		private HashSet<UEBuildModule> GetSharedPCHModuleDependencies(CppCompileEnvironment CompileEnvironment, FileItem PCHHeaderFile, ILogger Logger)
 		{
-			Dictionary<UEBuildModule, bool> ModuleToIncludePathsOnlyFlag = [];
+			Dictionary<UEBuildModule, bool> ModuleToIncludePathsOnlyFlag = new Dictionary<UEBuildModule, bool>();
 			FindModulesInPublicCompileEnvironment(ModuleToIncludePathsOnlyFlag);
 
-			HashSet<UEBuildModule> AllModuleDeps = [.. ModuleToIncludePathsOnlyFlag.Keys];
-			HashSet<UEBuildModule> ModuleDeps = [.. AllModuleDeps, this];
+			HashSet<UEBuildModule> AllModuleDeps = ModuleToIncludePathsOnlyFlag.Keys.ToHashSet();
+			HashSet<UEBuildModule> ModuleDeps = new HashSet<UEBuildModule>(AllModuleDeps) { this };
 
-			HashSet<string> VisitedIncludes = [];
-			HashSet<FileItem> FoundIncludeFileItems = [];
-			HashSet<string> IncludesNotFound = [];
+			HashSet<string> VisitedIncludes = new();
+			HashSet<FileItem> FoundIncludeFileItems = new();
+			HashSet<string> IncludesNotFound = new();
 
 			// pre-populate all the known file locations for this compile environment
 			foreach (DirectoryReference dir in CompileEnvironment.UserIncludePaths)
@@ -1409,14 +1390,14 @@ namespace UnrealBuildTool
 			}
 
 			FindIncludedHeaders(VisitedIncludes, FoundIncludeFileItems, IncludesNotFound, CompileEnvironment, PCHHeaderFile, Logger);
-			if (IncludesNotFound.Count != 0)
+			if (IncludesNotFound.Any())
 			{
 				// We are assuming that the includes not found aren't important. If this is found to be not true then we should return all the module dependencies
 				//return AllModuleDeps;
 			}
 
 			bool FoundAllModules = true;
-			HashSet<UEBuildModule> OptModules = [];
+			HashSet<UEBuildModule> OptModules = new();
 			foreach (FileItem IncludeFile in FoundIncludeFileItems)
 			{
 				if (CompileEnvironment.MetadataCache.UsesAPIDefine(IncludeFile) || CompileEnvironment.MetadataCache.ContainsReflectionMarkup(IncludeFile))
@@ -1427,7 +1408,7 @@ namespace UnrealBuildTool
 #if DEBUG
 						if (ModuleToIncludePathsOnlyFlag[FoundModule])
 						{
-							Logger.LogDebug("{Name} SharedPCH - '{Include}' is exporting types but '{Module}' isn't declared as a public dependency.", Name, IncludeFile.Location, FoundModule);
+							Logger.LogDebug("{0} SharedPCH - '{1}' is exporting types but '{2}' isn't declared as a public dependency.", Name, IncludeFile.Location, FoundModule);
 						}
 #endif
 
@@ -1436,7 +1417,7 @@ namespace UnrealBuildTool
 					else
 					{
 #if DEBUG
-						Logger.LogDebug("{Name} SharedPCH - '{Include}' is exporting types but the module this file belongs to isn't declared as a public dependency or include.", Name, IncludeFile.Location);
+						Logger.LogDebug("{0} SharedPCH - '{1}' is exporting types but the module this file belongs to isn't declared as a public dependency or include.", Name, IncludeFile.Location);
 #endif
 						FoundAllModules = false;
 					}
@@ -1444,7 +1425,7 @@ namespace UnrealBuildTool
 				else
 				{
 #if DEBUG
-					Logger.LogDebug("{Name} SharedPCH - '{Include}' is not exporting types so we are ignoring the dependency", Name, IncludeFile.Location);
+					Logger.LogDebug("{0} SharedPCH - '{1}' is not exporting types so we are ignoring the dependency", Name, IncludeFile.Location);
 #endif
 				}
 			}
@@ -1452,7 +1433,7 @@ namespace UnrealBuildTool
 			if (!FoundAllModules)
 			{
 #if DEBUG
-				Logger.LogDebug("{Name} SharedPCH - Is missing public dependencies. To be safe, this shared PCH will fall back to use all the module dependencies. Note that this could affect compile times.", Name);
+				Logger.LogDebug("{0} SharedPCH - Is missing public dependencies. To be safe, this shared PCH will fall back to use all the module dependencies. Note that this could affect compile times.", Name);
 #endif
 				return AllModuleDeps;
 			}
@@ -1489,7 +1470,7 @@ namespace UnrealBuildTool
 
 		static HashSet<string> GetImmutableDefinitions(List<string> Definitions)
 		{
-			HashSet<string> ImmutableDefinitions = [];
+			HashSet<string> ImmutableDefinitions = new();
 			foreach (string Definition in Definitions)
 			{
 				if (Definition.Contains("UE_IS_ENGINE_MODULE", StringComparison.Ordinal) ||
@@ -1523,7 +1504,7 @@ namespace UnrealBuildTool
 			FileItem WrapperFile = CreatePCHWrapperFile(WrapperLocation, DefinitionsFileItem!, HeaderFile, Graph);
 
 			// Create a new C++ environment that is used to create the PCH.
-			CppCompileEnvironment CompileEnvironment = new(ModuleCompileEnvironment);
+			CppCompileEnvironment CompileEnvironment = new CppCompileEnvironment(ModuleCompileEnvironment);
 			CompileEnvironment.Definitions.Clear();
 			CompileEnvironment.PrecompiledHeaderAction = PrecompiledHeaderAction.Create;
 			CompileEnvironment.PrecompiledHeaderIncludeFilename = WrapperFile.Location;
@@ -1538,7 +1519,7 @@ namespace UnrealBuildTool
 			}
 			else
 			{
-				Output = ToolChain.CompileAllCPPFiles(CompileEnvironment, [WrapperFile], IntermediateDirectory, Name, Graph);
+				Output = ToolChain.CompileAllCPPFiles(CompileEnvironment, new List<FileItem>() { WrapperFile }, IntermediateDirectory, Name, Graph);
 			}
 
 			return new PrecompiledHeaderInstance(WrapperFile, DefinitionsFileItem, CompileEnvironment, Output, GetImmutableDefinitions(ModuleCompileEnvironment.Definitions));
@@ -1562,7 +1543,7 @@ namespace UnrealBuildTool
 				// Modify definitions if we need to create a new shared pch for the include order
 				if (ModuleCompileEnvironment.IncludeOrderVersion != Template.BaseCompileEnvironment.IncludeOrderVersion)
 				{
-					Definitions = [.. Definitions];
+					Definitions = new List<string>(Definitions);
 					foreach (string OldDefine in EngineIncludeOrderHelper.GetDeprecationDefines(Template.BaseCompileEnvironment.IncludeOrderVersion))
 					{
 						Definitions.Remove(OldDefine);
@@ -1573,7 +1554,7 @@ namespace UnrealBuildTool
 				// Modify definitions if we need to create a new shared pch for project modules
 				if (ModuleCompileEnvironment.bTreatAsEngineModule != Template.BaseCompileEnvironment.bTreatAsEngineModule)
 				{
-					Definitions = [.. Definitions];
+					Definitions = new List<string>(Definitions);
 					Definitions.RemoveAll(x => x.Contains("UE_IS_ENGINE_MODULE", StringComparison.Ordinal));
 					Definitions.RemoveAll(x => x.Contains("UE_DEPRECATED_FORGAME", StringComparison.Ordinal));
 					Definitions.RemoveAll(x => x.Contains("UE_DEPRECATED_FORENGINE", StringComparison.Ordinal));
@@ -1586,7 +1567,7 @@ namespace UnrealBuildTool
 				// Modify definitions if we need to create a new shared pch for validating format strings
 				if (ModuleCompileEnvironment.bValidateFormatStrings != Template.BaseCompileEnvironment.bValidateFormatStrings)
 				{
-					Definitions = [.. Definitions];
+					Definitions = new List<string>(Definitions);
 					Definitions.RemoveAll(x => x.Contains("UE_VALIDATE_FORMAT_STRINGS", StringComparison.Ordinal));
 					Definitions.Add($"UE_VALIDATE_FORMAT_STRINGS={(ModuleCompileEnvironment.bValidateFormatStrings ? "1" : "0")}");
 				}
@@ -1594,7 +1575,7 @@ namespace UnrealBuildTool
 				// Modify definitions if we need to create a new shared pch for validating internal api usage
 				if (ModuleCompileEnvironment.bValidateInternalApi != Template.BaseCompileEnvironment.bValidateInternalApi)
 				{
-					Definitions = [.. Definitions];
+					Definitions = new List<string>(Definitions);
 					Definitions.RemoveAll(x => x.Contains("UE_VALIDATE_INTERNAL_API", StringComparison.Ordinal));
 					Definitions.Add($"UE_VALIDATE_INTERNAL_API={(ModuleCompileEnvironment.bValidateInternalApi ? "1" : "0")}");
 				}
@@ -1602,7 +1583,7 @@ namespace UnrealBuildTool
 				// Modify definitions if we need to create a new shared pch for validating experimental api usage
 				if (ModuleCompileEnvironment.bValidateExperimentalApi != Template.BaseCompileEnvironment.bValidateExperimentalApi)
 				{
-					Definitions = [.. Definitions];
+					Definitions = new List<string>(Definitions);
 					Definitions.RemoveAll(x => x.Contains("UE_VALIDATE_EXPERIMENTAL_API", StringComparison.Ordinal));
 					Definitions.Add($"UE_VALIDATE_EXPERIMENTAL_API={(ModuleCompileEnvironment.bValidateExperimentalApi ? "1" : "0")}");
 				}
@@ -1611,8 +1592,8 @@ namespace UnrealBuildTool
 				string Variant = GetSuffixForSharedPCH(ModuleCompileEnvironment, Template.BaseCompileEnvironment);
 
 				FileReference SharedDefinitionsLocation = FileReference.Combine(Template.OutputDir, $"SharedDefinitions.{Template.Module.Name}{Variant}.h");
-				List<string> NewDefinitions = [];
-				StringBuilder Writer = new();
+				List<string> NewDefinitions = new();
+				StringBuilder Writer = new StringBuilder();
 				WriteDefinitions($"Shared Definitions for {Template.Module.Name}{Variant}", [], Definitions, [], Writer);
 				FileItem SharedDefinitionsFileItem = Graph.CreateIntermediateTextFile(SharedDefinitionsLocation, Writer.ToString(), AllowAsync: false);
 
@@ -1621,7 +1602,7 @@ namespace UnrealBuildTool
 				FileItem WrapperFile = CreatePCHWrapperFile(WrapperLocation, SharedDefinitionsFileItem, Template.HeaderFile, Graph);
 
 				// Create the compile environment for this PCH
-				CppCompileEnvironment CompileEnvironment = new(Template.BaseCompileEnvironment);
+				CppCompileEnvironment CompileEnvironment = new CppCompileEnvironment(Template.BaseCompileEnvironment);
 				CompileEnvironment.Definitions.Clear();
 				CompileEnvironment.PrecompiledHeaderAction = PrecompiledHeaderAction.Create;
 				CompileEnvironment.PrecompiledHeaderIncludeFilename = WrapperFile.Location;
@@ -1632,8 +1613,8 @@ namespace UnrealBuildTool
 				if (Rules.Target.bChainPCHs)
 				{
 					// Find all the dependencies of this module
-					HashSet<UEBuildModule> ReferencedModules = [];
-					Template.Module.GetAllDependencyModules([], ReferencedModules, bIncludeDynamicallyLoaded: false, bForceCircular: false, bOnlyDirectDependencies: false);
+					HashSet<UEBuildModule> ReferencedModules = new HashSet<UEBuildModule>();
+					Template.Module.GetAllDependencyModules(new List<UEBuildModule>(), ReferencedModules, bIncludeDynamicallyLoaded: false, bForceCircular: false, bOnlyDirectDependencies: false);
 
 					PrecompiledHeaderTemplate? ParentTemplate = CompileEnvironment.SharedPCHs
 						.Skip(CompileEnvironment.SharedPCHs.IndexOf(Template) + 1)
@@ -1653,7 +1634,7 @@ namespace UnrealBuildTool
 				}
 				else
 				{
-					Output = ToolChain.CompileAllCPPFiles(CompileEnvironment, [WrapperFile], Template.OutputDir, "Shared", Graph);
+					Output = ToolChain.CompileAllCPPFiles(CompileEnvironment, new List<FileItem>() { WrapperFile }, Template.OutputDir, "Shared", Graph);
 				}
 				Instance = new PrecompiledHeaderInstance(WrapperFile, SharedDefinitionsFileItem, CompileEnvironment, Output, GetImmutableDefinitions(Template.BaseCompileEnvironment.Definitions))
 				{
@@ -1894,13 +1875,14 @@ namespace UnrealBuildTool
 			CppCompileEnvironment ModuleCompileEnvironment,
 			List<FileItem> NormalFiles,
 			List<FileItem> AdaptiveFiles,
-			IActionGraphBuilder Graph)
+			IActionGraphBuilder Graph,
+			ILogger Logger)
 		{
 			bool isEngine = Rules.File.IsUnderDirectory(Unreal.EngineDirectory) || Rules.bTreatAsEngineModule || Rules.PrivatePCHHeaderFile == null;
 			bool bAdaptiveUnityDisablesPCH = isEngine ? Target.bAdaptiveUnityDisablesPCH : Target.bAdaptiveUnityDisablesPCHForProject;
 			string adaptiveUnityDisablesPCHName = isEngine ? nameof(Target.bAdaptiveUnityDisablesPCH) : nameof(Target.bAdaptiveUnityDisablesPCHForProject);
 
-			CPPOutput OutputFiles = new();
+			CPPOutput OutputFiles = new CPPOutput();
 
 			if (NormalFiles.Count > 0)
 			{
@@ -1930,11 +1912,11 @@ namespace UnrealBuildTool
 					Graph.AddDiagnostic("[Adaptive Build] Enabling Edit & Continue for excluded files. Set bAdaptiveUnityEnablesEditAndContinue to false in BuildConfiguration.xml to change this behavior.");
 				}
 
-				Graph.AddDiagnostic($"[Adaptive Build] Excluded from {Name} unity file: {String.Join(", ", AdaptiveFiles.Select(File => Path.GetFileName(File.AbsolutePath)))}");
+				Graph.AddDiagnostic($"[Adaptive Build] Excluded from {Name} unity file: " + String.Join(", ", AdaptiveFiles.Select(File => Path.GetFileName(File.AbsolutePath))));
 
 				// Create the new compile environment. Always turn off PCH due to different compiler settings.
-				CppCompileEnvironment AdaptiveUnityEnvironment = new(ModuleCompileEnvironment);
-				AdaptiveUnityEnvironment.AdditionalPrerequisites = [.. AdaptiveUnityEnvironment.AdditionalPrerequisites.Concat(CompileEnvironment.AdditionalPrerequisites).Distinct()];
+				CppCompileEnvironment AdaptiveUnityEnvironment = new CppCompileEnvironment(ModuleCompileEnvironment);
+				AdaptiveUnityEnvironment.AdditionalPrerequisites = new(AdaptiveUnityEnvironment.AdditionalPrerequisites.Concat(CompileEnvironment.AdditionalPrerequisites).Distinct());
 				AdaptiveUnityEnvironment.UserIncludePaths.UnionWith(CompileEnvironment.SharedUserIncludePaths);
 				AdaptiveUnityEnvironment.bDisableStaticAnalysis = CompileEnvironment.bDisableStaticAnalysis;
 				if (Target.bAdaptiveUnityDisablesOptimizations)
@@ -1996,11 +1978,11 @@ namespace UnrealBuildTool
 
 		static CPPOutput CompileAdaptiveNonUnityFilesWithDedicatedPCH(UEToolChain ToolChain, CppCompileEnvironment CompileEnvironment, List<FileItem> Files, DirectoryReference IntermediateDirectory, string ModuleName, IActionGraphBuilder Graph)
 		{
-			CPPOutput Output = new();
+			CPPOutput Output = new CPPOutput();
 			foreach (FileItem File in Files)
 			{
 				// Build the contents of the wrapper file
-				StringBuilder WrapperContents = new();
+				StringBuilder WrapperContents = new StringBuilder();
 				{
 					string FileString = File.AbsolutePath;
 					if (File.Location.IsUnderDirectory(Unreal.RootDirectory))
@@ -2008,10 +1990,12 @@ namespace UnrealBuildTool
 						FileString = File.Location.MakeRelativeTo(Unreal.EngineSourceDirectory);
 					}
 					FileString = FileString.Replace('\\', '/');
-					WriteDefinitions($"Dedicated PCH for {FileString}", [], CompileEnvironment.Definitions, [], WrapperContents);
+					WriteDefinitions($"Dedicated PCH for {FileString}", Array.Empty<string>(), CompileEnvironment.Definitions, Array.Empty<string>(), WrapperContents);
 					WrapperContents.AppendLine();
-					using StreamReader Reader = new(File.Location.FullName);
-					CppIncludeParser.CopyIncludeDirectives(Reader, WrapperContents);
+					using (StreamReader Reader = new StreamReader(File.Location.FullName))
+					{
+						CppIncludeParser.CopyIncludeDirectives(Reader, WrapperContents);
+					}
 				}
 
 				// Write the PCH header
@@ -2019,25 +2003,25 @@ namespace UnrealBuildTool
 				FileItem DedicatedPchFile = Graph.CreateIntermediateTextFile(DedicatedPchLocation, WrapperContents.ToString());
 
 				// Create a new C++ environment to compile the PCH
-				CppCompileEnvironment PchEnvironment = new(CompileEnvironment);
+				CppCompileEnvironment PchEnvironment = new CppCompileEnvironment(CompileEnvironment);
 				PchEnvironment.Definitions.Clear();
 				PchEnvironment.UserIncludePaths.Add(File.Location.Directory); // Need to be able to include headers in the same directory as the source file
 				PchEnvironment.PrecompiledHeaderAction = PrecompiledHeaderAction.Create;
 				PchEnvironment.PrecompiledHeaderIncludeFilename = DedicatedPchFile.Location;
 
 				// Create the action to compile the PCH file.
-				CPPOutput PchOutput = ToolChain.CompileAllCPPFiles(PchEnvironment, [DedicatedPchFile], IntermediateDirectory, ModuleName, Graph);
+				CPPOutput PchOutput = ToolChain.CompileAllCPPFiles(PchEnvironment, new List<FileItem>() { DedicatedPchFile }, IntermediateDirectory, ModuleName, Graph);
 				Output.ObjectFiles.AddRange(PchOutput.ObjectFiles);
 
 				// Create a new C++ environment to compile the original file
-				CppCompileEnvironment FileEnvironment = new(CompileEnvironment);
+				CppCompileEnvironment FileEnvironment = new CppCompileEnvironment(CompileEnvironment);
 				FileEnvironment.Definitions.Clear();
 				FileEnvironment.PrecompiledHeaderAction = PrecompiledHeaderAction.Include;
 				FileEnvironment.PrecompiledHeaderIncludeFilename = DedicatedPchFile.Location;
 				FileEnvironment.PCHInstance = new PrecompiledHeaderInstance(DedicatedPchFile, DedicatedPchFile, PchEnvironment, PchOutput, GetImmutableDefinitions(CompileEnvironment.Definitions));
 
 				// Create the action to compile the PCH file.
-				CPPOutput FileOutput = ToolChain.CompileAllCPPFiles(FileEnvironment, [File], IntermediateDirectory, ModuleName, Graph);
+				CPPOutput FileOutput = ToolChain.CompileAllCPPFiles(FileEnvironment, new List<FileItem>() { File }, IntermediateDirectory, ModuleName, Graph);
 				Output.ObjectFiles.AddRange(FileOutput.ObjectFiles);
 			}
 			return Output;
@@ -2051,8 +2035,8 @@ namespace UnrealBuildTool
 		/// <param name="CompileEnvironment">The current compile environment</param>
 		/// <param name="LinkInputFiles">List of files that will be linked for the target</param>
 		/// <param name="Graph">List of build actions</param>
-		/// 
-		CppCompileEnvironment SetupPrecompiledHeaders(ReadOnlyTargetRules Target, UEToolChain? ToolChain, CppCompileEnvironment CompileEnvironment, List<FileItem> LinkInputFiles, IActionGraphBuilder Graph)
+		/// <param name="Logger">Logger for output</param>
+		CppCompileEnvironment SetupPrecompiledHeaders(ReadOnlyTargetRules Target, UEToolChain? ToolChain, CppCompileEnvironment CompileEnvironment, List<FileItem> LinkInputFiles, IActionGraphBuilder Graph, ILogger Logger)
 		{
 			if (Target.bUsePCHFiles && Rules.PCHUsage != ModuleRules.PCHUsageMode.NoPCHs)
 			{
@@ -2086,9 +2070,10 @@ namespace UnrealBuildTool
 						FileItem PrivateDefinitionsFileItem;
 						{
 							// Only add new definitions that are not already existing in the shared pch
-							List<string> NewDefinitions = [];
-							List<string> Undefinitions = [];
+							List<string> NewDefinitions = new();
+							List<string> Undefintions = new();
 							bool ModuleApiUndef = false;
+							bool ModuleNonAttributedUndef = false;
 							foreach (string Definition in CompileEnvironment.Definitions)
 							{
 								if (Instance.ImmutableDefinitions.Contains(Definition))
@@ -2102,12 +2087,19 @@ namespace UnrealBuildTool
 								if (!ModuleApiUndef && Definition.StartsWith(ModuleApiDefine, StringComparison.Ordinal))
 								{
 									ModuleApiUndef = true;
-									Undefinitions.Add(ModuleApiDefine);
+									Undefintions.Add(ModuleApiDefine);
+								}
+
+								// Remove the module _NON_ATTRIBUTED_API definition for cases where there are circular dependencies between the shared PCH module and modules using it
+								if (!ModuleNonAttributedUndef && Definition.StartsWith(ModuleNonAttributedDefine, StringComparison.Ordinal))
+								{
+									ModuleNonAttributedUndef = true;
+									Undefintions.Add(ModuleNonAttributedDefine);
 								}
 							}
 
 							StringBuilder Writer = new();
-							WriteDefinitions($"Shared PCH Definitions for {Name}", [Instance.DefinitionsFile.Name], NewDefinitions, Undefinitions, Writer);
+							WriteDefinitions($"Shared PCH Definitions for {Name}", new string[] { Instance.DefinitionsFile.Name }, NewDefinitions, Undefintions, Writer);
 							PrivateDefinitionsFileItem = Graph.CreateIntermediateTextFile(PrivateDefinitionsFile, Writer.ToString());
 						}
 
@@ -2149,7 +2141,7 @@ namespace UnrealBuildTool
 				FileReference PrivateDefinitionsFile = FileReference.Combine(IntermediateDirectory, PrivateDefinitionsName);
 				{
 					StringBuilder Writer = new();
-					WriteDefinitions("Definitions", [], CompileEnvironment.Definitions, [], Writer);
+					WriteDefinitions("Definitions", Array.Empty<string>(), CompileEnvironment.Definitions, Array.Empty<string>(), Writer);
 					CompileEnvironment.Definitions.Clear();
 
 					FileItem PrivateDefinitionsFileItem = Graph.CreateIntermediateTextFile(PrivateDefinitionsFile, Writer.ToString());
@@ -2173,7 +2165,7 @@ namespace UnrealBuildTool
 		static FileItem CreatePCHWrapperFile(FileReference OutputFile, FileItem DefinitionsFile, FileItem IncludedFile, IActionGraphBuilder Graph)
 		{
 			// Build the contents of the wrapper file
-			StringBuilder WrapperContents = new();
+			StringBuilder WrapperContents = new StringBuilder();
 			{
 				WrapperContents.AppendLine("// PCH for {0}", IncludedFile.Name);
 				WrapperContents.AppendLine("#include \"{0}\"", DefinitionsFile.Name);
@@ -2216,7 +2208,7 @@ namespace UnrealBuildTool
 				Writer.Append("#undef ").Append(Undefinition).AppendLine();
 			}
 
-			HashSet<string> Processed = [];
+			HashSet<string> Processed = new();
 			foreach (string Definition in Definitions.Where(x => Processed.Add(x)))
 			{
 				int EqualsIdx = Definition.IndexOf('=', StringComparison.Ordinal);
@@ -2245,14 +2237,14 @@ namespace UnrealBuildTool
 				if (InvalidIncludeDirectives == null)
 				{
 					// Find headers used by the source file.
-					Dictionary<string, FileReference> NameToHeaderFile = [];
+					Dictionary<string, FileReference> NameToHeaderFile = new Dictionary<string, FileReference>();
 					foreach (FileItem HeaderFile in HeaderFiles)
 					{
 						NameToHeaderFile[HeaderFile.Location.GetFileNameWithoutExtension()] = HeaderFile.Location;
 					}
 
 					// Find the directly included files for each source file, and make sure it includes the matching header if possible
-					InvalidIncludeDirectives = [];
+					InvalidIncludeDirectives = new();
 					if (Rules != null && Rules.IWYUSupport != IWYUSupport.None && Target.bEnforceIWYU)
 					{
 						foreach (FileItem CppFile in CppFiles)
@@ -2264,7 +2256,8 @@ namespace UnrealBuildTool
 								string ExpectedName = CppFile.Location.GetFileNameWithoutExtension();
 								if (!String.Equals(IncludeName, ExpectedName, StringComparison.OrdinalIgnoreCase))
 								{
-									if (NameToHeaderFile.TryGetValue(ExpectedName, out FileReference? HeaderFile))
+									FileReference? HeaderFile;
+									if (NameToHeaderFile.TryGetValue(ExpectedName, out HeaderFile))
 									{
 										InvalidIncludeDirectives.Add(new InvalidIncludeDirective(CppFile.Location, HeaderFile));
 									}
@@ -2434,19 +2427,24 @@ namespace UnrealBuildTool
 			{
 				return false;
 			}
-			return Setting switch
+			switch (Setting)
 			{
-				ModuleRules.CodeOptimization.Never => false,
-				ModuleRules.CodeOptimization.Default or ModuleRules.CodeOptimization.InNonDebugBuilds => Configuration != UnrealTargetConfiguration.Debug && (Configuration != UnrealTargetConfiguration.DebugGame || bIsEngineModule),
-				ModuleRules.CodeOptimization.InShippingBuildsOnly => (Configuration == UnrealTargetConfiguration.Shipping),
-				_ => true,
-			};
+				case ModuleRules.CodeOptimization.Never:
+					return false;
+				case ModuleRules.CodeOptimization.Default:
+				case ModuleRules.CodeOptimization.InNonDebugBuilds:
+					return Configuration != UnrealTargetConfiguration.Debug && (Configuration != UnrealTargetConfiguration.DebugGame || bIsEngineModule);
+				case ModuleRules.CodeOptimization.InShippingBuildsOnly:
+					return (Configuration == UnrealTargetConfiguration.Shipping);
+				default:
+					return true;
+			}
 		}
 
 		public CppCompileEnvironment CreateCompileEnvironmentForIntellisense(ReadOnlyTargetRules Target, CppCompileEnvironment BaseCompileEnvironment, ILogger Logger)
 		{
 			CppCompileEnvironment CompileEnvironment = CreateModuleCompileEnvironment(Target, BaseCompileEnvironment, Logger);
-			CompileEnvironment = SetupPrecompiledHeaders(Target, null, CompileEnvironment, [], new NullActionGraphBuilder(Logger));
+			CompileEnvironment = SetupPrecompiledHeaders(Target, null, CompileEnvironment, new List<FileItem>(), new NullActionGraphBuilder(Logger), Logger);
 			CreateHeaderForDefinitions(CompileEnvironment, IntermediateDirectory, null, new NullActionGraphBuilder(Logger));
 			return CompileEnvironment;
 		}
@@ -2460,12 +2458,11 @@ namespace UnrealBuildTool
 		/// <returns>The new module compile environment.</returns>
 		public CppCompileEnvironment CreateModuleCompileEnvironment(ReadOnlyTargetRules Target, CppCompileEnvironment BaseCompileEnvironment, ILogger Logger)
 		{
-			CppCompileEnvironment Result = new(BaseCompileEnvironment)
-			{
-				// Override compile environment
-				bUseUnity = Rules.bUseUnity,
-				bCodeCoverage = Target.bCodeCoverage
-			};
+			CppCompileEnvironment Result = new CppCompileEnvironment(BaseCompileEnvironment);
+
+			// Override compile environment
+			Result.bUseUnity = Rules.bUseUnity;
+			Result.bCodeCoverage = Target.bCodeCoverage;
 			Result.bOptimizeCode = ShouldEnableOptimization(Rules.OptimizeCode, Target.Configuration, Rules.bTreatAsEngineModule, Result.bCodeCoverage);
 			Result.bUseRTTI |= Rules.bUseRTTI;
 			Result.bVcRemoveUnreferencedComdat &= Rules.bVcRemoveUnreferencedComdat;
@@ -2627,7 +2624,7 @@ namespace UnrealBuildTool
 				Result.Definitions.Add(String.Format("UE_WITH_DEBUG_INFO={0}", BaseCompileEnvironment.bCreateDebugInfo ? "1" : "0"));
 
 				StringBuilder sb = new();
-				foreach ((uint id, DirectoryReference vfs, DirectoryReference local) in BaseCompileEnvironment.RootPaths)
+				foreach (var (id, vfs, local) in BaseCompileEnvironment.RootPaths)
 				{
 					sb.Append(vfs.FullName.Replace("\\", "\\\\")).Append(';');
 					sb.Append(local.FullName.Replace("\\", "\\\\")).Append(';');
@@ -2690,20 +2687,19 @@ namespace UnrealBuildTool
 		/// <returns>The new shared PCH compile environment.</returns>
 		public CppCompileEnvironment CreateSharedPCHCompileEnvironment(UEBuildTarget Target, CppCompileEnvironment BaseCompileEnvironment)
 		{
-			CppCompileEnvironment CompileEnvironment = new(BaseCompileEnvironment)
-			{
-				// Use the default optimization setting for
-				bOptimizeCode = ShouldEnableOptimization(ModuleRules.CodeOptimization.Default, Target.Configuration, Rules.bTreatAsEngineModule, Rules.bCodeCoverage),
-				bCodeCoverage = Rules.bCodeCoverage,
-				bTreatAsEngineModule = Rules.bTreatAsEngineModule,
-				bValidateFormatStrings = Rules.bValidateFormatStrings,
-				bValidateInternalApi = Rules.bValidateInternalApi,
-				bValidateExperimentalApi = Rules.bValidateExperimentalApi,
+			CppCompileEnvironment CompileEnvironment = new CppCompileEnvironment(BaseCompileEnvironment);
 
-				// Override compile environment
-				bIsBuildingDLL = !Target.ShouldCompileMonolithic(),
-				bIsBuildingLibrary = false
-			};
+			// Use the default optimization setting for
+			CompileEnvironment.bOptimizeCode = ShouldEnableOptimization(ModuleRules.CodeOptimization.Default, Target.Configuration, Rules.bTreatAsEngineModule, Rules.bCodeCoverage);
+			CompileEnvironment.bCodeCoverage = Rules.bCodeCoverage;
+			CompileEnvironment.bTreatAsEngineModule = Rules.bTreatAsEngineModule;
+			CompileEnvironment.bValidateFormatStrings = Rules.bValidateFormatStrings;
+			CompileEnvironment.bValidateInternalApi = Rules.bValidateInternalApi;
+			CompileEnvironment.bValidateExperimentalApi = Rules.bValidateExperimentalApi;
+
+			// Override compile environment
+			CompileEnvironment.bIsBuildingDLL = !Target.ShouldCompileMonolithic();
+			CompileEnvironment.bIsBuildingLibrary = false;
 
 			// Add a macro for when we're compiling an engine module, to enable additional compiler diagnostics through code.
 			CompileEnvironment.Definitions.Add($"UE_IS_ENGINE_MODULE={(Rules.bTreatAsEngineModule || Rules.Target.bWarnAboutMonolithicHeadersIncluded ? "1" : "0")}");
@@ -2724,7 +2720,7 @@ namespace UnrealBuildTool
 			CompileEnvironment.Definitions.AddRange(PublicDefinitions);
 
 			// Find all the modules that are part of the public compile environment for this module.
-			Dictionary<UEBuildModule, bool> ModuleToIncludePathsOnlyFlag = [];
+			Dictionary<UEBuildModule, bool> ModuleToIncludePathsOnlyFlag = new Dictionary<UEBuildModule, bool>();
 			FindModulesInPublicCompileEnvironment(ModuleToIncludePathsOnlyFlag);
 
 			// Set up additional root paths
@@ -2754,9 +2750,9 @@ namespace UnrealBuildTool
 		{
 			IReadOnlySet<string> ExcludedNames = UEBuildPlatform.GetBuildPlatform(Platform).GetExcludedFolderNames();
 
-			InputFileCollection InputFiles = new();
+			InputFileCollection InputFiles = new InputFileCollection();
 
-			SourceDirectories = [];
+			SourceDirectories = new HashSet<DirectoryReference>();
 			foreach (DirectoryReference Dir in ModuleDirectories)
 			{
 				DirectoryItem ModuleDirectoryItem = DirectoryItem.GetItemByDirectoryReference(Dir);
@@ -2777,7 +2773,8 @@ namespace UnrealBuildTool
 		/// <param name="Logger">Logger for output</param>
 		static void FindInputFilesFromDirectoryRecursive(DirectoryItem BaseDirectory, IReadOnlySet<string> ExcludedNames, HashSet<DirectoryReference> SourceDirectories, Dictionary<DirectoryItem, FileItem[]> DirectoryToSourceFiles, InputFileCollection InputFiles, ILogger Logger)
 		{
-			FileItem[] SourceFiles = FindInputFilesFromDirectory(BaseDirectory, InputFiles, out bool bIgnoreFileFound, Logger);
+			bool bIgnoreFileFound;
+			FileItem[] SourceFiles = FindInputFilesFromDirectory(BaseDirectory, InputFiles, out bIgnoreFileFound, Logger);
 
 			if (bIgnoreFileFound)
 			{
@@ -2801,14 +2798,14 @@ namespace UnrealBuildTool
 
 		// List of known extensions that are checked in FindInputFilesFromDirectory, so we can log out only if a hidden file
 		// is skipped, only if it has an extension that would normally be compiled
-		static readonly string[] KnownInputFileExtensions =
-		[
+		static string[] KnownInputFileExtensions =
+		{
 			".h", ".ipsh",
 			".cpp", ".cxx", ".ixx", ".c", ".cc",
 			".m", ".mm", ".swift",
 			".rc", ".ispc",
 			".proto",
-		];
+		};
 
 		/// <summary>
 		/// Finds the input files that should be built for this module, from a given directory
@@ -2821,7 +2818,7 @@ namespace UnrealBuildTool
 		static FileItem[] FindInputFilesFromDirectory(DirectoryItem BaseDirectory, InputFileCollection InputFiles, out bool bIgnoreFileFound, ILogger Logger)
 		{
 			bIgnoreFileFound = false;
-			List<FileItem> SourceFiles = [];
+			List<FileItem> SourceFiles = new List<FileItem>();
 			foreach (FileItem InputFile in BaseDirectory.EnumerateFiles())
 			{
 				if (InputFile.Name == ".ubtignore")
@@ -2830,7 +2827,7 @@ namespace UnrealBuildTool
 				}
 				// some text editors will leave temp files that start with . next to a, say, .cpp file, and we do not want to compile them,
 				// but skip files like .DS_Store that don't have an extension (so, no more .'s after the first)
-				else if (InputFile.Name.StartsWith('.') && InputFile.Name.IndexOf('.', 1) != -1)
+				else if (InputFile.Name.StartsWith(".") && InputFile.Name.IndexOf('.', 1) != -1)
 				{
 					if (KnownInputFileExtensions.Contains(InputFile.Location.GetExtension()))
 					{
@@ -2901,19 +2898,19 @@ namespace UnrealBuildTool
 				}
 				// if you add any extension checks here, make sure to update KnownInputFileExtensions
 			}
-			return [.. SourceFiles];
+			return SourceFiles.ToArray();
 		}
 
 		public static List<string> GenerateHeaderCpp(string HeaderName, string IncludeFileString)
 		{
-			HeaderName = HeaderName[..^2].Replace(".", "_").Replace("-", "_");
-
-			return [
-				"// This file is automatically generated at compile-time to include a user created header file.",
-				"#define UE_DIRECT_HEADER_COMPILE 1",
-				$"#define __COMPILING_{HeaderName}",
-				$"#include \"{IncludeFileString.Replace('\\', '/')}\"",
-			];
+			HeaderName = HeaderName.Substring(0, HeaderName.Length - 2).Replace(".", "_").Replace("-", "_");
+
+			List<string> GeneratedHeaderCppContents = new();
+			GeneratedHeaderCppContents.Add("// This file is automatically generated at compile-time to include a user created header file.");
+			GeneratedHeaderCppContents.Add("#define UE_DIRECT_HEADER_COMPILE 1");
+			GeneratedHeaderCppContents.Add($"#define __COMPILING_{HeaderName}");
+			GeneratedHeaderCppContents.Add($"#include \"{IncludeFileString.Replace('\\', '/')}\"");
+			return GeneratedHeaderCppContents;
 		}
 
 		/// <summary>
@@ -2924,10 +2921,11 @@ namespace UnrealBuildTool
 		/// <returns>Array of source files</returns>
 		public static FileItem[] GetSourceFiles(DirectoryItem Directory, ILogger Logger)
 		{
-			FileItem[] Files = FindInputFilesFromDirectory(Directory, new InputFileCollection(), out bool bIgnoreFileFound, Logger);
+			bool bIgnoreFileFound;
+			FileItem[] Files = FindInputFilesFromDirectory(Directory, new InputFileCollection(), out bIgnoreFileFound, Logger);
 			if (bIgnoreFileFound)
 			{
-				return [];
+				return Array.Empty<FileItem>();
 			}
 			return Files;
 		}
@@ -2940,27 +2938,50 @@ namespace UnrealBuildTool
 		/// <returns>Array of source files and headers</returns>
 		public static FileItem[] GetSourceFilesAndHeaders(DirectoryItem Directory, ILogger Logger)
 		{
-			InputFileCollection InputFiles = new();
-			List<FileItem> Files = [.. FindInputFilesFromDirectory(Directory, InputFiles, out bool bIgnoreFileFound, Logger)];
+			bool bIgnoreFileFound;
+			InputFileCollection InputFiles = new InputFileCollection();
+			List<FileItem> Files = new List<FileItem>(FindInputFilesFromDirectory(Directory, InputFiles, out bIgnoreFileFound, Logger));
 			if (bIgnoreFileFound)
 			{
-				return [];
+				return Array.Empty<FileItem>();
 			}
 			Files.AddRange(InputFiles.HeaderFiles);
-			return [.. Files];
+			return Files.ToArray();
 		}
 
 		/// <summary>
 		/// Checks a given directory path whether it exists and if it contains any Verse source files
 		/// </summary>
-		public static bool IsValidVerseDirectory(DirectoryReference MaybeVerseDirectory) => DirectoryReference.Exists(MaybeVerseDirectory)
-			&& Directory.EnumerateFiles(MaybeVerseDirectory.FullName, "*.v*", SearchOption.AllDirectories).Any(x => x.EndsWith(".verse") || x.EndsWith(".vmodule"));
+		public static bool IsValidVerseDirectory(DirectoryReference MaybeVerseDirectory)
+		{
+			if (!DirectoryReference.Exists(MaybeVerseDirectory))
+			{
+				return false;
+			}
 
+			foreach (string FilePath in Directory.EnumerateFiles(MaybeVerseDirectory.FullName, "*.v*", SearchOption.AllDirectories))
+			{
+				if (FilePath.EndsWith(".verse") || FilePath.EndsWith(".vmodule"))
+				{
+					return true;
+				}
+			}
+
+			return false;
+		}
+		
 
 		/// <summary>
 		/// Checks a given directory path whether it exists and if it contains any Proto source files
 		/// </summary>
-		public static bool IsValidProtoDirectory(DirectoryReference MaybeProtoDirectory) => DirectoryReference.Exists(MaybeProtoDirectory)
-			&& Directory.EnumerateFiles(MaybeProtoDirectory.FullName, "*.proto", SearchOption.AllDirectories).Any();
+		public static bool IsValidProtoDirectory(DirectoryReference MaybeProtoDirectory)
+		{
+			if (!DirectoryReference.Exists(MaybeProtoDirectory))
+			{
+				return false;
+			}
+
+			return Directory.EnumerateFiles(MaybeProtoDirectory.FullName, "*.proto", SearchOption.AllDirectories).Any();
+		}
 	}
 }
