dotnet resolve ide0057

#rnx

diff --git a/Engine/Source/Programs/Shared/EpicGames.UHT/Utils/UhtSession.cs b/Engine/Source/Programs/Shared/EpicGames.UHT/Utils/UhtSession.cs
--- a/Engine/Source/Programs/Shared/EpicGames.UHT/Utils/UhtSession.cs
+++ b/Engine/Source/Programs/Shared/EpicGames.UHT/Utils/UhtSession.cs
@@ -334,7 +334,7 @@ namespace EpicGames.UHT.Utils
 				{
 					if (relativePath == null || (filePath.Length - includePath.Length - 1) < relativePath.Length)
 					{
-						relativePath = filePath.Substring(includePath.Length + 1);
+						relativePath = filePath[(includePath.Length + 1)..];
 					}
 				}
 			}
@@ -1740,7 +1740,7 @@ namespace EpicGames.UHT.Utils
 			{
 				if (namespaceElements[index - 1].Namespace == null)
 				{
-					GetNamespace(namespaceElements, index - 1, fullName.Substring(0, fullName.Length - 2 - namespaceElements[index].Name.Length), headerNamespaceLookupCache);
+					GetNamespace(namespaceElements, index - 1, fullName[..(fullName.Length - 2 - namespaceElements[index].Name.Length)], headerNamespaceLookupCache);
 				}
 				namespaceObj = _namespaces.GetOrAdd(fullName, (x) =>
 				{
diff --git a/Engine/Source/Programs/UnrealBuildTool/Configuration/TargetDescriptor.cs b/Engine/Source/Programs/UnrealBuildTool/Configuration/TargetDescriptor.cs
--- a/Engine/Source/Programs/UnrealBuildTool/Configuration/TargetDescriptor.cs
+++ b/Engine/Source/Programs/UnrealBuildTool/Configuration/TargetDescriptor.cs
@@ -30,7 +30,7 @@ namespace UnrealBuildTool
 		{
 			if (Name.EndsWith(TEST_TARGETS_SUFFIX))
 			{
-				return Name.Substring(0, Name.Length - TEST_TARGETS_SUFFIX.Length);
+				return Name[..^TEST_TARGETS_SUFFIX.Length];
 			}
 			return Name;
 		}
@@ -238,10 +238,10 @@ namespace UnrealBuildTool
 						throw new BuildException("Missing suffix argument from -ModuleWithSuffix=Name,Suffix");
 					}
 
-					string ModuleName = ModuleWithSuffix.Substring(0, SuffixIdx);
+					string ModuleName = ModuleWithSuffix[..SuffixIdx];
 
 					int Suffix;
-					if (!Int32.TryParse(ModuleWithSuffix.Substring(SuffixIdx + 1), out Suffix))
+					if (!Int32.TryParse(ModuleWithSuffix[(SuffixIdx + 1)..], out Suffix))
 					{
 						throw new BuildException("Suffix for modules must be an integer");
 					}
diff --git a/Engine/Source/Programs/UnrealBuildTool/Configuration/TestModuleRules.cs b/Engine/Source/Programs/UnrealBuildTool/Configuration/TestModuleRules.cs
--- a/Engine/Source/Programs/UnrealBuildTool/Configuration/TestModuleRules.cs
+++ b/Engine/Source/Programs/UnrealBuildTool/Configuration/TestModuleRules.cs
@@ -391,11 +391,11 @@ namespace UnrealBuildTool
 				int PluginFolderIndex = ModuleDirectory.IndexOf("Plugins");
 				if (PluginFolderIndex >= 0)
 				{
-					return ModuleDirectory.Substring(0, PluginFolderIndex) + "Binaries";
+					return ModuleDirectory[..PluginFolderIndex] + "Binaries";
 				}
 				throw new Exception("Could not detect source folder path for module from directory " + ModuleDirectory);
 			}
-			return ModuleDirectory.Substring(0, SourceFolderIndex) + "Binaries";
+			return ModuleDirectory[..SourceFolderIndex] + "Binaries";
 		}
 
 		private static void AppendOrUpdateRunAllTestsNode(XElement Root, string MacroName, string Platform, Dictionary<string, string> ExpandArguments)
diff --git a/Engine/Source/Programs/UnrealBuildTool/Configuration/UEBuildModule.cs b/Engine/Source/Programs/UnrealBuildTool/Configuration/UEBuildModule.cs
--- a/Engine/Source/Programs/UnrealBuildTool/Configuration/UEBuildModule.cs
+++ b/Engine/Source/Programs/UnrealBuildTool/Configuration/UEBuildModule.cs
@@ -789,44 +789,44 @@ namespace UnrealBuildTool
 					{
 						if (MatchVariableName(Path, StartIdx, EndIdx, "EngineDir"))
 						{
-							Path = Unreal.EngineDirectory + Path.Substring(EndIdx + 1);
+							Path = Unreal.EngineDirectory + Path[(EndIdx + 1)..];
 						}
 						else if (MatchVariableName(Path, StartIdx, EndIdx, "ProjectDir"))
 						{
 							if (Rules.Target.ProjectFile == null)
 							{
-								Path = Unreal.EngineDirectory + Path.Substring(EndIdx + 1);
+								Path = Unreal.EngineDirectory + Path[(EndIdx + 1)..];
 							}
 							else
 							{
-								Path = Rules.Target.ProjectFile.Directory + Path.Substring(EndIdx + 1);
+								Path = Rules.Target.ProjectFile.Directory + Path[(EndIdx + 1)..];
 							}
 						}
 						else if (MatchVariableName(Path, StartIdx, EndIdx, "ModuleDir"))
 						{
-							Path = Rules.ModuleDirectory + Path.Substring(EndIdx + 1);
+							Path = Rules.ModuleDirectory + Path[(EndIdx + 1)..];
 						}
 						else if (MatchVariableName(Path, StartIdx, EndIdx, "PluginDir"))
 						{
-							Path = Rules.PluginDirectory + Path.Substring(EndIdx + 1);
+							Path = Rules.PluginDirectory + Path[(EndIdx + 1)..];
 						}
 						else if (BinaryOutputDir != null && MatchVariableName(Path, StartIdx, EndIdx, "BinaryOutputDir"))
 						{
-							Path = BinaryOutputDir.FullName + Path.Substring(EndIdx + 1);
+							Path = BinaryOutputDir.FullName + Path[(EndIdx + 1)..];
 						}
 						else if (TargetOutputDir != null && MatchVariableName(Path, StartIdx, EndIdx, "TargetOutputDir"))
 						{
-							Path = TargetOutputDir.FullName + Path.Substring(EndIdx + 1);
+							Path = TargetOutputDir.FullName + Path[(EndIdx + 1)..];
 						}
 						else
 						{
-							string Name = Path.Substring(StartIdx, EndIdx - StartIdx);
+							string Name = Path[StartIdx..EndIdx];
 							string? Value = Environment.GetEnvironmentVariable(Name);
 							if (String.IsNullOrEmpty(Value))
 							{
 								throw new BuildException("Environment variable '{0}' is not defined (referenced by {1})", Name, Rules.File);
 							}
-							Path = Value + Path.Substring(EndIdx + 1);
+							Path = Value + Path[(EndIdx + 1)..];
 						}
 						break;
 					}
diff --git a/Engine/Source/Programs/UnrealBuildTool/Configuration/UEBuildPlatform.cs b/Engine/Source/Programs/UnrealBuildTool/Configuration/UEBuildPlatform.cs
--- a/Engine/Source/Programs/UnrealBuildTool/Configuration/UEBuildPlatform.cs
+++ b/Engine/Source/Programs/UnrealBuildTool/Configuration/UEBuildPlatform.cs
@@ -271,7 +271,7 @@ namespace UnrealBuildTool
 						// find arguments that start with one of the hard-coded parameters
 						.Where(x => UATProjectParams.Any(y => x.StartsWith(y, StringComparison.OrdinalIgnoreCase)))
 						// treat the part after the = as a path to a uproject, and ask the NativeProjects class to find a .uproject
-						.Select(x => NativeProjects.FindProjectFile(x.Substring(x.IndexOf('=') + 1), Logger))
+						.Select(x => NativeProjects.FindProjectFile(x[(x.IndexOf('=') + 1)..], Logger))
 						// only use existant projects
 						.Where(x => x != null && FileReference.Exists(x));
 			}
diff --git a/Engine/Source/Programs/UnrealBuildTool/Configuration/UEBuildTarget.cs b/Engine/Source/Programs/UnrealBuildTool/Configuration/UEBuildTarget.cs
--- a/Engine/Source/Programs/UnrealBuildTool/Configuration/UEBuildTarget.cs
+++ b/Engine/Source/Programs/UnrealBuildTool/Configuration/UEBuildTarget.cs
@@ -3446,7 +3446,7 @@ namespace UnrealBuildTool
 					}
 					else
 					{
-						Writer.WriteLine(String.Format("#define {0} {1}", Definition.Substring(0, EqualsIdx), Definition.Substring(EqualsIdx + 1)));
+						Writer.WriteLine(String.Format("#define {0} {1}", Definition[..EqualsIdx], Definition[(EqualsIdx + 1)..]));
 					}
 				}
 			}
@@ -6083,7 +6083,7 @@ namespace UnrealBuildTool
 					string OutputFilePath = ExecutableBinary.OutputFilePath.FullName;
 					if (Platform == UnrealTargetPlatform.Mac && OutputFilePath.Contains(".app/Contents/MacOS"))
 					{
-						OutputFilePath = OutputFilePath.Substring(0, OutputFilePath.LastIndexOf(".app/Contents/MacOS") + 4);
+						OutputFilePath = OutputFilePath[..(OutputFilePath.LastIndexOf(".app/Contents/MacOS") + 4)];
 					}
 
 					DirectoryReference OutputDir = new DirectoryReference(OutputFilePath).ParentDirectory!;
diff --git a/Engine/Source/Programs/UnrealBuildTool/Executors/SNDBS.cs b/Engine/Source/Programs/UnrealBuildTool/Executors/SNDBS.cs
--- a/Engine/Source/Programs/UnrealBuildTool/Executors/SNDBS.cs
+++ b/Engine/Source/Programs/UnrealBuildTool/Executors/SNDBS.cs
@@ -352,7 +352,7 @@ namespace UnrealBuildTool
 						{
 							Writer.Write(++NumCompletedActions, NumActions);
 
-							Text = Args.Data.Substring(ProgressMarkupPrefix.Length).Trim();
+							Text = Args.Data[ProgressMarkupPrefix.Length..].Trim();
 							string[] ActionInfo = Text.Split(':');
 							Logger.LogInformation("[{NumCompletedActions}/{NumActions}] {ActionInfo0} {ActionInfo1}", NumCompletedActions, NumActions, ActionInfo[0], ActionInfo[1]);
 							CurrentStatus = ActionInfo[1];
diff --git a/Engine/Source/Programs/UnrealBuildTool/Executors/XGE.cs b/Engine/Source/Programs/UnrealBuildTool/Executors/XGE.cs
--- a/Engine/Source/Programs/UnrealBuildTool/Executors/XGE.cs
+++ b/Engine/Source/Programs/UnrealBuildTool/Executors/XGE.cs
@@ -754,7 +754,7 @@ namespace UnrealBuildTool
 							CommandDescription = ActionIndex != -1 ? Actions[ActionIndex].CommandDescription + " " : String.Empty;
 
 							// Strip out anything that is just an XGE timer. Some programs don't output anything except the progress text.
-							Text = Args.Data.Substring(MarkupLength);
+							Text = Args.Data[MarkupLength..];
 							if (Text.StartsWith(" (") && Text.EndsWith(")"))
 							{
 								// Write the progress text with the next line of output if the current doesn't have any status.
diff --git a/Engine/Source/Programs/UnrealBuildTool/Matchers/MicrosoftEventMatcher.cs b/Engine/Source/Programs/UnrealBuildTool/Matchers/MicrosoftEventMatcher.cs
--- a/Engine/Source/Programs/UnrealBuildTool/Matchers/MicrosoftEventMatcher.cs
+++ b/Engine/Source/Programs/UnrealBuildTool/Matchers/MicrosoftEventMatcher.cs
@@ -27,7 +27,7 @@ namespace UnrealBuildTool.Matchers
 			Match? match;
 			if (cursor.TryMatch(s_warningCodePattern, out match))
 			{
-				string prefix = cursor.CurrentLine!.Substring(0, match.Index);
+				string prefix = cursor.CurrentLine![..match.Index];
 
 				Match fileOrToolMatch = s_fileOrToolPattern.Match(prefix);
 				if (fileOrToolMatch.Success)
diff --git a/Engine/Source/Programs/UnrealBuildTool/Modes/BuildMode.cs b/Engine/Source/Programs/UnrealBuildTool/Modes/BuildMode.cs
--- a/Engine/Source/Programs/UnrealBuildTool/Modes/BuildMode.cs
+++ b/Engine/Source/Programs/UnrealBuildTool/Modes/BuildMode.cs
@@ -203,7 +203,7 @@ namespace UnrealBuildTool
 						FileReference RemoteLogFile = FileReference.Combine(BaseLogFile.Directory, BaseLogFile.GetFileNameWithoutExtension() + "_Remote.txt");
 
 						int RemoteIniIndex = TargetDesc.AdditionalArguments.FindIndex(x => x.ToLower().StartsWith("-remoteini="));
-						string? RemoteIni = RemoteIniIndex == -1 ? null : new string(TargetDesc.AdditionalArguments[RemoteIniIndex].Substring(11));
+						string? RemoteIni = RemoteIniIndex == -1 ? null : new string(TargetDesc.AdditionalArguments[RemoteIniIndex].AsSpan()[11..]);
 						RemoteMac RemoteMac = new RemoteMac(TargetDesc.ProjectFile, Logger, RemoteIni);
 						if (!RemoteMac.Build(TargetDesc, RemoteLogFile, bSkipPreBuildTargets, Logger))
 						{
diff --git a/Engine/Source/Programs/UnrealBuildTool/Modes/ClReproMode.cs b/Engine/Source/Programs/UnrealBuildTool/Modes/ClReproMode.cs
--- a/Engine/Source/Programs/UnrealBuildTool/Modes/ClReproMode.cs
+++ b/Engine/Source/Programs/UnrealBuildTool/Modes/ClReproMode.cs
@@ -38,7 +38,7 @@ namespace UnrealBuildTool.Modes
 			try
 			{
 				string root = Path.GetPathRoot(path)?.ToUpperInvariant() ?? String.Empty;
-				foreach (string name in path.Substring(root.Length).Split(System.IO.Path.DirectorySeparatorChar))
+				foreach (string name in path[root.Length..].Split(System.IO.Path.DirectorySeparatorChar))
 				{
 					root = System.IO.Directory.GetFileSystemEntries(root, name).First();
 				}
diff --git a/Engine/Source/Programs/UnrealBuildTool/Modes/FixIncludePathsMode.cs b/Engine/Source/Programs/UnrealBuildTool/Modes/FixIncludePathsMode.cs
--- a/Engine/Source/Programs/UnrealBuildTool/Modes/FixIncludePathsMode.cs
+++ b/Engine/Source/Programs/UnrealBuildTool/Modes/FixIncludePathsMode.cs
@@ -247,7 +247,7 @@ namespace UnrealBuildTool
 														if (FoundPreferredPath != null)
 														{
 															int end = FullPath.LastIndexOf(FoundPreferredPath) + FoundPreferredPath.Length;
-															PreferredInclude = FullPath.Substring(end);
+															PreferredInclude = FullPath[end..];
 
 															// Is the current include a shortened version of the preferred include path?
 															if (PreferredInclude != Include && PreferredInclude.Contains(Include))
@@ -278,7 +278,7 @@ namespace UnrealBuildTool
 
 															if (FoundDirectory != null)
 															{
-																PreferredInclude = FullPath.Substring(FoundDirectory.FullName.Length + 1);
+																PreferredInclude = FullPath[(FoundDirectory.FullName.Length + 1)..];
 															}
 														}
 													}
diff --git a/Engine/Source/Programs/UnrealBuildTool/Modes/IWYUMode.cs b/Engine/Source/Programs/UnrealBuildTool/Modes/IWYUMode.cs
--- a/Engine/Source/Programs/UnrealBuildTool/Modes/IWYUMode.cs
+++ b/Engine/Source/Programs/UnrealBuildTool/Modes/IWYUMode.cs
@@ -735,7 +735,7 @@ namespace UnrealBuildTool
 							{
 								if (Entry.Printable[0] == '<')
 								{
-									Entry.Printable = Entry.Printable.Substring(1, Entry.Printable.Length - 2);
+									Entry.Printable = Entry.Printable[1..^1];
 									Entry.System = true;
 								}
 							}
@@ -977,7 +977,7 @@ namespace UnrealBuildTool
 						{
 							if (fixedPath.StartsWith(Dir.FullName))
 							{
-								string printable = value.File.Substring(Dir.FullName.Length + 1);
+								string printable = value.File[(Dir.FullName.Length + 1)..];
 								IWYUIncludeEntry entry = new IWYUIncludeEntry() { Full = MissingInclude.Full, Printable = printable, System = false, Resolved = value, };
 								Info.Includes.Add(entry);
 								break;
@@ -1047,7 +1047,7 @@ namespace UnrealBuildTool
 					foreach (IWYUIncludeEntry Include in GeneratedCpp.Includes)
 					{
 						NameStart = Include.Full.LastIndexOf('/');
-						string IncludeName = Include.Full.Substring(NameStart + 1);
+						string IncludeName = Include.Full[(NameStart + 1)..];
 						if (HeaderName == IncludeName)
 						{
 							Infos.TryGetValue(Include.Full, out IncluderInfo);
@@ -1114,7 +1114,7 @@ namespace UnrealBuildTool
 					break;
 				}
 
-				FilePath = FilePath.Substring(0, LastIndexOfSlash);
+				FilePath = FilePath[..LastIndexOfSlash];
 				UEBuildModule? Module;
 				if (PathToModule.TryGetValue(FilePath, out Module))
 				{
@@ -1208,7 +1208,7 @@ namespace UnrealBuildTool
 						{
 							continue;
 						}
-						LineSpan = LineSpan.Slice(1).TrimStart();
+						LineSpan = LineSpan[1..].TrimStart();
 						if (LineSpan.StartsWith("if"))
 						{
 							// Include guards are special
@@ -1228,13 +1228,13 @@ namespace UnrealBuildTool
 						}
 						else if (LineSpan.StartsWith("include"))
 						{
-							ReadOnlySpan<char> IncludeSpan = LineSpan.Slice("include".Length).TrimStart();
+							ReadOnlySpan<char> IncludeSpan = LineSpan["include".Length..].TrimStart();
 							char LeadingIncludeChar = IncludeSpan[0];
 							if (LeadingIncludeChar == '"' || LeadingIncludeChar == '<')
 							{
-								ReadOnlySpan<char> Start = IncludeSpan.Slice(1);
+								ReadOnlySpan<char> Start = IncludeSpan[1..];
 								int EndIndex = Start.IndexOf(LeadingIncludeChar == '"' ? '"' : '>');
-								ReadOnlySpan<char> FileSpan = Start.Slice(0, EndIndex);
+								ReadOnlySpan<char> FileSpan = Start[..EndIndex];
 								string Include = FileSpan.ToString();
 								string? File = GetFullName(Include, Info.File, Info.Module, NameToModule);
 								if (File != null)
@@ -1356,7 +1356,7 @@ namespace UnrealBuildTool
 					int LastSlash = NameWithoutPath.LastIndexOf("/");
 					if (LastSlash != -1)
 					{
-						NameWithoutPath = NameWithoutPath.Substring(LastSlash + 1);
+						NameWithoutPath = NameWithoutPath[(LastSlash + 1)..];
 					}
 
 					string QuotedPrintable = Include.System ? $"<{Include.Printable}>" : $"\"{Include.Printable}\"";
@@ -1481,7 +1481,7 @@ namespace UnrealBuildTool
 					bool StartsWithHash = LineSpan.Length > 0 && LineSpan[0] == '#';
 					if (StartsWithHash)
 					{
-						LineSpan = LineSpan.Slice(1).TrimStart();
+						LineSpan = LineSpan[1..].TrimStart();
 					}
 
 					if (!StartsWithHash || !LineSpan.StartsWith("include"))
@@ -1529,7 +1529,7 @@ namespace UnrealBuildTool
 
 					bool ForceKeep = false;
 
-					ReadOnlySpan<char> IncludeSpan = LineSpan.Slice("include".Length).TrimStart();
+					ReadOnlySpan<char> IncludeSpan = LineSpan["include".Length..].TrimStart();
 					char LeadingIncludeChar = IncludeSpan[0];
 					if (LeadingIncludeChar != '"' && LeadingIncludeChar != '<')
 					{
@@ -1537,14 +1537,14 @@ namespace UnrealBuildTool
 						{
 							int Open = IncludeSpan.IndexOf('(') + 1;
 							int Close = IncludeSpan.IndexOf(')');
-							string ActualInclude = $"\"{IncludeSpan.Slice(Open, Close - Open).ToString()}.gen.cpp\"";
+							string ActualInclude = $"\"{IncludeSpan[Open..Close].ToString()}.gen.cpp\"";
 							IncludesToAdd.Remove(ActualInclude);
 						}
 						else if (IncludeSpan.IndexOf("COMPILED_PLATFORM_HEADER") != -1)
 						{
 							int Open = IncludeSpan.IndexOf('(') + 1;
 							int Close = IncludeSpan.IndexOf(')');
-							string ActualInclude = $"\"Linux/Linux{IncludeSpan.Slice(Open, Close - Open).ToString()}\"";
+							string ActualInclude = $"\"Linux/Linux{IncludeSpan[Open..Close].ToString()}\"";
 							IncludesToAdd.Remove(ActualInclude);
 						}
 						else
@@ -1562,10 +1562,10 @@ namespace UnrealBuildTool
 					}
 					else
 					{
-						int Index = IncludeSpan.Slice(1).IndexOf(LeadingIncludeChar == '"' ? '"' : '>');
+						int Index = IncludeSpan[1..].IndexOf(LeadingIncludeChar == '"' ? '"' : '>');
 						if (Index != -1)
 						{
-							IncludeSpan = IncludeSpan.Slice(0, Index + 2);
+							IncludeSpan = IncludeSpan[..(Index + 2)];
 						}
 
 						if (Line.Contains("IWYU pragma: ", StringComparison.Ordinal))
@@ -1696,7 +1696,7 @@ namespace UnrealBuildTool
 						bool StartsWithHash = OldLineSpan.Length > 0 && OldLineSpan[0] == '#';
 						if (StartsWithHash)
 						{
-							OldLineSpan = OldLineSpan.Slice(1).TrimStart();
+							OldLineSpan = OldLineSpan[1..].TrimStart();
 						}
 
 						bool IsInclude = StartsWithHash && OldLineSpan.StartsWith("include");
@@ -1825,7 +1825,7 @@ namespace UnrealBuildTool
 						{
 							return false;
 						}
-						int Index = IncludeSpan.Slice(1).IndexOf(LeadingIncludeChar == '"' ? '"' : '>');
+						int Index = IncludeSpan[1..].IndexOf(LeadingIncludeChar == '"' ? '"' : '>');
 						if (Index == -1)
 						{
 							return false;
diff --git a/Engine/Source/Programs/UnrealBuildTool/Modes/InlineGeneratedCpps.cs b/Engine/Source/Programs/UnrealBuildTool/Modes/InlineGeneratedCpps.cs
--- a/Engine/Source/Programs/UnrealBuildTool/Modes/InlineGeneratedCpps.cs
+++ b/Engine/Source/Programs/UnrealBuildTool/Modes/InlineGeneratedCpps.cs
@@ -158,7 +158,7 @@ namespace UnrealBuildTool
 
 										if (Line.Contains("//", StringComparison.InvariantCultureIgnoreCase))
 										{
-											Line = Line.Substring(0, Line.IndexOf("//"));
+											Line = Line[..Line.IndexOf("//")];
 										}
 
 										if (SameLineCommentRegex.IsMatch(Line))
diff --git a/Engine/Source/Programs/UnrealBuildTool/Modes/ParseMsvcTimingInfo.cs b/Engine/Source/Programs/UnrealBuildTool/Modes/ParseMsvcTimingInfo.cs
--- a/Engine/Source/Programs/UnrealBuildTool/Modes/ParseMsvcTimingInfo.cs
+++ b/Engine/Source/Programs/UnrealBuildTool/Modes/ParseMsvcTimingInfo.cs
@@ -309,7 +309,7 @@ namespace UnrealBuildTool
 				int TemplateIdx = ClassName.IndexOf('<');
 				if (TemplateIdx != -1)
 				{
-					ClassName = ClassName.Substring(0, TemplateIdx) + "<>";
+					ClassName = ClassName[..TemplateIdx] + "<>";
 				}
 
 				float Time;
diff --git a/Engine/Source/Programs/UnrealBuildTool/Modes/PipInstallMode.cs b/Engine/Source/Programs/UnrealBuildTool/Modes/PipInstallMode.cs
--- a/Engine/Source/Programs/UnrealBuildTool/Modes/PipInstallMode.cs
+++ b/Engine/Source/Programs/UnrealBuildTool/Modes/PipInstallMode.cs
@@ -902,8 +902,8 @@ namespace UnrealBuildTool.Modes
 				return;
 			}
 
-			Left = InStr.Substring(0, SplitIdx);
-			Right = InStr.Substring(SplitIdx + 1);
+			Left = InStr[..SplitIdx];
+			Right = InStr[(SplitIdx + 1)..];
 		}
 
 		private static string SplitRight(string InStr, char SplitChar)
diff --git a/Engine/Source/Programs/UnrealBuildTool/Platform/Android/AndroidPlatformSDK.cs b/Engine/Source/Programs/UnrealBuildTool/Platform/Android/AndroidPlatformSDK.cs
--- a/Engine/Source/Programs/UnrealBuildTool/Platform/Android/AndroidPlatformSDK.cs
+++ b/Engine/Source/Programs/UnrealBuildTool/Platform/Android/AndroidPlatformSDK.cs
@@ -173,7 +173,7 @@ namespace UnrealBuildTool
 				int EqualsIndex = RevisionString.IndexOf('=');
 				if (EqualsIndex > 0)
 				{
-					string[] RevisionParts = RevisionString.Substring(EqualsIndex + 1).Trim().Split('.');
+					string[] RevisionParts = RevisionString[(EqualsIndex + 1)..].Trim().Split('.');
 					int RevisionMinor = Int32.Parse(RevisionParts.Length > 1 ? RevisionParts[1] : "0");
 					char RevisionLetter = Convert.ToChar('a' + RevisionMinor);
 					NDKToolchainVersion = "r" + RevisionParts[0] + (RevisionMinor > 0 ? Char.ToString(RevisionLetter) : "");
@@ -363,7 +363,7 @@ namespace UnrealBuildTool
 			if (start != 1 && end != -1 && start < end)
 			{
 				++start;
-				Path = Source.Substring(start, end - start);
+				Path = Source[start..end];
 				return true;
 			}
 			else
@@ -453,7 +453,7 @@ namespace UnrealBuildTool
 					// JAVA_VERSION="17.0.6"
 					int VersionStartIndex = LineContents.IndexOf("\"");
 					int VersionStopIndex = LineContents.IndexOf(".");
-					Int32.TryParse(LineContents.Substring(VersionStartIndex, VersionStopIndex - VersionStartIndex), out JavaVersion);
+					Int32.TryParse(LineContents[VersionStartIndex..VersionStopIndex], out JavaVersion);
 					break;
 				}
 			}
diff --git a/Engine/Source/Programs/UnrealBuildTool/Platform/Android/AndroidToolChain.cs b/Engine/Source/Programs/UnrealBuildTool/Platform/Android/AndroidToolChain.cs
--- a/Engine/Source/Programs/UnrealBuildTool/Platform/Android/AndroidToolChain.cs
+++ b/Engine/Source/Programs/UnrealBuildTool/Platform/Android/AndroidToolChain.cs
@@ -324,7 +324,7 @@ namespace UnrealBuildTool
 			if (NDKVersion.Contains('-'))
 			{
 				int Version;
-				if (Int32.TryParse(NDKVersion.Substring(NDKVersion.LastIndexOf('-') + 1), out Version))
+				if (Int32.TryParse(NDKVersion[(NDKVersion.LastIndexOf('-') + 1)..], out Version))
 				{
 					if (Version > NDKVersionInt)
 					{
@@ -394,7 +394,7 @@ namespace UnrealBuildTool
 			if (ApiString.Contains('-'))
 			{
 				int Version;
-				if (Int32.TryParse(ApiString.Substring(ApiString.LastIndexOf('-') + 1), out Version))
+				if (Int32.TryParse(ApiString[(ApiString.LastIndexOf('-') + 1)..], out Version))
 				{
 					return (Version >= MinPlatform && Version <= MaxPlatform);
 				}
@@ -865,7 +865,7 @@ namespace UnrealBuildTool
 			string Lib = Path.GetFileNameWithoutExtension(FullLib);
 			if (Lib.StartsWith("lib"))
 			{
-				Lib = Lib.Substring(3);
+				Lib = Lib[3..];
 			}
 
 			// reject any libs we outright don't want to link with
@@ -1234,7 +1234,7 @@ namespace UnrealBuildTool
 				{
 					if (AdditionalLibrary.StartsWith("lib"))
 					{
-						LinkResponseArguments += String.Format(" \"-l{0}\"", AdditionalLibrary.Substring(3));
+						LinkResponseArguments += String.Format(" \"-l{0}\"", AdditionalLibrary[3..]);
 					}
 					else
 					{
@@ -1409,7 +1409,7 @@ namespace UnrealBuildTool
 
 			// fix the executable (replace the last -ar with -strip and keep any extension)
 			int ArIndex = StripExe.LastIndexOf("-ar");
-			StripExe = StripExe.Substring(0, ArIndex) + "-strip" + StripExe.Substring(ArIndex + 3);
+			StripExe = StripExe[..ArIndex] + "-strip" + StripExe[(ArIndex + 3)..];
 			return StripExe;
 		}
 
diff --git a/Engine/Source/Programs/UnrealBuildTool/Platform/Android/UEDeployAndroid.cs b/Engine/Source/Programs/UnrealBuildTool/Platform/Android/UEDeployAndroid.cs
--- a/Engine/Source/Programs/UnrealBuildTool/Platform/Android/UEDeployAndroid.cs
+++ b/Engine/Source/Programs/UnrealBuildTool/Platform/Android/UEDeployAndroid.cs
@@ -230,7 +230,7 @@ namespace UnrealBuildTool
 			if (ApiString.Contains('-'))
 			{
 				int Version;
-				if (Int32.TryParse(ApiString.Substring(ApiString.LastIndexOf('-') + 1), out Version))
+				if (Int32.TryParse(ApiString[(ApiString.LastIndexOf('-') + 1)..], out Version))
 				{
 					VersionInt = Version;
 				}
@@ -368,7 +368,7 @@ namespace UnrealBuildTool
 							{
 								if (PropertyLine.StartsWith("Pkg.Revision="))
 								{
-									RevisionString = PropertyLine.Substring(13);
+									RevisionString = PropertyLine[13..];
 									RevisionValue = GetRevisionValue(RevisionString);
 									break;
 								}
@@ -1312,17 +1312,17 @@ namespace UnrealBuildTool
 					int EndIndex = Output.IndexOf(']', ScanStart);
 					if (EndIndex > 0)
 					{
-						string SharedLib = Output.Substring(ScanStart, EndIndex - ScanStart);
+						string SharedLib = Output[ScanStart..EndIndex];
 
 						// strip off lib (not needed for System.loadLibrary()
 						if (SharedLib.StartsWith("lib"))
 						{
-							SharedLib = SharedLib.Substring(3);
+							SharedLib = SharedLib[3..];
 						}
 						// strip off .so (not needed for System.loadLibrary()
 						if (SharedLib.EndsWith(".so"))
 						{
-							SharedLib = SharedLib.Substring(0, SharedLib.Length - 3);
+							SharedLib = SharedLib[..^3];
 						}
 
 						DependencyList.Add(SharedLib);
@@ -2062,7 +2062,7 @@ namespace UnrealBuildTool
 			{
 				return true;
 			}
-			string CurrentApplicationName = Contents.Substring(KeyIndex, TagEnd - KeyIndex);
+			string CurrentApplicationName = Contents[KeyIndex..TagEnd];
 
 			// no need to do anything if matches
 			if (CurrentApplicationName == ApplicationDisplayName)
@@ -2181,12 +2181,12 @@ namespace UnrealBuildTool
 				int KeyIndex = SectionName.IndexOf('[');
 				if (KeyIndex > 0)
 				{
-					string KeyList = SectionName.Substring(KeyIndex + 1);
-					SectionName = SectionName.Substring(0, KeyIndex);
+					string KeyList = SectionName[(KeyIndex + 1)..];
+					SectionName = SectionName[..KeyIndex];
 					int CloseIndex = KeyList.IndexOf("]");
 					if (CloseIndex > 1)
 					{
-						NeededKeys = KeyList.Substring(0, CloseIndex).Split(',', StringSplitOptions.RemoveEmptyEntries);
+						NeededKeys = KeyList[..CloseIndex].Split(',', StringSplitOptions.RemoveEmptyEntries);
 						if (NeededKeys.Length == 0)
 						{
 							NeededKeys = null;
@@ -2816,11 +2816,11 @@ namespace UnrealBuildTool
 			string AppType = InTargetType == TargetType.Game ? "" : InTargetType.ToString();
 			if (CookFlavor.EndsWith("Client"))
 			{
-				CookFlavor = CookFlavor.Substring(0, CookFlavor.Length - 6);
+				CookFlavor = CookFlavor[..^6];
 			}
 			if (CookFlavor.EndsWith("Server"))
 			{
-				CookFlavor = CookFlavor.Substring(0, CookFlavor.Length - 6);
+				CookFlavor = CookFlavor[..^6];
 			}
 
 			//figure out which texture compressions are supported
@@ -2864,7 +2864,7 @@ namespace UnrealBuildTool
 			string CookedFlavors = (bETC2Enabled ? "ETC2," : "") +
 									(bDXTEnabled ? "DXT," : "") +
 									(bASTCEnabled ? "ASTC," : "");
-			CookedFlavors = (String.IsNullOrEmpty(CookedFlavors)) ? "" : CookedFlavors.Substring(0, CookedFlavors.Length - 1);
+			CookedFlavors = (String.IsNullOrEmpty(CookedFlavors)) ? "" : CookedFlavors[..^1];
 
 			StringBuilder Text = new StringBuilder();
 			Text.AppendLine(XML_HEADER);
@@ -3350,7 +3350,7 @@ namespace UnrealBuildTool
 						continue;
 					}
 
-					string XmlAppId = Line.Substring(StartIndex, EndIndex - StartIndex);
+					string XmlAppId = Line[StartIndex..EndIndex];
 
 					//validate that the AppId matches the .ini value for the GooglePlay AppId, assuming it's valid
 					if (!bInvalidIniAppId && IniAppId.CompareTo(XmlAppId) != 0)
@@ -3990,7 +3990,7 @@ namespace UnrealBuildTool
 
 			string NDKPath = Environment.GetEnvironmentVariable("NDKROOT")!.Replace("\\", "/");
 			int NDKVersionIndex = NDKPath.LastIndexOf("/");
-			string NDKVersion = NDKVersionIndex > 0 ? NDKPath.Substring(NDKVersionIndex + 1) : DEFAULT_NDK_VERSION;
+			string NDKVersion = NDKVersionIndex > 0 ? NDKPath[(NDKVersionIndex + 1)..] : DEFAULT_NDK_VERSION;
 			GradleProperties.AppendLine(String.Format("NDK_VERSION={0}", NDKVersion));
 			GradleProperties.AppendLine(String.Format("NDK_PATH={0}", NDKPath));
 
@@ -5398,8 +5398,8 @@ popd
 									{
 										if (Taskline.StartsWith("java.exe"))
 										{
-											string ProcessString = Taskline.Substring(8).TrimStart();
-											ProcessString = ProcessString.Substring(0, ProcessString.IndexOf(' '));
+											string ProcessString = Taskline[8..].TrimStart();
+											ProcessString = ProcessString[..ProcessString.IndexOf(' ')];
 											int ProcessId;
 											if (Int32.TryParse(ProcessString, out ProcessId))
 											{
@@ -5438,8 +5438,8 @@ popd
 										{
 											if (Taskline.StartsWith("java.exe"))
 											{
-												string ProcessString = Taskline.Substring(8).TrimStart();
-												ProcessString = ProcessString.Substring(0, ProcessString.IndexOf(' '));
+												string ProcessString = Taskline[8..].TrimStart();
+												ProcessString = ProcessString[..ProcessString.IndexOf(' ')];
 												int ProcessId;
 												if (Int32.TryParse(ProcessString, out ProcessId))
 												{
@@ -5535,7 +5535,7 @@ popd
 							continue;
 						}
 
-						string DirectoryName = Path.GetDirectoryName(Filename)!.Substring(UnrealBuildGradlePathLength);
+						string DirectoryName = Path.GetDirectoryName(Filename)![UnrealBuildGradlePathLength..];
 						string[] DirectoryNameComponents = DirectoryName.Split(Path.DirectorySeparatorChar, StringSplitOptions.RemoveEmptyEntries);
 						int BuildDirIndex = Array.IndexOf(DirectoryNameComponents, "build");
 						// if a build directory exists
@@ -5627,11 +5627,11 @@ popd
 										break;
 									}
 
-									int SourceVersionIndex2 = SourceManifest[Index].Substring(SourceVersionIndex + 22).IndexOf("\"");
-									string FixedSource = SourceManifest[Index].Substring(0, SourceVersionIndex + 21) + SourceManifest[Index].Substring(SourceVersionIndex + 22 + SourceVersionIndex2);
+									int SourceVersionIndex2 = SourceManifest[Index][(SourceVersionIndex + 22)..].IndexOf("\"");
+									string FixedSource = SourceManifest[Index][..(SourceVersionIndex + 21)] + SourceManifest[Index][(SourceVersionIndex + 22 + SourceVersionIndex2)..];
 
-									int DestVersionIndex2 = DestManifest[Index].Substring(DestVersionIndex + 22).IndexOf("\"");
-									string FixedDest = SourceManifest[Index].Substring(0, DestVersionIndex + 21) + DestManifest[Index].Substring(DestVersionIndex + 22 + DestVersionIndex2);
+									int DestVersionIndex2 = DestManifest[Index][(DestVersionIndex + 22)..].IndexOf("\"");
+									string FixedDest = SourceManifest[Index][..(DestVersionIndex + 21)] + DestManifest[Index][(DestVersionIndex + 22 + DestVersionIndex2)..];
 
 									if (FixedSource != FixedDest)
 									{
@@ -5769,7 +5769,7 @@ popd
 					ABIGradle.AppendLine("android {");
 					ABIGradle.AppendLine("\tdefaultConfig {");
 					ABIGradle.AppendLine("\t\tndk {");
-					ABIGradle.AppendLine(String.Format("\t\t\tabiFilters{0}", ABIFilter.Substring(1)));
+					ABIGradle.AppendLine(String.Format("\t\t\tabiFilters{0}", ABIFilter[1..]));
 					ABIGradle.AppendLine("\t\t}");
 					ABIGradle.AppendLine("\t}");
 					ABIGradle.AppendLine("}");
@@ -5788,8 +5788,8 @@ popd
 							continue;
 						}
 
-						int ManifestVersionIndex2 = ManifestContents[Index].Substring(ManifestVersionIndex + 22).IndexOf("\"");
-						ManifestContents[Index] = ManifestContents[Index].Substring(0, ManifestVersionIndex + 21) + BaseStoreVersion + ManifestContents[Index].Substring(ManifestVersionIndex + 22 + ManifestVersionIndex2);
+						int ManifestVersionIndex2 = ManifestContents[Index][(ManifestVersionIndex + 22)..].IndexOf("\"");
+						ManifestContents[Index] = ManifestContents[Index][..(ManifestVersionIndex + 21)] + BaseStoreVersion + ManifestContents[Index][(ManifestVersionIndex + 22 + ManifestVersionIndex2)..];
 						break;
 					}
 					File.WriteAllLines(ManifestFilename, ManifestContents);
@@ -6323,11 +6323,11 @@ import java.util.Collection;
 										int TraceEndIndex = Addition.IndexOf("\n", TraceStartIndex);
 										if (TraceEndIndex < 0)
 										{
-											Addition = Addition.Substring(0, TraceStartIndex);
+											Addition = Addition[..TraceStartIndex];
 										}
 										else
 										{
-											Addition = Addition.Substring(0, TraceStartIndex) + Addition.Substring(TraceEndIndex + 1);
+											Addition = Addition[..TraceStartIndex] + Addition[(TraceEndIndex + 1)..];
 										}
 										TraceStartIndex = Addition.IndexOf("Trace.", TraceStartIndex);
 									} while (TraceStartIndex >= 0);
@@ -6445,11 +6445,11 @@ import java.util.Collection;
 										int TraceEndIndex = Addition.IndexOf("\n", TraceStartIndex);
 										if (TraceEndIndex < 0)
 										{
-											Addition = Addition.Substring(0, TraceStartIndex);
+											Addition = Addition[..TraceStartIndex];
 										}
 										else
 										{
-											Addition = Addition.Substring(0, TraceStartIndex) + Addition.Substring(TraceEndIndex + 1);
+											Addition = Addition[..TraceStartIndex] + Addition[(TraceEndIndex + 1)..];
 										}
 										TraceStartIndex = Addition.IndexOf("Trace.", TraceStartIndex);
 									} while (TraceStartIndex >= 0);
@@ -6553,7 +6553,7 @@ import java.util.Collection;
 
 				if (Trimmed.StartsWith("repository "))
 				{
-					string DirectoryPath = Trimmed.Substring(11).Trim(' ').TrimEnd('/', '\\');
+					string DirectoryPath = Trimmed[11..].Trim(' ').TrimEnd('/', '\\');
 					DirectoryPath = DirectoryPath.Replace("$(ENGINEDIR)", EngineDir);
 					DirectoryPath = DirectoryPath.Replace("$(ANDROID_HOME)", AndroidHome);
 					DirectoryPath = DirectoryPath.Replace('\\', Path.DirectorySeparatorChar).Replace('/', Path.DirectorySeparatorChar);
@@ -6561,7 +6561,7 @@ import java.util.Collection;
 				}
 				else if (Trimmed.StartsWith("repositories "))
 				{
-					string DirectoryPath = Trimmed.Substring(13).Trim(' ').TrimEnd('/', '\\');
+					string DirectoryPath = Trimmed[13..].Trim(' ').TrimEnd('/', '\\');
 					DirectoryPath = DirectoryPath.Replace("$(ENGINEDIR)", EngineDir);
 					DirectoryPath = DirectoryPath.Replace("$(ANDROID_HOME)", AndroidHome);
 					DirectoryPath = DirectoryPath.Replace('\\', Path.DirectorySeparatorChar).Replace('/', Path.DirectorySeparatorChar);
diff --git a/Engine/Source/Programs/UnrealBuildTool/Platform/Apple/AppleExports.cs b/Engine/Source/Programs/UnrealBuildTool/Platform/Apple/AppleExports.cs
--- a/Engine/Source/Programs/UnrealBuildTool/Platform/Apple/AppleExports.cs
+++ b/Engine/Source/Programs/UnrealBuildTool/Platform/Apple/AppleExports.cs
@@ -106,11 +106,11 @@ namespace UnrealBuildTool
 
 			if (FilePath.StartsWith("/Engine/", StringComparison.OrdinalIgnoreCase))
 			{
-				return FileReference.Combine(Unreal.EngineDirectory, FilePath.Substring(8));
+				return FileReference.Combine(Unreal.EngineDirectory, FilePath[8..]);
 			}
 			else if (ProductDirectory != null && FilePath.StartsWith("/Game/", StringComparison.OrdinalIgnoreCase))
 			{
-				return FileReference.Combine(ProductDirectory, FilePath.Substring(6));
+				return FileReference.Combine(ProductDirectory, FilePath[6..]);
 			}
 			else if (FilePath.StartsWith("/", StringComparison.OrdinalIgnoreCase))
 			{
diff --git a/Engine/Source/Programs/UnrealBuildTool/Platform/IOS/IOSCodeSigning.cs b/Engine/Source/Programs/UnrealBuildTool/Platform/IOS/IOSCodeSigning.cs
--- a/Engine/Source/Programs/UnrealBuildTool/Platform/IOS/IOSCodeSigning.cs
+++ b/Engine/Source/Programs/UnrealBuildTool/Platform/IOS/IOSCodeSigning.cs
@@ -61,7 +61,7 @@ namespace UnrealBuildTool
 				string CleanPart = Part.Trim();
 				if (CleanPart.StartsWith("CN="))
 				{
-					CommonName = CleanPart.Substring(3);
+					CommonName = CleanPart[3..];
 				}
 			}
 
@@ -137,7 +137,7 @@ namespace UnrealBuildTool
 						{
 							start += header.Length;
 							int end = CertToolData.IndexOf(footer, start);
-							string base64 = CertToolData.Substring(start, (end - start));
+							string base64 = CertToolData[start..end];
 							byte[] certData = Convert.FromBase64String(base64);
 							X509Certificate2 cert = new X509Certificate2(certData);
 							FoundCerts.Add(cert);
@@ -229,16 +229,16 @@ namespace UnrealBuildTool
 				{
 					bValid = (Cert.NotBefore.ToUniversalTime() < Now) && (Cert.NotAfter.ToUniversalTime() > Now);
 				}
-				bool bPassesNameCheck = p.ApplicationIdentifier.Substring(p.ApplicationIdentifierPrefix.Length + 1) == CFBundleIdentifier;
+				bool bPassesNameCheck = p.ApplicationIdentifier[(p.ApplicationIdentifierPrefix.Length + 1)..] == CFBundleIdentifier;
 				bool bPassesCompanyCheck = false;
 				bool bPassesWildCardCheck = false;
 				if (p.ApplicationIdentifier.Contains('*'))
 				{
-					string CompanyName = p.ApplicationIdentifier.Substring(p.ApplicationIdentifierPrefix.Length + 1);
+					string CompanyName = p.ApplicationIdentifier[(p.ApplicationIdentifierPrefix.Length + 1)..];
 
 					if (CompanyName != "*")
 					{
-						CompanyName = CompanyName.Substring(0, CompanyName.LastIndexOf("."));
+						CompanyName = CompanyName[..CompanyName.LastIndexOf(".")];
 						bPassesCompanyCheck = CFBundleIdentifier.StartsWith(CompanyName);
 					}
 					else
@@ -983,17 +983,17 @@ namespace UnrealBuildTool
 					}
 					else if (Phase == 0)
 					{
-						bPassesNameCheck = TestProvision.ApplicationIdentifier.Substring(TestProvision.ApplicationIdentifierPrefix.Length + 1) == CFBundleIdentifier;
+						bPassesNameCheck = TestProvision.ApplicationIdentifier[(TestProvision.ApplicationIdentifierPrefix.Length + 1)..] == CFBundleIdentifier;
 						bNameMatch = bPassesNameCheck;
 					}
 					else if (Phase == 1)
 					{
 						if (TestProvision.ApplicationIdentifier.Contains('*'))
 						{
-							string CompanyName = TestProvision.ApplicationIdentifier.Substring(TestProvision.ApplicationIdentifierPrefix.Length + 1);
+							string CompanyName = TestProvision.ApplicationIdentifier[(TestProvision.ApplicationIdentifierPrefix.Length + 1)..];
 							if (CompanyName != "*")
 							{
-								CompanyName = CompanyName.Substring(0, CompanyName.LastIndexOf("."));
+								CompanyName = CompanyName[..CompanyName.LastIndexOf(".")];
 								bPassesNameCheck = CFBundleIdentifier.StartsWith(CompanyName);
 							}
 						}
@@ -1002,7 +1002,7 @@ namespace UnrealBuildTool
 					{
 						if (TestProvision.ApplicationIdentifier.Contains('*'))
 						{
-							string CompanyName = TestProvision.ApplicationIdentifier.Substring(TestProvision.ApplicationIdentifierPrefix.Length + 1);
+							string CompanyName = TestProvision.ApplicationIdentifier[(TestProvision.ApplicationIdentifierPrefix.Length + 1)..];
 							bPassesNameCheck = CompanyName == "*";
 						}
 					}
@@ -1171,7 +1171,7 @@ namespace UnrealBuildTool
 						int DotPosition = UbiquityKvstoreString.LastIndexOf("*");
 						if (DotPosition >= 0)
 						{
-							string TeamPrefix = DotPosition > 1 ? UbiquityKvstoreString.Substring(0, DotPosition - 1) : TeamIdentifier;
+							string TeamPrefix = DotPosition > 1 ? UbiquityKvstoreString[..(DotPosition - 1)] : TeamIdentifier;
 							string NewUbiquityKvstoreIdentifier = String.Format("{0}.{1}", TeamPrefix, CFBundleIdentifier);
 							XCentPList.SetValueForKey("com.apple.developer.ubiquity-kvstore-identifier", NewUbiquityKvstoreIdentifier);
 						}
@@ -1187,7 +1187,7 @@ namespace UnrealBuildTool
 							int DotPosition = UbiquityContainerIdentifiersGroups[i].LastIndexOf("*");
 							if (DotPosition >= 0)
 							{
-								string TeamPrefix = DotPosition > 1 ? UbiquityContainerIdentifiersGroups[i].Substring(0, DotPosition - 1) : TeamIdentifier;
+								string TeamPrefix = DotPosition > 1 ? UbiquityContainerIdentifiersGroups[i][..(DotPosition - 1)] : TeamIdentifier;
 								string NewUbiquityContainerIdentifier = String.Format("{0}.{1}", TeamPrefix, CFBundleIdentifier);
 								UbiquityContainerIdentifiersGroups[i] = NewUbiquityContainerIdentifier;
 							}
diff --git a/Engine/Source/Programs/UnrealBuildTool/Platform/IOS/IOSToolChain.cs b/Engine/Source/Programs/UnrealBuildTool/Platform/IOS/IOSToolChain.cs
--- a/Engine/Source/Programs/UnrealBuildTool/Platform/IOS/IOSToolChain.cs
+++ b/Engine/Source/Programs/UnrealBuildTool/Platform/IOS/IOSToolChain.cs
@@ -95,7 +95,7 @@ namespace UnrealBuildTool
 					else if (Target.Platform == UnrealTargetPlatform.IOS)
 					{
 						int Index = Binary.OutputFilePath.GetFileNameWithoutExtension().IndexOf("-");
-						string OutputFile = Binary.OutputFilePath.GetFileNameWithoutExtension().Substring(0, Index > 0 ? Index : Binary.OutputFilePath.GetFileNameWithoutExtension().Length);
+						string OutputFile = Binary.OutputFilePath.GetFileNameWithoutExtension()[..(Index > 0 ? Index : Binary.OutputFilePath.GetFileNameWithoutExtension().Length)];
 						FileReference AssetFile = FileReference.Combine(Binary.OutputFilePath.Directory, "Payload", OutputFile + ".app", "Assets.car");
 						BuildProducts.Add(AssetFile, BuildProductType.RequiredResource);
 
@@ -559,7 +559,7 @@ namespace UnrealBuildTool
 			int SuffixIdx = AppBundleName.IndexOf('-');
 			if (SuffixIdx != -1)
 			{
-				AppBundleName = AppBundleName.Substring(0, SuffixIdx);
+				AppBundleName = AppBundleName[..SuffixIdx];
 			}
 
 			// Append the .app suffix
@@ -1118,7 +1118,7 @@ namespace UnrealBuildTool
 						{
 							// For now, this is hard coded, but we need to loop over all modules, and copy bundled assets that need it
 							DirectoryReference LocalSource = DirectoryReference.Combine(FrameworkDirectory, Framework.CopyBundledAssets);
-							string BundleName = Framework.CopyBundledAssets.Substring(Framework.CopyBundledAssets.LastIndexOf('/') + 1);
+							string BundleName = Framework.CopyBundledAssets[(Framework.CopyBundledAssets.LastIndexOf('/') + 1)..];
 							FrameworkNameToSourceDir[BundleName] = LocalSource;
 						}
 
diff --git a/Engine/Source/Programs/UnrealBuildTool/Platform/IOS/MobileProvisionContents.cs b/Engine/Source/Programs/UnrealBuildTool/Platform/IOS/MobileProvisionContents.cs
--- a/Engine/Source/Programs/UnrealBuildTool/Platform/IOS/MobileProvisionContents.cs
+++ b/Engine/Source/Programs/UnrealBuildTool/Platform/IOS/MobileProvisionContents.cs
@@ -107,7 +107,7 @@ namespace UnrealBuildTool
 			{
 				throw new BuildException("Missing Bundle Identifier in MobileProvision");
 			}
-			return UniqueIdElement.InnerText.Substring(UniqueIdElement.InnerText.IndexOf('.') + 1);
+			return UniqueIdElement.InnerText[(UniqueIdElement.InnerText.IndexOf('.') + 1)..];
 		}
 
 		/// <summary>
diff --git a/Engine/Source/Programs/UnrealBuildTool/Platform/IOS/UEBuildIOS.cs b/Engine/Source/Programs/UnrealBuildTool/Platform/IOS/UEBuildIOS.cs
--- a/Engine/Source/Programs/UnrealBuildTool/Platform/IOS/UEBuildIOS.cs
+++ b/Engine/Source/Programs/UnrealBuildTool/Platform/IOS/UEBuildIOS.cs
@@ -561,7 +561,7 @@ namespace UnrealBuildTool
 						if (idx > 0)
 						{
 							idx += "<string>".Length;
-							MobileProvisionUUID = AllText.Substring(idx, AllText.IndexOf("</string>", idx) - idx);
+							MobileProvisionUUID = AllText[idx..AllText.IndexOf("</string>", idx)];
 						}
 					}
 					idx = AllText.IndexOf("<key>com.apple.developer.team-identifier</key>");
@@ -571,7 +571,7 @@ namespace UnrealBuildTool
 						if (idx > 0)
 						{
 							idx += "<string>".Length;
-							TeamUUID = AllText.Substring(idx, AllText.IndexOf("</string>", idx) - idx);
+							TeamUUID = AllText[idx..AllText.IndexOf("</string>", idx)];
 						}
 					}
 					idx = AllText.IndexOf("<key>application-identifier</key>");
@@ -581,8 +581,8 @@ namespace UnrealBuildTool
 						if (idx > 0)
 						{
 							idx += "<string>".Length;
-							string FullID = AllText.Substring(idx, AllText.IndexOf("</string>", idx) - idx);
-							BundleIdentifier = FullID.Substring(FullID.IndexOf('.') + 1);
+							string FullID = AllText[idx..AllText.IndexOf("</string>", idx)];
+							BundleIdentifier = FullID[(FullID.IndexOf('.') + 1)..];
 						}
 					}
 					idx = AllText.IndexOf("<key>Name</key>");
@@ -592,7 +592,7 @@ namespace UnrealBuildTool
 						if (idx > 0)
 						{
 							idx += "<string>".Length;
-							MobileProvisionName = AllText.Substring(idx, AllText.IndexOf("</string>", idx) - idx);
+							MobileProvisionName = AllText[idx..AllText.IndexOf("</string>", idx)];
 						}
 					}
 				}
@@ -638,7 +638,7 @@ namespace UnrealBuildTool
 						pindex += "PROVISION-".Length;
 						if (pindex < Line.Data.Length)
 						{
-							MobileProvisionFile = new FileReference(Line.Data.Substring(pindex));
+							MobileProvisionFile = new FileReference(Line.Data[pindex..]);
 						}
 					}
 				}
diff --git a/Engine/Source/Programs/UnrealBuildTool/Platform/IOS/UEDeployIOS.cs b/Engine/Source/Programs/UnrealBuildTool/Platform/IOS/UEDeployIOS.cs
--- a/Engine/Source/Programs/UnrealBuildTool/Platform/IOS/UEDeployIOS.cs
+++ b/Engine/Source/Programs/UnrealBuildTool/Platform/IOS/UEDeployIOS.cs
@@ -158,7 +158,7 @@ namespace UnrealBuildTool
 					if (Index != -1)
 					{
 						int Start = OldPList.IndexOf("<string>", Index) + ("<string>").Length;
-						CFBundleVersion = OldPList.Substring(Start, OldPList.IndexOf("</string>", Index) - Start);
+						CFBundleVersion = OldPList[Start..OldPList.IndexOf("</string>", Index)];
 						CFBundleVersion = CalculateUpdatedMinorVersionString(CFBundleVersion);
 					}
 					else
diff --git a/Engine/Source/Programs/UnrealBuildTool/Platform/Mac/ApplePlatformSDK.cs b/Engine/Source/Programs/UnrealBuildTool/Platform/Mac/ApplePlatformSDK.cs
--- a/Engine/Source/Programs/UnrealBuildTool/Platform/Mac/ApplePlatformSDK.cs
+++ b/Engine/Source/Programs/UnrealBuildTool/Platform/Mac/ApplePlatformSDK.cs
@@ -131,7 +131,7 @@ namespace UnrealBuildTool
 							"ASMapiInterfaceDLL", null) as string;
 						DllPath = String.IsNullOrEmpty(DllPath)
 							? null
-							: DllPath.Substring(0, DllPath.LastIndexOf('\\') + 1) + "MobileDevice.dll";
+							: DllPath[..(DllPath.LastIndexOf('\\') + 1)] + "MobileDevice.dll";
 
 						if (String.IsNullOrEmpty(DllPath) || !File.Exists(DllPath))
 						{
diff --git a/Engine/Source/Programs/UnrealBuildTool/Platform/Mac/MacToolChain.cs b/Engine/Source/Programs/UnrealBuildTool/Platform/Mac/MacToolChain.cs
--- a/Engine/Source/Programs/UnrealBuildTool/Platform/Mac/MacToolChain.cs
+++ b/Engine/Source/Programs/UnrealBuildTool/Platform/Mac/MacToolChain.cs
@@ -375,7 +375,7 @@ namespace UnrealBuildTool
 						int EngineDirStrIndex = LibraryDir.IndexOf("Engine");
 						if (EngineDirStrIndex >= 0)
 						{
-							LinkArguments.Add("-rpath \"@executable_path/../../../../../../" + LibraryDir.Substring(EngineDirStrIndex) + "\"");
+							LinkArguments.Add("-rpath \"@executable_path/../../../../../../" + LibraryDir[EngineDirStrIndex..] + "\"");
 						}
 					}
 				}
@@ -514,7 +514,7 @@ namespace UnrealBuildTool
 				StreamWriter FinalizeAppBundleScript = File.CreateText(FinalizeAppBundleScriptPath.FullName);
 				AppendMacLine(FinalizeAppBundleScript, "#!/bin/sh");
 				string BinariesPath = Path.GetDirectoryName(OutputFiles[0].AbsolutePath)!;
-				BinariesPath = Path.GetDirectoryName(BinariesPath.Substring(0, BinariesPath.IndexOf(".app")))!;
+				BinariesPath = Path.GetDirectoryName(BinariesPath[..BinariesPath.IndexOf(".app")])!;
 				AppendMacLine(FinalizeAppBundleScript, "cd \"{0}\"", BinariesPath.Replace("$", "\\$"));
 
 				string BundleVersion = LinkEnvironment.BundleVersion!;
@@ -889,7 +889,7 @@ namespace UnrealBuildTool
 					int AppIdx = MapFileBaseName.IndexOf(".app/Contents/MacOS");
 					if (AppIdx != -1)
 					{
-						MapFileBaseName = MapFileBaseName.Substring(0, AppIdx);
+						MapFileBaseName = MapFileBaseName[..AppIdx];
 					}
 
 					FileReference MapFilePath = new FileReference(MapFileBaseName + ".map");
@@ -1054,7 +1054,7 @@ namespace UnrealBuildTool
 		FileItem FinalizeAppBundle(ReadOnlyTargetRules Target, LinkEnvironment LinkEnvironment, FileItem Executable, IActionGraphBuilder Graph)
 		{
 			// Make a file item for the source and destination files
-			string FullDestPath = Executable.AbsolutePath.Substring(0, Executable.AbsolutePath.IndexOf(".app") + 4);
+			string FullDestPath = Executable.AbsolutePath[..(Executable.AbsolutePath.IndexOf(".app") + 4)];
 			FileItem DestFile = FileItem.GetItemByPath(FullDestPath);
 
 			// Make the compile action
@@ -1174,7 +1174,7 @@ namespace UnrealBuildTool
 					{
 						foreach (string ResourceFile in Directory.GetFiles(Resource.ResourcePath, "*", SearchOption.AllDirectories))
 						{
-							BuildProducts.Add(FileReference.Combine(BundleContentsDirectory!, Resource.BundleContentsSubdir!, ResourceFile.Substring(Path.GetDirectoryName(Resource.ResourcePath)!.Length + 1)), BuildProductType.RequiredResource);
+							BuildProducts.Add(FileReference.Combine(BundleContentsDirectory!, Resource.BundleContentsSubdir!, ResourceFile[(Path.GetDirectoryName(Resource.ResourcePath)!.Length + 1)..]), BuildProductType.RequiredResource);
 						}
 					}
 					else if (BundleContentsDirectory != null)
diff --git a/Engine/Source/Programs/UnrealBuildTool/Platform/Windows/VCToolChain.cs b/Engine/Source/Programs/UnrealBuildTool/Platform/Windows/VCToolChain.cs
--- a/Engine/Source/Programs/UnrealBuildTool/Platform/Windows/VCToolChain.cs
+++ b/Engine/Source/Programs/UnrealBuildTool/Platform/Windows/VCToolChain.cs
@@ -214,7 +214,7 @@ namespace UnrealBuildTool
 			}
 			else
 			{
-				AddDefinition(Arguments, Definition.Substring(0, ValueIdx), Definition.Substring(ValueIdx + 1));
+				AddDefinition(Arguments, Definition[..ValueIdx], Definition[(ValueIdx + 1)..]);
 			}
 		}
 
@@ -2123,7 +2123,7 @@ namespace UnrealBuildTool
 				string FileName = SourceFile.Name;
 				if (CompileEnvironment.CollidingNames != null && CompileEnvironment.CollidingNames.Contains(SourceFile))
 				{
-					string HashString = ContentHash.MD5(SourceFile.AbsolutePath.Substring(Unreal.RootDirectory.FullName.Length)).GetHashCode().ToString("X4");
+					string HashString = ContentHash.MD5(SourceFile.AbsolutePath[Unreal.RootDirectory.FullName.Length..]).GetHashCode().ToString("X4");
 					FileName = Path.GetFileNameWithoutExtension(FileName) + "_" + HashString + Path.GetExtension(FileName);
 				}
 
@@ -2473,7 +2473,7 @@ namespace UnrealBuildTool
 					else if (Target.WindowsPlatform.Compiler.IsClang())
 					{
 						FileReference ManifestFile = FileReference.Combine(Makefile.ProjectIntermediateDirectory, $"{Target.Name}TimingManifest.csv");
-						File.WriteAllLines(ManifestFile.FullName, TimingJsonFiles.Select(f => f.FullName.Remove(f.FullName.Length - ".json".Length)));
+						File.WriteAllLines(ManifestFile.FullName, TimingJsonFiles.Select(f => f.FullName[..^".json".Length]));
 
 						FileItem AggregateOutputFile = FileItem.GetItemByFileReference(FileReference.Combine(Makefile.ProjectIntermediateDirectory, $"{Target.Name}.trace.csv"));
 						FileItem HeadersOutputFile = FileItem.GetItemByFileReference(FileReference.Combine(Makefile.ProjectIntermediateDirectory, $"{Target.Name}.headers.csv"));
diff --git a/Engine/Source/Programs/UnrealBuildTool/ProjectFiles/CMake/CMakefileGenerator.cs b/Engine/Source/Programs/UnrealBuildTool/ProjectFiles/CMake/CMakefileGenerator.cs
--- a/Engine/Source/Programs/UnrealBuildTool/ProjectFiles/CMake/CMakefileGenerator.cs
+++ b/Engine/Source/Programs/UnrealBuildTool/ProjectFiles/CMake/CMakefileGenerator.cs
@@ -142,7 +142,7 @@ namespace UnrealBuildTool
 			{
 				if (String.IsNullOrEmpty(GameProjectName))
 				{
-					EngineFiles.Append("\t\"" + Utils.CleanDirectorySeparators(SourceFileRelativeToRoot, '/').Substring(3) + "\"\n");
+					EngineFiles.Append("\t\"" + Utils.CleanDirectorySeparators(SourceFileRelativeToRoot, '/')[3..] + "\"\n");
 				}
 				else
 				{
@@ -344,8 +344,8 @@ namespace UnrealBuildTool
 				int EqPos = PreprocessorDefinition.IndexOf("=");
 				if (EqPos >= 0)
 				{
-					string Key = PreprocessorDefinition.Substring(0, EqPos);
-					string Value = PreprocessorDefinition.Substring(EqPos).Replace("\"", "\\\"");
+					string Key = PreprocessorDefinition[..EqPos];
+					string Value = PreprocessorDefinition[EqPos..].Replace("\"", "\\\"");
 					PreprocessorDefinitionsList.Append("\t\"-D" + Key + Value + "\"\n");
 				}
 				else
diff --git a/Engine/Source/Programs/UnrealBuildTool/ProjectFiles/CodeLite/CodeLiteProject.cs b/Engine/Source/Programs/UnrealBuildTool/ProjectFiles/CodeLite/CodeLiteProject.cs
--- a/Engine/Source/Programs/UnrealBuildTool/ProjectFiles/CodeLite/CodeLiteProject.cs
+++ b/Engine/Source/Programs/UnrealBuildTool/ProjectFiles/CodeLite/CodeLiteProject.cs
@@ -131,12 +131,12 @@ namespace UnrealBuildTool
 							TargetName.Equals("UnrealEditor"))
 						{
 							int Idx = Unreal.EngineDirectory.FullName.Length;
-							CurrentFilePath = Path.GetDirectoryName(Path.GetFullPath(CurrentFile.Reference.FullName))!.Substring(Idx);
+							CurrentFilePath = Path.GetDirectoryName(Path.GetFullPath(CurrentFile.Reference.FullName))![Idx..];
 						}
 						else
 						{
 							int Idx = Path.GetDirectoryName(CurrentFile.Reference.FullName)!.IndexOf(ProjectNameRaw) + ProjectNameRaw.Length;
-							CurrentFilePath = Path.GetDirectoryName(CurrentFile.Reference.FullName)!.Substring(Idx);
+							CurrentFilePath = Path.GetDirectoryName(CurrentFile.Reference.FullName)![Idx..];
 						}
 					}
 					else if (ProjectTargetType == TargetType.Program)
@@ -145,7 +145,7 @@ namespace UnrealBuildTool
 						// We do not need all the editors subfolders to show the content. Find the correct programs subfolder.
 						//
 						int Idx = Path.GetDirectoryName(CurrentFile.Reference.FullName)!.IndexOf(TargetName) + TargetName.Length;
-						CurrentFilePath = Path.GetDirectoryName(CurrentFile.Reference.FullName)!.Substring(Idx);
+						CurrentFilePath = Path.GetDirectoryName(CurrentFile.Reference.FullName)![Idx..];
 					}
 
 					char[] Delimiters = new char[] { '/', '\\' };
diff --git a/Engine/Source/Programs/UnrealBuildTool/ProjectFiles/KDevelop/KDevelopGenerator.cs b/Engine/Source/Programs/UnrealBuildTool/ProjectFiles/KDevelop/KDevelopGenerator.cs
--- a/Engine/Source/Programs/UnrealBuildTool/ProjectFiles/KDevelop/KDevelopGenerator.cs
+++ b/Engine/Source/Programs/UnrealBuildTool/ProjectFiles/KDevelop/KDevelopGenerator.cs
@@ -346,8 +346,8 @@ namespace UnrealBuildTool
 			int EqualsIndex = Definition.IndexOf('=');
 			if (EqualsIndex >= 0)
 			{
-				Key = Definition.Substring(0, EqualsIndex);
-				Value = Definition.Substring(EqualsIndex + 1);
+				Key = Definition[..EqualsIndex];
+				Value = Definition[(EqualsIndex + 1)..];
 			}
 			else
 			{
diff --git a/Engine/Source/Programs/UnrealBuildTool/ProjectFiles/Make/MakefileGenerator.cs b/Engine/Source/Programs/UnrealBuildTool/ProjectFiles/Make/MakefileGenerator.cs
--- a/Engine/Source/Programs/UnrealBuildTool/ProjectFiles/Make/MakefileGenerator.cs
+++ b/Engine/Source/Programs/UnrealBuildTool/ProjectFiles/Make/MakefileGenerator.cs
@@ -84,7 +84,7 @@ namespace UnrealBuildTool
 			foreach (ProjectTarget TargetFile in AllProjects)
 			{
 				string TargetFileName = TargetFile.TargetFilePath.GetFileNameWithoutExtension();
-				string Basename = TargetFileName.Substring(0, TargetFileName.LastIndexOf(".Target", StringComparison.InvariantCultureIgnoreCase));
+				string Basename = TargetFileName[..TargetFileName.LastIndexOf(".Target", StringComparison.InvariantCultureIgnoreCase)];
 
 				foreach (UnrealTargetPlatform Platform in TargetFile.SupportedPlatforms.Concat(TargetFile.ExtraSupportedPlatforms).Distinct().OrderBy(x => x.ToString()))
 				{
@@ -115,7 +115,7 @@ namespace UnrealBuildTool
 			foreach (ProjectTarget TargetFile in AllProjects)
 			{
 				string TargetFileName = TargetFile.TargetFilePath.GetFileNameWithoutExtension();
-				string Basename = TargetFileName.Substring(0, TargetFileName.LastIndexOf(".Target", StringComparison.InvariantCultureIgnoreCase));
+				string Basename = TargetFileName[..TargetFileName.LastIndexOf(".Target", StringComparison.InvariantCultureIgnoreCase)];
 				string MakeProjectCmdArg = "";
 
 				if (TargetFile.UnrealProjectFilePath != null)
diff --git a/Engine/Source/Programs/UnrealBuildTool/ProjectFiles/ProjectFile.cs b/Engine/Source/Programs/UnrealBuildTool/ProjectFiles/ProjectFile.cs
--- a/Engine/Source/Programs/UnrealBuildTool/ProjectFiles/ProjectFile.cs
+++ b/Engine/Source/Programs/UnrealBuildTool/ProjectFiles/ProjectFile.cs
@@ -243,8 +243,8 @@ namespace UnrealBuildTool
 			int EqualsIndex = Definition.IndexOf('=');
 			if (EqualsIndex >= 0)
 			{
-				Key = Definition.Substring(0, EqualsIndex);
-				Value = Definition.Substring(EqualsIndex + 1);
+				Key = Definition[..EqualsIndex];
+				Value = Definition[(EqualsIndex + 1)..];
 			}
 			else
 			{
diff --git a/Engine/Source/Programs/UnrealBuildTool/ProjectFiles/ProjectFileGenerator.cs b/Engine/Source/Programs/UnrealBuildTool/ProjectFiles/ProjectFileGenerator.cs
--- a/Engine/Source/Programs/UnrealBuildTool/ProjectFiles/ProjectFileGenerator.cs
+++ b/Engine/Source/Programs/UnrealBuildTool/ProjectFiles/ProjectFileGenerator.cs
@@ -973,7 +973,7 @@ namespace UnrealBuildTool
 
 				if (PrimaryProjectPathString.Length > 2 && PrimaryProjectPathString.Substring(1, 1) == ":")
 				{
-					string NewPrimaryProjectDrive = PrimaryProjectPathString.Substring(0, 1);
+					string NewPrimaryProjectDrive = PrimaryProjectPathString[..1];
 					if (!String.IsNullOrEmpty(NewPrimaryProjectDrive))
 					{
 						PrimaryProjectName += "-" + NewPrimaryProjectDrive;
@@ -1554,15 +1554,15 @@ namespace UnrealBuildTool
 					if (CurArgument.StartsWith("-Platforms=", StringComparison.InvariantCultureIgnoreCase))
 					{
 						// Parse the list... will be in Foo+Bar+New format
-						string PlatformList = CurArgument.Substring(11);
+						string PlatformList = CurArgument[11..];
 						while (PlatformList.Length > 0)
 						{
 							string PlatformString = PlatformList;
 							int PlusIdx = PlatformList.IndexOf("+");
 							if (PlusIdx != -1)
 							{
-								PlatformString = PlatformList.Substring(0, PlusIdx);
-								PlatformList = PlatformList.Substring(PlusIdx + 1);
+								PlatformString = PlatformList[..PlusIdx];
+								PlatformList = PlatformList[(PlusIdx + 1)..];
 							}
 							else
 							{
@@ -1901,11 +1901,11 @@ namespace UnrealBuildTool
 				{
 					if (LocalRoot.EndsWith("\\") || LocalRoot.EndsWith("/"))
 					{
-						Search = Search.Substring(LocalRoot.Length - 1);
+						Search = Search[(LocalRoot.Length - 1)..];
 					}
 					else
 					{
-						Search = Search.Substring(LocalRoot.Length);
+						Search = Search[LocalRoot.Length..];
 					}
 				}
 
@@ -1916,11 +1916,11 @@ namespace UnrealBuildTool
 					{
 						if (ProjectRoot.EndsWith("\\") || ProjectRoot.EndsWith("/"))
 						{
-							Search = Search.Substring(ProjectRoot.Length - 1);
+							Search = Search[(ProjectRoot.Length - 1)..];
 						}
 						else
 						{
-							Search = Search.Substring(ProjectRoot.Length);
+							Search = Search[ProjectRoot.Length..];
 						}
 					}
 				}
diff --git a/Engine/Source/Programs/UnrealBuildTool/ProjectFiles/QMake/QMakefileGenerator.cs b/Engine/Source/Programs/UnrealBuildTool/ProjectFiles/QMake/QMakefileGenerator.cs
--- a/Engine/Source/Programs/UnrealBuildTool/ProjectFiles/QMake/QMakefileGenerator.cs
+++ b/Engine/Source/Programs/UnrealBuildTool/ProjectFiles/QMake/QMakefileGenerator.cs
@@ -51,8 +51,8 @@ namespace UnrealBuildTool
 			int EqualsIndex = Definition.IndexOf('=');
 			if (EqualsIndex >= 0)
 			{
-				Key = Definition.Substring(0, EqualsIndex);
-				Value = Definition.Substring(EqualsIndex + 1);
+				Key = Definition[..EqualsIndex];
+				Value = Definition[(EqualsIndex + 1)..];
 			}
 			else
 			{
@@ -321,7 +321,7 @@ namespace UnrealBuildTool
 							{
 								if (String.IsNullOrEmpty(GameProjectName))
 								{
-									QMakeSourceFilesListBuilder.Append("\t\"" + SourceFileRelativeToRoot.Substring(3) + "\" \\\n");
+									QMakeSourceFilesListBuilder.Append("\t\"" + SourceFileRelativeToRoot[3..] + "\" \\\n");
 								}
 								else
 								{
@@ -341,7 +341,7 @@ namespace UnrealBuildTool
 								if (String.IsNullOrEmpty(GameProjectName))
 								{
 									// SourceFileRelativeToRoot = SourceFileRelativeToRoot.Substring (3);
-									QMakeHeaderFilesListBuilder.Append("\t\"" + SourceFileRelativeToRoot.Substring(3) + "\" \\\n");
+									QMakeHeaderFilesListBuilder.Append("\t\"" + SourceFileRelativeToRoot[3..] + "\" \\\n");
 								}
 								else
 								{
@@ -361,7 +361,7 @@ namespace UnrealBuildTool
 								if (String.IsNullOrEmpty(GameProjectName))
 								{
 									// SourceFileRelativeToRoot = SourceFileRelativeToRoot.Substring (3);
-									QMakeConfigFilesListBuilder.Append("\t\"" + SourceFileRelativeToRoot.Substring(3) + "\" \\\n");
+									QMakeConfigFilesListBuilder.Append("\t\"" + SourceFileRelativeToRoot[3..] + "\" \\\n");
 								}
 								else
 								{
diff --git a/Engine/Source/Programs/UnrealBuildTool/ProjectFiles/Rider/RiderProjectFile.cs b/Engine/Source/Programs/UnrealBuildTool/ProjectFiles/Rider/RiderProjectFile.cs
--- a/Engine/Source/Programs/UnrealBuildTool/ProjectFiles/Rider/RiderProjectFile.cs
+++ b/Engine/Source/Programs/UnrealBuildTool/ProjectFiles/Rider/RiderProjectFile.cs
@@ -270,7 +270,7 @@ namespace UnrealBuildTool
 					string? Line = Reader.ReadLine();
 					if (Line != null)
 					{
-						CorrectFilePathPch = Line.Substring("// PCH for ".Length).Trim();
+						CorrectFilePathPch = Line["// PCH for ".Length..].Trim();
 						return true;
 					}
 				}
diff --git a/Engine/Source/Programs/UnrealBuildTool/ProjectFiles/VisualStudio/VCProject.cs b/Engine/Source/Programs/UnrealBuildTool/ProjectFiles/VisualStudio/VCProject.cs
--- a/Engine/Source/Programs/UnrealBuildTool/ProjectFiles/VisualStudio/VCProject.cs
+++ b/Engine/Source/Programs/UnrealBuildTool/ProjectFiles/VisualStudio/VCProject.cs
@@ -1239,7 +1239,7 @@ namespace UnrealBuildTool
 					}
 					else
 					{
-						FilterRelativeSourceDirectory = FilterRelativeSourceDirectory.Substring(0, LastSeparatorIdx);
+						FilterRelativeSourceDirectory = FilterRelativeSourceDirectory[..LastSeparatorIdx];
 					}
 
 					LocalAliasedFiles.Add(new AliasedFile(CurFile.Reference, ProjectRelativeSourceFile, FilterRelativeSourceDirectory));
@@ -2033,8 +2033,8 @@ namespace UnrealBuildTool
 						string SplitDirectory;
 						if (SlashIndex != -1)
 						{
-							SplitDirectory = PathRemaining.Substring(0, SlashIndex);
-							PathRemaining = PathRemaining.Substring(SplitDirectory.Length + 1);
+							SplitDirectory = PathRemaining[..SlashIndex];
+							PathRemaining = PathRemaining[(SplitDirectory.Length + 1)..];
 						}
 						else
 						{
diff --git a/Engine/Source/Programs/UnrealBuildTool/ProjectFiles/VisualStudioCode/VSCodeProjectFileGenerator.cs b/Engine/Source/Programs/UnrealBuildTool/ProjectFiles/VisualStudioCode/VSCodeProjectFileGenerator.cs
--- a/Engine/Source/Programs/UnrealBuildTool/ProjectFiles/VisualStudioCode/VSCodeProjectFileGenerator.cs
+++ b/Engine/Source/Programs/UnrealBuildTool/ProjectFiles/VisualStudioCode/VSCodeProjectFileGenerator.cs
@@ -192,7 +192,7 @@ namespace UnrealBuildTool
 			public void EndObject()
 			{
 				Lines[^1] = Lines[^1].TrimEnd(',');
-				TabString = TabString.Remove(TabString.Length - 1);
+				TabString = TabString[..^1];
 				Lines.Add(TabString + "},");
 			}
 
@@ -206,7 +206,7 @@ namespace UnrealBuildTool
 			public void EndArray()
 			{
 				Lines[^1] = Lines[^1].TrimEnd(',');
-				TabString = TabString.Remove(TabString.Length - 1);
+				TabString = TabString[..^1];
 				Lines.Add(TabString + "],");
 			}
 
diff --git a/Engine/Source/Programs/UnrealBuildTool/ProjectFiles/Xcode/XcodeProject.Scheme.cs b/Engine/Source/Programs/UnrealBuildTool/ProjectFiles/Xcode/XcodeProject.Scheme.cs
--- a/Engine/Source/Programs/UnrealBuildTool/ProjectFiles/Xcode/XcodeProject.Scheme.cs
+++ b/Engine/Source/Programs/UnrealBuildTool/ProjectFiles/Xcode/XcodeProject.Scheme.cs
@@ -54,7 +54,7 @@ namespace UnrealBuildTool.XcodeProjectXcconfig
 					int OldCommandLineArgumentsEnd = OldContents.IndexOf("</CommandLineArguments>");
 					if (OldCommandLineArgumentsStart != -1 && OldCommandLineArgumentsEnd != -1)
 					{
-						OldCommandLineArguments = OldContents.Substring(OldCommandLineArgumentsStart, OldCommandLineArgumentsEnd - OldCommandLineArgumentsStart);
+						OldCommandLineArguments = OldContents[OldCommandLineArgumentsStart..OldCommandLineArgumentsEnd];
 					}
 				}
 
diff --git a/Engine/Source/Programs/UnrealBuildTool/ProjectFiles/Xcode/XcodeProject.cs b/Engine/Source/Programs/UnrealBuildTool/ProjectFiles/Xcode/XcodeProject.cs
--- a/Engine/Source/Programs/UnrealBuildTool/ProjectFiles/Xcode/XcodeProject.cs
+++ b/Engine/Source/Programs/UnrealBuildTool/ProjectFiles/Xcode/XcodeProject.cs
@@ -1584,7 +1584,7 @@ namespace UnrealBuildTool.XcodeProjectXcconfig
 				};
 
 				// make a string like --exclude=/Info.plist --exclude=/Manifest_* ...
-				string ExcludeString = String.Join(" ", Exclusions.Select(x => (x[0] == '+' ? "--include" : "--exclude") + $"=\\\"{x.Substring(1)}\\\""));
+				string ExcludeString = String.Join(" ", Exclusions.Select(x => (x[0] == '+' ? "--include" : "--exclude") + $"=\\\"{x[1..]}\\\""));
 
 
 				// copy uecommandline.txt for IOS type platforms
diff --git a/Engine/Source/Programs/UnrealBuildTool/ProjectFiles/Xcode/XcodeProjectLegacy.cs b/Engine/Source/Programs/UnrealBuildTool/ProjectFiles/Xcode/XcodeProjectLegacy.cs
--- a/Engine/Source/Programs/UnrealBuildTool/ProjectFiles/Xcode/XcodeProjectLegacy.cs
+++ b/Engine/Source/Programs/UnrealBuildTool/ProjectFiles/Xcode/XcodeProjectLegacy.cs
@@ -1583,7 +1583,7 @@ namespace UnrealBuildTool.XcodeProjectLegacy
 				int OldCommandLineArgumentsEnd = OldContents.IndexOf("</CommandLineArguments>");
 				if (OldCommandLineArgumentsStart != -1 && OldCommandLineArgumentsEnd != -1)
 				{
-					OldCommandLineArguments = OldContents.Substring(OldCommandLineArgumentsStart, OldCommandLineArgumentsEnd - OldCommandLineArgumentsStart);
+					OldCommandLineArguments = OldContents[OldCommandLineArgumentsStart..OldCommandLineArgumentsEnd];
 				}
 			}
 
diff --git a/Engine/Source/Programs/UnrealBuildTool/System/ActionHistory.cs b/Engine/Source/Programs/UnrealBuildTool/System/ActionHistory.cs
--- a/Engine/Source/Programs/UnrealBuildTool/System/ActionHistory.cs
+++ b/Engine/Source/Programs/UnrealBuildTool/System/ActionHistory.cs
@@ -295,7 +295,7 @@ namespace UnrealBuildTool
 					int PlatformEndOffset = Location.FullName.IndexOf(Path.DirectorySeparatorChar, PlatformOffset);
 					if (PlatformEndOffset != -1)
 					{
-						string PlatformName = Location.FullName.Substring(PlatformOffset, PlatformEndOffset - PlatformOffset);
+						string PlatformName = Location.FullName[PlatformOffset..PlatformEndOffset];
 						return FileReference.Combine(BaseDir, "Intermediate", "Build", PlatformName, "ActionHistory.bin");
 					}
 				}
@@ -318,7 +318,7 @@ namespace UnrealBuildTool
 							}
 							if (Idx == 2)
 							{
-								return FileReference.Combine(BaseDir, Location.FullName.Substring(NameOffset, EndOffset - NameOffset), "ActionHistory.bin");
+								return FileReference.Combine(BaseDir, Location.FullName[NameOffset..EndOffset], "ActionHistory.bin");
 							}
 						}
 					}
diff --git a/Engine/Source/Programs/UnrealBuildTool/System/BuildHostPlatform.cs b/Engine/Source/Programs/UnrealBuildTool/System/BuildHostPlatform.cs
--- a/Engine/Source/Programs/UnrealBuildTool/System/BuildHostPlatform.cs
+++ b/Engine/Source/Programs/UnrealBuildTool/System/BuildHostPlatform.cs
@@ -303,10 +303,10 @@ namespace UnrealBuildTool
 				{
 					string Line = FileLine.Trim();
 					int PIDEnd = Line.IndexOf(' ');
-					string PIDString = Line.Substring(0, PIDEnd);
+					string PIDString = Line[..PIDEnd];
 					if (PIDString != "PID")
 					{
-						string Filename = Line.Substring(PIDEnd + 1);
+						string Filename = Line[(PIDEnd + 1)..];
 						int Pid = Int32.Parse(PIDString);
 						try
 						{
@@ -371,7 +371,7 @@ namespace UnrealBuildTool
 					const int SharingModeLength = 6;
 					int SMStart = Line.IndexOf("SM=");
 					int PathStart = SMStart + SharingModeLength;
-					string Module = Line.Substring(PathStart).Trim();
+					string Module = Line[PathStart..].Trim();
 					if (!Modules.Contains(Module))
 					{
 						Modules.Add(Module);
diff --git a/Engine/Source/Programs/UnrealBuildTool/System/CommandLine.cs b/Engine/Source/Programs/UnrealBuildTool/System/CommandLine.cs
--- a/Engine/Source/Programs/UnrealBuildTool/System/CommandLine.cs
+++ b/Engine/Source/Programs/UnrealBuildTool/System/CommandLine.cs
@@ -127,7 +127,7 @@ namespace UnrealBuildTool
 				{
 					// Get the length of the argument prefix
 					int EqualsIdx = Argument.IndexOfAny(new char[] { '=', ':' });
-					string Prefix = (EqualsIdx == -1) ? Argument : Argument.Substring(0, EqualsIdx + 1);
+					string Prefix = (EqualsIdx == -1) ? Argument : Argument[..(EqualsIdx + 1)];
 
 					// Check if there's a matching argument registered
 					Parameter? Parameter;
@@ -149,7 +149,7 @@ namespace UnrealBuildTool
 						}
 						else if (EqualsIdx != -1)
 						{
-							AssignValue(Parameter, Argument.Substring(EqualsIdx + 1), TargetObject, AssignedFieldToParameter, Logger);
+							AssignValue(Parameter, Argument[(EqualsIdx + 1)..], TargetObject, AssignedFieldToParameter, Logger);
 						}
 						else if (NonNullableType(Parameter.FieldInfo.FieldType) == typeof(bool))
 						{
diff --git a/Engine/Source/Programs/UnrealBuildTool/System/ConfigCache.cs b/Engine/Source/Programs/UnrealBuildTool/System/ConfigCache.cs
--- a/Engine/Source/Programs/UnrealBuildTool/System/ConfigCache.cs
+++ b/Engine/Source/Programs/UnrealBuildTool/System/ConfigCache.cs
@@ -343,11 +343,11 @@ namespace UnrealBuildTool
 			{
 				if (CmdLineArg.StartsWith(IniConfigArgPrefix, StringComparison.InvariantCultureIgnoreCase))
 				{
-					OverrideStrings.Add(CmdLineArg.Substring(IniConfigArgPrefix.Length));
+					OverrideStrings.Add(CmdLineArg[IniConfigArgPrefix.Length..]);
 				}
 				if (CmdLineArg.StartsWith(CustomConfigPrefix, StringComparison.InvariantCultureIgnoreCase))
 				{
-					CustomConfig = CmdLineArg.Substring(CustomConfigPrefix.Length);
+					CustomConfig = CmdLineArg[CustomConfigPrefix.Length..];
 				}
 				if (CmdLineArg.Equals("-dumpiniloads", StringComparison.InvariantCultureIgnoreCase))
 				{
@@ -775,7 +775,7 @@ namespace UnrealBuildTool
 				string Line = Lines[Index];
 				if (Line.StartsWith('+') || Line.StartsWith('-') || Line.StartsWith('.') || Line.StartsWith('!'))
 				{
-					Line = Line.Substring(1);
+					Line = Line[1..];
 				}
 
 				// look for last array entry in case of multiples (or the only line for non-array type)
diff --git a/Engine/Source/Programs/UnrealBuildTool/System/ConfigFile.cs b/Engine/Source/Programs/UnrealBuildTool/System/ConfigFile.cs
--- a/Engine/Source/Programs/UnrealBuildTool/System/ConfigFile.cs
+++ b/Engine/Source/Programs/UnrealBuildTool/System/ConfigFile.cs
@@ -341,7 +341,7 @@ namespace UnrealBuildTool
 			foreach (string Setting in SettingLines)
 			{
 				// Locate and break off the section name
-				string SectionName = Setting.Remove(Setting.IndexOf(':')).Trim(new char[] { '[', ']' });
+				string SectionName = Setting[..Setting.IndexOf(':')].Trim(new char[] { '[', ']' });
 				// lookup remaps
 				RemapSectionOrKey(SectionNameRemap, SectionName, $"which is a config section (found in a string in code - search your .ini files for it, the source file is unknown)");
 				SectionKeyRemap.TryGetValue(SectionName, out CurrentRemap);
@@ -357,7 +357,7 @@ namespace UnrealBuildTool
 
 					if (CurrentSection != null)
 					{
-						string IniKeyValue = Setting.Substring(Setting.IndexOf(':') + 1);
+						string IniKeyValue = Setting[(Setting.IndexOf(':') + 1)..];
 						TryAddConfigLine(CurrentSection, CurrentRemap, "unknown source file", IniKeyValue, 0, IniKeyValue.Length, DefaultAction, Sections);
 					}
 				}
@@ -403,7 +403,7 @@ namespace UnrealBuildTool
 			// RemoveKey actions do not require a value
 			if (Action == ConfigLineAction.RemoveKey && EqualsIdx == -1)
 			{
-				Section.Lines.Add(new ConfigLine(Action, Line.Substring(KeyStartIdx).Trim(), ""));
+				Section.Lines.Add(new ConfigLine(Action, Line[KeyStartIdx..].Trim(), ""));
 				return true;
 			}
 
@@ -445,8 +445,8 @@ namespace UnrealBuildTool
 			}
 
 			// Add it to the config section
-			string Key = Line.Substring(KeyStartIdx, KeyEndIdx - KeyStartIdx);
-			string Value = Line.Substring(ValueStartIdx, ValueEndIdx - ValueStartIdx);
+			string Key = Line[KeyStartIdx..KeyEndIdx];
+			string Value = Line[ValueStartIdx..ValueEndIdx];
 
 			// remap the key if needed
 			string NewKey = RemapSectionOrKey(KeyRemap, Key, $"which is a config key in section [{Section.Name}], in '{Filename}'");
@@ -454,7 +454,7 @@ namespace UnrealBuildTool
 			// look for a section:name remap
 			if (!NewKey.Equals(Key) && NewKey.Contains(":", StringComparison.CurrentCulture))
 			{
-				string SectionName = NewKey.Substring(0, NewKey.IndexOf(':'));
+				string SectionName = NewKey[..NewKey.IndexOf(':')];
 				ConfigFileSection? CurrentSection;
 				if (!Sections.TryGetValue(SectionName, out CurrentSection))
 				{
@@ -462,7 +462,7 @@ namespace UnrealBuildTool
 					Sections.Add(SectionName, CurrentSection);
 				}
 
-				string KeyName = NewKey.Substring(NewKey.IndexOf(':') + 1);
+				string KeyName = NewKey[(NewKey.IndexOf(':') + 1)..];
 
 				CurrentSection.Lines.Add(new ConfigLine(Action, KeyName, Value, HasStrippedValueQuotes));
 
diff --git a/Engine/Source/Programs/UnrealBuildTool/System/ConfigHierarchy.cs b/Engine/Source/Programs/UnrealBuildTool/System/ConfigHierarchy.cs
--- a/Engine/Source/Programs/UnrealBuildTool/System/ConfigHierarchy.cs
+++ b/Engine/Source/Programs/UnrealBuildTool/System/ConfigHierarchy.cs
@@ -584,7 +584,7 @@ namespace UnrealBuildTool
 		{
 			if (Text.EndsWith("f"))
 			{
-				return Single.TryParse(Text.Substring(0, Text.Length - 1), out Value);
+				return Single.TryParse(Text[..^1], out Value);
 			}
 			else
 			{
@@ -602,7 +602,7 @@ namespace UnrealBuildTool
 		{
 			if (Text.EndsWith("f"))
 			{
-				return Double.TryParse(Text.Substring(0, Text.Length - 1), out Value);
+				return Double.TryParse(Text[..^1], out Value);
 			}
 			else
 			{
@@ -1458,7 +1458,7 @@ namespace UnrealBuildTool
 				return new Dictionary<string, string>();
 			}
 			// strip ()
-			Input = Input.Substring(1, Input.Length - 2);
+			Input = Input[1..^1];
 
 			List<string> Props = new List<string>();
 
@@ -1472,19 +1472,19 @@ namespace UnrealBuildTool
 				// comma first? easy
 				if (NextComma != -1 && NextComma < NextQuote)
 				{
-					Props.Add(Input.Substring(TokenStart, NextComma - TokenStart));
+					Props.Add(Input[TokenStart..NextComma]);
 					TokenStart = NextComma + 1;
 				}
 				// comma but no quotes
 				else if (NextComma != -1 && NextQuote == -1)
 				{
-					Props.Add(Input.Substring(TokenStart, NextComma - TokenStart));
+					Props.Add(Input[TokenStart..NextComma]);
 					TokenStart = NextComma + 1;
 				}
 				// neither found, use the rest
 				else if (NextComma == -1 && NextQuote == -1)
 				{
-					Props.Add(Input.Substring(TokenStart));
+					Props.Add(Input[TokenStart..]);
 					break;
 				}
 				// quote first? look for quote after
@@ -1495,7 +1495,7 @@ namespace UnrealBuildTool
 					if (NextQuote + 1 == StrLen)
 					{
 						// use the rest of the string
-						Props.Add(Input.Substring(TokenStart));
+						Props.Add(Input[TokenStart..]);
 						break;
 					}
 					// it's expected that the following character is a comma, if not, give up
@@ -1523,8 +1523,8 @@ namespace UnrealBuildTool
 					continue;
 				}
 
-				string Key = Prop.Substring(0, Equals);
-				string Value = Prop.Substring(Equals + 1);
+				string Key = Prop[..Equals];
+				string Value = Prop[(Equals + 1)..];
 				// trim off any quotes around the entire value
 				Value = Value.Trim(" \"".ToCharArray());
 				Key = Key.Trim(" ".ToCharArray());
diff --git a/Engine/Source/Programs/UnrealBuildTool/System/CppIncludeLookup.cs b/Engine/Source/Programs/UnrealBuildTool/System/CppIncludeLookup.cs
--- a/Engine/Source/Programs/UnrealBuildTool/System/CppIncludeLookup.cs
+++ b/Engine/Source/Programs/UnrealBuildTool/System/CppIncludeLookup.cs
@@ -343,7 +343,7 @@ namespace UnrealBuildTool
 		static byte[]? ParseIncludeName(ReadOnlySpan<byte> Span, int Idx)
 		{
 			ReadOnlySpan<byte> IncludeSpan = IncludeBytes.Span;
-			if (!Span.Slice(Idx).StartsWith(IncludeSpan))
+			if (!Span[Idx..].StartsWith(IncludeSpan))
 			{
 				return null;
 			}
@@ -423,7 +423,7 @@ namespace UnrealBuildTool
 
 			// Normalize the filename
 			byte[] Name = new byte[MaxNameIdx - MinNameIdx];
-			Span.Slice(MinNameIdx, MaxNameIdx - MinNameIdx).CopyTo(Name);
+			Span[MinNameIdx..MaxNameIdx].CopyTo(Name);
 			return Name;
 		}
 	}
diff --git a/Engine/Source/Programs/UnrealBuildTool/System/CppRootPaths.cs b/Engine/Source/Programs/UnrealBuildTool/System/CppRootPaths.cs
--- a/Engine/Source/Programs/UnrealBuildTool/System/CppRootPaths.cs
+++ b/Engine/Source/Programs/UnrealBuildTool/System/CppRootPaths.cs
@@ -244,7 +244,7 @@ namespace UnrealBuildTool
 			{
 				if (reference.IsUnderDirectory(item.Value.local))
 				{
-					vfsPath = item.Value.vfs.FullName + reference.FullName.Substring(item.Value.local.FullName.Length);
+					vfsPath = item.Value.vfs.FullName + reference.FullName[item.Value.local.FullName.Length..];
 					if (OperatingSystem.IsWindows())
 					{
 						vfsPath = vfsPath.Replace('\\', '/');
@@ -257,7 +257,7 @@ namespace UnrealBuildTool
 			{
 				if (reference.IsUnderDirectory(item.Value.local))
 				{
-					vfsPath = item.Value.vfs.FullName + reference.FullName.Substring(item.Value.local.FullName.Length);
+					vfsPath = item.Value.vfs.FullName + reference.FullName[item.Value.local.FullName.Length..];
 					if (OperatingSystem.IsWindows())
 					{
 						vfsPath = vfsPath.Replace('\\', '/');
@@ -283,7 +283,7 @@ namespace UnrealBuildTool
 			{
 				if (reference.IsUnderDirectory(item.Value.vfs))
 				{
-					return FileReference.FromString(item.Value.local.FullName + reference.FullName.Substring(item.Value.vfs.FullName.Length));
+					return FileReference.FromString(item.Value.local.FullName + reference.FullName[item.Value.vfs.FullName.Length..]);
 				}
 			}
 
@@ -291,7 +291,7 @@ namespace UnrealBuildTool
 			{
 				if (reference.IsUnderDirectory(item.Value.vfs))
 				{
-					return FileReference.FromString(item.Value.local.FullName + reference.FullName.Substring(item.Value.vfs.FullName.Length));
+					return FileReference.FromString(item.Value.local.FullName + reference.FullName[item.Value.vfs.FullName.Length..]);
 				}
 			}
 
diff --git a/Engine/Source/Programs/UnrealBuildTool/System/DeployExports.cs b/Engine/Source/Programs/UnrealBuildTool/System/DeployExports.cs
--- a/Engine/Source/Programs/UnrealBuildTool/System/DeployExports.cs
+++ b/Engine/Source/Programs/UnrealBuildTool/System/DeployExports.cs
@@ -33,7 +33,7 @@ namespace UnrealBuildTool
 
 			if (DeployFolder.Length > DeployFolderMaxLength)
 			{
-				DeployFolder = DeployFolder.Substring(0, DeployFolderMaxLength);
+				DeployFolder = DeployFolder[..DeployFolderMaxLength];
 				char[] CharsToTrim = { '-' };
 				DeployFolder = DeployFolder.Trim(CharsToTrim);
 			}
diff --git a/Engine/Source/Programs/UnrealBuildTool/System/EncryptionAndSigning.cs b/Engine/Source/Programs/UnrealBuildTool/System/EncryptionAndSigning.cs
--- a/Engine/Source/Programs/UnrealBuildTool/System/EncryptionAndSigning.cs
+++ b/Engine/Source/Programs/UnrealBuildTool/System/EncryptionAndSigning.cs
@@ -218,7 +218,7 @@ namespace UnrealBuildTool
 		{
 			if (InString.StartsWith("0x"))
 			{
-				InString = InString.Substring(2);
+				InString = InString[2..];
 			}
 			return InString.TrimStart('0');
 		}
@@ -506,15 +506,15 @@ namespace UnrealBuildTool
 		{
 			if (InString.StartsWith("0x"))
 			{
-				InString = InString.Substring(2);
+				InString = InString[2..];
 			}
 
 			List<byte> Bytes = new List<byte>();
 			while (InString.Length > 0)
 			{
 				int CharsToParse = Math.Min(2, InString.Length);
-				string Value = InString.Substring(InString.Length - CharsToParse);
-				InString = InString.Substring(0, InString.Length - CharsToParse);
+				string Value = InString[^CharsToParse..];
+				InString = InString[..^CharsToParse];
 				Bytes.Add(Byte.Parse(Value, System.Globalization.NumberStyles.AllowHexSpecifier));
 			}
 
@@ -532,7 +532,7 @@ namespace UnrealBuildTool
 
 			if (InString.Length > InRequiredLength)
 			{
-				InString = InString.Substring(0, InRequiredLength);
+				InString = InString[..InRequiredLength];
 			}
 
 			foreach (char C in InString)
diff --git a/Engine/Source/Programs/UnrealBuildTool/System/HotReload.cs b/Engine/Source/Programs/UnrealBuildTool/System/HotReload.cs
--- a/Engine/Source/Programs/UnrealBuildTool/System/HotReload.cs
+++ b/Engine/Source/Programs/UnrealBuildTool/System/HotReload.cs
@@ -825,12 +825,12 @@ namespace UnrealBuildTool
 				}
 				if (SuffixEndIdx == FileName.Length || FileName[SuffixEndIdx] == '-' || FileName[SuffixEndIdx] == '.')
 				{
-					FileName = FileName.Substring(0, NameEndIdx) + FileName.Substring(SuffixEndIdx);
+					FileName = FileName[..NameEndIdx] + FileName[SuffixEndIdx..];
 				}
 			}
 
 			// NOTE: Formatting of this string must match the code in ModuleManager.cpp, MakeUniqueModuleFilename
-			string NewFileName = String.Format("{0}-{1:D4}{2}", FileName.Substring(0, NameEndIdx), Suffix, FileName.Substring(NameEndIdx));
+			string NewFileName = String.Format("{0}-{1:D4}{2}", FileName[..NameEndIdx], Suffix, FileName[NameEndIdx..]);
 
 			return FileReference.Combine(File.Directory, NewFileName);
 		}
@@ -854,7 +854,7 @@ namespace UnrealBuildTool
 				}
 				else if ((Idx == 0 || !IsBaseFileNameCharacter(Text[Idx - 1])) && (Idx + OldFileName.Length == Text.Length || !IsBaseFileNameCharacter(Text[Idx + OldFileName.Length])))
 				{
-					Text = Text.Substring(0, Idx) + NewFileName + Text.Substring(Idx + OldFileName.Length);
+					Text = Text[..Idx] + NewFileName + Text[(Idx + OldFileName.Length)..];
 					StartIdx = Idx + NewFileName.Length;
 				}
 				else
@@ -947,7 +947,7 @@ namespace UnrealBuildTool
 							}
 						}
 
-						FileReference OldDependenciesFile = new FileReference(Arguments[DependenciesIdx].Substring(DependenciesPrefix.Length));
+						FileReference OldDependenciesFile = new FileReference(Arguments[DependenciesIdx][DependenciesPrefix.Length..]);
 						FileItem OldDependenciesFileItem = Action.ProducedItems.First(x => x.Location == OldDependenciesFile);
 						NewAction.ProducedItems.Remove(OldDependenciesFileItem);
 
@@ -972,7 +972,7 @@ namespace UnrealBuildTool
 						}
 					}
 
-					FileReference OldResponseFile = Action.RootPaths.GetLocalPath(new FileReference(Arguments[ResponseFileIdx].Substring(1).Trim('\"')));
+					FileReference OldResponseFile = Action.RootPaths.GetLocalPath(new FileReference(Arguments[ResponseFileIdx][1..].Trim('\"')));
 					FileReference NewResponseFile = OldResponseFile.ChangeExtension(".lc.rsp");
 
 					NewAction.PrerequisiteItems.Remove(FileItem.GetItemByFileReference(OldResponseFile));
@@ -986,7 +986,7 @@ namespace UnrealBuildTool
 						string ResponseLine = ResponseLines[Idx];
 						if (ResponseLine.StartsWith(OutputFilePrefix, StringComparison.Ordinal))
 						{
-							FileReference OldOutputFile = Action.RootPaths.GetLocalPath(new FileReference(ResponseLine.Substring(OutputFilePrefix.Length).Trim('\"')));
+							FileReference OldOutputFile = Action.RootPaths.GetLocalPath(new FileReference(ResponseLine[OutputFilePrefix.Length..].Trim('\"')));
 							FileItem OldOutputFileItem = Action.ProducedItems.First(x => x.Location == OldOutputFile);
 							NewAction.ProducedItems.Remove(OldOutputFileItem);
 
@@ -1020,11 +1020,11 @@ namespace UnrealBuildTool
 							string ResponseLine = ResponseLines[Idx];
 							if (ResponseLine.StartsWith(SourceDependencyPrefix, StringComparison.Ordinal))
 							{
-								FileReference OldSourceDependencyFile = Action.RootPaths.GetLocalPath(new FileReference(ResponseLine.Substring(SourceDependencyPrefix.Length).Trim().Trim('\"')));
+								FileReference OldSourceDependencyFile = Action.RootPaths.GetLocalPath(new FileReference(ResponseLine[SourceDependencyPrefix.Length..].Trim().Trim('\"')));
 								FileItem OldSourceDependencyFileItem = Action.ProducedItems.First(x => x.Location == OldSourceDependencyFile);
 								NewAction.ProducedItems.Remove(OldSourceDependencyFileItem);
 
-								FileReference NewSourceDependencyFile = FileReference.FromString(OldSourceDependencyFile.FullName.Substring(0, OldSourceDependencyFile.FullName.Length - extension.Length) + extraExt + extension);
+								FileReference NewSourceDependencyFile = FileReference.FromString(OldSourceDependencyFile.FullName[..^extension.Length] + extraExt + extension);
 								FileItem NewSourceDependencyFileItem = FileItem.GetItemByFileReference(NewSourceDependencyFile);
 								NewAction.ProducedItems.Add(NewSourceDependencyFileItem);
 								NewAction.DependencyListFile = NewSourceDependencyFileItem;
@@ -1107,7 +1107,7 @@ namespace UnrealBuildTool
 						int ResponseExtensionIndex = OriginalPrerequisiteItem.AbsolutePath.IndexOf(ResponseFileExtension, StringComparison.InvariantCultureIgnoreCase);
 						if (ResponseExtensionIndex != -1)
 						{
-							string OriginalResponseFilePathWithoutExtension = OriginalPrerequisiteItem.AbsolutePath.Substring(0, ResponseExtensionIndex);
+							string OriginalResponseFilePathWithoutExtension = OriginalPrerequisiteItem.AbsolutePath[..ResponseExtensionIndex];
 							string OriginalResponseFilePath = OriginalResponseFilePathWithoutExtension + ResponseFileExtension;
 
 							string NewResponseFilePath = ReplaceBaseFileName(OriginalResponseFilePath, OriginalFileNameWithoutExtension, NewFileNameWithoutExtension);
@@ -1118,7 +1118,7 @@ namespace UnrealBuildTool
 						int ScriptExtensionIndex = OriginalPrerequisiteItem.AbsolutePath.IndexOf(ScriptExtension, StringComparison.InvariantCultureIgnoreCase);
 						if (ScriptExtensionIndex != -1)
 						{
-							string OriginalScriptFilePathWithoutExtension = OriginalPrerequisiteItem.AbsolutePath.Substring(0, ScriptExtensionIndex);
+							string OriginalScriptFilePathWithoutExtension = OriginalPrerequisiteItem.AbsolutePath[..ScriptExtensionIndex];
 							string OriginalScriptFilePath = OriginalScriptFilePathWithoutExtension + ScriptExtension;
 
 							string NewScriptFilePath = ReplaceBaseFileName(OriginalScriptFilePath, OriginalFileNameWithoutExtension, NewFileNameWithoutExtension);
@@ -1258,7 +1258,7 @@ namespace UnrealBuildTool
 					}
 
 					// Read the metadata file
-					FileReference TargetInfoFile = new FileReference(Arguments.Substring(FileNameIdx, FileNameEndIdx - FileNameIdx));
+					FileReference TargetInfoFile = new FileReference(Arguments[FileNameIdx..FileNameEndIdx]);
 					if (!FileReference.Exists(TargetInfoFile))
 					{
 						throw new Exception(String.Format("Unable to find metadata file to patch action graph ({0})", TargetInfoFile));
@@ -1294,7 +1294,7 @@ namespace UnrealBuildTool
 						NewAction.PrerequisiteItems.RemoveWhere(x => x.Location == TargetInfoFile);
 						NewAction.PrerequisiteItems.Add(FileItem.GetItemByFileReference(HotReloadTargetInfoFile));
 
-						NewAction.CommandArguments = Arguments.Substring(0, FileNameIdx) + HotReloadTargetInfoFile + Arguments.Substring(FileNameEndIdx);
+						NewAction.CommandArguments = Arguments[..FileNameIdx] + HotReloadTargetInfoFile + Arguments[FileNameEndIdx..];
 
 						Action.Inner = NewAction;
 					}
diff --git a/Engine/Source/Programs/UnrealBuildTool/System/InstalledPlatformInfo.cs b/Engine/Source/Programs/UnrealBuildTool/System/InstalledPlatformInfo.cs
--- a/Engine/Source/Programs/UnrealBuildTool/System/InstalledPlatformInfo.cs
+++ b/Engine/Source/Programs/UnrealBuildTool/System/InstalledPlatformInfo.cs
@@ -240,7 +240,7 @@ namespace UnrealBuildTool
 			}
 			// Get the remainder of the string after the match
 			MatchIndex += Match.Length;
-			TrimmedLine = TrimmedLine.Substring(MatchIndex);
+			TrimmedLine = TrimmedLine[MatchIndex..];
 			if (String.IsNullOrEmpty(TrimmedLine))
 			{
 				return false;
@@ -254,7 +254,7 @@ namespace UnrealBuildTool
 				int QuoteEnd = Result.LastIndexOf('\"');
 				if (QuoteEnd > 0)
 				{
-					Result = Result.Substring(1, QuoteEnd - 1);
+					Result = Result[1..QuoteEnd];
 				}
 			}
 			return true;
diff --git a/Engine/Source/Programs/UnrealBuildTool/System/PluginDescriptor.cs b/Engine/Source/Programs/UnrealBuildTool/System/PluginDescriptor.cs
--- a/Engine/Source/Programs/UnrealBuildTool/System/PluginDescriptor.cs
+++ b/Engine/Source/Programs/UnrealBuildTool/System/PluginDescriptor.cs
@@ -291,7 +291,7 @@ namespace UnrealBuildTool
 			// a category name with escaped quotes. Remove them here to make sure we can list them in the right category.
 			if (Category != null && Category.Length >= 2 && Category.StartsWith("\"") && Category.EndsWith("\""))
 			{
-				Category = Category.Substring(1, Category.Length - 2);
+				Category = Category[1..^1];
 			}
 
 			RawObject.TryGetStringField("CreatedBy", out CreatedBy);
diff --git a/Engine/Source/Programs/UnrealBuildTool/System/RulesAssembly.cs b/Engine/Source/Programs/UnrealBuildTool/System/RulesAssembly.cs
--- a/Engine/Source/Programs/UnrealBuildTool/System/RulesAssembly.cs
+++ b/Engine/Source/Programs/UnrealBuildTool/System/RulesAssembly.cs
@@ -345,7 +345,7 @@ namespace UnrealBuildTool
 				string NameWithoutTarget = Name;
 				if (NameWithoutTarget.EndsWith("Target"))
 				{
-					NameWithoutTarget = NameWithoutTarget.Substring(0, NameWithoutTarget.Length - 6);
+					NameWithoutTarget = NameWithoutTarget[..^6];
 				}
 				if (TargetNameToTargetFile.TryGetValue(NameWithoutTarget, out File))
 				{
diff --git a/Engine/Source/Programs/UnrealBuildTool/System/RulesDocumentation.cs b/Engine/Source/Programs/UnrealBuildTool/System/RulesDocumentation.cs
--- a/Engine/Source/Programs/UnrealBuildTool/System/RulesDocumentation.cs
+++ b/Engine/Source/Programs/UnrealBuildTool/System/RulesDocumentation.cs
@@ -290,7 +290,7 @@ namespace UnrealBuildTool
 			{
 				if (Lines[Idx].StartsWith("*") || Lines[Idx].StartsWith("-"))
 				{
-					Writer.WriteLine("        * {0}", Lines[Idx].Substring(1).TrimStart());
+					Writer.WriteLine("        * {0}", Lines[Idx][1..].TrimStart());
 				}
 				else
 				{
@@ -364,7 +364,7 @@ namespace UnrealBuildTool
 							Writer.WriteLine("        <ul>");
 							for (; Idx < Lines.Count && (Lines[Idx].StartsWith("*") || Lines[Idx].StartsWith("-")); Idx++)
 							{
-								Writer.WriteLine("          <li>{0}</li>", Lines[Idx].Substring(1).TrimStart());
+								Writer.WriteLine("          <li>{0}</li>", Lines[Idx][1..].TrimStart());
 							}
 							Writer.WriteLine("        </ul>");
 						}
@@ -382,7 +382,7 @@ namespace UnrealBuildTool
 		{
 			if (FieldType.IsGenericType)
 			{
-				return String.Format("{0}&lt;{1}&gt;", FieldType.Name.Substring(0, FieldType.Name.IndexOf('`')), String.Join(", ", FieldType.GenericTypeArguments.Select(x => GetPrettyTypeName(x))));
+				return String.Format("{0}&lt;{1}&gt;", FieldType.Name[..FieldType.Name.IndexOf('`')], String.Join(", ", FieldType.GenericTypeArguments.Select(x => GetPrettyTypeName(x))));
 			}
 			else
 			{
diff --git a/Engine/Source/Programs/UnrealBuildTool/System/SourceFileMetadataCache.cs b/Engine/Source/Programs/UnrealBuildTool/System/SourceFileMetadataCache.cs
--- a/Engine/Source/Programs/UnrealBuildTool/System/SourceFileMetadataCache.cs
+++ b/Engine/Source/Programs/UnrealBuildTool/System/SourceFileMetadataCache.cs
@@ -327,7 +327,7 @@ namespace UnrealBuildTool
 				ReadOnlySpan<char> LineSpan = Line.AsSpan().TrimStart();
 				if (LineSpan.StartsWith("#include") && InsideDeprecationScope == 0)
 				{
-					ReadOnlySpan<char> IncludeSpan = LineSpan.Slice("#include".Length).TrimStart();
+					ReadOnlySpan<char> IncludeSpan = LineSpan["#include".Length..].TrimStart();
 					if (IncludeSpan.IsEmpty)
 					{
 						continue;
@@ -350,7 +350,7 @@ namespace UnrealBuildTool
 
 					if (TrimQuotation)
 					{
-						IncludeSpan = IncludeSpan.Slice(1);
+						IncludeSpan = IncludeSpan[1..];
 					}
 
 					if (IncludeSpan.Contains("HEADER_UNIT_IGNORE", StringComparison.OrdinalIgnoreCase))
@@ -371,7 +371,7 @@ namespace UnrealBuildTool
 						EndIndex++;
 					}
 
-					IncludeSpan = IncludeSpan.Slice(0, EndIndex);
+					IncludeSpan = IncludeSpan[..EndIndex];
 					Includes.Add(IncludeSpan.ToString());
 				}
 				else if (InsideDeprecationScope != 0)
diff --git a/Engine/Source/Programs/UnrealBuildTool/System/SourceFileWorkingSet.cs b/Engine/Source/Programs/UnrealBuildTool/System/SourceFileWorkingSet.cs
--- a/Engine/Source/Programs/UnrealBuildTool/System/SourceFileWorkingSet.cs
+++ b/Engine/Source/Programs/UnrealBuildTool/System/SourceFileWorkingSet.cs
@@ -245,13 +245,13 @@ namespace UnrealBuildTool
 				int ArrowIdx = Args.Data.IndexOf(" -> ", MinIdx, MaxIdx - MinIdx);
 				if (ArrowIdx == -1)
 				{
-					AddPath(Args.Data.Substring(MinIdx, MaxIdx - MinIdx));
+					AddPath(Args.Data[MinIdx..MaxIdx]);
 				}
 				else
 				{
-					AddPath(Args.Data.Substring(MinIdx, ArrowIdx - MinIdx));
+					AddPath(Args.Data[MinIdx..ArrowIdx]);
 					int ArrowEndIdx = ArrowIdx + 4;
-					AddPath(Args.Data.Substring(ArrowEndIdx, MaxIdx - ArrowEndIdx));
+					AddPath(Args.Data[ArrowEndIdx..MaxIdx]);
 				}
 			}
 		}
diff --git a/Engine/Source/Programs/UnrealBuildTool/System/Subnet.cs b/Engine/Source/Programs/UnrealBuildTool/System/Subnet.cs
--- a/Engine/Source/Programs/UnrealBuildTool/System/Subnet.cs
+++ b/Engine/Source/Programs/UnrealBuildTool/System/Subnet.cs
@@ -53,8 +53,8 @@ namespace UnrealBuildTool
 		public static Subnet Parse(string Text)
 		{
 			int SlashIdx = Text.IndexOf('/');
-			IPAddress Address = IPAddress.Parse(Text.Substring(0, SlashIdx));
-			return new Subnet(Address, Int32.Parse(Text.Substring(SlashIdx + 1)));
+			IPAddress Address = IPAddress.Parse(Text[..SlashIdx]);
+			return new Subnet(Address, Int32.Parse(Text[(SlashIdx + 1)..]));
 		}
 
 		/// <summary>
diff --git a/Engine/Source/Programs/UnrealBuildTool/System/TargetMakefile.cs b/Engine/Source/Programs/UnrealBuildTool/System/TargetMakefile.cs
--- a/Engine/Source/Programs/UnrealBuildTool/System/TargetMakefile.cs
+++ b/Engine/Source/Programs/UnrealBuildTool/System/TargetMakefile.cs
@@ -750,7 +750,7 @@ namespace UnrealBuildTool
 				// Check if any manifests are out of date
 				foreach (FileReference Manifest in Makefile.AdditionalArguments!
 					.Where(x => x.StartsWith("-Manifest=", StringComparison.OrdinalIgnoreCase))
-					.Select(x => FileReference.FromString(x.Substring("-Manifest=".Length).Trim('"').Trim('\'')))
+					.Select(x => FileReference.FromString(x["-Manifest=".Length..].Trim('"').Trim('\'')))
 					.Where(x => x != null))
 				{
 					if (!FileReference.Exists(Manifest))
diff --git a/Engine/Source/Programs/UnrealBuildTool/System/TargetReceipt.cs b/Engine/Source/Programs/UnrealBuildTool/System/TargetReceipt.cs
--- a/Engine/Source/Programs/UnrealBuildTool/System/TargetReceipt.cs
+++ b/Engine/Source/Programs/UnrealBuildTool/System/TargetReceipt.cs
@@ -442,13 +442,13 @@ namespace UnrealBuildTool
 			const string EnginePrefix = "$(EngineDir)";
 			if (Path.StartsWith(EnginePrefix, StringComparison.InvariantCultureIgnoreCase))
 			{
-				return new FileReference(EngineDir.FullName + Path.Substring(EnginePrefix.Length));
+				return new FileReference(EngineDir.FullName + Path[EnginePrefix.Length..]);
 			}
 
 			const string ProjectPrefix = "$(ProjectDir)";
 			if (ProjectDir != null && Path.StartsWith(ProjectPrefix, StringComparison.InvariantCultureIgnoreCase))
 			{
-				return new FileReference(ProjectDir.FullName + Path.Substring(ProjectPrefix.Length));
+				return new FileReference(ProjectDir.FullName + Path[ProjectPrefix.Length..]);
 			}
 
 			return new FileReference(Path);
diff --git a/Engine/Source/Programs/UnrealBuildTool/System/UnrealPluginLanguage.cs b/Engine/Source/Programs/UnrealBuildTool/System/UnrealPluginLanguage.cs
--- a/Engine/Source/Programs/UnrealBuildTool/System/UnrealPluginLanguage.cs
+++ b/Engine/Source/Programs/UnrealBuildTool/System/UnrealPluginLanguage.cs
@@ -761,7 +761,7 @@ namespace UnrealBuildTool
 				}
 
 				// Extract the variable name from the string
-				string Name = Result.Substring(Idx + 3, EndIdx - (Idx + 3));
+				string Name = Result[(Idx + 3)..EndIdx];
 
 				// Find the value for it, either from the dictionary or the environment block
 				bool Value;
@@ -775,7 +775,7 @@ namespace UnrealBuildTool
 				}
 
 				// Replace the variable, or skip past it
-				Result = Result.Substring(0, Idx) + Value.ToString().ToLower() + Result.Substring(EndIdx + 1);
+				Result = Result[..Idx] + Value.ToString().ToLower() + Result[(EndIdx + 1)..];
 			}
 			for (int Idx = Result.IndexOf("$I("); Idx != -1; Idx = Result.IndexOf("$I(", Idx))
 			{
@@ -787,7 +787,7 @@ namespace UnrealBuildTool
 				}
 
 				// Extract the variable name from the string
-				string Name = Result.Substring(Idx + 3, EndIdx - (Idx + 3));
+				string Name = Result[(Idx + 3)..EndIdx];
 
 				// Find the value for it, either from the dictionary or the environment block
 				int Value;
@@ -801,7 +801,7 @@ namespace UnrealBuildTool
 				}
 
 				// Replace the variable, or skip past it
-				Result = Result.Substring(0, Idx) + Value.ToString() + Result.Substring(EndIdx + 1);
+				Result = Result[..Idx] + Value.ToString() + Result[(EndIdx + 1)..];
 			}
 			for (int Idx = Result.IndexOf("$S("); Idx != -1; Idx = Result.IndexOf("$S(", Idx))
 			{
@@ -813,7 +813,7 @@ namespace UnrealBuildTool
 				}
 
 				// Extract the variable name from the string
-				string Name = Result.Substring(Idx + 3, EndIdx - (Idx + 3));
+				string Name = Result[(Idx + 3)..EndIdx];
 
 				// Find the value for it, either from the dictionary or the environment block
 				string? Value;
@@ -827,7 +827,7 @@ namespace UnrealBuildTool
 				}
 
 				// Replace the variable, or skip past it
-				Result = Result.Substring(0, Idx) + Value + Result.Substring(EndIdx + 1);
+				Result = Result[..Idx] + Value + Result[(EndIdx + 1)..];
 			}
 			for (int Idx = Result.IndexOf("$E("); Idx != -1; Idx = Result.IndexOf("$E(", Idx))
 			{
@@ -839,7 +839,7 @@ namespace UnrealBuildTool
 				}
 
 				// Extract the variable name from the string
-				string Name = Result.Substring(Idx + 3, EndIdx - (Idx + 3));
+				string Name = Result[(Idx + 3)..EndIdx];
 
 				// Find the value for it, either from the dictionary or the environment block
 				XElement? Value;
@@ -853,7 +853,7 @@ namespace UnrealBuildTool
 				}
 
 				// Replace the variable, or skip past it
-				Result = Result.Substring(0, Idx) + Value + Result.Substring(EndIdx + 1);
+				Result = Result[..Idx] + Value + Result[(EndIdx + 1)..];
 			}
 			return Result;
 		}
@@ -1051,12 +1051,12 @@ namespace UnrealBuildTool
 			if (Index >= 0)
 			{
 				XNamespace? Namespace;
-				string NamespaceAlias = Name.Substring(0, Index);
+				string NamespaceAlias = Name[..Index];
 				if (!XMLNameSpaceAliases.TryGetValue(NamespaceAlias, out Namespace))
 				{
 					Namespace = XMLDefaultNameSpace;
 				}
-				Name = Name.Substring(Index + 1);
+				Name = Name[(Index + 1)..];
 				return Namespace;
 			}
 			else
@@ -1347,9 +1347,9 @@ namespace UnrealBuildTool
 									XElement? Target = CurrentElement;
 									if (Tag.Length > 1)
 									{
-										if (!CurrentContext.ElementVariables.TryGetValue(Tag.Substring(1), out Target))
+										if (!CurrentContext.ElementVariables.TryGetValue(Tag[1..], out Target))
 										{
-											if (!GlobalContext.ElementVariables.TryGetValue(Tag.Substring(1), out Target))
+											if (!GlobalContext.ElementVariables.TryGetValue(Tag[1..], out Target))
 											{
 												Logger.LogWarning("Missing element variable '{Tag}' in '{Node}' (skipping instruction)", Tag, TraceNodeString(CurrentContext, Node));
 												continue;
@@ -1384,9 +1384,9 @@ namespace UnrealBuildTool
 									XElement? Target = CurrentElement;
 									if (Tag.Length > 1)
 									{
-										if (!CurrentContext.ElementVariables.TryGetValue(Tag.Substring(1), out Target))
+										if (!CurrentContext.ElementVariables.TryGetValue(Tag[1..], out Target))
 										{
-											if (!GlobalContext.ElementVariables.TryGetValue(Tag.Substring(1), out Target))
+											if (!GlobalContext.ElementVariables.TryGetValue(Tag[1..], out Target))
 											{
 												Logger.LogInformation("\nMissing element variable '{Tag}' in '{Node}' (skipping instruction)", Tag, TraceNodeString(CurrentContext, Node));
 												continue;
@@ -1633,9 +1633,9 @@ namespace UnrealBuildTool
 									XElement? Target = CurrentElement;
 									if (Tag.Length > 1)
 									{
-										if (!CurrentContext.ElementVariables.TryGetValue(Tag.Substring(1), out Target))
+										if (!CurrentContext.ElementVariables.TryGetValue(Tag[1..], out Target))
 										{
-											if (!GlobalContext.ElementVariables.TryGetValue(Tag.Substring(1), out Target))
+											if (!GlobalContext.ElementVariables.TryGetValue(Tag[1..], out Target))
 											{
 												Logger.LogWarning("Missing element variable '{Name}' in '{Node}' (skipping instruction)", Tag, TraceNodeString(CurrentContext, Node));
 												continue;
@@ -1681,9 +1681,9 @@ namespace UnrealBuildTool
 									XElement? Target = CurrentElement;
 									if (Tag.Length > 1)
 									{
-										if (!CurrentContext.ElementVariables.TryGetValue(Tag.Substring(1), out Target))
+										if (!CurrentContext.ElementVariables.TryGetValue(Tag[1..], out Target))
 										{
-											if (!GlobalContext.ElementVariables.TryGetValue(Tag.Substring(1), out Target))
+											if (!GlobalContext.ElementVariables.TryGetValue(Tag[1..], out Target))
 											{
 												Logger.LogWarning("Missing element variable '{Tag}' in '{Node}' (skipping instruction)", Tag, TraceNodeString(CurrentContext, Node));
 												continue;
@@ -1740,7 +1740,7 @@ namespace UnrealBuildTool
 								}
 								if (Index < Value.Length)
 								{
-									Value = Value.Substring(0, Index);
+									Value = Value[..Index];
 								}
 
 								// trim leading newlines
@@ -1752,7 +1752,7 @@ namespace UnrealBuildTool
 
 								if (Index < Value.Length)
 								{
-									GlobalContext.StringVariables["Output"] += Value.Substring(Index);
+									GlobalContext.StringVariables["Output"] += Value[Index..];
 								}
 							}
 						}
@@ -2309,9 +2309,9 @@ namespace UnrealBuildTool
 								{
 									if (Tag.Length > 1)
 									{
-										if (!CurrentContext.ElementVariables.TryGetValue(Tag.Substring(1), out Element))
+										if (!CurrentContext.ElementVariables.TryGetValue(Tag[1..], out Element))
 										{
-											if (!GlobalContext.ElementVariables.TryGetValue(Tag.Substring(1), out Element))
+											if (!GlobalContext.ElementVariables.TryGetValue(Tag[1..], out Element))
 											{
 												Logger.LogWarning("Missing element variable '{Tag}' in '{Node}' (skipping instruction)", Tag, TraceNodeString(CurrentContext, Node));
 												continue;
@@ -2336,9 +2336,9 @@ namespace UnrealBuildTool
 								{
 									if (Tag.Length > 1)
 									{
-										if (!CurrentContext.ElementVariables.TryGetValue(Tag.Substring(1), out Element))
+										if (!CurrentContext.ElementVariables.TryGetValue(Tag[1..], out Element))
 										{
-											if (!GlobalContext.ElementVariables.TryGetValue(Tag.Substring(1), out Element))
+											if (!GlobalContext.ElementVariables.TryGetValue(Tag[1..], out Element))
 											{
 												Logger.LogWarning("Missing element variable '{Tag}' in '{Node}' (skipping instruction)", Tag, TraceNodeString(CurrentContext, Node));
 												continue;
@@ -2366,9 +2366,9 @@ namespace UnrealBuildTool
 								{
 									if (Tag.Length > 1)
 									{
-										if (!CurrentContext.ElementVariables.TryGetValue(Tag.Substring(1), out Element))
+										if (!CurrentContext.ElementVariables.TryGetValue(Tag[1..], out Element))
 										{
-											if (!GlobalContext.ElementVariables.TryGetValue(Tag.Substring(1), out Element))
+											if (!GlobalContext.ElementVariables.TryGetValue(Tag[1..], out Element))
 											{
 												Logger.LogWarning("Missing element variable '{Tag}' in '{Node}' (skipping instruction)", Tag, TraceNodeString(CurrentContext, Node));
 												continue;
@@ -3166,7 +3166,7 @@ namespace UnrealBuildTool
 			if (start != 1 && end != -1 && start < end)
 			{
 				++start;
-				Path = Source.Substring(start, end - start);
+				Path = Source[start..end];
 				return true;
 			}
 			else
@@ -3214,12 +3214,12 @@ namespace UnrealBuildTool
 			}
 			else if (Line.StartsWith("-"))
 			{
-				Line = Line.Substring(1).TrimStart();
+				Line = Line[1..].TrimStart();
 				return ParseAction.Remove;
 			}
 			else if (Line.StartsWith("+"))
 			{
-				Line = Line.Substring(1).TrimStart();
+				Line = Line[1..].TrimStart();
 				return ParseAction.Add;
 			}
 			else
@@ -3274,7 +3274,7 @@ namespace UnrealBuildTool
 					bMultiLine = TrimmedLine.EndsWith("\\");
 					if (bMultiLine)
 					{
-						TrimmedLine = TrimmedLine.Substring(0, TrimmedLine.Length - 1).TrimEnd();
+						TrimmedLine = TrimmedLine[..^1].TrimEnd();
 					}
 					if (!bWasMultiLine)
 					{
@@ -3362,12 +3362,12 @@ namespace UnrealBuildTool
 			{
 				throw new IniParsingException("Failed to find value when parsing {0}, line {1}: {2}", Filename, LineIndex, TrimmedLine);
 			}
-			Key = TrimmedLine.Substring(0, AssignIndex).Trim();
+			Key = TrimmedLine[..AssignIndex].Trim();
 			if (String.IsNullOrEmpty(Key))
 			{
 				throw new IniParsingException("Empty key when parsing {0}, line {1}: {2}", Filename, LineIndex, TrimmedLine);
 			}
-			Value = TrimmedLine.Substring(AssignIndex + 1).Trim();
+			Value = TrimmedLine[(AssignIndex + 1)..].Trim();
 			if (Value.StartsWith("\""))
 			{
 				// Remove quotes
@@ -3376,7 +3376,7 @@ namespace UnrealBuildTool
 				{
 					throw new IniParsingException("Mismatched quotes when parsing {0}, line {1}: {2}", Filename, LineIndex, TrimmedLine);
 				}
-				Value = Value.Substring(1, Value.Length - 2);
+				Value = Value[1..^1];
 			}
 		}
 
@@ -3439,7 +3439,7 @@ namespace UnrealBuildTool
 			}
 
 			// comment could follow the ] but will just be trimmed out
-			string SectionName = TrimmedLine.Substring(1, SectionEndIndex - 1);
+			string SectionName = TrimmedLine[1..SectionEndIndex];
 			if (String.IsNullOrEmpty(SectionName))
 			{
 				throw new IniParsingException("Empty section name when parsing {0}, line {1}: {2}", Filename, LineIndex, TrimmedLine);
diff --git a/Engine/Source/Programs/UnrealBuildTool/System/Utils.cs b/Engine/Source/Programs/UnrealBuildTool/System/Utils.cs
--- a/Engine/Source/Programs/UnrealBuildTool/System/Utils.cs
+++ b/Engine/Source/Programs/UnrealBuildTool/System/Utils.cs
@@ -63,12 +63,12 @@ namespace UnrealBuildTool
 				string EnvironmentVariable = M.ToString();
 				if (EnvironmentVariable.StartsWith("$") && EnvironmentVariable.EndsWith(")"))
 				{
-					EnvironmentVariable = EnvironmentVariable.Substring(1, EnvironmentVariable.Length - 2).Replace("(", "");
+					EnvironmentVariable = EnvironmentVariable[1..^1].Replace("(", "");
 				}
 
 				if (EnvironmentVariable.StartsWith("%") && EnvironmentVariable.EndsWith("%"))
 				{
-					EnvironmentVariable = EnvironmentVariable.Substring(1, EnvironmentVariable.Length - 2);
+					EnvironmentVariable = EnvironmentVariable[1..^1];
 				}
 
 				// Resolve environment variable.				
@@ -102,7 +102,7 @@ namespace UnrealBuildTool
 				}
 
 				// Extract the variable name from the string
-				string Name = Result.Substring(Idx + 2, EndIdx - (Idx + 2));
+				string Name = Result[(Idx + 2)..EndIdx];
 
 				// Find the value for it, either from the dictionary or the environment block
 				string? Value = null;
@@ -120,7 +120,7 @@ namespace UnrealBuildTool
 				}
 
 				// Replace the variable, or skip past it
-				Result = Result.Substring(0, Idx) + Value + Result.Substring(EndIdx + 1);
+				Result = Result[..Idx] + Value + Result[(EndIdx + 1)..];
 			}
 			return Result;
 		}
@@ -628,12 +628,12 @@ namespace UnrealBuildTool
 					if (C != UseDirectorySeparatorChar)
 					{
 						C = UseDirectorySeparatorChar;
-						CleanPath ??= new StringBuilder(FilePath.Substring(0, Index), FilePath.Length);
+						CleanPath ??= new StringBuilder(FilePath[..Index], FilePath.Length);
 					}
 
 					if (bCanCheckDoubleSeparators && C == PrevC)
 					{
-						CleanPath ??= new StringBuilder(FilePath.Substring(0, Index), FilePath.Length);
+						CleanPath ??= new StringBuilder(FilePath[..Index], FilePath.Length);
 						continue;
 					}
 				}
@@ -715,7 +715,7 @@ namespace UnrealBuildTool
 					return InPath;
 				}
 
-				LocalString = LocalString.Substring(0, PreviousSeparatorIndex) + LocalString.Substring(Index + ParentDirLength);
+				LocalString = LocalString[..PreviousSeparatorIndex] + LocalString[(Index + ParentDirLength)..];
 			}
 
 			LocalString = LocalString.Replace("./", "");
@@ -781,7 +781,7 @@ namespace UnrealBuildTool
 			// If we added a directory separator character earlier on, remove it now
 			if (!SourcePathEndsWithDirectorySeparator && AlwaysTreatSourceAsDirectory && RelativePath.EndsWith(Path.DirectorySeparatorChar.ToString()))
 			{
-				RelativePath = RelativePath.Substring(0, RelativePath.Length - 1);
+				RelativePath = RelativePath[..^1];
 			}
 
 			return RelativePath;
@@ -948,7 +948,7 @@ namespace UnrealBuildTool
 			}
 			else
 			{
-				return Filename.Substring(0, DotIndex);
+				return Filename[..DotIndex];
 			}
 		}
 
diff --git a/Engine/Source/Programs/UnrealBuildTool/System/XmlConfig.cs b/Engine/Source/Programs/UnrealBuildTool/System/XmlConfig.cs
--- a/Engine/Source/Programs/UnrealBuildTool/System/XmlConfig.cs
+++ b/Engine/Source/Programs/UnrealBuildTool/System/XmlConfig.cs
@@ -815,7 +815,7 @@ namespace UnrealBuildTool
 				{
 					continue;
 				}
-				string[] split = key.Substring("UnrealBuildTool_".Length).Split("__");
+				string[] split = key["UnrealBuildTool_".Length..].Split("__");
 				string category = split.First();
 				string field = split.Last();
 				if (!contents.TryGetValue(category, out Dictionary<string, string>? fields))
@@ -1049,7 +1049,7 @@ namespace UnrealBuildTool
 						{
 							if (lines[idx].StartsWith("*", StringComparison.OrdinalIgnoreCase) || lines[idx].StartsWith("-", StringComparison.OrdinalIgnoreCase))
 							{
-								writer.WriteLine("        * {0}", lines[idx].Substring(1).TrimStart());
+								writer.WriteLine("        * {0}", lines[idx][1..].TrimStart());
 							}
 							else
 							{
@@ -1109,7 +1109,7 @@ namespace UnrealBuildTool
 									writer.WriteLine("        <ul>");
 									for (; idx < lines.Count && (lines[idx].StartsWith("*", StringComparison.OrdinalIgnoreCase) || lines[idx].StartsWith("-", StringComparison.OrdinalIgnoreCase)); idx++)
 									{
-										writer.WriteLine("          <li>{0}</li>", lines[idx].Substring(1).TrimStart());
+										writer.WriteLine("          <li>{0}</li>", lines[idx][1..].TrimStart());
 									}
 									writer.WriteLine("        </ul>");
 								}
diff --git a/Engine/Source/Programs/UnrealBuildTool/ToolChain/AppleToolChain.cs b/Engine/Source/Programs/UnrealBuildTool/ToolChain/AppleToolChain.cs
--- a/Engine/Source/Programs/UnrealBuildTool/ToolChain/AppleToolChain.cs
+++ b/Engine/Source/Programs/UnrealBuildTool/ToolChain/AppleToolChain.cs
@@ -187,7 +187,7 @@ namespace UnrealBuildTool
 				{
 					throw new BuildException("Unable to get xcode version");
 				}
-				if (Int32.Parse(InstalledSdkVersion.Substring(0, 2)) < 13)
+				if (Int32.Parse(InstalledSdkVersion[..2]) < 13)
 				{
 					throw new BuildException("Building for macOS, iOS and tvOS requires Xcode 13.4.1 or newer, Xcode " + InstalledSdkVersion + " detected");
 				}
diff --git a/Engine/Source/Programs/UnrealBuildTool/ToolChain/ClangSpecificFileAction.cs b/Engine/Source/Programs/UnrealBuildTool/ToolChain/ClangSpecificFileAction.cs
--- a/Engine/Source/Programs/UnrealBuildTool/ToolChain/ClangSpecificFileAction.cs
+++ b/Engine/Source/Programs/UnrealBuildTool/ToolChain/ClangSpecificFileAction.cs
@@ -68,7 +68,7 @@ namespace UnrealBuildTool
 			string UniqueDummyName = Filename;
 
 			int FileNameIndex = CommandArguments.IndexOf(DummyName);
-			string DummyPath = CommandArguments.Substring(2, FileNameIndex + DummyName.Length - 2);
+			string DummyPath = CommandArguments[2..(FileNameIndex + DummyName.Length)];
 
 			List<string> NewRspLines = new();
 
@@ -122,7 +122,7 @@ namespace UnrealBuildTool
 				Action.ProducedItems.Add(PreprocessedItem);
 			}
 
-			FileReference ResponseFile = RootPaths.GetLocalPath(FileReference.FromString(Action.CommandArguments.Substring(1).Trim('"')));
+			FileReference ResponseFile = RootPaths.GetLocalPath(FileReference.FromString(Action.CommandArguments[1..].Trim('"')));
 			File.WriteAllLines(ResponseFile.FullName, NewRspLines);
 
 			return Action;
diff --git a/Engine/Source/Programs/UnrealBuildTool/ToolChain/RemoteMac.cs b/Engine/Source/Programs/UnrealBuildTool/ToolChain/RemoteMac.cs
--- a/Engine/Source/Programs/UnrealBuildTool/ToolChain/RemoteMac.cs
+++ b/Engine/Source/Programs/UnrealBuildTool/ToolChain/RemoteMac.cs
@@ -187,12 +187,12 @@ namespace UnrealBuildTool
 			int PortIdx = ServerName.LastIndexOf(':');
 			if (PortIdx != -1)
 			{
-				string Port = ServerName.Substring(PortIdx + 1);
+				string Port = ServerName[(PortIdx + 1)..];
 				if (!Int32.TryParse(Port, out ServerPort))
 				{
 					throw new BuildException("Unable to parse port number from '{0}'", ServerName);
 				}
-				ServerName = ServerName.Substring(0, PortIdx);
+				ServerName = ServerName[..PortIdx];
 			}
 
 			// If a user name is not set, use the current user
@@ -662,8 +662,8 @@ namespace UnrealBuildTool
 					continue;
 				}
 
-				string Key = LocalArgument.Substring(0, EqualsIdx);
-				string Value = LocalArgument.Substring(EqualsIdx + 1);
+				string Key = LocalArgument[..EqualsIdx];
+				string Value = LocalArgument[(EqualsIdx + 1)..];
 
 				if (Key.Equals("-Log", StringComparison.InvariantCultureIgnoreCase))
 				{
@@ -797,7 +797,7 @@ namespace UnrealBuildTool
 			{
 				if (RemotePath.StartsWith(Mapping.RemoteDirectory, StringComparison.InvariantCultureIgnoreCase) && RemotePath.Length > Mapping.RemoteDirectory.Length && RemotePath[Mapping.RemoteDirectory.Length] == '/')
 				{
-					return FileReference.Combine(Mapping.LocalDirectory, RemotePath.Substring(Mapping.RemoteDirectory.Length + 1));
+					return FileReference.Combine(Mapping.LocalDirectory, RemotePath[(Mapping.RemoteDirectory.Length + 1)..]);
 				}
 			}
 			throw new BuildException("Unable to map remote path '{0}' to local path", RemotePath);
@@ -834,7 +834,7 @@ namespace UnrealBuildTool
 			{
 				throw new BuildException("Invalid local path for converting to cygwin format ({0}).", InPath);
 			}
-			return String.Format("/cygdrive/{0}{1}", InPath.FullName.Substring(0, 1), InPath.FullName.Substring(2).Replace('\\', '/'));
+			return String.Format("/cygdrive/{0}{1}", InPath.FullName[..1], InPath.FullName[2..].Replace('\\', '/'));
 		}
 
 		/// <summary>
@@ -1363,7 +1363,7 @@ namespace UnrealBuildTool
 				{
 					if (Path.StartsWith(Mapping.RemoteDirectory))
 					{
-						Path = Mapping.LocalDirectory + Path.Substring(Mapping.RemoteDirectory.Length).Replace('/', '\\');
+						Path = Mapping.LocalDirectory + Path[Mapping.RemoteDirectory.Length..].Replace('/', '\\');
 						break;
 					}
 				}
diff --git a/Engine/Source/Programs/UnrealBuildTool/UnrealBuildTool.cs b/Engine/Source/Programs/UnrealBuildTool/UnrealBuildTool.cs
--- a/Engine/Source/Programs/UnrealBuildTool/UnrealBuildTool.cs
+++ b/Engine/Source/Programs/UnrealBuildTool/UnrealBuildTool.cs
@@ -588,7 +588,7 @@ namespace UnrealBuildTool
 							}
 						}
 
-						OverrideTempDirectory = DirectoryReference.Combine(OverrideTempDirectory, ContentHash.MD5(Encoding.UTF8.GetBytes(Unreal.UnrealBuildToolDllPath.FullName)).ToString().Substring(0, 8));
+						OverrideTempDirectory = DirectoryReference.Combine(OverrideTempDirectory, ContentHash.MD5(Encoding.UTF8.GetBytes(Unreal.UnrealBuildToolDllPath.FullName)).ToString()[..8]);
 						DirectoryReference.CreateDirectory(OverrideTempDirectory);
 
 						Logger.LogDebug("Setting temp directory to '{Path}'", OverrideTempDirectory);
