dotnet analysis issues

#rnx

diff --git a/Engine/Source/Programs/Shared/EpicGames.UHT/Tables/UhtCodeGeneratorInjectorTable.cs b/Engine/Source/Programs/Shared/EpicGames.UHT/Tables/UhtCodeGeneratorInjectorTable.cs
--- a/Engine/Source/Programs/Shared/EpicGames.UHT/Tables/UhtCodeGeneratorInjectorTable.cs
+++ b/Engine/Source/Programs/Shared/EpicGames.UHT/Tables/UhtCodeGeneratorInjectorTable.cs
@@ -130,7 +130,7 @@ namespace EpicGames.UHT.Tables
 					MethodInfo methodRhs = rhs.Method;
 					
 					// Sort by module
-					int moduleComparison = String.Compare(methodLhs.Module.Name, methodLhs.Module.Name, StringComparison.Ordinal);
+					int moduleComparison = String.Compare(methodLhs.Module.Name, methodRhs.Module.Name, StringComparison.Ordinal);
 					if (moduleComparison != 0)
 					{
 						return moduleComparison;
diff --git a/Engine/Source/Programs/UnrealBuildTool/Executors/UnrealBuildAccelerator/UBAAgentCoordinatorHorde.cs b/Engine/Source/Programs/UnrealBuildTool/Executors/UnrealBuildAccelerator/UBAAgentCoordinatorHorde.cs
--- a/Engine/Source/Programs/UnrealBuildTool/Executors/UnrealBuildAccelerator/UBAAgentCoordinatorHorde.cs
+++ b/Engine/Source/Programs/UnrealBuildTool/Executors/UnrealBuildAccelerator/UBAAgentCoordinatorHorde.cs
@@ -183,6 +183,10 @@ namespace UnrealBuildTool
 			// REALLY UGLY. Remove this loop. We know InitAsync has been called before this one so right now it is a hack to just loop.
 			while (_client == null)
 			{
+				if (cancellationToken.IsCancellationRequested)
+				{
+					return null;
+				}
 				Thread.Sleep(10);
 			}
 
@@ -929,7 +933,7 @@ namespace UnrealBuildTool
 
 						bool isSatisfied = queueWeight <= queueThreshold || (activeCores + queuedCores) >= _hordeConfig.HordeMaxCores || _cancellationSource!.IsCancellationRequested;
 
-						hordeSession.UpdateHordeStatus(isSatisfied ? " - Requests paused" : $" - Requesting agent{("....".Substring(0, (requestCounter++ % 4)))}");
+						hordeSession.UpdateHordeStatus(isSatisfied ? " - Requests paused" : $" - Requesting agent{("...."[..(requestCounter++ % 4)])}");
 
 						if (isSatisfied)
 						{
