UnrealBuildTool: Disable analysis warning C6393, it fires on generated arrays

#rnx

diff --git a/Engine/Source/Programs/UnrealBuildTool/Platform/Windows/VCToolChain.cs b/Engine/Source/Programs/UnrealBuildTool/Platform/Windows/VCToolChain.cs
--- a/Engine/Source/Programs/UnrealBuildTool/Platform/Windows/VCToolChain.cs
+++ b/Engine/Source/Programs/UnrealBuildTool/Platform/Windows/VCToolChain.cs
@@ -598,6 +598,9 @@ namespace UnrealBuildTool
 				{
 					Arguments.Add("/wd6326"); // Potential comparison of a constant with another constant
 				}
+
+				// A lookup table of size 365 isn't sufficient to handle leap years https://learn.microsoft.com/en-us/cpp/code-quality/c6393
+				Arguments.Add("/wd6393");
 			}
 
 			// Prevents the compiler from displaying its logo for each invocation.
