UnrealBuildTool: Remove "Searching for third-party source files..." output spam

#rnx

diff --git a/Engine/Source/Programs/UnrealBuildTool/ProjectFiles/ProjectFileGenerator.cs b/Engine/Source/Programs/UnrealBuildTool/ProjectFiles/ProjectFileGenerator.cs
--- a/Engine/Source/Programs/UnrealBuildTool/ProjectFiles/ProjectFileGenerator.cs
+++ b/Engine/Source/Programs/UnrealBuildTool/ProjectFiles/ProjectFileGenerator.cs
@@ -2653,11 +2653,6 @@ namespace UnrealBuildTool
 					// for every third-party module, unless we were configured to do so.
 					bool SearchSubdirectories = !IsThirdPartyModule || bGatherThirdPartySource;
 
-					if (bGatherThirdPartySource)
-					{
-						Logger.LogInformation("Searching for third-party source files...");
-					}
-
 					// Find all of the source files (and other files) and add them to the project
 					List<FileReference> FoundFiles = SourceFileSearch.FindModuleSourceFiles(CurModuleFile, SearchSubdirectories: SearchSubdirectories);
 					// remove any target files, they are technically not part of the module and are explicitly added when the project is created
