Work around Gradle bug whereby when only assets changes, it attempts to patch the existing APK from previous packaging by adding asset deltas resulting in growth of the final APK payload. Workaround consists is detection of changes to assets and wiping the existing APK to force Gradle to repackage it correctly.

diff --git a/Engine/Source/Programs/UnrealBuildTool/Platform/Android/UEDeployAndroid.cs b/Engine/Source/Programs/UnrealBuildTool/Platform/Android/UEDeployAndroid.cs
--- a/Engine/Source/Programs/UnrealBuildTool/Platform/Android/UEDeployAndroid.cs
+++ b/Engine/Source/Programs/UnrealBuildTool/Platform/Android/UEDeployAndroid.cs
@@ -5022,6 +5022,7 @@ dependencies.implementation('androidx.tracing:tracing:1.2.0')
 				string UnrealBuildPath = Path.Combine(IntermediateAndroidPath, Arch.ToString());
 				string UnrealBuildGradlePath = Path.Combine(UnrealBuildPath, "gradle");
 
+				bool bInApkAssetChange = false;
 				// If we are packaging for Amazon then we need to copy the  file to the correct location
 				Logger.LogInformation("bPackageDataInsideApk = {PackageDataInsideApk}", bPackageDataInsideApk);
 				if (bPackageDataInsideApk)
@@ -5036,16 +5037,25 @@ dependencies.implementation('androidx.tracing:tracing:1.2.0')
 						Logger.LogInformation("Obb file exists...");
 						string DestFileName = Path.Combine(ObbFileDestination, "main.obb.png"); // Need a rename to turn off compression
 						string SrcFileName = ObbFileLocation;
-						CopyIfDifferent(SrcFileName, DestFileName, true, false);
+						bInApkAssetChange = CopyIfDifferent(SrcFileName, DestFileName, true, false);
 					}
 				}
 				else // try to remove the file it we aren't packaging inside the APK
 				{
 					string ObbFileDestination = UnrealBuildPath + "/assets";
 					string DestFileName = Path.Combine(ObbFileDestination, "main.obb.png");
-					SafeDeleteFile(DestFileName);
+					bInApkAssetChange = SafeDeleteFile(DestFileName);
 				}
 
+				if (bInApkAssetChange)
+				{
+					// Workaround for Gradle bug whereby a change in the contents of assets triggers a patch build whereby main.obb.png is duplicated in the APK causing the APK to grow in size. 
+					// Workaround deletes existing Gradle output APK to rebuild when main.obb.png is modified or deleted.
+					string GradleConfig = bForDistribution ? "release" : "debug";
+					string GradeOutAPK = Path.Combine(UnrealBuildGradlePath, "app", "build", "outputs", "apk", GradleConfig, "app-" + GradleConfig + ".apk");
+					SafeDeleteFile(GradeOutAPK);
+				}
+				
 				// See if we need to stage a UECommandLine.txt file in assets
 				string CommandLineDestFileName = Path.Combine(UnrealBuildPath, "assets", "UECommandLine.txt");
 				if (File.Exists(CommandLineSourceFileName))
