Fixed VersionScriptFile not being created robustly

#android
#13451

diff --git a/Engine/Source/Programs/UnrealBuildTool/Platform/Android/AndroidToolChain.cs b/Engine/Source/Programs/UnrealBuildTool/Platform/Android/AndroidToolChain.cs
--- a/Engine/Source/Programs/UnrealBuildTool/Platform/Android/AndroidToolChain.cs
+++ b/Engine/Source/Programs/UnrealBuildTool/Platform/Android/AndroidToolChain.cs
@@ -312,9 +312,9 @@ namespace UnrealBuildTool
 			return Ini.GetBool("/Script/AndroidRuntimeSettings.AndroidRuntimeSettings", "bDisableLibCppSharedDependencyValidation", out bDisableLibCppSharedDependencyValidation) && bDisableLibCppSharedDependencyValidation;
 		}
 
-		private string GetVersionScriptFilename(LinkEnvironment LinkEnvironment)
+		private FileReference GetVersionScriptFileName(LinkEnvironment LinkEnvironment)
 		{
-			return Path.Combine(LinkEnvironment.IntermediateDirectory!.FullName, "ExportSymbols.ldscript");
+			return FileReference.Combine(LinkEnvironment.IntermediateDirectory!, "ExportSymbols.ldscript");
 		}
 
 		public int GetNdkApiLevelInt(int MinNdk = MinimumNDKAPILevel)
@@ -703,7 +703,7 @@ namespace UnrealBuildTool
 			}
 		}
 
-		protected virtual string GetLinkArguments(LinkEnvironment LinkEnvironment, UnrealArch Architecture)
+		protected virtual string GetLinkArguments(LinkEnvironment LinkEnvironment, UnrealArch Architecture, IActionGraphBuilder Graph)
 		{
 			string Result = string.Join(' ', AdditionalPaths(LinkEnvironment.RootPaths));
 
@@ -758,13 +758,10 @@ namespace UnrealBuildTool
 
 			if (!LinkEnvironment.AdditionalArguments.Contains("--version-script"))
 			{
-				string VersionScriptFile = GetVersionScriptFilename(LinkEnvironment);
-				using (StreamWriter Writer = File.CreateText(VersionScriptFile))
-				{
-					// Make all symbols hidden (except new/delete operators and ones called from Java)
-					Writer.WriteLine("{ global: _Znwm*; _Znam*; _ZdlPv*; _ZdaPv*; Java_*; ANativeActivity_onCreate; JNI_OnLoad; local: *; };");
-				}
-				Result += " -Wl,--version-script=\"" + VersionScriptFile + "\"";
+				FileReference VersionScriptFileName = GetVersionScriptFileName(LinkEnvironment);
+				// Make all symbols hidden (except new/delete operators and ones called from Java)
+				Graph.CreateIntermediateTextFile(VersionScriptFileName, "{ global: _Znwm*; _Znam*; _ZdlPv*; _ZdaPv*; Java_*; ANativeActivity_onCreate; JNI_OnLoad; local: *; };");
+				Result += " -Wl,--version-script=\"" + VersionScriptFileName + "\"";
 			}
 			else
 			{
@@ -1084,7 +1081,7 @@ namespace UnrealBuildTool
 			LinkAction.WorkingDirectory = LinkEnvironment.IntermediateDirectory!;
 
 			// Get link arguments.
-			LinkAction.CommandArguments = LinkEnvironment.bIsBuildingLibrary ? GetArArguments(LinkEnvironment) : GetLinkArguments(LinkEnvironment, Arch);
+			LinkAction.CommandArguments = LinkEnvironment.bIsBuildingLibrary ? GetArArguments(LinkEnvironment) : GetLinkArguments(LinkEnvironment, Arch, Graph);
 
 			// Add the output file as a production of the link action.
 			FileItem OutputFile;
@@ -1274,8 +1271,8 @@ namespace UnrealBuildTool
 			// Only execute linking on the local PC.
 			LinkAction.bCanExecuteRemotely = false;
 
-			string VersionScriptFileItem = GetVersionScriptFilename(LinkEnvironment);
-			LinkAction.PrerequisiteItems.Add(FileItem.GetItemByPath(VersionScriptFileItem));
+			FileReference VersionScriptFileName = GetVersionScriptFileName(LinkEnvironment);
+			LinkAction.PrerequisiteItems.Add(FileItem.GetItemByFileReference(VersionScriptFileName));
 
 			Logger.LogInformation("Link: {LinkActionCommandPathFullName} {LinkActionCommandArguments}", LinkAction.CommandPath.FullName, LinkAction.CommandArguments);
 
