UnrealBuildTool: Track inline files in generated unity files

#rnx

diff --git a/Engine/Source/Programs/UnrealBuildTool/Configuration/UEBuildModuleCPP.cs b/Engine/Source/Programs/UnrealBuildTool/Configuration/UEBuildModuleCPP.cs
--- a/Engine/Source/Programs/UnrealBuildTool/Configuration/UEBuildModuleCPP.cs
+++ b/Engine/Source/Programs/UnrealBuildTool/Configuration/UEBuildModuleCPP.cs
@@ -694,11 +694,11 @@ namespace UnrealBuildTool
 							string Key = Prefix + ListOfInlinedGenCppsItem;
 							if (GeneratedFiles.Remove(Key, out FileItem? FoundGenCppFile))
 							{
-								if (!CppCompileEnvironment.FileInlineGenCPPMap.ContainsKey(CPPFileItem))
+								if (!CppCompileEnvironment.FileInlineSourceMap.ContainsKey(CPPFileItem))
 								{
-									CppCompileEnvironment.FileInlineGenCPPMap[CPPFileItem] = new List<FileItem>();
+									CppCompileEnvironment.FileInlineSourceMap[CPPFileItem] = [];
 								}
-								CppCompileEnvironment.FileInlineGenCPPMap[CPPFileItem].Add(FoundGenCppFile);
+								CppCompileEnvironment.FileInlineSourceMap[CPPFileItem].Add(FoundGenCppFile);
 							}
 							else
 							{
@@ -706,7 +706,7 @@ namespace UnrealBuildTool
 							}
 						}
 					}
-					ModuleCompileEnvironment.FileInlineGenCPPMap = CppCompileEnvironment.FileInlineGenCPPMap;
+					ModuleCompileEnvironment.FileInlineSourceMap = CppCompileEnvironment.FileInlineSourceMap;
 
 					if (Rules.CppCompileWarningSettings.NonInlinedGenCppWarningLevel >= WarningLevel.Warning)
 					{
@@ -919,8 +919,8 @@ namespace UnrealBuildTool
 				GeneratedCPPCompileEnvironment.bUseUnity = true;
 				Unity.GenerateUnitySource(Target, GeneratedFileItems, new List<FileItem>(), GeneratedCPPCompileEnvironment, WorkingSet, (Rules.ShortName ?? Name) + ".gen", IntermediateDirectory, Graph, SourceFileToUnityFile,
 					out List<FileItem> NormalGeneratedFiles, out List<FileItem> AdaptiveGeneratedFiles, NumIncludedBytesPerUnityCPP);
-				ModuleCompileEnvironment.FileInlineGenCPPMap = new(GeneratedCPPCompileEnvironment.FileInlineGenCPPMap);
-				CppCompileEnvironment.FileInlineGenCPPMap = new(GeneratedCPPCompileEnvironment.FileInlineGenCPPMap);
+				ModuleCompileEnvironment.FileInlineSourceMap = new(GeneratedCPPCompileEnvironment.FileInlineSourceMap);
+				CppCompileEnvironment.FileInlineSourceMap = new(GeneratedCPPCompileEnvironment.FileInlineSourceMap);
 				LinkInputFiles.AddRange(CompileFilesWithToolChain(Target, ToolChain, GeneratedCPPCompileEnvironment, ModuleCompileEnvironment, NormalGeneratedFiles, AdaptiveGeneratedFiles, Graph, Logger).ObjectFiles);
 			}
 			else
@@ -970,7 +970,7 @@ namespace UnrealBuildTool
 				{
 					Unity.GenerateUnitySource(Target, CPPFiles, InputFiles.HeaderFiles, CppCompileEnvironment, WorkingSet, Rules.ShortName ?? Name, IntermediateDirectory, Graph, SourceFileToUnityFile,
 						out List<FileItem> NormalFiles, out List<FileItem> AdaptiveFiles, NumIncludedBytesPerUnityCPP);
-					ModuleCompileEnvironment.FileInlineGenCPPMap = new(CppCompileEnvironment.FileInlineGenCPPMap);
+					ModuleCompileEnvironment.FileInlineSourceMap = new(CppCompileEnvironment.FileInlineSourceMap);
 					LinkInputFiles.AddRange(CompileFilesWithToolChain(Target, ToolChain, CppCompileEnvironment, ModuleCompileEnvironment, NormalFiles, AdaptiveFiles, Graph, Logger).ObjectFiles);
 				}
 				else
diff --git a/Engine/Source/Programs/UnrealBuildTool/Platform/Windows/VCToolChain.cs b/Engine/Source/Programs/UnrealBuildTool/Platform/Windows/VCToolChain.cs
--- a/Engine/Source/Programs/UnrealBuildTool/Platform/Windows/VCToolChain.cs
+++ b/Engine/Source/Programs/UnrealBuildTool/Platform/Windows/VCToolChain.cs
@@ -2276,8 +2276,7 @@ namespace UnrealBuildTool
 					}
 				}
 
-				List<FileItem>? InlinedFiles;
-				if (CompileEnvironment.FileInlineGenCPPMap.TryGetValue(SourceFile, out InlinedFiles))
+				if (CompileEnvironment.FileInlineSourceMap.TryGetValue(SourceFile, out HashSet<FileItem>? InlinedFiles))
 				{
 					CompileAction.AdditionalPrerequisiteItems.AddRange(InlinedFiles);
 				}
diff --git a/Engine/Source/Programs/UnrealBuildTool/System/CppCompileEnvironment.cs b/Engine/Source/Programs/UnrealBuildTool/System/CppCompileEnvironment.cs
--- a/Engine/Source/Programs/UnrealBuildTool/System/CppCompileEnvironment.cs
+++ b/Engine/Source/Programs/UnrealBuildTool/System/CppCompileEnvironment.cs
@@ -636,7 +636,7 @@ namespace UnrealBuildTool
 		/// <summary>
 		/// A dictionary of the source file items and the inlined gen.cpp files contained in it
 		/// </summary>
-		public Dictionary<FileItem, List<FileItem>> FileInlineGenCPPMap = [];
+		public Dictionary<FileItem, HashSet<FileItem>> FileInlineSourceMap = [];
 		
 		/// <summary>
 		/// Non-default naming conventions that generated CPP files can have (this uses a filter to check the file name for matches) 
@@ -896,7 +896,7 @@ namespace UnrealBuildTool
 			AdditionalPrerequisites.AddRange(Other.AdditionalPrerequisites);
 			ExtraGeneratedCPPFileTypes = Other.ExtraGeneratedCPPFileTypes;
 			CollidingNames = Other.CollidingNames;
-			FileInlineGenCPPMap = new(Other.FileInlineGenCPPMap);
+			FileInlineSourceMap = new(Other.FileInlineSourceMap);
 			Definitions.AddRange(Other.Definitions);
 			AdditionalResponseFiles.AddRange(Other.AdditionalResponseFiles);
 			AdditionalArguments = Other.AdditionalArguments;
diff --git a/Engine/Source/Programs/UnrealBuildTool/System/Unity.cs b/Engine/Source/Programs/UnrealBuildTool/System/Unity.cs
--- a/Engine/Source/Programs/UnrealBuildTool/System/Unity.cs
+++ b/Engine/Source/Programs/UnrealBuildTool/System/Unity.cs
@@ -276,7 +276,7 @@ namespace UnrealBuildTool
 					: $"{ModulePrefix}{BaseName}{UnityExt}";
 				FileReference UnitySourceFilePath = FileReference.Combine(IntermediateDirectory, UnitySourceFileName);
 
-				List<FileItem> InlinedGenCPPFilesInUnity = [];
+				HashSet<FileItem> InlinedSourceFilesInUnity = [];
 
 				// Add source files to the unity file
 				foreach (FileItem SourceFile in UnityFile.Files)
@@ -292,11 +292,11 @@ namespace UnrealBuildTool
 					}
 					OutputUnitySourceWriter.WriteLine("#include \"{0}\"", SourceFileString.Replace('\\', '/'));
 
-					List<FileItem>? InlinedGenCPPFiles;
-					if (CompileEnvironment.FileInlineGenCPPMap.TryGetValue(SourceFile, out InlinedGenCPPFiles))
+					if (CompileEnvironment.FileInlineSourceMap.TryGetValue(SourceFile, out HashSet<FileItem>? InlinedFiles))
 					{
-						InlinedGenCPPFilesInUnity.AddRange(InlinedGenCPPFiles);
+						InlinedSourceFilesInUnity.UnionWith(InlinedFiles);
 					}
+					InlinedSourceFilesInUnity.Add(SourceFile);
 				}
 
 				// Write the unity file to the intermediate folder.
@@ -304,7 +304,7 @@ namespace UnrealBuildTool
 				NewSourceFiles.Add(UnitySourceFile);
 
 				// Store all the inlined gen.cpp files
-				CompileEnvironment.FileInlineGenCPPMap[UnitySourceFile] = InlinedGenCPPFilesInUnity;
+				CompileEnvironment.FileInlineSourceMap[UnitySourceFile] = InlinedSourceFilesInUnity;
 
 				// Store the mapping of source files to unity files in the makefile
 				foreach (FileItem SourceFile in UnityFile.Files)
diff --git a/Engine/Source/Programs/UnrealBuildTool/ToolChain/ClangToolChain.cs b/Engine/Source/Programs/UnrealBuildTool/ToolChain/ClangToolChain.cs
--- a/Engine/Source/Programs/UnrealBuildTool/ToolChain/ClangToolChain.cs
+++ b/Engine/Source/Programs/UnrealBuildTool/ToolChain/ClangToolChain.cs
@@ -1076,7 +1076,7 @@ namespace UnrealBuildTool
 			CompileAction.PrerequisiteItems.UnionWith(CompileEnvironment.AdditionalPrerequisites);
 			CompileAction.PrerequisiteItems.Add(SourceFile);
 
-			if (CompileEnvironment.FileInlineGenCPPMap.TryGetValue(SourceFile, out List<FileItem>? InlinedFiles))
+			if (CompileEnvironment.FileInlineSourceMap.TryGetValue(SourceFile, out HashSet<FileItem>? InlinedFiles))
 			{
 				CompileAction.PrerequisiteItems.UnionWith(InlinedFiles);
 			}
diff --git a/Engine/Source/Programs/UnrealBuildTool/ToolChain/IWYUToolChain.cs b/Engine/Source/Programs/UnrealBuildTool/ToolChain/IWYUToolChain.cs
--- a/Engine/Source/Programs/UnrealBuildTool/ToolChain/IWYUToolChain.cs
+++ b/Engine/Source/Programs/UnrealBuildTool/ToolChain/IWYUToolChain.cs
@@ -240,8 +240,7 @@ namespace UnrealBuildTool
 			Arguments.Add("-Xiwyu --write_json_path=\"" + fileItem.AbsolutePath.Replace('\\', '/') + "\" ");
 			if (SourceFile.HasExtension(".cpp"))
 			{
-				List<FileItem>? InlinedFiles;
-				if (CompileEnvironment.FileInlineGenCPPMap.TryGetValue(SourceFile, out InlinedFiles))
+				if (CompileEnvironment.FileInlineSourceMap.TryGetValue(SourceFile, out HashSet<FileItem>? InlinedFiles))
 				{
 					foreach (FileItem InlinedFile in InlinedFiles)
 					{
