UnrealBuildTool: Fix -DisableUnity command line flag not working for static analysis

#rnx

diff --git a/Engine/Source/Programs/UnrealBuildTool/Configuration/TargetRules.cs b/Engine/Source/Programs/UnrealBuildTool/Configuration/TargetRules.cs
--- a/Engine/Source/Programs/UnrealBuildTool/Configuration/TargetRules.cs
+++ b/Engine/Source/Programs/UnrealBuildTool/Configuration/TargetRules.cs
@@ -3174,6 +3174,10 @@ namespace UnrealBuildTool
 			// Determine intermediate environment overrides based on command line flags
 			if (StaticAnalyzer != StaticAnalyzer.None)
 			{
+				if (!bUseUnityBuild)
+				{
+					bUseUnityBuildOverride = false;
+				}
 				// Always override environment for analyzing, regardless of other settings
 				IntermediateEnvironment = UnrealIntermediateEnvironment.Analyze;
 			}
