UnrealBuildTool: PVS-Studio toolchain working with -SingleFile

#rnx

diff --git a/Engine/Source/Programs/UnrealBuildTool/Modes/BuildMode.cs b/Engine/Source/Programs/UnrealBuildTool/Modes/BuildMode.cs
--- a/Engine/Source/Programs/UnrealBuildTool/Modes/BuildMode.cs
+++ b/Engine/Source/Programs/UnrealBuildTool/Modes/BuildMode.cs
@@ -931,11 +931,24 @@ namespace UnrealBuildTool
 				}
 
 				// There is no special action for this file.. let's look for an action that has this exact file as input and use that (for example ispc files)
-				LinkedAction? Action = Actions.Find(x => x.PrerequisiteItems.Contains(SourceFile));
+				IEnumerable<LinkedAction> prereqActions = Actions.Where(x => x.PrerequisiteItems.Contains(SourceFile));
 
-				if (Action != null)
+				if (prereqActions.Any())
 				{
-					ProducedItems.AddRange(Action.ProducedItems);
+					// Hack to get PVS working with -SingleFile
+					LinkedAction? pvsFinalizeAction = Actions.FirstOrDefault(x => x.CommandDescription == "Process PVS-Studio Results");
+					if (pvsFinalizeAction != null)
+					{
+						prereqActions = prereqActions.Except([pvsFinalizeAction]);
+						IExternalAction? newPvsFinalizeAction = PVSToolChain.SingleFileFinalizeOutput(pvsFinalizeAction, prereqActions);
+						if (newPvsFinalizeAction != null)
+						{
+							Actions.Add(new LinkedAction(newPvsFinalizeAction, Target));
+							ProducedItems.AddRange(newPvsFinalizeAction.ProducedItems);
+						}
+					}
+
+					ProducedItems.AddRange(prereqActions.SelectMany(x => x.ProducedItems));
 				}
 				else if (!BuildConfiguration.bIgnoreInvalidFiles)
 				{
diff --git a/Engine/Source/Programs/UnrealBuildTool/Platform/Windows/PVSToolChain.cs b/Engine/Source/Programs/UnrealBuildTool/Platform/Windows/PVSToolChain.cs
--- a/Engine/Source/Programs/UnrealBuildTool/Platform/Windows/PVSToolChain.cs
+++ b/Engine/Source/Programs/UnrealBuildTool/Platform/Windows/PVSToolChain.cs
@@ -2,6 +2,7 @@
 
 using EpicGames.Core;
 using Microsoft.Extensions.Logging;
+using Microsoft.Extensions.Logging.Abstractions;
 using Newtonsoft.Json;
 using System;
 using System.Collections.Generic;
@@ -545,7 +546,7 @@ namespace UnrealBuildTool
 		static readonly Version _minAnalyzerVersion = new Version("7.30");
 		static readonly Version _analysisPathsSkipVersion = new Version("7.34");
 
-		const string OutputFileExtension = ".PVS-Studio.log";
+		static string OutputFileExtension = ".PVS-Studio.log";
 
 		public PVSToolChain(ReadOnlyTargetRules Target, VCToolChain InInnerToolchain, ILogger Logger)
 			: base(Logger)
@@ -876,6 +877,7 @@ namespace UnrealBuildTool
 
 			analyzeAction.CommandArguments = String.Join(' ', arguments);
 
+			analyzeAction.PrerequisiteItems.UnionWith(preprocessAction.AdditionalPrerequisiteItems);
 			analyzeAction.PrerequisiteItems.Add(sourceFileItem);
 			analyzeAction.PrerequisiteItems.Add(configFileItem);
 			analyzeAction.PrerequisiteItems.Add(preprocessedFileItem);
@@ -984,6 +986,7 @@ namespace UnrealBuildTool
 			finalizeAction.CommandDescription = "Process PVS-Studio Results";
 			finalizeAction.PrerequisiteItems.Add(inputFileListItem);
 			finalizeAction.PrerequisiteItems.Add(ignoredFileListItem);
+			finalizeAction.PrerequisiteItems.Add(rootPathsItem);
 			finalizeAction.PrerequisiteItems.UnionWith(makefile.OutputItems);
 			finalizeAction.PrerequisiteItems.UnionWith(compileSourceFiles);
 			finalizeAction.ProducedItems.Add(FileItem.GetItemByFileReference(outputFile));
@@ -993,6 +996,46 @@ namespace UnrealBuildTool
 			makefile.OutputItems.AddRange(finalizeAction.ProducedItems);
 		}
 
+		public static IExternalAction? SingleFileFinalizeOutput(IExternalAction finalizeAction, IEnumerable<IExternalAction> prereqActions)
+		{
+			Match match = Regex.Match(finalizeAction.CommandArguments, @"-Input=""(.*?)"" -Output=""(.*?)"" -Ignored=""(.*?)"" -RootPaths=""(.*?)"" -PrintLevel=(\d+)");
+			if (match.Success)
+			{
+				FileItem input = FileItem.GetItemByPath(match.Groups[1].Value);
+				FileItem output = FileItem.GetItemByPath(match.Groups[2].Value);
+				FileItem ignored = FileItem.GetItemByPath(match.Groups[3].Value);
+				FileItem rootPaths = FileItem.GetItemByPath(match.Groups[4].Value);
+				int printLevel = int.Parse(match.Groups[5].Value);
+
+				FileItem newInput = FileItem.GetItemByFileReference(input.Location.ChangeExtension($".single{input.Location.GetExtension()}"));
+				FileItem newOutput = FileItem.GetItemByFileReference(output.Location.ChangeExtension($".single{output.Location.GetExtension()}"));
+
+				string[] inputFileLines = FileReference.ReadAllLines(input.Location);
+				IEnumerable<FileItem> inputFiles = inputFileLines.Select(x => x.Trim()).Where(x => x.Length > 0).Select(FileItem.GetItemByPath);
+				IEnumerable<FileItem> newInputFiles = inputFiles.Intersect(prereqActions.SelectMany(x => x.ProducedItems).Where(x => x.Location.HasExtension(OutputFileExtension)));
+
+				Utils.WriteFileIfChanged(newInput, newInputFiles.Select(x => x.FullName), NullLogger.Instance);
+
+				string arguments = $"-Input=\"{newInput.Location}\" -Output=\"{newOutput.Location}\" -Ignored=\"{ignored.Location}\" -RootPaths=\"{rootPaths.Location}\" -PrintLevel={printLevel}";
+				Action singleFinalizeAction = new(finalizeAction);
+				singleFinalizeAction.CommandDescription = "Process PVS-Studio Single File Results";
+				singleFinalizeAction.CommandArguments = $"{singleFinalizeAction.CommandArguments[..singleFinalizeAction.CommandArguments.IndexOf("-Input=")]}{arguments}";
+				singleFinalizeAction.PrerequisiteItems.Clear();
+				singleFinalizeAction.PrerequisiteItems.Add(newInput);
+				singleFinalizeAction.PrerequisiteItems.Add(ignored);
+				singleFinalizeAction.PrerequisiteItems.Add(rootPaths);
+				singleFinalizeAction.PrerequisiteItems.UnionWith(newInputFiles);
+				singleFinalizeAction.ProducedItems.Clear();
+				singleFinalizeAction.ProducedItems.Add(newOutput);
+				singleFinalizeAction.ProducedItems.Add(FileItem.GetItemByPath(newOutput.FullName + "_does_not_exist")); // Force the gather step to always execute
+				singleFinalizeAction.DeleteItems.Clear();
+				singleFinalizeAction.DeleteItems.UnionWith(singleFinalizeAction.ProducedItems);
+
+				return singleFinalizeAction;
+			}
+			return null;
+		}
+
 		public override List<string> GetISPCCompileTargets(UnrealTargetPlatform Platform, UnrealArch Arch) => InnerToolChain.GetISPCCompileTargets(Platform, Arch);
 
 		public override string GetISPCOSTarget(UnrealTargetPlatform Platform) => InnerToolChain.GetISPCOSTarget(Platform);
