UnrealBuildTool: Update inaccurate comment

#rnx

diff --git a/Engine/Source/Programs/UnrealBuildTool/Modes/BuildMode.cs b/Engine/Source/Programs/UnrealBuildTool/Modes/BuildMode.cs
--- a/Engine/Source/Programs/UnrealBuildTool/Modes/BuildMode.cs
+++ b/Engine/Source/Programs/UnrealBuildTool/Modes/BuildMode.cs
@@ -930,7 +930,7 @@ namespace UnrealBuildTool
 					}
 				}
 
-				// There is no special action for this file.. let's look for an action that has this exact file as input and use that (for example ispc files)
+				// There is no special action for this file.. let's look for any actions that has this exact file as input and use that (for example ispc files)
 				IEnumerable<LinkedAction> prereqActions = Actions.Where(x => x.PrerequisiteItems.Contains(SourceFile));
 
 				if (prereqActions.Any())
