Warn if a library on Android is not 16kb page size compatible

diff --git a/Engine/Source/Programs/UnrealBuildTool/Platform/Android/UEDeployAndroid.cs b/Engine/Source/Programs/UnrealBuildTool/Platform/Android/UEDeployAndroid.cs
--- a/Engine/Source/Programs/UnrealBuildTool/Platform/Android/UEDeployAndroid.cs
+++ b/Engine/Source/Programs/UnrealBuildTool/Platform/Android/UEDeployAndroid.cs
@@ -8,6 +8,8 @@ using System.IO.Compression;
 using System.Linq;
 using System.Security.Cryptography;
 using System.Text;
+using System.Text.Json;
+using System.Text.Json.Serialization;
 using System.Xml.Linq;
 using EpicGames.Core;
 using Microsoft.Extensions.Logging;
@@ -1339,6 +1341,125 @@ namespace UnrealBuildTool
 			return DependencyList;
 		}
 
+		private static bool IsSO16KbCompatible(string FileName, UnrealArch UnrealArch)
+		{
+			string ReadElfExecutable = AndroidToolChain.GetReadElfExecutablePath(UnrealArch);
+			string StdOutput = Utils.RunLocalProcessAndReturnStdOut(ReadElfExecutable, $"--elf-output-style=JSON --program-headers \"{FileName}\"");
+
+			ReadElfOutput Output = default;
+			bool bParsedOk = false;
+			try
+			{
+				// expected in NDK >= 27
+				List<ReadElfOutput>? JsonOutput = JsonSerializer.Deserialize<List<ReadElfOutput>>(StdOutput);
+				if (JsonOutput != null)
+				{
+					Output = JsonOutput.First();
+					bParsedOk = true;
+				}
+			}
+			catch (Exception)
+			{
+				// ignored
+			}
+
+			if (!bParsedOk)
+			{
+				try
+				{
+					// expected in NDK 25
+					List<Dictionary<string, ReadElfOutput>>? JsonOutput = JsonSerializer.Deserialize<List<Dictionary<string, ReadElfOutput>>>(StdOutput);
+					if (JsonOutput != null)
+					{
+						Output = JsonOutput.First().First().Value;
+						bParsedOk = true;
+					}
+				}
+				catch (Exception)
+				{
+					// ignored
+				}
+			}
+
+			if (!bParsedOk || Output.ProgramHeaders == null || Output.ProgramHeaders.Count == 0)
+			{
+				// don't fail the build if we can't check alignment
+				return true;
+			}
+
+			bool bIsAlignmentOk = true;
+
+			foreach (ReadElfProgramHeaderWrapper Wrapper in Output.ProgramHeaders)
+			{
+				ReadElfProgramHeader ProgramHeader = Wrapper.ProgramHeader;
+
+				if (ProgramHeader.Type.Name == "PT_LOAD")
+				{
+					bIsAlignmentOk &= ProgramHeader.Alignment >= 16 * 1024;
+				}
+			}
+
+			return bIsAlignmentOk;
+		}
+
+		private struct ReadElfOutput
+		{
+			public required List<ReadElfProgramHeaderWrapper> ProgramHeaders { get; set; }
+		}
+
+		private struct ReadElfProgramHeaderWrapper
+		{
+			public ReadElfProgramHeader ProgramHeader { get; set; }
+		}
+
+		private struct ReadElfProgramHeader
+		{
+			public ReadElfProgramHeaderType Type { get; set; }
+			public int Alignment { get; set; }
+		}
+
+		private struct ReadElfProgramHeaderType
+		{
+			// only present in NDK >=27
+			[JsonPropertyName("Name")]
+			public string MaybeName { get; set; }
+
+			// in NDK 25 this is a string, in NDK >=27 it's an int
+			[JsonPropertyName("Value")]
+			public object MaybeValue { get; set; }
+
+			private string MaybeValueAsString => MaybeValue?.ToString() ?? string.Empty;
+
+			[JsonIgnore]
+			public string Name => !string.IsNullOrEmpty(MaybeName) ? MaybeName : MaybeValueAsString;
+		}
+
+		private void CheckSOInFolderAre16KbCompatible(string FolderPath, UnrealArch UnrealArch)
+		{
+			try
+			{
+				IEnumerable<string> SOFiles = Directory.EnumerateFiles(FolderPath, "*.so", SearchOption.AllDirectories);
+
+				foreach (var SOFile in SOFiles)
+				{
+					if (!IsSO16KbCompatible(SOFile, UnrealArch))
+					{
+						Logger.LogWarning("Library '{SOFile}' is NOT compatible with 16kb page size alignment.", SOFile);
+					}
+				}
+			}
+			catch (Exception)
+			{
+				// ignored
+			}
+		}
+
+		private void CheckSOInGradleFolderAre16KbCompatible(string GradlePath, UnrealArch UnrealArch)
+		{
+			CheckSOInFolderAre16KbCompatible($"{GradlePath}/jniLibs/arm64-v8a", UnrealArch);
+			CheckSOInFolderAre16KbCompatible($"{GradlePath}/libs/arm64-v8a", UnrealArch);
+		}
+
 		private static void CopySO(string FinalName, string SourceName)
 		{
 			// check to see if file is newer than last time we copied
@@ -5362,6 +5483,8 @@ popd
 				CreateGradlePropertiesFiles(ToolChain, Arch, MinSDKVersion, TargetSDKVersion, CompileSDKVersion, BuildToolsVersion, PackageName, DestApkName, NDKArch,
 					UnrealBuildFilesPath, GameBuildFilesPath, UnrealBuildGradleAppPath, UnrealBuildPath, UnrealBuildGradlePath, bForDistribution, bSkipGradleBuild, RequiredOBBFiles, bJavaPerfettoMarkers);
 
+				CheckSOInGradleFolderAre16KbCompatible(UnrealBuildGradleMainPath, Arch);
+
 				if (!bSkipGradleBuild)
 				{
 					string GradleScriptPath = Path.Combine(UnrealBuildGradlePath, "gradlew");
