[Linux] Fix up incorrect error message in toolchain version

diff --git a/Engine/Source/Programs/UnrealBuildTool/Platform/Linux/LinuxToolChain.cs b/Engine/Source/Programs/UnrealBuildTool/Platform/Linux/LinuxToolChain.cs
--- a/Engine/Source/Programs/UnrealBuildTool/Platform/Linux/LinuxToolChain.cs
+++ b/Engine/Source/Programs/UnrealBuildTool/Platform/Linux/LinuxToolChain.cs
@@ -57,7 +57,7 @@ namespace UnrealBuildTool
 			// prevent unknown clangs since the build is likely to fail on too old or too new compilers
 			if ((CompilerVersionLessThan(20, 0, 0) || CompilerVersionGreaterOrEqual(21, 0, 0)) && !Options.HasFlag(ClangToolChainOptions.UseAutoRTFMCompiler))
 			{
-				throw new BuildLogEventException("This version of the Unreal Engine can only be compiled with clang 18.x. clang {Version} may not build it - please use a different version.", Info.ClangVersion);
+				throw new BuildLogEventException("This version of the Unreal Engine can only be compiled with clang 20.x. clang {Version} may not build it - please use a different version.", Info.ClangVersion);
 			}
 		}
 
