UnrealBuilldTool: Remove break preventing UpdateProducingAttributes from being called

#rnx

diff --git a/Engine/Source/Programs/UnrealBuildTool/System/ActionGraph.cs b/Engine/Source/Programs/UnrealBuildTool/System/ActionGraph.cs
--- a/Engine/Source/Programs/UnrealBuildTool/System/ActionGraph.cs
+++ b/Engine/Source/Programs/UnrealBuildTool/System/ActionGraph.cs
@@ -743,7 +743,6 @@ namespace UnrealBuildTool
 					}
 
 					bIsOutdated = true;
-					break;
 				}
 
 				// If the produced file doesn't exist or has zero size, consider it outdated.  The zero size check is to detect cases
@@ -765,7 +764,6 @@ namespace UnrealBuildTool
 					// If any of the produced items doesn't exist, the action is outdated.
 					Logger.LogTrace("{StatusDescription}: Produced item \"{ProducedItem}\" doesn't exist.", RootAction.StatusDescription, ProducedItem.Location);
 					bIsOutdated = true;
-					break;
 				}
 			}
 
