bAllowLTCG returns true if bPGOProfile or bPGOOptimize are set to true

#rnx

diff --git a/Engine/Source/Programs/UnrealBuildTool/Configuration/TargetRules.cs b/Engine/Source/Programs/UnrealBuildTool/Configuration/TargetRules.cs
--- a/Engine/Source/Programs/UnrealBuildTool/Configuration/TargetRules.cs
+++ b/Engine/Source/Programs/UnrealBuildTool/Configuration/TargetRules.cs
@@ -2138,12 +2138,17 @@ namespace UnrealBuildTool
 		public bool bUseIncrementalLinking { get; set; }
 
 		/// <summary>
-		/// Whether to allow the use of link time code generation (LTCG).
+		/// Whether to allow the use of link time code generation (LTCG). Implied when PGO is enabled
 		/// </summary>
 		[CommandLine("-LTCG")]
 		[XmlConfigFile(Category = "BuildConfiguration")]
-		public bool bAllowLTCG { get; set; }
-
+		public bool bAllowLTCG
+		{
+			get => bAllowLTCGPrivate || bPGOProfile || bPGOOptimize;
+			set => bAllowLTCGPrivate = value;
+		}
+		private bool bAllowLTCGPrivate;
+		
 		/// <summary>
 		/// When Link Time Code Generation (LTCG) is enabled, whether to 
 		/// prefer using the lighter weight version on supported platforms.
