UnrealBuildTool: -SingleFileBuildDependents should ignore generated code

#rnx

diff --git a/Engine/Source/Programs/UnrealBuildTool/Configuration/TargetDescriptor.cs b/Engine/Source/Programs/UnrealBuildTool/Configuration/TargetDescriptor.cs
--- a/Engine/Source/Programs/UnrealBuildTool/Configuration/TargetDescriptor.cs
+++ b/Engine/Source/Programs/UnrealBuildTool/Configuration/TargetDescriptor.cs
@@ -85,7 +85,7 @@ namespace UnrealBuildTool
 		public string? WorkingDir = null;
 
 		/// <summary>
-		/// Will build all files that directly include any of the files provided in -SingleFile
+		/// Will build all files that directly include any of the files provided in -SingleFile, not including generated code
 		/// </summary>
 		[CommandLine("-SingleFileBuildDependents")]
 		public bool bSingleFileBuildDependents;
diff --git a/Engine/Source/Programs/UnrealBuildTool/Modes/BuildMode.cs b/Engine/Source/Programs/UnrealBuildTool/Modes/BuildMode.cs
--- a/Engine/Source/Programs/UnrealBuildTool/Modes/BuildMode.cs
+++ b/Engine/Source/Programs/UnrealBuildTool/Modes/BuildMode.cs
@@ -996,7 +996,12 @@ namespace UnrealBuildTool
 				CacheDir = ProjectFile.Directory;
 			}
 
-			IEnumerable<FileItem> CodeFiles = Actions.SelectMany(x => x.PrerequisiteItems).Distinct().Where(x => IsCode(x.Location) && !x.Name.StartsWith(Unity.ModulePrefix));
+			// Ignore generated files
+			IEnumerable<FileItem> CodeFiles = 
+				Actions
+				.SelectMany(x => x.PrerequisiteItems)
+				.Distinct()
+				.Where(x => IsCode(x.Location) && !x.Name.StartsWith(Unity.ModulePrefix) && !x.HasExtension(".gen.cpp"));
 
 			FileReference IncludeCache = FileReference.Combine(CacheDir, "Intermediate", "HeaderLookup.dat");
 			CppIncludeLookup IncludeLookup = new CppIncludeLookup(IncludeCache);
@@ -1011,15 +1016,10 @@ namespace UnrealBuildTool
 			// Found files can be a lot (Commandline.h gives 130.000+ files). And we only want to build the ones that are in existing actions
 			IEnumerable<FileReference> NewFiles = FoundFiles.Intersect(UsedFiles).Distinct();
 
-			int genFileCount = NewFiles.Where(x => x.GetFileName().EndsWith(".gen.cpp", StringComparison.OrdinalIgnoreCase)).Count();
-			Logger.LogInformation("Found {NewFilesCount} source files", NewFiles.Count() - genFileCount);
-			if (genFileCount > 0)
-			{
-				Logger.LogInformation("Found {NewFilesCount} generated source files", genFileCount);
-			}
+			Logger.LogInformation("Found {NewFilesCount} source files", NewFiles.Count());
 			foreach (FileReference NewFile in NewFiles.Order())
 			{
-				Logger.LogDebug("  {NewFile}", NewFile);
+				Logger.LogTrace("  {NewFile}", NewFile);
 			}
 
 			return NewFiles.Concat(SpecificFilesToCompile).Distinct().Order().ToList();
