UBA: Handle verse-clang++ as clang++

#rnx

diff --git a/Engine/Source/Programs/UnrealBuildAccelerator/Core/Private/UbaApplicationRules.cpp b/Engine/Source/Programs/UnrealBuildAccelerator/Core/Private/UbaApplicationRules.cpp
--- a/Engine/Source/Programs/UnrealBuildAccelerator/Core/Private/UbaApplicationRules.cpp
+++ b/Engine/Source/Programs/UnrealBuildAccelerator/Core/Private/UbaApplicationRules.cpp
@@ -759,6 +759,7 @@ namespace uba
 			//{ L"cmd.exe"),		new ApplicationRules() },
 #else
 			{ ApplicationHash(TC("clang++")),					new ApplicationRulesClangPlusPlusExe() },
+			{ ApplicationHash(TC("verse-clang++")),				new ApplicationRulesClangPlusPlusExe() },
 			{ ApplicationHash(TC("ld.lld")),					new ApplicationRulesLdLLdExe() },
 			{ ApplicationHash(TC("ispc")),						new ApplicationRulesISPCExe() },
 			{ ApplicationHash(TC("shadercompileworker")),		new ApplicationRulesShaderCompileWorker() },
diff --git a/Engine/Source/Programs/UnrealBuildTool/ToolChain/ISPCToolChain.cs b/Engine/Source/Programs/UnrealBuildTool/ToolChain/ISPCToolChain.cs
--- a/Engine/Source/Programs/UnrealBuildTool/ToolChain/ISPCToolChain.cs
+++ b/Engine/Source/Programs/UnrealBuildTool/ToolChain/ISPCToolChain.cs
@@ -381,7 +381,7 @@ namespace UnrealBuildTool
 			IEnumerable<string> strings = [
 				target.CacheSalt ?? String.Empty,			// Optional salt
 				target.BuildEnvironment.ToString(),			// Unique or Share environment
-				target.bMergeModules ? "MM2" : String.Empty	// Merged modules will modify all object files
+				target.bMergeModules ? "MM2" : String.Empty // Merged modules will modify all object files
 			];
 			return strings.Aggregate(GetStringHash(String.Empty), (acc, str) => GetStringHash(str, acc));
 		}
@@ -425,6 +425,7 @@ namespace UnrealBuildTool
 				target.LinkType.ToString(), // Modular or monolithic (we want this for compiles too in order to end up in same bucket
 				!compileEnvironment.Architectures.bIsMultiArch ? compileEnvironment.Architecture.ToString() : compileEnvironment.Architectures.ToString(), // arm64 or x64
 				target.StaticAnalyzer != StaticAnalyzer.None && !compileEnvironment.bDisableStaticAnalysis ? target.StaticAnalyzer.ToString() : String.Empty, // If static analyzer
+				compileEnvironment.bUseAutoRTFMCompiler ? "AUTORTFM" : String.Empty, // AutoRTFM
 				!compileEnvironment.bUseUnity ? "NU" : String.Empty, // Non-unity
 			];
 
@@ -468,7 +469,8 @@ namespace UnrealBuildTool
 				linkEnvironment.Configuration.ToString(), // Configuration
 				target.LinkType.ToString(), // Modular or monolithic
 				!linkEnvironment.Architectures.bIsMultiArch ? linkEnvironment.Architecture.ToString() : linkEnvironment.Architectures.ToString(), // arm64 or x64
-				target.bUseUnityBuild == false ? "NU" : String.Empty // Non-unity
+				target.bUseAutoRTFMCompiler ? "AUTORTFM" : String.Empty, // AutoRTFM
+				!target.bUseUnityBuild ? "NU" : String.Empty // Non-unity
 			];
 
 			// Optimize settings
