UnrealBuildTool: Need to take VFS into account for cache bucketing

#rnx

diff --git a/Engine/Source/Programs/UnrealBuildTool/ToolChain/ISPCToolChain.cs b/Engine/Source/Programs/UnrealBuildTool/ToolChain/ISPCToolChain.cs
--- a/Engine/Source/Programs/UnrealBuildTool/ToolChain/ISPCToolChain.cs
+++ b/Engine/Source/Programs/UnrealBuildTool/ToolChain/ISPCToolChain.cs
@@ -425,6 +425,7 @@ namespace UnrealBuildTool
 				target.LinkType.ToString(), // Modular or monolithic (we want this for compiles too in order to end up in same bucket
 				!compileEnvironment.Architectures.bIsMultiArch ? compileEnvironment.Architecture.ToString() : compileEnvironment.Architectures.ToString(), // arm64 or x64
 				target.StaticAnalyzer != StaticAnalyzer.None && !compileEnvironment.bDisableStaticAnalysis ? target.StaticAnalyzer.ToString() : String.Empty, // If static analyzer
+				compileEnvironment.RootPaths.bUseVfs ? "VFS" : String.Empty,
 				compileEnvironment.bUseAutoRTFMCompiler ? "AUTORTFM" : String.Empty, // AutoRTFM
 				!compileEnvironment.bUseUnity ? "NU" : String.Empty, // Non-unity
 			];
@@ -469,6 +470,7 @@ namespace UnrealBuildTool
 				linkEnvironment.Configuration.ToString(), // Configuration
 				target.LinkType.ToString(), // Modular or monolithic
 				!linkEnvironment.Architectures.bIsMultiArch ? linkEnvironment.Architecture.ToString() : linkEnvironment.Architectures.ToString(), // arm64 or x64
+				linkEnvironment.RootPaths.bUseVfs ? "VFS" : String.Empty,
 				target.bUseAutoRTFMCompiler ? "AUTORTFM" : String.Empty, // AutoRTFM
 				!target.bUseUnityBuild ? "NU" : String.Empty // Non-unity
 			];
