Roll back CL 44950096 (bAllowLTCG returns true if bPGOProfile or bPGOOptimize are set to true) - causing Android link failures with PGO

#android

#rnx

diff --git a/Engine/Source/Programs/UnrealBuildTool/Configuration/TargetRules.cs b/Engine/Source/Programs/UnrealBuildTool/Configuration/TargetRules.cs
--- a/Engine/Source/Programs/UnrealBuildTool/Configuration/TargetRules.cs
+++ b/Engine/Source/Programs/UnrealBuildTool/Configuration/TargetRules.cs
@@ -2138,16 +2138,11 @@ namespace UnrealBuildTool
 		public bool bUseIncrementalLinking { get; set; }
 
 		/// <summary>
-		/// Whether to allow the use of link time code generation (LTCG). Implied when PGO is enabled
+		/// Whether to allow the use of link time code generation (LTCG).
 		/// </summary>
 		[CommandLine("-LTCG")]
 		[XmlConfigFile(Category = "BuildConfiguration")]
-		public bool bAllowLTCG
-		{
-			get => bAllowLTCGPrivate || bPGOProfile || bPGOOptimize;
-			set => bAllowLTCGPrivate = value;
-		}
-		private bool bAllowLTCGPrivate;
+		public bool bAllowLTCG { get; set; }
 		
 		/// <summary>
 		/// When Link Time Code Generation (LTCG) is enabled, whether to 
