Fix incorrect Targets when using App Extensions and in Xcode compilation.
 - add a new json .uappex to allow extensions to specify required frameworks and other details
 - works with multiple extensions

diff --git a/Engine/Source/Programs/UnrealBuildTool/ProjectFiles/Xcode/XcodeProject.cs b/Engine/Source/Programs/UnrealBuildTool/ProjectFiles/Xcode/XcodeProject.cs
--- a/Engine/Source/Programs/UnrealBuildTool/ProjectFiles/Xcode/XcodeProject.cs
+++ b/Engine/Source/Programs/UnrealBuildTool/ProjectFiles/Xcode/XcodeProject.cs
@@ -5,6 +5,7 @@ using System.Collections.Generic;
 using System.IO;
 using System.Linq;
 using System.Text;
+using System.Text.Json;
 using EpicGames.Core;
 using Microsoft.CodeAnalysis;
 using Microsoft.Extensions.Logging;
@@ -16,7 +17,7 @@ Here's how this works:
 
   * An XcodeProjectFile (subclass of generic ProjectFile class) is created, along with it - UnrealData and XcodeFileCollection objects are made
   * High level code calls AddModuleForIntelliSense() which this code will use to cache information about the Modules in the project (including build settings, etc)
-    * These are used to determine what source files can be indexed together (we use native xcode code compilation for indexing, so we make compiling succesful for best index)
+    * These are used to determine what source files can be indexed together (we use native xcode code compilation for indexing, so we make compiling successful for best index)
     * A few #defines are removed or modified (FOO_API, etc) which would otherwise make every module a separate target
   * High level code then calls WriteProjectFile() which is the meat of all this
   * This code then creates a hierarchy/reference-chain of xcode project nodes (an xcode project is basically a series of Guid/object pairs that reference each other)
@@ -1184,7 +1185,7 @@ namespace UnrealBuildTool.XcodeProjectXcconfig
 		public string Name;
 
 		// QAGame, QAGameEditor, etc
-		private string UnrealTargetName;
+		protected string UnrealTargetName;
 
 		// list of build configs this target supports (for instance, the Index target only indexes a Development config) 
 		public XcodeBuildConfigList? BuildConfigList;
@@ -1343,6 +1344,18 @@ namespace UnrealBuildTool.XcodeProjectXcconfig
 		}
 	}
 
+	#region Read-only accessor properties
+#pragma warning disable CS1591
+	public class ExtensionConfig
+	{
+		public int FileVersion { get; set; }
+		public string? ExtensionType { get; set; }
+		public string? Description { get; set; }
+		public string[] Frameworks { get; set; } = Array.Empty<string>();
+	}
+#pragma warning restore CS1591
+	#endregion
+	
 	class XcodeRunTarget : XcodeTarget
 	{
 		private string ProductGuid = XcodeProjectFileGenerator.MakeXcodeGuid();
@@ -1386,12 +1399,32 @@ namespace UnrealBuildTool.XcodeProjectXcconfig
 				if (Project.UnrealData.UProjectFileLocation != null)
 				{
 					DirectoryReference ProjectDirectory = Project.UnrealData.ProductDirectory;
-					string ProjectExtensionsDir = Path.Combine(ProjectDirectory.FullName, "Build/IOS/Extensions");
+					string ProjectExtensionsDir = Path.Combine(ProjectDirectory.FullName, "Build", "IOS", "Extensions");
 					if (Directory.Exists(ProjectExtensionsDir))
 					{
 						foreach (DirectoryInfo DI in new System.IO.DirectoryInfo(ProjectExtensionsDir).EnumerateDirectories())
 						{
 							Logger.LogInformation("Creating App Extension Target for '{Name}' in '{FullName}'", DI.Name, DI.FullName);
+							
+							// Get the .uappex details file for this extension
+							ExtensionConfig? ExtConfig;
+							string UExtPath = Path.Combine(DI.FullName, DI.Name) + ".uappex";
+							if (File.Exists(UExtPath))
+							{
+								string JsonText = File.ReadAllText(UExtPath);
+								var Options = new JsonSerializerOptions{ PropertyNameCaseInsensitive = true };
+								ExtConfig = JsonSerializer.Deserialize<ExtensionConfig>(JsonText, Options);
+								if (ExtConfig == null)
+								{
+									Logger.LogWarning("Error loading extension '{Name}.uappex' file. Skipping Extension...", DI.Name);
+									continue;
+								}
+							}
+							else
+							{
+								Logger.LogWarning("There is no extension '{Name}.uappex' file. Skipping Extension...", DI.Name);
+								continue;
+							}
 
 							// assume each Extension in here will create a resulting Extension.appex
 							string Extension = DI.Name + ".appex";
@@ -1400,14 +1433,14 @@ namespace UnrealBuildTool.XcodeProjectXcconfig
 							// is it an extension we support? Grab the Plist file for it and check
 							// Also, only some AppExtensions need to embed
 							String AppExtPlist = DI.FullName + "/Info.plist";
-							PListHelper plistDoc = PListHelper.CreateFromFile(AppExtPlist);
+							PListHelper PlistDoc = PListHelper.CreateFromFile(AppExtPlist);
 							
 							XcodeCopyAppExtFilesBuildPhase EmbedAppExt = new XcodeCopyAppExtFilesBuildPhase(Project.FileCollection);
 							
 							// Get the source files
 							foreach (FileSystemInfo FSI in DI.EnumerateFileSystemInfos())
 							{
-								if (FSI.Name.StartsWith(".") || FSI.Name.EndsWith(".plist"))
+								if (FSI.Name.StartsWith(".") || FSI.Name.EndsWith(".uappex"))
 								{
 									continue;
 								}
@@ -1417,9 +1450,7 @@ namespace UnrealBuildTool.XcodeProjectXcconfig
 								XcodeSourceFile SourceFile = new XcodeSourceFile(FileRef, FileRef.Directory, XcodeProjectFileGenerator.MakeXcodeGuid());
 								Project.FileCollection.ProcessFile(SourceFile, true, false, "Source");
 							}
-							
-							Project.FileCollection.AddFileReference(ExtensionGuid, Extension, "explicitFileType", "wrapper.app-extension", "BUILT_PRODUCTS_DIR", "Products");
-							
+
 							DirectoryReference FinalAppExtDir = DirectoryReference.Combine(ProjectDirectory, DI.Name);
 							EmbedAppExt.AddExtension(FinalAppExtDir, ExtensionGuid);
 							
@@ -1427,7 +1458,7 @@ namespace UnrealBuildTool.XcodeProjectXcconfig
 							References.Add(EmbedAppExt);
 							
 							// Also need a (Native)Run Target for this extension
-							XcodeAppExtRunTarget RunTarget = new XcodeAppExtRunTarget(Project, DI.Name, ExtensionGuid, UnrealData.AllConfigs[0].ProjectTarget!.TargetRules!.Type, UnrealTargetPlatform.IOS, BuildTarget, ProjectFile, DI, Logger);
+							XcodeAppExtRunTarget RunTarget = new XcodeAppExtRunTarget(Project, ExtConfig, DI.Name, ExtensionGuid, UnrealData.AllConfigs[0].ProjectTarget!.TargetRules!.Type, UnrealTargetPlatform.IOS, BuildTarget, ProjectFile, DI, Logger);
 							Project.AppExtRunTargets.Add(RunTarget);
 							Project.References.Add(RunTarget);
 						}
@@ -1680,18 +1711,18 @@ namespace UnrealBuildTool.XcodeProjectXcconfig
 		private void ProcessAssets(XcodeResourcesBuildPhase ResourcesBuildPhase)
 		{
 			List<string> StoryboardPaths = new List<string>()
-				{
-					"$(Project)/Build/$(Platform)/Resources/Interface/LaunchScreen.storyboardc",
-					"$(Project)/Build/$(Platform)/Resources/Interface/LaunchScreen.storyboard",
-					"$(Project)/Build/Apple/Resources/Interface/LaunchScreen.storyboardc",
-					"$(Project)/Build/Apple/Resources/Interface/LaunchScreen.storyboard",
-					"$(Engine)/Build/$(Platform)/Resources/Interface/LaunchScreen.storyboardc",
-					"$(Engine)/Build/$(Platform)/Resources/Interface/LaunchScreen.storyboard",
-					"$(Engine)/Build/Apple/Resources/Interface/LaunchScreen.storyboardc",
-					"$(Engine)/Build/Apple/Resources/Interface/LaunchScreen.storyboard",
-					"$(Project)/Build/IOS/Resources/Interface/LaunchScreen.storyboard",
-					"$(Engine)/Build/IOS/Resources/Interface/LaunchScreen.storyboardc",
-				};
+			{
+				"$(Project)/Build/$(Platform)/Resources/Interface/LaunchScreen.storyboardc",
+				"$(Project)/Build/$(Platform)/Resources/Interface/LaunchScreen.storyboard",
+				"$(Project)/Build/Apple/Resources/Interface/LaunchScreen.storyboardc",
+				"$(Project)/Build/Apple/Resources/Interface/LaunchScreen.storyboard",
+				"$(Engine)/Build/$(Platform)/Resources/Interface/LaunchScreen.storyboardc",
+				"$(Engine)/Build/$(Platform)/Resources/Interface/LaunchScreen.storyboard",
+				"$(Engine)/Build/Apple/Resources/Interface/LaunchScreen.storyboardc",
+				"$(Engine)/Build/Apple/Resources/Interface/LaunchScreen.storyboard",
+				"$(Project)/Build/IOS/Resources/Interface/LaunchScreen.storyboard",
+				"$(Engine)/Build/IOS/Resources/Interface/LaunchScreen.storyboardc",
+			};
 
 			// look for Assets (in normal place, or an alternate for Programs)
 			string? StoryboardPath;
@@ -1700,8 +1731,8 @@ namespace UnrealBuildTool.XcodeProjectXcconfig
 
 			//default to IOS path for other platforms (eg VisionOS)
 			if (Platform != UnrealTargetPlatform.TVOS &&
-			Platform != UnrealTargetPlatform.IOS &&
-			Platform != UnrealTargetPlatform.Mac)
+			    Platform != UnrealTargetPlatform.IOS &&
+			    Platform != UnrealTargetPlatform.Mac)
 			{
 				AssetsSubPath = $"Build/IOS/Resources/Assets.xcassets";
 			}
@@ -2165,8 +2196,6 @@ namespace UnrealBuildTool.XcodeProjectXcconfig
 			Xcconfig.Write();
 
 			// Now for each config write out the specific settings
-			DirectoryReference? GameDir = UnrealData.UProjectFileLocation?.Directory;
-			string? GamePath = GameDir != null ? XcodeFileCollection.ConvertPath(GameDir.FullName) : null;
 			foreach (UnrealBuildConfig Config in UnrealData.AllConfigs)
 			{
 				XcodeBuildConfig? MatchedConfig = BuildConfigList!.BuildConfigs.FirstOrDefault(x => x.Info == Config);
@@ -2274,20 +2303,22 @@ namespace UnrealBuildTool.XcodeProjectXcconfig
 		private string ProductGuid = XcodeProjectFileGenerator.MakeXcodeGuid();
 		private UnrealData UnrealData;
 		private DirectoryInfo ExtensionPath;
-		public UnrealTargetPlatform Platform;
-		public TargetType TargetType;
-		public string TargetName;
-		public string AppGuid;
+		private UnrealTargetPlatform Platform;
+		private TargetType TargetType;
+		private string TargetName;
+		private string AppGuid;
+		private ExtensionConfig ExtConfig;
 		
-		public XcodeAppExtRunTarget(XcodeProject Project, string TargetName, string AppGuid, TargetType TargetType, UnrealTargetPlatform Platform, XcodeBuildTarget? BuildTarget, XcodeProjectFile ProjectFile, DirectoryInfo ExtensionPath, ILogger Logger)
+		public XcodeAppExtRunTarget(XcodeProject Project, ExtensionConfig? ExtConfig, string TargetName, string AppGuid, TargetType TargetType, UnrealTargetPlatform Platform, XcodeBuildTarget? BuildTarget, XcodeProjectFile ProjectFile, DirectoryInfo ExtensionPath, ILogger Logger)
 			: base(XcodeTarget.Type.App_Ext, Project.UnrealData,
-				  TargetName + (XcodeProjectFileGenerator.PerPlatformMode == XcodePerPlatformMode.OneWorkspacePerPlatform ? "" : $"_{Platform}"))
+				TargetName + (XcodeProjectFileGenerator.PerPlatformMode == XcodePerPlatformMode.OneWorkspacePerPlatform ? "" : $"_{Platform}"))
 		{
 			this.TargetType = TargetType;
 			this.Platform = Platform;
 			this.AppGuid = AppGuid;
 			this.TargetName = TargetName;
 			this.ExtensionPath = ExtensionPath;
+			this.ExtConfig = ExtConfig ?? new ExtensionConfig();
 			UnrealData = Project.UnrealData;
 
 			BuildConfigList = new XcodeBuildConfigList(Platform, Name, Project.UnrealData, bIncludeAllPlatformsInConfig: false);
@@ -2297,7 +2328,7 @@ namespace UnrealBuildTool.XcodeProjectXcconfig
 			{
 				XcodeFileCollection AppExtFileList = new XcodeFileCollection(ProjectFile);	
 
-				// add all of the files to be natively compiled by this target
+				// add all the files to be natively compiled by this target
 				XcodeSourcesBuildPhase SourcesBuildPhase = new XcodeSourcesBuildPhase();
 				BuildPhases.Add(SourcesBuildPhase);
 				References.Add(SourcesBuildPhase);
@@ -2311,18 +2342,12 @@ namespace UnrealBuildTool.XcodeProjectXcconfig
 
 					string FileRef = XcodeProjectFileGenerator.MakeXcodeGuid();
 					FileReference FR = new FileReference(FSI.FullName);
-					XcodeSourceFile xsf = new XcodeSourceFile(FR, FR.Directory, FileRef);
-					SourcesBuildPhase.AddFile(xsf);
-					Project.FileCollection.AllFiles.Add(xsf);
-					Project.FileCollection.BuildableFilesToResponseFile[xsf] = new();
-				}
-				
-				foreach (KeyValuePair<XcodeSourceFile, Dictionary<UnrealTargetPlatform, FileReference>> Pair in Project.FileCollection.BuildableFilesToResponseFile)
-				{
-					// only add files that found a module to be part of (since we can't build without the build settings that come from a module)
-					if (Pair.Value.ContainsKey(Platform))
+					XcodeSourceFile XSourceFile = new XcodeSourceFile(FR, FR.Directory, FileRef);
+					SourcesBuildPhase.AddFile(XSourceFile);
+					if (!Project.FileCollection.AllFiles.Contains(XSourceFile))
 					{
-						SourcesBuildPhase.AddFile(Pair.Key);
+						Project.FileCollection.AllFiles.Add(XSourceFile);
+						Project.FileCollection.BuildableFilesToResponseFile[XSourceFile] = new();
 					}
 				}
 				
@@ -2331,7 +2356,6 @@ namespace UnrealBuildTool.XcodeProjectXcconfig
 				References.Add(ResourcesBuildPhase);
 
 				ProcessFrameworks(ResourcesBuildPhase, Project, ProjectFile, Logger);
-				ProcessAssets(ResourcesBuildPhase);
 			}
 
 			if (BuildTarget != null)
@@ -2347,19 +2371,6 @@ namespace UnrealBuildTool.XcodeProjectXcconfig
 			}
 		}
 
-		private void ProcessAssets(XcodeResourcesBuildPhase ResourcesBuildPhase)
-		{
-			// Get the source files
-			foreach (FileSystemInfo FSI in ExtensionPath.EnumerateFileSystemInfos())
-			{
-				if (FSI.Name.StartsWith(".") || FSI.Name.EndsWith(".plist"))
-				{
-					continue;
-				}
-				ResourcesBuildPhase.AddResource(new FileReference(FSI.FullName)); 
-			}
-		}
-		
 		protected void ProcessFrameworks(XcodeResourcesBuildPhase ResourcesBuildPhase, XcodeProject Project, XcodeProjectFile ProjectFile, ILogger Logger)
 		{
 			// look up to see if we had cached any Frameworks
@@ -2376,8 +2387,6 @@ namespace UnrealBuildTool.XcodeProjectXcconfig
 			IEnumerable<UEBuildFramework>? Frameworks;
 			if (XcodeProjectFileGenerator.TargetFrameworks.TryGetValue(FrameworkKey, out Frameworks))
 			{
-				XcodeFrameworkBuildPhase FrameworkPhase = new XcodeFrameworkBuildPhase(Project.FileCollection);
-
 				// filter frameworks that need to installed into the .app (either the framework or a bundle inside a .zip)
 				IEnumerable<UEBuildFramework> InstalledFrameworks = Frameworks.Where(x => x.bCopyFramework || !String.IsNullOrEmpty(x.CopyBundledAssets));
 
@@ -2398,14 +2407,19 @@ namespace UnrealBuildTool.XcodeProjectXcconfig
 						ResourcesBuildPhase.AddFolderResource(DirectoryReference.Combine(BundleRootDir, Framework.CopyBundledAssets), "Resources");
 					}
 				}
-
-				DirectoryReference userNotifFw = new DirectoryReference("UserNotifications.framework");
+			}
+			
+			XcodeFrameworkBuildPhase FrameworkPhase = new XcodeFrameworkBuildPhase(Project.FileCollection);
+			foreach (string Framework in ExtConfig.Frameworks)
+			{
+				DirectoryReference FrameworkRef = new DirectoryReference(Framework + ".framework");
 				string FileRefGuid = XcodeProjectFileGenerator.MakeXcodeGuid();
-				FrameworkPhase.AddFramework(userNotifFw, FileRefGuid);
 
-				BuildPhases.Add(FrameworkPhase);
-				References.Add(FrameworkPhase);
+				FrameworkPhase.AddFramework(FrameworkRef, FileRefGuid);
 			}
+
+			BuildPhases.Add(FrameworkPhase);
+			References.Add(FrameworkPhase);
 		}
 		
 		public override void WriteXcconfigFile(ILogger Logger)
@@ -2567,8 +2581,6 @@ namespace UnrealBuildTool.XcodeProjectXcconfig
 			Xcconfig.AppendLine($"UE_XCODE_BUILD_MODE = Default");
 			Xcconfig.AppendLine($"UE_PRODUCT_NAME = {Name}");
 			Xcconfig.AppendLine($"UE_PRODUCT_NAME_STRIPPED = {Name.Replace("_", "").Replace(" ", "")}");
-			// Xcconfig.AppendLine($"UE_PRODUCT_NAME = {UnrealData.ProductName}");
-			// Xcconfig.AppendLine($"UE_PRODUCT_NAME_STRIPPED = {UnrealData.ProductName.Replace("_", "").Replace(" ", "")}");
 			Xcconfig.AppendLine($"UE_DISPLAY_NAME = {UnrealData.DisplayName}");
 			Xcconfig.AppendLine($"UE_SIGNING_PREFIX = {SigningPrefix}");
 			Xcconfig.AppendLine($"UE_PLATFORM_NAME = {Platform}");
@@ -2596,7 +2608,7 @@ namespace UnrealBuildTool.XcodeProjectXcconfig
 			Xcconfig.AppendLine($"SDKROOT = {SDKRoot}");
 
 			// Xcode creates the Build Dir (where the .app is) by combining {SYMROOT}/{CONFIGURATION}{EFFECTIVE_PLATFORM_NAME}, so we set SYMROOT
-			// to the Parent directory of the diectory the binary is in, CONFIGURATION to nothing, and EFFECTIVE_PLATFORM_NAME to the directory
+			// to the Parent directory of the directory the binary is in, CONFIGURATION to nothing, and EFFECTIVE_PLATFORM_NAME to the directory
 			// the binary is in
 			Xcconfig.AppendLine("");
 			Xcconfig.AppendLine($"// These settings combined will tell Xcode to write to Binaries/{Platform} (instead of something like Binaries/Development-iphoneos)");
@@ -2617,13 +2629,15 @@ namespace UnrealBuildTool.XcodeProjectXcconfig
 			Xcconfig.AppendLine("");
 			Xcconfig.AppendLine("// Project settings");
 			Xcconfig.AppendLine($"TARGETED_DEVICE_FAMILY = {SupportedDevices}");
-			Xcconfig.AppendLine($"PRODUCT_BUNDLE_IDENTIFIER = {BundleIdentifier}.{Name}");
+			Xcconfig.AppendLine($"PRODUCT_BUNDLE_IDENTIFIER = $(UE_SIGNING_PREFIX).{UnrealTargetName}.{Name}");
 			Xcconfig.AppendLine($"{DeploymentTargetKey} = {DeploymentTarget}");
 
 			Xcconfig.AppendLine("");
 			Xcconfig.AppendLine("// Plist settings");
 			// Xcconfig.AppendLine($"INFOPLIST_FILE = {PlistMetadata.XcodeProjectRelative}");
-			Xcconfig.AppendLine($"INFOPLIST_FILE = ../../Build/IOS/Extensions/{Name}/Info.plist");
+			
+			DirectoryReference GameDir = UnrealData.UProjectFileLocation?.Directory ?? new DirectoryReference("./");
+			Xcconfig.AppendLine($"INFOPLIST_FILE = {GameDir}/Build/IOS/Extensions/{Name}/Info.plist");
 			Xcconfig.AppendLine($"INFOPLIST_KEY_CFBundleDisplayName = {Name}");
 			if (PlistMetadata.Mode == MetadataMode.UpdateTemplate)
 			{
@@ -2676,8 +2690,6 @@ namespace UnrealBuildTool.XcodeProjectXcconfig
 			Xcconfig.Write();
 
 			// Now for each config write out the specific settings
-			DirectoryReference? GameDir = UnrealData.UProjectFileLocation?.Directory;
-			string? GamePath = GameDir != null ? XcodeFileCollection.ConvertPath(GameDir.FullName) : null;
 			foreach (UnrealBuildConfig Config in UnrealData.AllConfigs)
 			{
 				XcodeBuildConfig? MatchedConfig = BuildConfigList!.BuildConfigs.FirstOrDefault(x => x.Info == Config);
