Iris - Iris support in the engine is now always enabled. It's still optional to use it.

diff --git a/Engine/Source/Programs/ReplicationSystemLowLevelTests/ReplicationSystemLowLevelTests.Target.cs b/Engine/Source/Programs/ReplicationSystemLowLevelTests/ReplicationSystemLowLevelTests.Target.cs
--- a/Engine/Source/Programs/ReplicationSystemLowLevelTests/ReplicationSystemLowLevelTests.Target.cs
+++ b/Engine/Source/Programs/ReplicationSystemLowLevelTests/ReplicationSystemLowLevelTests.Target.cs
@@ -18,7 +18,6 @@ public class ReplicationSystemLowLevelTestsTarget : TestTargetRules
 
 		// Network config
 		bWithPushModel = true;
-		bUseIris = true;
 
 		bEnableTrace = true;
 		GlobalDefinitions.Add("UE_NET_TEST_FAKE_REP_TAGS=1");
diff --git a/Engine/Source/Programs/ReplicationSystemTest/ReplicationSystemTest.Target.cs b/Engine/Source/Programs/ReplicationSystemTest/ReplicationSystemTest.Target.cs
--- a/Engine/Source/Programs/ReplicationSystemTest/ReplicationSystemTest.Target.cs
+++ b/Engine/Source/Programs/ReplicationSystemTest/ReplicationSystemTest.Target.cs
@@ -37,7 +37,6 @@ public class ReplicationSystemTestTarget : TargetRules
 
 		// Network config
 		bWithPushModel = true;		
-		bUseIris = true;
 
 		bEnableTrace = true;
 		GlobalDefinitions.Add("UE_NET_TEST_FAKE_REP_TAGS=1");
diff --git a/Engine/Source/Programs/UnrealBuildTool/Configuration/ModuleRules.cs b/Engine/Source/Programs/UnrealBuildTool/Configuration/ModuleRules.cs
--- a/Engine/Source/Programs/UnrealBuildTool/Configuration/ModuleRules.cs
+++ b/Engine/Source/Programs/UnrealBuildTool/Configuration/ModuleRules.cs
@@ -1619,33 +1619,18 @@ namespace UnrealBuildTool
 		/// </summary>
 		public void SetupIrisSupport(ReadOnlyTargetRules target, bool bAddAsPublicDependency = false)
 		{
-			if (target.bUseIris == true)
-			{
-				PublicDefinitions.Add("UE_WITH_IRIS=1");
-				PrivateDefinitions.Add("UE_NET_HAS_IRIS_FASTARRAY_BINDING=1");
+			PrivateDefinitions.Add("UE_NET_HAS_IRIS_FASTARRAY_BINDING=1");
 
-				if (bAddAsPublicDependency)
-				{
-					PublicDependencyModuleNames.Add("IrisCore");
-				}
-				else
-				{
-					PrivateDependencyModuleNames.Add("IrisCore");
-				}
+			// To be removed once UE_WITH_IRIS has been removed in code.
+			PublicDefinitions.Add("UE_WITH_IRIS=1");
+
+			if (bAddAsPublicDependency)
+			{
+				PublicDependencyModuleNames.Add("IrisCore");
 			}
 			else
 			{
-				PublicDefinitions.Add("UE_WITH_IRIS=0");
-
-				// If we compile without Iris we have a stub Iris module for UHT dependencies
-				if (bAddAsPublicDependency)
-				{
-					PublicDependencyModuleNames.Add("IrisStub");
-				}
-				else
-				{
-					PrivateDependencyModuleNames.Add("IrisStub");
-				}
+				PrivateDependencyModuleNames.Add("IrisCore");
 			}
 		}
 
diff --git a/Engine/Source/Programs/UnrealBuildTool/Configuration/ReadOnlyTargetRules.cs b/Engine/Source/Programs/UnrealBuildTool/Configuration/ReadOnlyTargetRules.cs
--- a/Engine/Source/Programs/UnrealBuildTool/Configuration/ReadOnlyTargetRules.cs
+++ b/Engine/Source/Programs/UnrealBuildTool/Configuration/ReadOnlyTargetRules.cs
@@ -157,8 +157,6 @@ namespace UnrealBuildTool
 
 		public bool bCompileIoStoreOnDemand => Inner.bCompileIoStoreOnDemand;
 
-		public bool bUseIris => Inner.bUseIris;
-
 		public bool bTrackRHIResourceInfoForTest => Inner.bTrackRHIResourceInfoForTest;
 
 		public bool bCompilePython => Inner.bCompilePython;
diff --git a/Engine/Source/Programs/UnrealBuildTool/Configuration/TargetRules.cs b/Engine/Source/Programs/UnrealBuildTool/Configuration/TargetRules.cs
--- a/Engine/Source/Programs/UnrealBuildTool/Configuration/TargetRules.cs
+++ b/Engine/Source/Programs/UnrealBuildTool/Configuration/TargetRules.cs
@@ -1077,14 +1077,6 @@ namespace UnrealBuildTool
 		[CommandLine("-CompileIoStoreOnDemand", Value = "true")]
 		public bool bCompileIoStoreOnDemand { get; set; } = false;
 
-		/// <summary>
-		/// Whether to use Iris.
-		/// </summary>
-		[RequiresUniqueBuildEnvironment]
-		[CommandLine("-NoUseIris", Value = "false")]
-		[CommandLine("-UseIris", Value = "true")]
-		public bool bUseIris { get; set; } = true;
-
 		/// <summary>
 		/// Whether to track owner (asset name) of RHI resource for Test configuration.
 		/// Useful for ListShaderMaps and ListShaderLibraries commands.
diff --git a/Engine/Source/Runtime/Core/Core.Build.cs b/Engine/Source/Runtime/Core/Core.Build.cs
--- a/Engine/Source/Runtime/Core/Core.Build.cs
+++ b/Engine/Source/Runtime/Core/Core.Build.cs
@@ -269,14 +269,8 @@ public class Core : ModuleRules
 		}
 
 		// Setup definitions to include / exclude Iris modifications to UObject Note: Only the definition is required as we do not depend on Iris in any way.
-		if (Target.bUseIris == true)
-		{
-			PublicDefinitions.Add("UE_WITH_IRIS=1");
-		}
-		else
-		{
-			PublicDefinitions.Add("UE_WITH_IRIS=0");
-		}
+		// To be removed once UE_WITH_IRIS has been removed in code.
+		PublicDefinitions.Add("UE_WITH_IRIS=1");
 		
 		if (Target.Platform == UnrealTargetPlatform.Win64
 			&& Target.Configuration != UnrealTargetConfiguration.Shipping)
diff --git a/Engine/Source/Runtime/Experimental/Iris/Core/IrisCore.Build.cs b/Engine/Source/Runtime/Experimental/Iris/Core/IrisCore.Build.cs
--- a/Engine/Source/Runtime/Experimental/Iris/Core/IrisCore.Build.cs
+++ b/Engine/Source/Runtime/Experimental/Iris/Core/IrisCore.Build.cs
@@ -16,20 +16,8 @@ namespace UnrealBuildTool.Rules
 				}
 			);
 
-			if (Target.bUseIris == true)
-			{
-				PublicDefinitions.Add("UE_WITH_IRIS=1");
-			}
-			else
-			{
-				PublicDefinitions.Add("UE_WITH_IRIS=0");
-
-				// This module should not be precompiled if Iris is disabled
-				if (Target.bBuildAllModules == true)
-				{
-					PrecompileForTargets = PrecompileTargetsType.None;
-				}
-			}
+			// To be removed once UE_WITH_IRIS has been removed in code.
+			PublicDefinitions.Add("UE_WITH_IRIS=1");
 
 			CppCompileWarningSettings.UnsafeTypeCastWarningLevel = WarningLevel.Error;
 		}
diff --git a/Engine/Source/Runtime/Experimental/Iris/Stub/IrisStub.Build.cs b/Engine/Source/Runtime/Experimental/Iris/Stub/IrisStub.Build.cs
deleted file mode 100644
diff --git a/Engine/Source/Runtime/Experimental/Iris/Stub/Private/Iris/IrisStubModule.cpp b/Engine/Source/Runtime/Experimental/Iris/Stub/Private/Iris/IrisStubModule.cpp
deleted file mode 100644
diff --git a/Engine/Source/Runtime/Experimental/Iris/Stub/Public/Iris/DataStream/DataStream.h b/Engine/Source/Runtime/Experimental/Iris/Stub/Public/Iris/DataStream/DataStream.h
deleted file mode 100644
diff --git a/Engine/Source/Runtime/Experimental/Iris/Stub/Public/Iris/ReplicationSystem/ChunkedDataStream/ChunkedDataStream.h b/Engine/Source/Runtime/Experimental/Iris/Stub/Public/Iris/ReplicationSystem/ChunkedDataStream/ChunkedDataStream.h
deleted file mode 100644
diff --git a/Engine/Source/Runtime/Experimental/Iris/Stub/Public/Iris/ReplicationSystem/Filtering/NetObjectConnectionFilter.h b/Engine/Source/Runtime/Experimental/Iris/Stub/Public/Iris/ReplicationSystem/Filtering/NetObjectConnectionFilter.h
deleted file mode 100644
diff --git a/Engine/Source/Runtime/Experimental/Iris/Stub/Public/Iris/ReplicationSystem/Filtering/NetObjectFilter.h b/Engine/Source/Runtime/Experimental/Iris/Stub/Public/Iris/ReplicationSystem/Filtering/NetObjectFilter.h
deleted file mode 100644
diff --git a/Engine/Source/Runtime/Experimental/Iris/Stub/Public/Iris/ReplicationSystem/Filtering/NetObjectGridFilter.h b/Engine/Source/Runtime/Experimental/Iris/Stub/Public/Iris/ReplicationSystem/Filtering/NetObjectGridFilter.h
deleted file mode 100644
diff --git a/Engine/Source/Runtime/Experimental/Iris/Stub/Public/Iris/ReplicationSystem/NetBlob/NetBlobHandler.h b/Engine/Source/Runtime/Experimental/Iris/Stub/Public/Iris/ReplicationSystem/NetBlob/NetBlobHandler.h
deleted file mode 100644
diff --git a/Engine/Source/Runtime/Experimental/Iris/Stub/Public/Iris/ReplicationSystem/NetBlob/SequentialPartialNetBlobHandler.h b/Engine/Source/Runtime/Experimental/Iris/Stub/Public/Iris/ReplicationSystem/NetBlob/SequentialPartialNetBlobHandler.h
deleted file mode 100644
diff --git a/Engine/Source/Runtime/Experimental/Iris/Stub/Public/Iris/ReplicationSystem/NetObjectFactory.h b/Engine/Source/Runtime/Experimental/Iris/Stub/Public/Iris/ReplicationSystem/NetObjectFactory.h
deleted file mode 100644
diff --git a/Engine/Source/Runtime/Experimental/Iris/Stub/Public/Iris/ReplicationSystem/ObjectReplicationBridge.h b/Engine/Source/Runtime/Experimental/Iris/Stub/Public/Iris/ReplicationSystem/ObjectReplicationBridge.h
deleted file mode 100644
diff --git a/Engine/Source/Runtime/Experimental/Iris/Stub/Public/Iris/ReplicationSystem/Prioritization/NetObjectCountLimiter.h b/Engine/Source/Runtime/Experimental/Iris/Stub/Public/Iris/ReplicationSystem/Prioritization/NetObjectCountLimiter.h
deleted file mode 100644
diff --git a/Engine/Source/Runtime/Experimental/Iris/Stub/Public/Iris/ReplicationSystem/Prioritization/NetObjectPrioritizer.h b/Engine/Source/Runtime/Experimental/Iris/Stub/Public/Iris/ReplicationSystem/Prioritization/NetObjectPrioritizer.h
deleted file mode 100644
diff --git a/Engine/Source/Runtime/Experimental/Iris/Stub/Public/Iris/ReplicationSystem/Prioritization/SphereNetObjectPrioritizer.h b/Engine/Source/Runtime/Experimental/Iris/Stub/Public/Iris/ReplicationSystem/Prioritization/SphereNetObjectPrioritizer.h
deleted file mode 100644
diff --git a/Engine/Source/Runtime/Experimental/Iris/Stub/Public/Iris/ReplicationSystem/Prioritization/SphereWithOwnerBoostNetObjectPrioritizer.h b/Engine/Source/Runtime/Experimental/Iris/Stub/Public/Iris/ReplicationSystem/Prioritization/SphereWithOwnerBoostNetObjectPrioritizer.h
deleted file mode 100644
diff --git a/Engine/Source/Runtime/Experimental/Iris/Stub/Public/Iris/ReplicationSystem/ReplicationBridge.h b/Engine/Source/Runtime/Experimental/Iris/Stub/Public/Iris/ReplicationSystem/ReplicationBridge.h
deleted file mode 100644
diff --git a/Engine/Source/Runtime/Experimental/Iris/Stub/Public/Iris/Serialization/IrisObjectReferencePackageMap.h b/Engine/Source/Runtime/Experimental/Iris/Stub/Public/Iris/Serialization/IrisObjectReferencePackageMap.h
deleted file mode 100644
diff --git a/Engine/Source/Runtime/Experimental/Iris/Stub/Public/Iris/Serialization/NetSerializer.h b/Engine/Source/Runtime/Experimental/Iris/Stub/Public/Iris/Serialization/NetSerializer.h
deleted file mode 100644
diff --git a/Engine/Source/Runtime/Experimental/Iris/Stub/Public/Iris/Serialization/PolymorphicNetSerializer.h b/Engine/Source/Runtime/Experimental/Iris/Stub/Public/Iris/Serialization/PolymorphicNetSerializer.h
deleted file mode 100644
diff --git a/Engine/Source/Runtime/Experimental/Iris/Stub/Readme.txt b/Engine/Source/Runtime/Experimental/Iris/Stub/Readme.txt
deleted file mode 100644
