PR #13708: Force synchronous writes on type library input files

diff --git a/Engine/Source/Programs/UnrealBuildTool/Platform/Windows/VCToolChain.cs b/Engine/Source/Programs/UnrealBuildTool/Platform/Windows/VCToolChain.cs
--- a/Engine/Source/Programs/UnrealBuildTool/Platform/Windows/VCToolChain.cs
+++ b/Engine/Source/Programs/UnrealBuildTool/Platform/Windows/VCToolChain.cs
@@ -2834,7 +2834,7 @@ namespace UnrealBuildTool
 			}
 			Contents.AppendLine();
 
-			FileItem InputFile = Graph.CreateIntermediateTextFile(OutputFile.ChangeExtension(".cpp"), Contents.ToString());
+			FileItem InputFile = Graph.CreateIntermediateTextFile(OutputFile.ChangeExtension(".cpp"), Contents.ToString(), false);
 
 			// Build the argument list for the compiler
 			FileItem ObjectFile = FileItem.GetItemByFileReference(OutputFile.ChangeExtension(".obj"));
@@ -2849,7 +2849,7 @@ namespace UnrealBuildTool
 
 			AddIncludePathArguments(CompileEnvironment, EnvVars.ToolChain, Arguments);
 
-			FileItem ResponseFile = Graph.CreateIntermediateTextFile(OutputFile.ChangeExtension(ResponseExt), Arguments);
+			FileItem ResponseFile = Graph.CreateIntermediateTextFile(OutputFile.ChangeExtension(ResponseExt), Arguments, false);
 
 			// Build the command for touching the output file(s) to update the time stamp
 			string OutfileTlhFullName = NormalizeCommandLinePath(OutputFile, RootPaths).Replace('/', '\\'); // Must be backward slash
@@ -2864,7 +2864,7 @@ namespace UnrealBuildTool
 			BatchFileContents.AppendLine(TouchTlhActionCommand);
 			BatchFileContents.AppendLine(TouchTliActionCommand);
 
-			FileItem BatchFile = Graph.CreateIntermediateTextFile(OutputFile.ChangeExtension(".bat"), BatchFileContents.ToString());
+			FileItem BatchFile = Graph.CreateIntermediateTextFile(OutputFile.ChangeExtension(".bat"), BatchFileContents.ToString(), false);
 
 			// Create the batch file action that will compile then touch the output file
 			Action CompileAction = Graph.CreateAction(ActionType.Compile);
