UnrealBuildTool: Expose UEBuildPlatform.IsPlatformAvailable for use by platform-specific TargetRules during compilation.
#rnx

diff --git a/Engine/Source/Programs/UnrealBuildTool/Configuration/TargetRules.cs b/Engine/Source/Programs/UnrealBuildTool/Configuration/TargetRules.cs
--- a/Engine/Source/Programs/UnrealBuildTool/Configuration/TargetRules.cs
+++ b/Engine/Source/Programs/UnrealBuildTool/Configuration/TargetRules.cs
@@ -3572,5 +3572,13 @@ namespace UnrealBuildTool
 			GetBuildSettingsInfo(diagnostics);
 			diagnostics.ForEach(x => Logger.LogWarning("Warning: {Message}", x));
 		}
+
+		/// <summary>
+		/// Determines whether a given platform, and its sdk, is available
+		/// </summary>
+		public bool IsPlatformAvailable(UnrealTargetPlatform TargetPlatform)
+		{
+			return UEBuildPlatform.IsPlatformAvailable(TargetPlatform);
+		}
 	}
 }
