UnrealBuildTool: Fix SN-DBS remote task with RHI generated header

#rnx

diff --git a/Engine/Source/Programs/UnrealBuildTool/Executors/SNDBS.cs b/Engine/Source/Programs/UnrealBuildTool/Executors/SNDBS.cs
--- a/Engine/Source/Programs/UnrealBuildTool/Executors/SNDBS.cs
+++ b/Engine/Source/Programs/UnrealBuildTool/Executors/SNDBS.cs
@@ -442,14 +442,17 @@ namespace UnrealBuildTool
 					@"pattern1=^COMPILED_PLATFORM_HEADER\(\s*([^ ,]+)\s*\)",
 					$"expansions1={String.Join("|", Platforms.Select(Name => $"{Name}/{Name}$1|{Name}$1"))}",
 
-					@"pattern2=^COMPILED_PLATFORM_HEADER_WITH_PREFIX\(\s*([^ ,]+)\s*,\s*([^ ,]+)\s*\)",
-					$"expansions2={String.Join("|", Platforms.Select(Name => $"$1/{Name}/{Name}$2|$1/{Name}$2"))}",
+					@"pattern2=^COMPILED_PLATFORM_HEADER_GENERATED\(\s*([^ ,]+)\s*\)",
+					$"expansions2={String.Join("|", Platforms.Select(Name => $"{Name}$1"))}",
 
-					@"pattern3=^[A-Z]{5}_PLATFORM_HEADER_NAME\(\s*([^ ,]+)\s*\)",
-					$"expansions3={String.Join("|", Platforms.Select(Name => $"{Name}/{Name}$1|{Name}$1"))}",
+					@"pattern3=^COMPILED_PLATFORM_HEADER_WITH_PREFIX\(\s*([^ ,]+)\s*,\s*([^ ,]+)\s*\)",
+					$"expansions3={String.Join("|", Platforms.Select(Name => $"$1/{Name}/{Name}$2|$1/{Name}$2"))}",
 
-					@"pattern4=^[A-Z]{5}_PLATFORM_HEADER_NAME_WITH_PREFIX\(\s*([^ ,]+)\s*,\s*([^ ,]+)\s*\)",
-					$"expansions4={String.Join("|", Platforms.Select(Name => $"$1/{Name}/{Name}$2|$1/{Name}$2"))}"
+					@"pattern4=^[A-Z]{5}_PLATFORM_HEADER_NAME\(\s*([^ ,]+)\s*\)",
+					$"expansions4={String.Join("|", Platforms.Select(Name => $"{Name}/{Name}$1|{Name}$1"))}",
+
+					@"pattern5=^[A-Z]{5}_PLATFORM_HEADER_NAME_WITH_PREFIX\(\s*([^ ,]+)\s*,\s*([^ ,]+)\s*\)",
+					$"expansions5={String.Join("|", Platforms.Select(Name => $"$1/{Name}/{Name}$2|$1/{Name}$2"))}"
 				};
 
 				File.WriteAllText(IncludeRewriteRulesFile.FullName, String.Join(Environment.NewLine, new[] { "[computed-include-rules]" }.Concat(Lines)));
diff --git a/Engine/Source/Runtime/Core/Public/HAL/PreprocessorHelpers.h b/Engine/Source/Runtime/Core/Public/HAL/PreprocessorHelpers.h
--- a/Engine/Source/Runtime/Core/Public/HAL/PreprocessorHelpers.h
+++ b/Engine/Source/Runtime/Core/Public/HAL/PreprocessorHelpers.h
@@ -84,6 +84,10 @@
 #define COMPILED_PLATFORM_HEADER(Suffix) PREPROCESSOR_TO_STRING(PREPROCESSOR_JOIN(PLATFORM_HEADER_NAME/PLATFORM_HEADER_NAME, Suffix))
 #endif
 
+// Creates a string that can be used to include a header in the platform extension form "PlatformHeader.h", but will
+// not include a directory like COMPILED_PLATFORM_HEADER does, generally for UBT generated headers.
+#define COMPILED_PLATFORM_HEADER_GENERATED(Suffix) PREPROCESSOR_TO_STRING(PREPROCESSOR_JOIN(PLATFORM_HEADER_NAME, Suffix))
+
 #if PLATFORM_IS_EXTENSION
 // Creates a string that can be used to include a header with the platform in its name, like
 // "Prefix/PlatformNameSuffix.h". When using this you should add "// IWYU pragma: export" at the end of the line.
diff --git a/Engine/Source/Runtime/RHI/Public/RHIShaderPlatform.h b/Engine/Source/Runtime/RHI/Public/RHIShaderPlatform.h
--- a/Engine/Source/Runtime/RHI/Public/RHIShaderPlatform.h
+++ b/Engine/Source/Runtime/RHI/Public/RHIShaderPlatform.h
@@ -58,7 +58,7 @@ enum EShaderPlatform : uint16
 static_assert(SP_NumPlatforms <= (1 << SP_NumBits), "SP_NumPlatforms will not fit on SP_NumBits");
 
 #if USE_STATIC_SHADER_PLATFORM_ENUMS
-#include PREPROCESSOR_TO_STRING(PREPROCESSOR_JOIN(PLATFORM_HEADER_NAME, RHIShaderPlatform.h))
+#include COMPILED_PLATFORM_HEADER_GENERATED(RHIShaderPlatform.h)
 #include COMPILED_PLATFORM_HEADER(StaticShaderPlatform.inl)
 #else
 struct FStaticShaderPlatform final
