UnrealBuildTool: Properly track SDK external dependencies in makefile

#rnx

diff --git a/Engine/Source/Programs/Shared/EpicGames.Build/UEBuildPlatformSDK.cs b/Engine/Source/Programs/Shared/EpicGames.Build/UEBuildPlatformSDK.cs
--- a/Engine/Source/Programs/Shared/EpicGames.Build/UEBuildPlatformSDK.cs
+++ b/Engine/Source/Programs/Shared/EpicGames.Build/UEBuildPlatformSDK.cs
@@ -281,6 +281,8 @@ namespace EpicGames.Core
 
 		protected readonly ILogger Logger = InLogger;
 
+		protected readonly HashSet<FileReference> LoadedConfigFiles = [];
+
 		#region Global SDK Registration
 
 		/// <summary>
@@ -413,6 +415,9 @@ namespace EpicGames.Core
 		// True if this Sdk is allowed to be used by this host - if not, we can skip a lot 
 		public bool bIsSdkAllowedOnHost;
 
+		// External dependencies for this Sdk (config files)
+		public IEnumerable<FileItem> ExternalDependencies => LoadedConfigFiles.Select(FileItem.GetItemByFileReference);
+
 		public SDKCollection GetAllSDKInfo()
 		{
 			SDKCollection AllSdks = new SDKCollection(this);
@@ -1152,6 +1157,8 @@ namespace EpicGames.Core
 				throw new Exception($"Failed to parse SDK version file '{SDKConfigFile}'");
 			}
 
+			LoadedConfigFiles.Add(SDKConfigFile);
+
 			// look for special key to chain to other - after processing everything else (we only add from the "parent" settings not already in the map)
 			List<string> Parents = [];
 			foreach (string Key in LoadedDictionary.Keys)
diff --git a/Engine/Source/Programs/UnrealBuildTool/Platform/Windows/VCToolChain.cs b/Engine/Source/Programs/UnrealBuildTool/Platform/Windows/VCToolChain.cs
--- a/Engine/Source/Programs/UnrealBuildTool/Platform/Windows/VCToolChain.cs
+++ b/Engine/Source/Programs/UnrealBuildTool/Platform/Windows/VCToolChain.cs
@@ -195,6 +195,9 @@ namespace UnrealBuildTool
 		{
 			base.GetExternalDependencies(ExternalDependencies);
 
+			UEBuildPlatformSDK? SDK = UEBuildPlatformSDK.GetSDKForPlatform(UnrealTargetPlatform.Win64.ToString());
+			ExternalDependencies.UnionWith(SDK?.ExternalDependencies ?? []);
+
 			ExternalDependencies.Add(FileItem.GetItemByFileReference(EnvVars.CompilerPath));
 			ExternalDependencies.Add(FileItem.GetItemByFileReference(EnvVars.LinkerPath));
 			ExternalDependencies.Add(FileItem.GetItemByFileReference(EnvVars.ResourceCompilerPath));
@@ -3153,7 +3156,16 @@ namespace UnrealBuildTool
 				}
 				foreach (string SystemLibrary in LinkEnvironment.SystemLibraries)
 				{
-					InputFileNames.Add($"\"{SystemLibrary}\"");
+					if (Path.IsPathFullyQualified(SystemLibrary))
+					{
+						FileItem SystemLibraryItem = FileItem.GetItemByPath(SystemLibrary);
+						InputFileNames.Add($"\"{NormalizeCommandLinePath(SystemLibraryItem, RootPaths)}\"");
+						PrerequisiteItems.Add(SystemLibraryItem);
+					}
+					else
+					{
+						InputFileNames.Add($"\"{SystemLibrary}\"");
+					}
 				}
 
 				foreach (FileItem NatvisFile in LinkEnvironment.DebuggerVisualizerFiles)
diff --git a/Engine/Source/Programs/UnrealBuildTool/ToolChain/ISPCToolChain.cs b/Engine/Source/Programs/UnrealBuildTool/ToolChain/ISPCToolChain.cs
--- a/Engine/Source/Programs/UnrealBuildTool/ToolChain/ISPCToolChain.cs
+++ b/Engine/Source/Programs/UnrealBuildTool/ToolChain/ISPCToolChain.cs
@@ -46,17 +46,8 @@ namespace UnrealBuildTool
 
 			if (TargetRules != null)
 			{
-				string Platform = TargetRules.Platform == UnrealTargetPlatform.Win64 ? "Windows" : TargetRules.Platform.ToString();
-				FileItem PlatformConfig = FileItem.GetItemByFileReference(FileReference.Combine(Unreal.EngineDirectory, "Platforms", Platform, "Config", $"{Platform}_SDK.json"));
-				if (PlatformConfig.Exists)
-				{
-					ExternalDependencies.Add(PlatformConfig);
-				}
-				PlatformConfig = FileItem.GetItemByFileReference(FileReference.Combine(Unreal.EngineDirectory, "Config", Platform, $"{Platform}_SDK.json"));
-				if (PlatformConfig.Exists)
-				{
-					ExternalDependencies.Add(PlatformConfig);
-				}
+				UEBuildPlatformSDK? SDK = UEBuildPlatformSDK.GetSDKForPlatform(TargetRules.Platform.ToString());
+				ExternalDependencies.UnionWith(SDK?.ExternalDependencies ?? []);
 			}
 		}
 
