UnrealBuildTool: UBA PathHash for DIA SDK

#rnx

diff --git a/Engine/Source/Programs/UnrealBuildTool/Platform/Windows/VCToolChain.cs b/Engine/Source/Programs/UnrealBuildTool/Platform/Windows/VCToolChain.cs
--- a/Engine/Source/Programs/UnrealBuildTool/Platform/Windows/VCToolChain.cs
+++ b/Engine/Source/Programs/UnrealBuildTool/Platform/Windows/VCToolChain.cs
@@ -3765,6 +3765,11 @@ namespace UnrealBuildTool
 				EpicGames.UBA.Utils.RegisterPathHash(EnvVars.ToolChainDir.FullName, EnvVars.ToolChainVersion.ToString());
 			}
 
+			if (!String.IsNullOrEmpty(Target.WindowsPlatform.DiaSdkDir))
+			{
+				EpicGames.UBA.Utils.RegisterPathHash(Target.WindowsPlatform.DiaSdkDir, EnvVars.ToolChainVersion.ToString());
+			}
+
 			// Validate PGO data
 			if (EnvVars.Compiler.IsClang() && GlobalLinkEnvironment.bPGOOptimize && Target.WindowsPlatform.bIgnoreStalePGOData)
 			{
