UnrealBuildTool: FIx clang bForceWarningsAsError matcher for Linux clang

#rnx

diff --git a/Engine/Source/Programs/UnrealBuildTool/Executors/ParallelExecutor.cs b/Engine/Source/Programs/UnrealBuildTool/Executors/ParallelExecutor.cs
--- a/Engine/Source/Programs/UnrealBuildTool/Executors/ParallelExecutor.cs
+++ b/Engine/Source/Programs/UnrealBuildTool/Executors/ParallelExecutor.cs
@@ -244,15 +244,15 @@ namespace UnrealBuildTool
 
 			await Process.WaitForExitAsync(CancellationToken);
 
-			List<string> LogLines = Console.OutputEncoding.GetString(StdOutStream.GetBuffer(), 0, Convert.ToInt32(StdOutStream.Length)).Split(LineEndingSplit, StringSplitOptions.RemoveEmptyEntries).ToList();
+			List<string> LogLines = [.. Console.OutputEncoding.GetString(StdOutStream.GetBuffer(), 0, Convert.ToInt32(StdOutStream.Length)).Split(LineEndingSplit, StringSplitOptions.RemoveEmptyEntries)];
 			int ExitCode = Process.ExitCode;
 			TimeSpan ProcessorTime = Process.TotalProcessorTime;
 			TimeSpan ExecutionTime = Process.ExitTime - Process.StartTime;
 
-			if (ExitCode == 0 && Action.bForceWarningsAsError && LogLines.Any(x => x.Contains("): warning: ", StringComparison.OrdinalIgnoreCase)))
+			if (ExitCode == 0 && Action.bForceWarningsAsError && LogLines.Any(x => x.Contains(": warning: ", StringComparison.OrdinalIgnoreCase)))
 			{
 				ExitCode = 1;
-				LogLines = LogLines.Select(x => x.Replace("): warning: ", "): error: ", StringComparison.OrdinalIgnoreCase)).ToList();
+				LogLines = [.. LogLines.Select(x => x.Replace(": warning: ", ": error: ", StringComparison.OrdinalIgnoreCase))];
 			}
 
 			return new ExecuteResults(LogLines, ExitCode, ExecutionTime, ProcessorTime, Action, AdditionalDescription);
diff --git a/Engine/Source/Programs/UnrealBuildTool/Executors/UnrealBuildAccelerator/UBAExecutor.cs b/Engine/Source/Programs/UnrealBuildTool/Executors/UnrealBuildAccelerator/UBAExecutor.cs
--- a/Engine/Source/Programs/UnrealBuildTool/Executors/UnrealBuildAccelerator/UBAExecutor.cs
+++ b/Engine/Source/Programs/UnrealBuildTool/Executors/UnrealBuildAccelerator/UBAExecutor.cs
@@ -792,10 +792,10 @@ namespace UnrealBuildTool
 
 				if (exitCode == 0)
 				{
-					if (action.bForceWarningsAsError && logLines.Any(x => x.Contains("): warning: ")))
+					if (action.bForceWarningsAsError && logLines.Any(x => x.Contains(": warning: ", StringComparison.OrdinalIgnoreCase)))
 					{
 						exitCode = 1;
-						logLines = logLines.Select(x => x.Replace("): warning: ", "): error: ", StringComparison.OrdinalIgnoreCase)).ToList();
+						logLines = [.. logLines.Select(x => x.Replace(": warning: ", ": error: ", StringComparison.OrdinalIgnoreCase))];
 					}
 					else
 					{
