[UBT]
* Changed so AndroidProjectGenerator runs vswhere.exe in task which is waited on when needed

diff --git a/Engine/Source/Programs/UnrealBuildTool/Platform/Android/AndroidProjectGenerator.cs b/Engine/Source/Programs/UnrealBuildTool/Platform/Android/AndroidProjectGenerator.cs
--- a/Engine/Source/Programs/UnrealBuildTool/Platform/Android/AndroidProjectGenerator.cs
+++ b/Engine/Source/Programs/UnrealBuildTool/Platform/Android/AndroidProjectGenerator.cs
@@ -5,6 +5,7 @@ using System.Collections.Generic;
 using System.Diagnostics;
 using System.IO;
 using System.Text;
+using System.Threading.Tasks;
 using EpicGames.Core;
 using Microsoft.Extensions.Logging;
 using UnrealBuildBase;
@@ -20,18 +21,22 @@ namespace UnrealBuildTool
 		/// Whether Android Game Development Extension is installed in the system. See https://developer.android.com/games/agde for more details.
 		/// May be disabled by using -noagde on commandline
 		/// </summary>
-		private bool AGDEInstalled = false;
+		private bool AGDEInstalled { get => AGDEInstalledTask.Result; }
+
+		private Task<bool> AGDEInstalledTask;
 
 		public AndroidProjectGenerator(CommandLineArguments Arguments, ILogger Logger)
 			: base(Arguments, Logger)
 		{
-			AGDEInstalled = false;
-			if (OperatingSystem.IsWindows() && !Arguments.HasOption("-noagde"))
+			AGDEInstalledTask = Task.Run(() =>
 			{
-				AGDEInstalled = Microsoft.Win32.Registry.LocalMachine.OpenSubKey(@"SOFTWARE\WOW6432Node\Google\AndroidGameDevelopmentExtension")?.ValueCount > 0;
-
-				if (!AGDEInstalled)
+				if (OperatingSystem.IsWindows() && !Arguments.HasOption("-noagde"))
 				{
+					if (Microsoft.Win32.Registry.LocalMachine.OpenSubKey(@"SOFTWARE\WOW6432Node\Google\AndroidGameDevelopmentExtension")?.ValueCount > 0)
+					{
+						return true;
+					}
+
 					try
 					{
 						string? programFiles86 = Environment.GetEnvironmentVariable("ProgramFiles(x86)");
@@ -51,7 +56,7 @@ namespace UnrealBuildTool
 									};
 									p.StartInfo = info;
 									p.Start();
-									AGDEInstalled = p.StandardOutput.ReadToEnd().Contains("Google.VisualStudio.Android.dll");
+									return p.StandardOutput.ReadToEnd().Contains("Google.VisualStudio.Android.dll");
 								}
 							}
 						}
@@ -61,7 +66,8 @@ namespace UnrealBuildTool
 						Logger.LogInformation("Failed to identify AGDE installation status: {Message}", ex.Message);
 					}
 				}
-			}
+				return false;
+			});
 		}
 
 		/// <summary>
