[UBT]
* Changed SourceFileSearch to use DirectoryItem instead of DirectoryReference. Improved performance

diff --git a/Engine/Source/Programs/UnrealBuildTool/System/SourceFileSearch.cs b/Engine/Source/Programs/UnrealBuildTool/System/SourceFileSearch.cs
--- a/Engine/Source/Programs/UnrealBuildTool/System/SourceFileSearch.cs
+++ b/Engine/Source/Programs/UnrealBuildTool/System/SourceFileSearch.cs
@@ -53,6 +53,11 @@ namespace UnrealBuildTool
 		/// <param name="SearchSubdirectories">True to include subdirectories, otherwise we only search the list of base directories</param>
 		/// <param name="SearchedDirectories">If non-null, is updated with a list of the directories searched</param>
 		public static List<FileReference> FindFiles(DirectoryReference DirectoryToSearch, List<string>? SubdirectoryNamesToExclude = null, bool SearchSubdirectories = true, HashSet<DirectoryReference>? SearchedDirectories = null)
+		{
+			return FindFiles(DirectoryItem.GetItemByDirectoryReference(DirectoryToSearch), SubdirectoryNamesToExclude, SearchSubdirectories, SearchedDirectories);
+		}
+
+		public static List<FileReference> FindFiles(DirectoryItem DirectoryToSearch, List<string>? SubdirectoryNamesToExclude = null, bool SearchSubdirectories = true, HashSet<DirectoryReference>? SearchedDirectories = null)
 		{
 			// Build a list of directory names that we don't want to search under. We always ignore intermediate directories.
 			string[] ExcludedDirectorySuffixes;
@@ -78,26 +83,26 @@ namespace UnrealBuildTool
 			return FoundFiles;
 		}
 
-		static void FindFilesInternal(DirectoryReference Directory, string[] ExcludedDirectorySuffixes, List<FileReference> FoundFiles, HashSet<DirectoryReference>? SearchedDirectories)
+		static void FindFilesInternal(DirectoryItem Directory, string[] ExcludedDirectorySuffixes, List<FileReference> FoundFiles, HashSet<DirectoryReference>? SearchedDirectories)
 		{
-			SearchedDirectories?.Add(Directory);
+			SearchedDirectories?.Add(Directory.Location);
 
-			foreach (FileReference File in DirectoryLookupCache.EnumerateFiles(Directory))
+			foreach (FileItem File in Directory.EnumerateFiles())
 			{
-				if (ShouldInclude(File, DefaultExcludedFileSuffixes))
+				if (ShouldInclude(File.Location, DefaultExcludedFileSuffixes))
 				{
-					FoundFiles.Add(File);
+					FoundFiles.Add(File.Location);
 				}
 			}
 		}
 
-		static void FindFilesInternalRecursive(DirectoryReference Directory, string[] ExcludedDirectorySuffixes, List<FileReference> FoundFiles, HashSet<DirectoryReference>? SearchedDirectories)
+		static void FindFilesInternalRecursive(DirectoryItem Directory, string[] ExcludedDirectorySuffixes, List<FileReference> FoundFiles, HashSet<DirectoryReference>? SearchedDirectories)
 		{
 			FindFilesInternal(Directory, ExcludedDirectorySuffixes, FoundFiles, SearchedDirectories);
 
-			foreach (DirectoryReference SubDirectory in DirectoryLookupCache.EnumerateDirectories(Directory))
+			foreach (DirectoryItem SubDirectory in Directory.EnumerateDirectories())
 			{
-				if (ShouldInclude(SubDirectory, ExcludedDirectorySuffixes))
+				if (ShouldInclude(SubDirectory.Location, ExcludedDirectorySuffixes))
 				{
 					FindFilesInternalRecursive(SubDirectory, ExcludedDirectorySuffixes, FoundFiles, SearchedDirectories);
 				}
