[UBT]
* Cached CppCompileWarnings results to not query them over and over again. (These calls are very expensive)

diff --git a/Engine/Source/Programs/UnrealBuildTool/Configuration/CompileWarnings/WarningLevelResolverRegistry.cs b/Engine/Source/Programs/UnrealBuildTool/Configuration/CompileWarnings/WarningLevelResolverRegistry.cs
--- a/Engine/Source/Programs/UnrealBuildTool/Configuration/CompileWarnings/WarningLevelResolverRegistry.cs
+++ b/Engine/Source/Programs/UnrealBuildTool/Configuration/CompileWarnings/WarningLevelResolverRegistry.cs
@@ -32,24 +32,13 @@ namespace UnrealBuildTool
 	internal class WarningLevelResolverRegistry
 	{
 		private readonly ILogger? _logger = Log.Logger;
-		private static WarningLevelResolverRegistry? s_resolverRegistry;
+		private static Lazy<WarningLevelResolverRegistry> s_resolverRegistry = new(() => new());
 		private readonly Dictionary<string, IWarningLevelResolver> _attributeNameToBehaviour = [];
 
 		/// <summary>
 		/// The resolver name to resolver representation of the registry.
 		/// </summary>
-		internal static IReadOnlyDictionary<string, IWarningLevelResolver> Resolvers
-		{
-			get
-			{
-				if (s_resolverRegistry == null)
-				{
-					s_resolverRegistry = new WarningLevelResolverRegistry();
-				}
-
-				return s_resolverRegistry._attributeNameToBehaviour;
-			}
-		}
+		internal static IReadOnlyDictionary<string, IWarningLevelResolver> Resolvers => s_resolverRegistry.Value._attributeNameToBehaviour;
 
 		/// <summary>
 		/// Requests a resolver by name.
diff --git a/Engine/Source/Programs/UnrealBuildTool/Configuration/CppCompileWarnings.cs b/Engine/Source/Programs/UnrealBuildTool/Configuration/CppCompileWarnings.cs
--- a/Engine/Source/Programs/UnrealBuildTool/Configuration/CppCompileWarnings.cs
+++ b/Engine/Source/Programs/UnrealBuildTool/Configuration/CppCompileWarnings.cs
@@ -3,6 +3,7 @@
 using EpicGames.Core;
 using Microsoft.Extensions.Logging;
 using System;
+using System.Collections.Concurrent;
 using System.Collections.Generic;
 using System.Linq;
 using System.Reflection;
@@ -1483,7 +1484,9 @@ namespace UnrealBuildTool
 
 		#endregion -- Public API --
 
-		private static IEnumerable<MemberInfo>? s_membersWithApplyWarningsAttribute = null;
+		private static Lazy<List<MemberInfo>> s_membersWithApplyWarningsAttribute = new(() => GetMembersWithAttributes(typeof(CppCompileWarnings), typeof(ApplyWarningsAttribute), typeof(WarningLevel)).ToList());
+		private static Lazy<List<MemberInfo>> s_membersWithWarningLevelDefaultAttribute = new(() => GetMembersWithAttributes(typeof(CppCompileWarnings), typeof(WarningLevelDefaultAttribute), typeof(WarningLevel)).ToList());
+		private static ConcurrentDictionary<MemberInfo, Lazy<Dictionary<Type, List<WarningLevelDefaultAttribute>>>>[] s_groupedAttributes = { [], [], [] };
 
 		/// <summary>
 		/// Generates the compiler warnings command line arguments for the requested <see cref="UEToolChain"/> type, and version.
@@ -1507,9 +1510,7 @@ namespace UnrealBuildTool
 			List<string> warningArguments = [];
 			List<ValueTuple<ApplyWarningsAttribute, WarningLevel>> warningAttributesBuffer = [];
 
-			s_membersWithApplyWarningsAttribute ??= GetMembersWithAttributes(typeof(CppCompileWarnings), typeof(ApplyWarningsAttribute), typeof(WarningLevel));
-
-			foreach (MemberInfo memberInfo in s_membersWithApplyWarningsAttribute)
+			foreach (MemberInfo memberInfo in s_membersWithApplyWarningsAttribute.Value)
 			{
 				ApplyWarningsAttribute? applyWarningsAttribute = null;
 				IEnumerable<ApplyWarningsAttribute?> applyWarningsAttributes = memberInfo.GetCustomAttributes<ApplyWarningsAttribute>().Where(x => x.CanApplyToContext(toolChainContext));
@@ -1553,9 +1554,7 @@ namespace UnrealBuildTool
 		{
 			cppCompileWarnings._logger?.LogTrace("Applying default C++ Compiler warnings for the Target({TargetName}) context.", cppCompileWarnings.BuildContextProvider?._targetRulesPrivate?.Name);
 
-			IEnumerable<MemberInfo> memberInfos = GetMembersWithAttributes(typeof(CppCompileWarnings), typeof(WarningLevelDefaultAttribute), typeof(WarningLevel));
-
-			foreach (MemberInfo mi in memberInfos)
+			foreach (MemberInfo mi in s_membersWithWarningLevelDefaultAttribute.Value)
 			{
 				WarningLevel defaultWarningLevel = WarningLevel.Default;
 				defaultWarningLevel = CoalesceDefaultWarningFromMemberInfo(mi, InitializationContext.Target, defaultWarningLevel, cppCompileWarnings.BuildContextProvider, cppCompileWarnings._logger);
@@ -1571,18 +1570,24 @@ namespace UnrealBuildTool
 
 		internal static WarningLevel CoalesceDefaultWarningFromMemberInfo(MemberInfo mi, InitializationContext defaultWarningContext, WarningLevel defaultWarningLevel, BuildSystemContext? buildSystemContext, ILogger? logger)
 		{
-			IList<WarningLevelDefaultAttribute> attributes = [.. mi.GetCustomAttributes<WarningLevelDefaultAttribute>().Where(prop => prop.Context == InitializationContext.Any || prop.Context == defaultWarningContext)];
-			Type fallbackType = typeof(BasicWarningLevelDefaultAttribute);
-			Dictionary<Type, List<WarningLevelDefaultAttribute>> groupedAttributes = attributes
-				.GroupBy(attr => attr.GetType())
-				.OrderBy(g => g.Key == fallbackType ? 1 : 0)
-				.ToDictionary(g => g.Key, g => g.ToList());
+			// This is really slow so needs to be cached.
+			Dictionary<Type, List<WarningLevelDefaultAttribute>> groupedAttributes = s_groupedAttributes[(int)defaultWarningContext].GetOrAdd(mi, new Lazy<Dictionary<Type, List<WarningLevelDefaultAttribute>>>(() =>
+			{
+				IList<WarningLevelDefaultAttribute> attributes = [.. mi.GetCustomAttributes<WarningLevelDefaultAttribute>().Where(prop => prop.Context == InitializationContext.Any || prop.Context == defaultWarningContext)];
+				Type fallbackType = typeof(BasicWarningLevelDefaultAttribute);
+				return attributes
+					.GroupBy(attr => attr.GetType())
+					.OrderBy(g => g.Key == fallbackType ? 1 : 0)
+					.ToDictionary(g => g.Key, g => g.ToList());
+			})).Value;
+
+			IReadOnlyDictionary<string, IWarningLevelResolver> resolvers = WarningLevelResolverRegistry.Resolvers;
 
 			foreach (KeyValuePair<Type, List<WarningLevelDefaultAttribute>> item in groupedAttributes)
 			{
-				if (WarningLevelResolverRegistry.Resolvers.ContainsKey(item.Key.Name))
+				if (resolvers.ContainsKey(item.Key.Name))
 				{
-					WarningLevel resolvedWarning = WarningLevelResolverRegistry.Resolvers[item.Key.Name].ResolveWarningLevelDefault(item.Value, buildSystemContext);
+					WarningLevel resolvedWarning = resolvers[item.Key.Name].ResolveWarningLevelDefault(item.Value, buildSystemContext);
 
 					defaultWarningLevel = resolvedWarning;
 				}
@@ -1641,9 +1646,7 @@ namespace UnrealBuildTool
 
 			BuildSettingsVersion buildSettingsVersion = BuildContextProvider != null ? BuildContextProvider._buildContext.GetBuildSettings() : BuildSettingsVersion.V1;
 
-			IEnumerable<MemberInfo> memberInfos = GetMembersWithAttributes(typeof(CppCompileWarnings), typeof(WarningLevelDefaultAttribute), typeof(WarningLevel));
-
-			foreach (MemberInfo mi in memberInfos)
+			foreach (MemberInfo mi in s_membersWithWarningLevelDefaultAttribute.Value)
 			{
 				WarningLevel defaultWarningLevel = WarningLevel.Default;
 				defaultWarningLevel = CoalesceDefaultWarningFromMemberInfo(mi, InitializationContext.Constructor, defaultWarningLevel, BuildContextProvider, _logger);
