[UBT]
* Cache object properties with ConfigSubObjectAttribute since this code is very expensive to run over and over again
* Treat build environment as unique when generating project files in order to not have to run expensive code (sln output is identical regardless)

diff --git a/Engine/Source/Programs/UnrealBuildTool/Configuration/TargetRules.cs b/Engine/Source/Programs/UnrealBuildTool/Configuration/TargetRules.cs
--- a/Engine/Source/Programs/UnrealBuildTool/Configuration/TargetRules.cs
+++ b/Engine/Source/Programs/UnrealBuildTool/Configuration/TargetRules.cs
@@ -1,6 +1,7 @@
 // Copyright Epic Games, Inc. All Rights Reserved.
 
 using System;
+using System.Collections.Concurrent;
 using System.Collections.Generic;
 using System.Diagnostics.CodeAnalysis;
 using System.IO;
@@ -2998,6 +2999,13 @@ namespace UnrealBuildTool
 				return;
 			}
 
+			// RequiresUniqueEnvironment is expensive so early out when generating project files
+			if (IntermediateEnvironment == UnrealIntermediateEnvironment.GenerateProjectFiles)
+			{
+				BuildEnvironment = TargetBuildEnvironment.Unique;
+				return;
+			}
+
 			if (UEBuildPlatformSDK.bHasAnySDKOverride)
 			{
 				// check if any platform needs a unique environment for the sdk check
@@ -3401,6 +3409,23 @@ namespace UnrealBuildTool
 			}
 		}
 
+		static readonly ConcurrentDictionary<Type, Lazy<List<PropertyInfo>>> s_subPropsCache = new();
+
+		static List<PropertyInfo> GetSubObjectProperties(Type t) => s_subPropsCache.GetOrAdd(t, static type => new Lazy<List<PropertyInfo>>(() =>
+		{
+			var props = type.GetProperties(BindingFlags.Instance | BindingFlags.Public);
+			var list = new List<PropertyInfo>();
+			foreach (PropertyInfo pi in props)
+			{
+				if (pi.IsDefined(typeof(ConfigSubObjectAttribute), inherit: false))
+				{
+					list.Add(pi);
+				}
+			}
+			return list;
+		})).Value;
+
+
 		/// <summary>
 		/// Finds all the subobjects which can be configured by command line options and config files
 		/// </summary>
@@ -3417,7 +3442,7 @@ namespace UnrealBuildTool
 					yield return value;
 				}
 			}
-			foreach (PropertyInfo property in GetType().GetProperties(BindingFlags.Public | BindingFlags.Instance).Where(property => property.GetCustomAttribute<ConfigSubObjectAttribute>() != null))
+			foreach (PropertyInfo property in GetSubObjectProperties(GetType()))
 			{
 				object? value = property.GetValue(this);
 				if (value != null)
