[UBA]
* Added Lazy to concurrent dictionary usage in RulesCompiler for when creating assemly  to prevent same assembly from being compiled at the same time
* Moved global mutex to inside concurrent dictionary lazy func to only lock when actually compiling

diff --git a/Engine/Source/Programs/UnrealBuildTool/System/RulesCompiler.cs b/Engine/Source/Programs/UnrealBuildTool/System/RulesCompiler.cs
--- a/Engine/Source/Programs/UnrealBuildTool/System/RulesCompiler.cs
+++ b/Engine/Source/Programs/UnrealBuildTool/System/RulesCompiler.cs
@@ -3,6 +3,7 @@
 using System;
 using System.Collections.Concurrent;
 using System.Collections.Generic;
+using System.Diagnostics;
 using System.IO;
 using System.Linq;
 using EpicGames.Core;
@@ -60,6 +61,11 @@ namespace UnrealBuildTool
 			}
 		}
 
+		/// <summary>
+		/// Lock for creating engine rule assembly
+		/// </summary>
+		private static object EngineRulesAssemblyLock = new();
+
 		/// <summary>
 		/// The cached rules assembly for engine modules and targets.
 		/// </summary>
@@ -69,7 +75,7 @@ namespace UnrealBuildTool
 		/// Map of assembly names we've already compiled and loaded to their Assembly and list of game folders.  This is used to prevent
 		/// trying to recompile the same assembly when ping-ponging between different types of targets
 		/// </summary>
-		private static ConcurrentDictionary<FileReference, RulesAssembly> LoadedAssemblyMap = new ConcurrentDictionary<FileReference, RulesAssembly>();
+		private static ConcurrentDictionary<FileReference, Lazy<RulesAssembly>> LoadedAssemblyMap = new();
 
 		/// <summary>
 		/// Creates the engine rules assembly
@@ -81,33 +87,35 @@ namespace UnrealBuildTool
 		/// <returns>New rules assembly</returns>
 		public static RulesAssembly CreateEngineRulesAssembly(bool bUsePrecompiled, bool bSkipCompile, bool bForceCompile, ILogger Logger)
 		{
-			// Prevent multiple conflicting processes building rules assembly at the same time
-			string MutexName = GlobalSingleInstanceMutex.GetUniqueMutexForPath("UnrealBuildTool_CreateTargetRulesAssembly", Unreal.EngineDirectory);
-			using (new GlobalSingleInstanceMutex(MutexName, true))
+			lock (EngineRulesAssemblyLock)
 			{
 				if (EngineRulesAssembly == null)
 				{
-					List<PluginInfo> EnginePlugins = new List<PluginInfo>();
-					List<PluginInfo> MarketplacePlugins = new List<PluginInfo>();
-
-					DirectoryReference MarketplaceDirectory = DirectoryReference.Combine(Unreal.EngineDirectory, "Plugins", "Marketplace");
-					foreach (PluginInfo PluginInfo in Plugins.ReadEnginePlugins(Unreal.EngineDirectory))
+					// Prevent multiple conflicting processes building rules assembly at the same time
+					using GlobalSingleInstanceMutex _ = new(GlobalSingleInstanceMutex.GetUniqueMutexForPath("UnrealBuildTool_CreateTargetRulesAssembly", Unreal.EngineDirectory), true);
 					{
-						if (PluginInfo.File.IsUnderDirectory(MarketplaceDirectory))
-						{
-							MarketplacePlugins.Add(PluginInfo);
-						}
-						else
+						List<PluginInfo> EnginePlugins = new List<PluginInfo>();
+						List<PluginInfo> MarketplacePlugins = new List<PluginInfo>();
+
+						DirectoryReference MarketplaceDirectory = DirectoryReference.Combine(Unreal.EngineDirectory, "Plugins", "Marketplace");
+						foreach (PluginInfo PluginInfo in Plugins.ReadEnginePlugins(Unreal.EngineDirectory))
 						{
-							EnginePlugins.Add(PluginInfo);
+							if (PluginInfo.File.IsUnderDirectory(MarketplaceDirectory))
+							{
+								MarketplacePlugins.Add(PluginInfo);
+							}
+							else
+							{
+								EnginePlugins.Add(PluginInfo);
+							}
 						}
-					}
 
-					EngineRulesAssembly = CreateEngineRulesAssemblyInternal(Unreal.GetExtensionDirs(Unreal.EngineDirectory), ProjectFileGenerator.EngineRulesAssemblyName, EnginePlugins, Unreal.IsEngineInstalled() || bUsePrecompiled, bSkipCompile, bForceCompile, null, Logger);
+						EngineRulesAssembly = CreateEngineRulesAssemblyInternal(Unreal.GetExtensionDirs(Unreal.EngineDirectory), ProjectFileGenerator.EngineRulesAssemblyName, EnginePlugins, Unreal.IsEngineInstalled() || bUsePrecompiled, bSkipCompile, bForceCompile, null, Logger);
 
-					if (MarketplacePlugins.Count > 0)
-					{
-						EngineRulesAssembly = CreateMarketplaceRulesAssembly(MarketplacePlugins, Unreal.IsEngineInstalled() || bUsePrecompiled, bSkipCompile, bForceCompile, EngineRulesAssembly, Logger);
+						if (MarketplacePlugins.Count > 0)
+						{
+							EngineRulesAssembly = CreateMarketplaceRulesAssembly(MarketplacePlugins, Unreal.IsEngineInstalled() || bUsePrecompiled, bSkipCompile, bForceCompile, EngineRulesAssembly, Logger);
+						}
 					}
 				}
 				return EngineRulesAssembly;
@@ -240,101 +248,99 @@ namespace UnrealBuildTool
 		/// <returns>New rules assembly</returns>
 		public static RulesAssembly CreateProjectRulesAssembly(FileReference ProjectFileName, bool bUsePrecompiled, bool bSkipCompile, bool bForceCompile, ILogger Logger)
 		{
-			// Prevent multiple conflicting processes building rules assembly at the same time
-			string MutexName = GlobalSingleInstanceMutex.GetUniqueMutexForPath("UnrealBuildTool_CreateTargetRulesAssembly", ProjectFileName);
-			using (new GlobalSingleInstanceMutex(MutexName, true))
+			// Check if there's an existing assembly for this project
+			return LoadedAssemblyMap.GetOrAdd(ProjectFileName, new Lazy<RulesAssembly>(() =>
 			{
-				// Check if there's an existing assembly for this project
-				return LoadedAssemblyMap.GetOrAdd(ProjectFileName, _ =>
-				{
-					Logger.LogTrace("Creating project rules assembly for {Project}...", ProjectFileName);
+				// Prevent multiple conflicting processes building rules assembly at the same time
+				using GlobalSingleInstanceMutex _ = new(GlobalSingleInstanceMutex.GetUniqueMutexForPath("UnrealBuildTool_CreateTargetRulesAssembly", ProjectFileName), true);
 
-					ProjectDescriptor Project = ProjectDescriptor.FromFile(ProjectFileName);
+				Logger.LogTrace("Creating project rules assembly for {Project}...", ProjectFileName);
 
-					// Create the parent assembly
-					RulesAssembly Parent = CreateEngineRulesAssembly(bUsePrecompiled, bSkipCompile, bForceCompile, Logger);
+				ProjectDescriptor Project = ProjectDescriptor.FromFile(ProjectFileName);
 
-					DirectoryReference MainProjectDirectory = ProjectFileName.Directory;
-					//DirectoryReference MainProjectSourceDirectory = DirectoryReference.Combine(MainProjectDirectory, "Source");
+				// Create the parent assembly
+				RulesAssembly Parent = CreateEngineRulesAssembly(bUsePrecompiled, bSkipCompile, bForceCompile, Logger);
 
-					// Create a scope for things in this assembly
-					RulesScope Scope = new RulesScope("Project", Parent.Scope);
+				DirectoryReference MainProjectDirectory = ProjectFileName.Directory;
+				//DirectoryReference MainProjectSourceDirectory = DirectoryReference.Combine(MainProjectDirectory, "Source");
 
-					// Create a new context for modules created by this assembly
-					ModuleRulesContext DefaultModuleContext = new ModuleRulesContext(Scope, MainProjectDirectory);
-					DefaultModuleContext.bCanBuildDebugGame = true;
-					DefaultModuleContext.bCanHotReload = true;
-					DefaultModuleContext.bClassifyAsGameModuleForUHT = true;
-					DefaultModuleContext.bCanUseForSharedPCH = false;
+				// Create a scope for things in this assembly
+				RulesScope Scope = new RulesScope("Project", Parent.Scope);
 
-					// gather modules from project and platforms
-					Dictionary<FileReference, ModuleRulesContext> ModuleFiles = new Dictionary<FileReference, ModuleRulesContext>();
-					List<FileReference> TargetFiles = new List<FileReference>();
+				// Create a new context for modules created by this assembly
+				ModuleRulesContext DefaultModuleContext = new ModuleRulesContext(Scope, MainProjectDirectory);
+				DefaultModuleContext.bCanBuildDebugGame = true;
+				DefaultModuleContext.bCanHotReload = true;
+				DefaultModuleContext.bClassifyAsGameModuleForUHT = true;
+				DefaultModuleContext.bCanUseForSharedPCH = false;
 
-					// Find all the project directories
-					List<DirectoryReference> ProjectDirectories = Unreal.GetExtensionDirs(ProjectFileName.Directory);
-					if (Project.AdditionalRootDirectories != null)
-					{
-						ProjectDirectories.AddRange(Project.AdditionalRootDirectories);
-					}
+				// gather modules from project and platforms
+				Dictionary<FileReference, ModuleRulesContext> ModuleFiles = new Dictionary<FileReference, ModuleRulesContext>();
+				List<FileReference> TargetFiles = new List<FileReference>();
 
-					// Find all the rules/plugins under the project source directories
-					foreach (DirectoryReference ProjectDirectory in ProjectDirectories)
-					{
-						DirectoryReference ProjectSourceDirectory = DirectoryReference.Combine(ProjectDirectory, "Source");
+				// Find all the project directories
+				List<DirectoryReference> ProjectDirectories = Unreal.GetExtensionDirs(ProjectFileName.Directory);
+				if (Project.AdditionalRootDirectories != null)
+				{
+					ProjectDirectories.AddRange(Project.AdditionalRootDirectories);
+				}
 
-						AddModuleRulesWithContext(ProjectSourceDirectory, DefaultModuleContext, ModuleFiles);
-						TargetFiles.AddRange(Rules.FindAllRulesFiles(ProjectSourceDirectory, Rules.RulesFileType.Target));
-					}
+				// Find all the rules/plugins under the project source directories
+				foreach (DirectoryReference ProjectDirectory in ProjectDirectories)
+				{
+					DirectoryReference ProjectSourceDirectory = DirectoryReference.Combine(ProjectDirectory, "Source");
 
-					// Find all the project plugins
-					List<PluginInfo> ProjectPlugins = new List<PluginInfo>();
-					ProjectPlugins.AddRange(Plugins.ReadProjectPlugins(MainProjectDirectory));
+					AddModuleRulesWithContext(ProjectSourceDirectory, DefaultModuleContext, ModuleFiles);
+					TargetFiles.AddRange(Rules.FindAllRulesFiles(ProjectSourceDirectory, Rules.RulesFileType.Target));
+				}
+
+				// Find all the project plugins
+				List<PluginInfo> ProjectPlugins = new List<PluginInfo>();
+				ProjectPlugins.AddRange(Plugins.ReadProjectPlugins(MainProjectDirectory));
 
-					// Add the project's additional plugin directories plugins too
-					if (Project.AdditionalPluginDirectories != null)
+				// Add the project's additional plugin directories plugins too
+				if (Project.AdditionalPluginDirectories != null)
+				{
+					foreach (DirectoryReference AdditionalPluginDirectory in Project.AdditionalPluginDirectories)
 					{
-						foreach (DirectoryReference AdditionalPluginDirectory in Project.AdditionalPluginDirectories)
-						{
-							ProjectPlugins.AddRange(Plugins.ReadAdditionalPlugins(AdditionalPluginDirectory, Logger));
-						}
+						ProjectPlugins.AddRange(Plugins.ReadAdditionalPlugins(AdditionalPluginDirectory, Logger));
 					}
+				}
 
-					Logger.LogTrace(" Found {Count} Plugins:", ProjectPlugins.Count);
-					ProjectPlugins.ForEach(x => { Logger.LogTrace("  {Plugin}", x.File); });
+				Logger.LogTrace(" Found {Count} Plugins:", ProjectPlugins.Count);
+				ProjectPlugins.ForEach(x => { Logger.LogTrace("  {Plugin}", x.File); });
 
-					// Find all the plugin module rules as well as plugin test target and module rules
-					FindModuleRulesForPlugins(ProjectPlugins, DefaultModuleContext, ModuleFiles);
-					FindTestRulesForPlugins(ProjectPlugins, DefaultModuleContext, ModuleFiles, TargetFiles);
+				// Find all the plugin module rules as well as plugin test target and module rules
+				FindModuleRulesForPlugins(ProjectPlugins, DefaultModuleContext, ModuleFiles);
+				FindTestRulesForPlugins(ProjectPlugins, DefaultModuleContext, ModuleFiles, TargetFiles);
 
-					Logger.LogTrace(" Found {Count} Modules:", ModuleFiles.Count);
-					foreach (KeyValuePair<FileReference, ModuleRulesContext> Item in ModuleFiles)
-					{
-						Logger.LogTrace("  {Module}", Item.Key);
-					}
+				Logger.LogTrace(" Found {Count} Modules:", ModuleFiles.Count);
+				foreach (KeyValuePair<FileReference, ModuleRulesContext> Item in ModuleFiles)
+				{
+					Logger.LogTrace("  {Module}", Item.Key);
+				}
 
-					// Add the games project's intermediate source folder
-					DirectoryReference ProjectIntermediateSourceDirectory = DirectoryReference.Combine(MainProjectDirectory, "Intermediate", "Source");
-					if (DirectoryReference.Exists(ProjectIntermediateSourceDirectory))
-					{
-						AddModuleRulesWithContext(ProjectIntermediateSourceDirectory, DefaultModuleContext, ModuleFiles);
-						TargetFiles.AddRange(Rules.FindAllRulesFiles(ProjectIntermediateSourceDirectory, Rules.RulesFileType.Target));
-					}
+				// Add the games project's intermediate source folder
+				DirectoryReference ProjectIntermediateSourceDirectory = DirectoryReference.Combine(MainProjectDirectory, "Intermediate", "Source");
+				if (DirectoryReference.Exists(ProjectIntermediateSourceDirectory))
+				{
+					AddModuleRulesWithContext(ProjectIntermediateSourceDirectory, DefaultModuleContext, ModuleFiles);
+					TargetFiles.AddRange(Rules.FindAllRulesFiles(ProjectIntermediateSourceDirectory, Rules.RulesFileType.Target));
+				}
 
-					RulesAssembly ProjectRulesAssembly;
-					// Compile the assembly. If there are no module or target files, just use the parent assembly.
-					FileReference AssemblyFileName = FileReference.Combine(MainProjectDirectory, "Intermediate", "Build", "BuildRules", ProjectFileName.GetFileNameWithoutExtension() + "ModuleRules" + FrameworkAssemblyExtension);
-					if (ModuleFiles.Count == 0 && TargetFiles.Count == 0 && ProjectPlugins.Count == 0)
-					{
-						ProjectRulesAssembly = Parent;
-					}
-					else
-					{
-						ProjectRulesAssembly = new RulesAssembly(Scope, new List<DirectoryReference> { MainProjectDirectory }, ProjectPlugins, ModuleFiles, TargetFiles, AssemblyFileName, bContainsEngineModules: false, DefaultBuildSettings: null, bReadOnly: Unreal.IsProjectInstalled(), bSkipCompile: bSkipCompile, bForceCompile: bForceCompile, Parent: Parent, Logger: Logger);
-					}
-					return ProjectRulesAssembly;
-				});
-			}
+				RulesAssembly ProjectRulesAssembly;
+				// Compile the assembly. If there are no module or target files, just use the parent assembly.
+				FileReference AssemblyFileName = FileReference.Combine(MainProjectDirectory, "Intermediate", "Build", "BuildRules", ProjectFileName.GetFileNameWithoutExtension() + "ModuleRules" + FrameworkAssemblyExtension);
+				if (ModuleFiles.Count == 0 && TargetFiles.Count == 0 && ProjectPlugins.Count == 0)
+				{
+					ProjectRulesAssembly = Parent;
+				}
+				else
+				{
+					ProjectRulesAssembly = new RulesAssembly(Scope, new List<DirectoryReference> { MainProjectDirectory }, ProjectPlugins, ModuleFiles, TargetFiles, AssemblyFileName, bContainsEngineModules: false, DefaultBuildSettings: null, bReadOnly: Unreal.IsProjectInstalled(), bSkipCompile: bSkipCompile, bForceCompile: bForceCompile, Parent: Parent, Logger: Logger);
+				}
+				return ProjectRulesAssembly;
+			})).Value;
 		}
 
 		/// <summary>
@@ -350,60 +356,58 @@ namespace UnrealBuildTool
 		/// <returns>The new rules assembly</returns>
 		public static RulesAssembly CreatePluginRulesAssembly(FileReference PluginFileName, bool bSkipCompile, bool bForceCompile, RulesAssembly Parent, bool bBuildPluginAsLocal, bool bContainsEngineModules, ILogger Logger)
 		{
-			// Prevent multiple conflicting processes building rules assembly at the same time
-			string MutexName = GlobalSingleInstanceMutex.GetUniqueMutexForPath("UnrealBuildTool_CreateTargetRulesAssembly", PluginFileName);
-			using (new GlobalSingleInstanceMutex(MutexName, true))
+			// Check if there's an existing assembly for this project
+			return LoadedAssemblyMap.GetOrAdd(PluginFileName, new Lazy<RulesAssembly>(() =>
 			{
-				// Check if there's an existing assembly for this project
-				return LoadedAssemblyMap.GetOrAdd(PluginFileName, _ =>
-				{
-					Logger.LogTrace("Creating plugin rules assembly for {Plugin}...", PluginFileName);
+				// Prevent multiple conflicting processes building rules assembly at the same time
+				using GlobalSingleInstanceMutex _ = new(GlobalSingleInstanceMutex.GetUniqueMutexForPath("UnrealBuildTool_CreateTargetRulesAssembly", PluginFileName), true);
 
-					// Find all the rules source files
-					Dictionary<FileReference, ModuleRulesContext> ModuleFiles = new Dictionary<FileReference, ModuleRulesContext>();
-					List<FileReference> TargetFiles = new List<FileReference>();
+				Logger.LogTrace("Creating plugin rules assembly for {Plugin}...", PluginFileName);
 
-					// Create a list of plugins for this assembly. We need to override the parent plugin, if it exists, due to overriding the
-					// setting for bClassifyAsGameModuleForUHT below.
-					List<PluginInfo> ForeignPlugins = new List<PluginInfo>();
-					if (!Parent.EnumeratePlugins().Any(x => x.ChoiceVersion != null && x.ChoiceVersion.File == PluginFileName && x.ChoiceVersion.Type == PluginType.Engine))
+				// Find all the rules source files
+				Dictionary<FileReference, ModuleRulesContext> ModuleFiles = new Dictionary<FileReference, ModuleRulesContext>();
+				List<FileReference> TargetFiles = new List<FileReference>();
+
+				// Create a list of plugins for this assembly. We need to override the parent plugin, if it exists, due to overriding the
+				// setting for bClassifyAsGameModuleForUHT below.
+				List<PluginInfo> ForeignPlugins = new List<PluginInfo>();
+				if (!Parent.EnumeratePlugins().Any(x => x.ChoiceVersion != null && x.ChoiceVersion.File == PluginFileName && x.ChoiceVersion.Type == PluginType.Engine))
+				{
+					PluginInfo ForeignPluginInfo = new PluginInfo(PluginFileName, PluginType.External)
 					{
-						PluginInfo ForeignPluginInfo = new PluginInfo(PluginFileName, PluginType.External)
-						{
-							bExplicitPluginTarget = true
-						};
-						ForeignPlugins.Add(ForeignPluginInfo);
-					}
+						bExplicitPluginTarget = true
+					};
+					ForeignPlugins.Add(ForeignPluginInfo);
+				}
 
-					Logger.LogTrace(" Found {Count} ForeignPlugins:", ForeignPlugins.Count);
-					ForeignPlugins.ForEach(x => { Logger.LogTrace("  {ForeignPlugin}", x.File); });
+				Logger.LogTrace(" Found {Count} ForeignPlugins:", ForeignPlugins.Count);
+				ForeignPlugins.ForEach(x => { Logger.LogTrace("  {ForeignPlugin}", x.File); });
 
-					// Create a new scope for the plugin. It should not reference anything else.
-					RulesScope Scope = new RulesScope("Plugin", Parent.Scope);
+				// Create a new scope for the plugin. It should not reference anything else.
+				RulesScope Scope = new RulesScope("Plugin", Parent.Scope);
 
-					// Find all the modules
-					ModuleRulesContext PluginModuleContext = new ModuleRulesContext(Scope, PluginFileName.Directory);
-					PluginModuleContext.bClassifyAsGameModuleForUHT = !bContainsEngineModules;
-					if (bBuildPluginAsLocal)
-					{
-						PluginModuleContext.bCanBuildDebugGame = true;
-						PluginModuleContext.bCanHotReload = true;
-						PluginModuleContext.bClassifyAsGameModuleForUHT = true;
-						PluginModuleContext.bCanUseForSharedPCH = false;
-					}
-					FindModuleRulesForPlugins(ForeignPlugins, PluginModuleContext, ModuleFiles);
+				// Find all the modules
+				ModuleRulesContext PluginModuleContext = new ModuleRulesContext(Scope, PluginFileName.Directory);
+				PluginModuleContext.bClassifyAsGameModuleForUHT = !bContainsEngineModules;
+				if (bBuildPluginAsLocal)
+				{
+					PluginModuleContext.bCanBuildDebugGame = true;
+					PluginModuleContext.bCanHotReload = true;
+					PluginModuleContext.bClassifyAsGameModuleForUHT = true;
+					PluginModuleContext.bCanUseForSharedPCH = false;
+				}
+				FindModuleRulesForPlugins(ForeignPlugins, PluginModuleContext, ModuleFiles);
 
-					Logger.LogTrace(" Found {Count} Modules:", ModuleFiles.Count);
-					foreach (KeyValuePair<FileReference, ModuleRulesContext> Item in ModuleFiles)
-					{
-						Logger.LogTrace("  {Module}", Item.Key);
-					}
+				Logger.LogTrace(" Found {Count} Modules:", ModuleFiles.Count);
+				foreach (KeyValuePair<FileReference, ModuleRulesContext> Item in ModuleFiles)
+				{
+					Logger.LogTrace("  {Module}", Item.Key);
+				}
 
-					// Compile the assembly
-					FileReference AssemblyFileName = FileReference.Combine(PluginFileName.Directory, "Intermediate", "Build", "BuildRules", Path.GetFileNameWithoutExtension(PluginFileName.FullName) + "ModuleRules" + FrameworkAssemblyExtension);
-					return new RulesAssembly(Scope, new List<DirectoryReference> { PluginFileName.Directory }, ForeignPlugins, ModuleFiles, TargetFiles, AssemblyFileName, bContainsEngineModules, DefaultBuildSettings: null, bReadOnly: false, bSkipCompile: bSkipCompile, bForceCompile: bForceCompile, Parent: Parent, Logger: Logger);
-				});
-			}
+				// Compile the assembly
+				FileReference AssemblyFileName = FileReference.Combine(PluginFileName.Directory, "Intermediate", "Build", "BuildRules", Path.GetFileNameWithoutExtension(PluginFileName.FullName) + "ModuleRules" + FrameworkAssemblyExtension);
+				return new RulesAssembly(Scope, new List<DirectoryReference> { PluginFileName.Directory }, ForeignPlugins, ModuleFiles, TargetFiles, AssemblyFileName, bContainsEngineModules, DefaultBuildSettings: null, bReadOnly: false, bSkipCompile: bSkipCompile, bForceCompile: bForceCompile, Parent: Parent, Logger: Logger);
+			})).Value;
 		}
 
 		/// <summary>
@@ -420,6 +424,7 @@ namespace UnrealBuildTool
 		/// <returns>The compiled rules assembly</returns>
 		public static RulesAssembly CreateTargetRulesAssembly(FileReference? ProjectFile, string TargetName, bool bSkipRulesCompile, bool bForceRulesCompile, bool bUsePrecompiled, FileReference? ForeignPlugin, bool bBuildPluginAsLocal, ILogger Logger)
 		{
+			using var _ = Timeline.ScopeEvent("CreateRulesAssemly");
 			RulesAssembly RulesAssembly;
 			if (ProjectFile != null)
 			{
@@ -506,9 +511,9 @@ namespace UnrealBuildTool
 				return FileName.FullName;
 			}
 
-			foreach (RulesAssembly RulesAssembly in LoadedAssemblyMap.Values)
+			foreach (Lazy<RulesAssembly> RulesAssembly in LoadedAssemblyMap.Values)
 			{
-				if (RulesAssembly.TryGetFileNameFromType(ExistingType, out FileName))
+				if (RulesAssembly.Value.TryGetFileNameFromType(ExistingType, out FileName))
 				{
 					return FileName.FullName;
 				}
