[UBT]
* Fixed so Plugins.ReadPluginsFromDirectory is running things in parallel
* Added plugin info directory cache

diff --git a/Engine/Source/Programs/UnrealBuildTool/System/Plugins.cs b/Engine/Source/Programs/UnrealBuildTool/System/Plugins.cs
--- a/Engine/Source/Programs/UnrealBuildTool/System/Plugins.cs
+++ b/Engine/Source/Programs/UnrealBuildTool/System/Plugins.cs
@@ -3,6 +3,7 @@
 using EpicGames.Core;
 using Microsoft.Extensions.Logging;
 using System;
+using System.Collections.Concurrent;
 using System.Collections.Generic;
 using System.Diagnostics;
 using System.IO;
@@ -233,10 +234,8 @@ namespace UnrealBuildTool
 	/// </summary>
 	public static class Plugins
 	{
-		/// <summary>
-		/// Cache of plugins under each directory
-		/// </summary>
-		static Dictionary<DirectoryReference, List<PluginInfo>> PluginInfoCache = new Dictionary<DirectoryReference, List<PluginInfo>>();
+		static ConcurrentDictionary<string, PluginInfo> PluginInfoCache = [];
+		static ConcurrentDictionary<DirectoryReference, Lazy<List<PluginInfo>>> PluginDirectoryCache = [];
 
 		/// <summary>
 		/// Invalidate cached plugin data so that we can pickup new things
@@ -244,7 +243,7 @@ namespace UnrealBuildTool
 		/// </summary>
 		public static void InvalidateCaches_SLOW()
 		{
-			PluginInfoCache = new Dictionary<DirectoryReference, List<PluginInfo>>();
+			PluginInfoCache = new();
 			PluginsBase.InvalidateCache_SLOW();
 			DirectoryItem.ResetAllCachedInfo_SLOW();
 		}
@@ -256,15 +255,9 @@ namespace UnrealBuildTool
 		/// <returns>The cached plugin info, or null if not found.</returns>
 		public static PluginInfo? GetPlugin(string PluginName)
 		{
-			foreach (KeyValuePair<DirectoryReference, List<PluginInfo>> Pair in PluginInfoCache)
+			if (PluginInfoCache.TryGetValue(PluginName, out PluginInfo? plugin))
 			{
-				foreach (PluginInfo Plugin in Pair.Value)
-				{
-					if (String.Equals(Plugin.Name, PluginName, StringComparison.InvariantCultureIgnoreCase))
-					{
-						return Plugin;
-					}
-				}
+				return plugin;
 			}
 			return null;
 		}
@@ -349,6 +342,7 @@ namespace UnrealBuildTool
 		/// <returns>Sequence of the found PluginInfo object.</returns>
 		public static IReadOnlyList<PluginInfo> ReadEnginePlugins(DirectoryReference EngineDirectory)
 		{
+			using var _ = Timeline.ScopeEvent("ReadEnginePlugins"); 
 			return ReadPluginsFromDirectory(EngineDirectory, "Plugins", PluginType.Engine);
 		}
 
@@ -390,7 +384,7 @@ namespace UnrealBuildTool
 		{
 			foreach (UnrealPlatformGroup Group in UnrealPlatformGroup.GetValidGroups())
 			{
-				if (Group.ToString().Equals(Suffix, StringComparison.InvariantCultureIgnoreCase))
+				if (Group.ToString().Equals(Suffix, StringComparison.OrdinalIgnoreCase))
 				{
 					return true;
 				}
@@ -398,7 +392,7 @@ namespace UnrealBuildTool
 
 			foreach (UnrealTargetPlatform Platform in UnrealTargetPlatform.GetValidPlatforms())
 			{
-				if (Platform.ToString().Equals(Suffix, StringComparison.InvariantCultureIgnoreCase))
+				if (Platform.ToString().Equals(Suffix, StringComparison.OrdinalIgnoreCase))
 				{
 					return true;
 				}
@@ -411,9 +405,10 @@ namespace UnrealBuildTool
 		///  Attempt to merge a child plugin up into a parent plugin (via file naming scheme). Very little merging happens
 		///  but it does allow for platform extensions to extend a plugin with module files
 		/// </summary>
+		/// <param name="Plugins"></param>
 		/// <param name="Child">Child plugin that needs to merge to a main, parent plugin</param>
 		/// <param name="Filename">Child plugin's filename, used to determine the parent's name</param>
-		private static void TryMergeWithParent(PluginInfo Child, FileReference Filename)
+		private static void TryMergeWithParent(List<PluginInfo> Plugins, PluginInfo Child, FileReference Filename)
 		{
 			// find the parent
 			PluginInfo? Parent = null;
@@ -422,14 +417,7 @@ namespace UnrealBuildTool
 			if (Tokens.Length == 2)
 			{
 				string ParentPluginName = Tokens[0];
-				foreach (KeyValuePair<DirectoryReference, List<PluginInfo>> Pair in PluginInfoCache)
-				{
-					Parent = Pair.Value.FirstOrDefault(x => x.Name.Equals(ParentPluginName, StringComparison.InvariantCultureIgnoreCase) && x.LoadedFrom == Child.LoadedFrom);
-					if (Parent != null)
-					{
-						break;
-					}
-				}
+				Parent = Plugins.FirstOrDefault(x => x.Name.Equals(ParentPluginName, StringComparison.OrdinalIgnoreCase) && x.LoadedFrom == Child.LoadedFrom);
 			}
 			else
 			{
@@ -450,6 +438,16 @@ namespace UnrealBuildTool
 				return;
 			}
 
+			// The reason we need a lock here is because there could be multiple targets using the same plugins
+			// and ReadPluginsFromDirectory can be called in parallel from multiple targets causing ChildPlugins
+			// to be populated from different threads.. not the same ChildPlugins but some from each thread.
+			lock (Parent)
+			{
+				MergeWithParent(Parent, Child);
+			}
+		}
+		private static void MergeWithParent(PluginInfo Parent, PluginInfo Child)
+		{
 			// add our uplugin file to the existing plugin to be used to search for modules later
 			Parent.ChildFiles.Add(Child.File);
 
@@ -617,51 +615,59 @@ namespace UnrealBuildTool
 		/// <returns>Sequence of the found PluginInfo object.</returns>
 		public static IReadOnlyList<PluginInfo> ReadPluginsFromDirectory(DirectoryReference RootDirectory, string Subdirectory, PluginType Type)
 		{
-			// look for directories in RootDirectory and and extension directories under RootDirectory
-			List<DirectoryReference> RootDirectories = Unreal.GetExtensionDirs(RootDirectory, Subdirectory);
-
-			Dictionary<PluginInfo, FileReference> ChildPlugins = new Dictionary<PluginInfo, FileReference>();
-			List<PluginInfo> AllParentPlugins = new List<PluginInfo>();
+			DirectoryReference pluginDir = DirectoryReference.Combine(RootDirectory, Subdirectory);
 
-			foreach (DirectoryReference Dir in RootDirectories)
+			return PluginDirectoryCache.GetOrAdd(pluginDir, new Lazy<List<PluginInfo>>(() =>
 			{
-				if (!DirectoryReference.Exists(Dir))
-				{
-					continue;
-				}
+				// look for directories in RootDirectory and and extension directories under RootDirectory
+				List<DirectoryReference> RootDirectories = Unreal.GetExtensionDirs(RootDirectory, Subdirectory);
 
-				List<PluginInfo>? Plugins;
-				if (!PluginInfoCache.TryGetValue(Dir, out Plugins))
+				Dictionary<PluginInfo, FileReference> ChildPlugins = new Dictionary<PluginInfo, FileReference>();
+				List<PluginInfo> AllParentPlugins = new List<PluginInfo>();
+
+				foreach (DirectoryReference Dir in RootDirectories)
 				{
-					Plugins = new List<PluginInfo>();
-					foreach (FileReference PluginFileName in PluginsBase.EnumeratePlugins(Dir))
+					if (!DirectoryReference.Exists(Dir))
+					{
+						continue;
+					}
+
+					Parallel2.ForEach(PluginsBase.EnumeratePlugins(Dir), (PluginFileName) =>
 					{
 						PluginInfo Plugin = new PluginInfo(PluginFileName, Type);
 
 						// is there a parent to merge up into?
 						if (Plugin.Descriptor.bIsPluginExtension)
 						{
-							ChildPlugins.Add(Plugin, PluginFileName);
+							lock (ChildPlugins)
+							{
+								ChildPlugins.Add(Plugin, PluginFileName);
+							}
 						}
 						else
 						{
-							Plugins.Add(Plugin);
+							lock (AllParentPlugins)
+							{
+								AllParentPlugins.Add(Plugin);
+							}
 						}
-					}
-					PluginInfoCache.Add(Dir, Plugins);
+					});
 				}
 
-				// gather all of the plugins into one list
-				AllParentPlugins.AddRange(Plugins);
-			}
-
-			// now that all parent plugins are read in, we can let the children look up the parents
-			foreach (KeyValuePair<PluginInfo, FileReference> Pair in ChildPlugins)
-			{
-				TryMergeWithParent(Pair.Key, Pair.Value);
-			}
+				// now that all parent plugins are read in, we can let the children look up the parents
+				foreach (KeyValuePair<PluginInfo, FileReference> Pair in ChildPlugins)
+				{
+					TryMergeWithParent(AllParentPlugins, Pair.Key, Pair.Value);
+				}
 
-			return AllParentPlugins;
+				// Add to cache. Note that we TryAdd because there is a chance the plugin is already added through another RootDirectory
+				// (VersePrototypes.uproject has AdditionalPluginDirectory set to a subdir of a plugin folder so PluginDirectoryCache will not catch that)
+				foreach (PluginInfo plugin in AllParentPlugins)
+				{
+					PluginInfoCache.TryAdd(plugin.Name, plugin);
+				}
+				return AllParentPlugins;
+			})).Value;
 		}
 
 		/// <summary>
@@ -686,7 +692,7 @@ namespace UnrealBuildTool
 			{
 				foreach (PluginReferenceDescriptor PluginReference in Project.Plugins)
 				{
-					if (String.Equals(PluginReference.Name, Plugin.Name, StringComparison.CurrentCultureIgnoreCase) && !PluginReference.bOptional)
+					if (String.Equals(PluginReference.Name, Plugin.Name, StringComparison.OrdinalIgnoreCase) && !PluginReference.bOptional)
 					{
 						bEnabled = PluginReference.IsEnabledForPlatform(Platform) && PluginReference.IsEnabledForTargetConfiguration(Configuration) && PluginReference.IsEnabledForTarget(TargetType);
 					}
