[UBT]
* Cache the EncryptionKeys per ProjectKeyChainFile since it is queried multiple times when generating solution (creates rules for lots of configurations and targets that share some file)

diff --git a/Engine/Source/Programs/UnrealBuildTool/System/EncryptionAndSigning.cs b/Engine/Source/Programs/UnrealBuildTool/System/EncryptionAndSigning.cs
--- a/Engine/Source/Programs/UnrealBuildTool/System/EncryptionAndSigning.cs
+++ b/Engine/Source/Programs/UnrealBuildTool/System/EncryptionAndSigning.cs
@@ -1,6 +1,7 @@
 // Copyright Epic Games, Inc. All Rights Reserved.
 
 using System;
+using System.Collections.Concurrent;
 using System.Collections.Generic;
 using System.Linq;
 using System.Text.Json;
@@ -244,6 +245,8 @@ namespace UnrealBuildTool
 			return true;
 		}
 
+		static ConcurrentDictionary<FileReference, Lazy<List<EncryptionKey>?>> s_encryptionKeyCache = [];
+
 		/// <summary>
 		/// Parse crypto settings from INI file
 		/// </summary>
@@ -411,15 +414,15 @@ namespace UnrealBuildTool
 					if (GameIni.GetString("ContentEncryption", "ProjectKeyChain", out Filename))
 					{
 						FileReference ProjectKeyChainFile = FileReference.Combine(InProjectDirectory, "Content", Filename);
-						if (FileReference.Exists(ProjectKeyChainFile))
+						List<EncryptionKey>? EncryptionKeysInFile = s_encryptionKeyCache.GetOrAdd(ProjectKeyChainFile, static ProjectKeyChainFile => new Lazy<List<EncryptionKey>?>(() =>
 						{
-							List<EncryptionKey> EncryptionKeys = new List<EncryptionKey>();
-
-							if (Settings.SecondaryEncryptionKeys != null)
+							if (!FileReference.Exists(ProjectKeyChainFile))
 							{
-								EncryptionKeys.AddRange(Settings.SecondaryEncryptionKeys);
+								return null;
 							}
 
+							List<EncryptionKey> EncryptionKeys = [];
+
 							string[] Lines = FileReference.ReadAllLines(ProjectKeyChainFile);
 							foreach (string Line in Lines)
 							{
@@ -442,6 +445,19 @@ namespace UnrealBuildTool
 								}
 							}
 
+							return EncryptionKeys;
+						})).Value;
+
+						if (EncryptionKeysInFile != null)
+						{
+							List<EncryptionKey> EncryptionKeys = new();
+							if (Settings.SecondaryEncryptionKeys != null)
+							{
+								EncryptionKeys.AddRange(Settings.SecondaryEncryptionKeys);
+							}
+
+							EncryptionKeys.AddRange(EncryptionKeysInFile);
+
 							Settings.SecondaryEncryptionKeys = EncryptionKeys.ToArray();
 						}
 					}
