[UBA]
* VCProject Changed to Parallel2.ForEach to get rid of the joins
* VCProject - sort local aliased files now when they can be produced in concurrent logic so they are not sorted

diff --git a/Engine/Source/Programs/UnrealBuildTool/ProjectFiles/ProjectFile.cs b/Engine/Source/Programs/UnrealBuildTool/ProjectFiles/ProjectFile.cs
--- a/Engine/Source/Programs/UnrealBuildTool/ProjectFiles/ProjectFile.cs
+++ b/Engine/Source/Programs/UnrealBuildTool/ProjectFiles/ProjectFile.cs
@@ -181,10 +181,9 @@ namespace UnrealBuildTool
 		/// <summary>
 		/// Set of all files that have been added
 		/// </summary>
-		HashSet<FileReference> AliasedFilesSet = new HashSet<FileReference>();
+		Dictionary<FileReference, AliasedFile> AliasedFilesLookup = [];
 
-		/// Aliased (i.e. files is custom filter tree) in this project
-		public readonly List<AliasedFile> AliasedFiles = new List<AliasedFile>();
+		public Dictionary<FileReference, AliasedFile>.ValueCollection AliasedFiles => AliasedFilesLookup.Values;
 
 		/// <summary>
 		/// Adds aliased file to the project.
@@ -192,10 +191,7 @@ namespace UnrealBuildTool
 		/// <param name="File">Aliased file.</param>
 		public void AddAliasedFileToProject(AliasedFile File)
 		{
-			if (AliasedFilesSet.Add(File.Location))
-			{
-				AliasedFiles.Add(File);
-			}
+			AliasedFilesLookup.TryAdd(File.Location, File);
 		}
 
 		/// <summary>
@@ -206,7 +202,7 @@ namespace UnrealBuildTool
 		public void AddFileToProject(FileReference FilePath, DirectoryReference BaseFolder)
 		{
 			// Check if hasn't already been added as an aliased file
-			if (AliasedFilesSet.Contains(FilePath))
+			if (AliasedFilesLookup.ContainsKey(FilePath))
 			{
 				return;
 			}
diff --git a/Engine/Source/Programs/UnrealBuildTool/ProjectFiles/VisualStudio/VCProject.cs b/Engine/Source/Programs/UnrealBuildTool/ProjectFiles/VisualStudio/VCProject.cs
--- a/Engine/Source/Programs/UnrealBuildTool/ProjectFiles/VisualStudio/VCProject.cs
+++ b/Engine/Source/Programs/UnrealBuildTool/ProjectFiles/VisualStudio/VCProject.cs
@@ -1161,14 +1161,14 @@ namespace UnrealBuildTool
 			{
 				// Find out how many source files there are in each directory
 				ConcurrentDictionary<DirectoryReference, int> SourceDirToCount = new();
-				Parallel.ForEach(SourceFiles.Where(sf => sf.Reference.HasExtension(".cpp")), SourceFile =>
+				Parallel2.ForEach(SourceFiles.Where(sf => sf.Reference.HasExtension(".cpp")), SourceFile =>
 				{
 					SourceDirToCount.AddOrUpdate(SourceFile.Reference.Directory, _ => 1, (k, v) => v + 1);
 				});
 
 				// Figure out the most common include paths
 				ConcurrentDictionary<DirectoryReference, int> IncludePathToCount = new();
-				Parallel.ForEach(SourceDirToCount, Pair =>
+				Parallel2.ForEach(SourceDirToCount, Pair =>
 				{
 					if (TryGetBuildEnvironment(Pair.Key, out BuildEnvironment? OutBuildEnvironment))
 					{
@@ -1207,7 +1207,7 @@ namespace UnrealBuildTool
 			}
 
 			// Gather source folder and file info
-			List<AliasedFile> LocalAliasedFiles = new List<AliasedFile>(AliasedFiles);
+			List<AliasedFile> LocalAliasedFiles = new(AliasedFiles);
 			ConcurrentDictionary<DirectoryReference, string?> DirectoryToPchFile = new();
 			ConcurrentDictionary<DirectoryReference, string> DirectoryToForceIncludePaths = new();
 			ConcurrentDictionary<DirectoryReference, string> DirectoryToIncludeSearchPaths = new();
@@ -1244,7 +1244,13 @@ namespace UnrealBuildTool
 
 					LocalAliasedFiles.Add(new AliasedFile(CurFile.Reference, ProjectRelativeSourceFile, FilterRelativeSourceDirectory));
 				}
-				Parallel.ForEach(LocalAliasedFiles, LocalAliasedFile =>
+
+				LocalAliasedFiles.Sort((a, b) =>
+				{
+					return a.Location.CompareTo(b.Location);
+				});
+
+				Parallel2.ForEach(LocalAliasedFiles, LocalAliasedFile =>
 				{
 					// get the filetype as represented to Visual Studio
 					VCFileType FileType = GetVCFileType(LocalAliasedFile.FileSystemPath);
@@ -2057,10 +2063,8 @@ namespace UnrealBuildTool
 
 				foreach (string LeadingDirectory in AllDirectoriesInPath)
 				{
-					if (!FilterDirectories.Contains(LeadingDirectory))
+					if (FilterDirectories.Add(LeadingDirectory))
 					{
-						FilterDirectories.Add(LeadingDirectory);
-
 						// Generate a unique stable GUID for this folder by hashing the LeadingDirectory string
 						// NOTE: When saving generated project files, we ignore differences in GUIDs if every other part of the file
 						//       matches identically with the pre-existing file
diff --git a/Engine/Source/Programs/UnrealBuildTool/System/NativeProjects.cs b/Engine/Source/Programs/UnrealBuildTool/System/NativeProjects.cs
--- a/Engine/Source/Programs/UnrealBuildTool/System/NativeProjects.cs
+++ b/Engine/Source/Programs/UnrealBuildTool/System/NativeProjects.cs
@@ -45,6 +45,7 @@ namespace UnrealBuildTool
 				{
 					if (CachedTargetNameToProjectFile == null)
 					{
+						using var _ = Timeline.ScopeEvent("Find target files");
 						Dictionary<string, FileReference> TargetNameToProjectFile = new Dictionary<string, FileReference>();
 						foreach (FileReference ProjectFile in EnumerateProjectFiles(Logger))
 						{
