[UBA]
* Disabled tiered compilation since the overhead of running it is higher than the time it takes to run UBT

diff --git a/Engine/Source/Programs/UnrealBuildTool/runtimeconfig.template.json b/Engine/Source/Programs/UnrealBuildTool/runtimeconfig.template.json
--- a/Engine/Source/Programs/UnrealBuildTool/runtimeconfig.template.json
+++ b/Engine/Source/Programs/UnrealBuildTool/runtimeconfig.template.json
@@ -1,8 +1,12 @@
 {
 	"configProperties": {
 		"System.GC.Server": true,
+		"System.GC.Concurrent": true,
 		"System.GC.HeapCount": 32,
 		"System.Globalization.Invariant": true,
-		"System.Globalization.UseNls": true
+		"System.Globalization.UseNls": true,
+		"System.Runtime.TieredCompilation": false,
+		"System.Runtime.TieredCompilation.QuickJit": false,
+		"System.Runtime.TieredCompilation.QuickJitForLoops": false
 	}
 }
\ No newline at end of file
