[UBT - Project generation optimizations]
* Find modules .csproj files in parallel
* Discover c# programs in parallel
* Create c# rules assemblies in parallel
* Change order of project file writing since most expensive ones seem to end up in the end

diff --git a/Engine/Source/Programs/UnrealBuildTool/ProjectFiles/ProjectFileGenerator.cs b/Engine/Source/Programs/UnrealBuildTool/ProjectFiles/ProjectFileGenerator.cs
--- a/Engine/Source/Programs/UnrealBuildTool/ProjectFiles/ProjectFileGenerator.cs
+++ b/Engine/Source/Programs/UnrealBuildTool/ProjectFiles/ProjectFileGenerator.cs
@@ -547,11 +547,25 @@ namespace UnrealBuildTool
 
 			// Find all the modules .csproj files to add
 			List<FileReference> ModuleFiles = Rules.FindAllRulesSourceFiles(RulesFileType, GameFolders, ForeignPlugins: null, AdditionalSearchPaths: BuildFolders);
-			foreach (FileReference ProjectFile in ModuleFiles)
+
+			Dictionary<FileReference, VCSharpProjectFile> projectFiles = [];
+			Parallel.ForEach(ModuleFiles, (ProjectFile) =>
 			{
 				if (FileReference.Exists(ProjectFile))
 				{
+					using var _ = Timeline.ScopeEvent(ProjectFile.GetFileName());
 					VCSharpProjectFile Project = new VCSharpProjectFile(ProjectFile, Logger);
+					lock (projectFiles)
+					{
+						projectFiles.Add(ProjectFile, Project);
+					}
+				}
+			});
+
+			foreach (FileReference ProjectFile in ModuleFiles)
+			{
+				if (projectFiles.TryGetValue(ProjectFile, out VCSharpProjectFile? Project))
+				{
 
 					Project.ShouldBuildForAllSolutionTargets = false;//true;
 					AddExistingProjectFile(Project, bForceDevelopmentConfiguration: true);
@@ -649,16 +663,18 @@ namespace UnrealBuildTool
 			List<FileReference> FoundProjects = new List<FileReference>();
 			List<DirectoryReference> ProjectDirs = new List<DirectoryReference>();
 
+			ParallelQueue Queue = new();
+
 			if (bIncludeEnginePrograms)
 			{
 				DirectoryReference EngineExtras = DirectoryReference.Combine(Unreal.EngineDirectory, "Extras");
 				ProjectDirs.Add(EngineExtras);
-				DiscoverCSharpProgramProjectsRecursively(EngineExtras, FoundProjects);
+				DiscoverCSharpProgramProjectsRecursivelyInternal(EngineExtras, FoundProjects, Queue);
 
 				ProjectDirs = ProjectDirs.Union(AllEngineDirectories).ToList();
 				foreach (DirectoryReference EngineDir in AllEngineDirectories)
 				{
-					DiscoverCSharpProgramProjectsRecursively(EngineDir, FoundProjects);
+					DiscoverCSharpProgramProjectsRecursivelyInternal(EngineDir, FoundProjects, Queue);
 				}
 			}
 
@@ -669,9 +685,11 @@ namespace UnrealBuildTool
 				ProjectDirs = ProjectDirs.Union(AllGameDirectories).ToList();
 				foreach (DirectoryReference GameDir in AllGameDirectories)
 				{
-					DiscoverCSharpProgramProjectsRecursively(GameDir, FoundProjects);
+					DiscoverCSharpProgramProjectsRecursivelyInternal(GameDir, FoundProjects, Queue);
 				}
 			}
+			Queue.Drain();
+
 
 			string[] UnsupportedPlatformNames = Utils.MakeListOfUnsupportedPlatforms(SupportedPlatforms, bIncludeUnbuildablePlatforms: true, Logger).ToArray();
 			foreach (FileReference FoundProject in FoundProjects)
@@ -715,7 +733,7 @@ namespace UnrealBuildTool
 			}
 		}
 
-		private VCSharpProjectFile CreateRulesAssemblyProject(Dictionary<RulesAssembly, DirectoryReference> RulesAssemblies, RulesAssembly RulesAssembly, DirectoryReference FSPathBase, ILogger Logger)
+		private static VCSharpProjectFile CreateRulesAssemblyProject(Dictionary<RulesAssembly, DirectoryReference> RulesAssemblies, RulesAssembly RulesAssembly, DirectoryReference FSPathBase, ILogger Logger)
 		{
 			DirectoryReference ProjectFilesDirectory = DirectoryReference.Combine(FSPathBase, "Intermediate/Build/BuildRulesProjects", RulesAssembly.GetSimpleAssemblyName()!);
 			DirectoryReference.CreateDirectory(ProjectFilesDirectory);
@@ -783,12 +801,17 @@ namespace UnrealBuildTool
 			return ModuleFilesProject;
 		}
 
-		private static void DiscoverCSharpProgramProjectsRecursively(DirectoryReference SearchFolder, List<FileReference> FoundProjects)
+		private static void DiscoverCSharpProgramProjectsRecursivelyInternal(DirectoryReference SearchFolder, List<FileReference> FoundProjects, ParallelQueue Queue)
 		{
 			// Scan all the files in this directory
 			bool bSearchSubFolders = true;
 			foreach (FileReference File in DirectoryLookupCache.EnumerateFiles(SearchFolder))
 			{
+				if (File.GetFileName() == ".ubtignore") // Folder should be ignored, nothing to see here
+				{
+					bSearchSubFolders = false;
+				}
+
 				// If we find a csproj or sln, we should not recurse this directory.
 				bool bIsCsProj = File.HasExtension(".csproj");
 				bool bIsSln = File.HasExtension(".sln");
@@ -801,7 +824,10 @@ namespace UnrealBuildTool
 				// For csproj files, add them to the sln if the .uprogram file also exists.
 				if (bIsCsProj && FileReference.Exists(File.ChangeExtension(".uprogram")))
 				{
-					FoundProjects.Add(File);
+					lock (FoundProjects)
+					{
+						FoundProjects.Add(File);
+					}
 				}
 			}
 
@@ -810,7 +836,7 @@ namespace UnrealBuildTool
 			{
 				foreach (DirectoryReference SubDirectory in DirectoryLookupCache.EnumerateDirectories(SearchFolder))
 				{
-					DiscoverCSharpProgramProjectsRecursively(SubDirectory, FoundProjects);
+					Queue.Enqueue(() => DiscoverCSharpProgramProjectsRecursivelyInternal(SubDirectory, FoundProjects, Queue));
 				}
 			}
 		}
@@ -1328,18 +1354,24 @@ namespace UnrealBuildTool
 					if (bIncludeEnginePrograms)
 					{
 						// Discover C# programs which should additionally be included in the solution.
-						DiscoverCSharpProgramProjects(AllEngineDirectories, AllGameProjects, ProgramsFolder, Logger);
+						using (Timeline.ScopeEvent("DiscoverCsharpPrograms", Logger))
+						{
+							DiscoverCSharpProgramProjects(AllEngineDirectories, AllGameProjects, ProgramsFolder, Logger);
+						}
 					}
 
-					foreach (KeyValuePair<RulesAssembly, DirectoryReference> RulesAssemblyEntry in RulesAssemblies)
+					Parallel2.ForEach(RulesAssemblies, (RulesAssemblyEntry) =>
 					{
 						if (RulesAssemblyEntry.Key.GetSimpleAssemblyName() != null)
 						{
 							VCSharpProjectFile RulesProject = CreateRulesAssemblyProject(RulesAssemblies, RulesAssemblyEntry.Key, RulesAssemblyEntry.Value, Logger);
 							AddExistingProjectFile(RulesProject, bForceDevelopmentConfiguration: false);
-							RootFolder.AddSubFolder("Rules").ChildProjects.Add(RulesProject);
+							lock (RootFolder)
+							{
+								RootFolder.AddSubFolder("Rules").ChildProjects.Add(RulesProject);
+							}
 						}
-					}
+					});
 				}
 
 				// Eliminate all redundant project folders.  E.g., folders which contain only one project and that project
@@ -1451,8 +1483,9 @@ namespace UnrealBuildTool
 			using (ProgressWriter Progress = new ProgressWriter("Writing Query Target Info", true, Logger))
 			{
 				int TasksFinished = 0;
-				Parallel.ForEach(ProjectFileToRules, Item =>
+				Parallel2.ForEach(ProjectFileToRules, Item =>
 				{
+					using var t = Timeline.ScopeEvent(Item.Key?.GetFileName()??"Unknown");
 					QueryTargetsMode.WriteTargetInfo(Item.Key, Item.Value, QueryTargetsMode.GetDefaultOutputFile(Item.Key), new CommandLineArguments(GetTargetArguments(Arguments)), Logger);
 
 					lock (Progress)
@@ -3429,15 +3462,19 @@ namespace UnrealBuildTool
 			{
 				int TotalProjectFileCount = GeneratedProjectFiles.Count + 1;    // +1 for the primary project file, which we'll save next
 
+				GeneratedProjectFiles.Reverse(); // It is likely the biggest projects were added last
+
 				int ProjectsFinished = 0;
 				bool ProjectCreationFailed = false;
-				System.Threading.Tasks.Parallel.ForEach(GeneratedProjectFiles, CurProject =>
+				Parallel2.ForEach(GeneratedProjectFiles, CurProject =>
 				{
 					if (ProjectCreationFailed)
 					{
 						return;
 					}
 
+					using var _ = Timeline.ScopeEvent(CurProject.ProjectFilePath.GetFileName());
+
 					if (!CurProject.WriteProjectFile(SupportedPlatforms, SupportedConfigurations, PlatformProjectGenerators, Logger))
 					{
 						ProjectCreationFailed = true;
@@ -3709,7 +3746,10 @@ namespace UnrealBuildTool
 			// Existing projects must always have a GUID.  This will throw an exception if one isn't found.
 			InProject.LoadGUIDFromExistingProject();
 
-			OtherProjectFiles.Add(InProject);
+			lock (OtherProjectFiles)
+			{
+				OtherProjectFiles.Add(InProject);
+			}
 		}
 
 		/// <summary>
