Remove permissions not allowed by [REDACTED_PERSON]

#android
#13633

diff --git a/Engine/Source/Programs/UnrealBuildTool/Platform/Android/UEDeployAndroid.cs b/Engine/Source/Programs/UnrealBuildTool/Platform/Android/UEDeployAndroid.cs
--- a/Engine/Source/Programs/UnrealBuildTool/Platform/Android/UEDeployAndroid.cs
+++ b/Engine/Source/Programs/UnrealBuildTool/Platform/Android/UEDeployAndroid.cs
@@ -3301,13 +3301,7 @@ namespace UnrealBuildTool
 
 				Text.AppendLine("\t<uses-feature android:glEsVersion=\"" + AndroidToolChain.GetGLESVersion(bBuildForES31) + "\" android:required=\"true\" />");
 				Text.AppendLine("\t<uses-permission android:name=\"android.permission.INTERNET\"/>");
-				if (Configuration != "Shipping" || !bUseExternalFilesDir)
-				{
-					Text.AppendLine("\t<uses-permission android:name=\"android.permission.WRITE_EXTERNAL_STORAGE\"/>");
-				}
 				Text.AppendLine("\t<uses-permission android:name=\"android.permission.ACCESS_NETWORK_STATE\"/>");
-				Text.AppendLine("\t<uses-permission android:name=\"android.permission.WAKE_LOCK\"/>");
-				//	Text.AppendLine("\t<uses-permission android:name=\"android.permission.READ_PHONE_STATE\"/>");
 				Text.AppendLine("\t<uses-permission android:name=\"com.android.vending.CHECK_LICENSE\"/>");
 				Text.AppendLine("\t<uses-permission android:name=\"android.permission.ACCESS_WIFI_STATE\"/>");
 
@@ -3323,8 +3317,15 @@ namespace UnrealBuildTool
 					Text.AppendLine("\t<uses-permission android:name=\"android.permission.RECEIVE_BOOT_COMPLETED\"/>");
 				}
 
+				// Apps requesting these permissions are rejected by the Quest App Store
 				if (!bPackageForMetaQuest)
 				{
+					if (Configuration != "Shipping" || !bUseExternalFilesDir)
+					{
+						Text.AppendLine("\t<uses-permission android:name=\"android.permission.WRITE_EXTERNAL_STORAGE\"/>");
+					}
+					Text.AppendLine("\t<uses-permission android:name=\"android.permission.WAKE_LOCK\"/>");
+					//Text.AppendLine("\t<uses-permission android:name=\"android.permission.READ_PHONE_STATE\"/>");
 					Text.AppendLine("\t<uses-permission android:name=\"android.permission.MODIFY_AUDIO_SETTINGS\"/>");
 					Text.AppendLine("\t<uses-permission android:name=\"android.permission.VIBRATE\"/>");
 				}
