[UBT - Project generation optimizations]
* Added so source file population is done in parallel
* Run engine and project populations in parallel

diff --git a/Engine/Source/Programs/UnrealBuildTool/ProjectFiles/ProjectFileGenerator.cs b/Engine/Source/Programs/UnrealBuildTool/ProjectFiles/ProjectFileGenerator.cs
--- a/Engine/Source/Programs/UnrealBuildTool/ProjectFiles/ProjectFileGenerator.cs
+++ b/Engine/Source/Programs/UnrealBuildTool/ProjectFiles/ProjectFileGenerator.cs
@@ -1124,115 +1124,125 @@ namespace UnrealBuildTool
 					AddProjectsForAllTargets(PlatformProjectGenerators, AllGameProjects, AllTargetFiles, Arguments, EngineProjects, GameProjects, ProjectFileToUProjectFile, ProgramProjects, RulesAssemblies, Logger);
 				}
 
-				// Add projects for mods
-				AddProjectsForMods(GameProjects, ModProjects);
-
-				// Add all game projects and game config files
-				using (Timeline.ScopeEvent("Adding projects for all game projects", Logger))
+				Task GameTask = Task.Run(() =>
 				{
-					AddAllGameProjects(GameProjects);
-				}
 
-				// Set the game to be the default project
-				if (ModProjects.Count > 0)
-				{
-					DefaultProject = ModProjects.First();
-				}
-				else if (bGeneratingGameProjectFiles && GameProjects.Count > 0)
-				{
-					DefaultProject = GameProjects.First();
-				}
+					// Add projects for mods
+					AddProjectsForMods(GameProjects, ModProjects);
+
+					// Add all game projects and game config files
+					using (Timeline.ScopeEvent("Adding projects for all game projects", Logger))
+					{
+						AddAllGameProjects(GameProjects);
+					}
+
+					// Set the game to be the default project
+					if (ModProjects.Count > 0)
+					{
+						DefaultProject = ModProjects.First();
+					}
+					else if (bGeneratingGameProjectFiles && GameProjects.Count > 0)
+					{
+						DefaultProject = GameProjects.First();
+					}
+				});
 
 				//Related Debug Project Files - Tuple here has the related Debug Project, SolutionFolder
 				List<Tuple<ProjectFile, string>> DebugProjectFiles = new List<Tuple<ProjectFile, string>>();
 
-				// Place projects into root level solution folders
-				if (bIncludeEngineSource)
+				Task EngineTask = Task.Run(() =>
 				{
-					using ITimelineEvent EngineTimer = Timeline.ScopeEvent("Adding Engine Source");
-					// If we're still missing an engine project because we don't have any targets for it, make one up.
-					if (EngineProjects.Count == 0)
+					// Place projects into root level solution folders
+					if (bIncludeEngineSource)
 					{
-						FileReference ProjectFilePath = FileReference.Combine(IntermediateProjectFilesPath, EngineProjectFileNameBase + ProjectFileExtension);
+						using ITimelineEvent EngineTimer = Timeline.ScopeEvent("Adding Engine Source");
+						// If we're still missing an engine project because we don't have any targets for it, make one up.
+						if (EngineProjects.Count == 0)
+						{
+							FileReference ProjectFilePath = FileReference.Combine(IntermediateProjectFilesPath, EngineProjectFileNameBase + ProjectFileExtension);
 
-						ProjectFile EngineProject = FindOrAddProject(ProjectFilePath, Unreal.EngineDirectory, true, out _);
+							ProjectFile EngineProject = FindOrAddProject(ProjectFilePath, Unreal.EngineDirectory, true, out _);
 
-						EngineProject.IsForeignProject = false;
-						EngineProject.IsGeneratedProject = true;
-						EngineProject.IsStubProject = true;
+							EngineProject.IsForeignProject = false;
+							EngineProject.IsGeneratedProject = true;
+							EngineProject.IsStubProject = true;
 
-						EngineProjects.Add(EngineProject);
-					}
+							EngineProjects.Add(EngineProject);
+						}
 
-					if (EngineProjects.Count > 0)
-					{
-						// put all engine projects into solution
-						PrimaryProjectFolder EngineFolder = RootFolder.AddSubFolder("Engine");
-						foreach (ProjectFile Project in EngineProjects)
+						if (EngineProjects.Count > 0)
 						{
-							EngineFolder.ChildProjects.Add(Project);
-						}
+							// put all engine projects into solution
+							PrimaryProjectFolder EngineFolder = RootFolder.AddSubFolder("Engine");
+							foreach (ProjectFile Project in EngineProjects)
+							{
+								EngineFolder.ChildProjects.Add(Project);
+							}
 
-						// put all this stuff into one of the engine projects
-						ProjectFile EngineProject = EngineProjects.First(x => x.ProjectFilePath.GetFileNameWithoutAnyExtensions() == EngineProjectFileNameBase);
+							// put all this stuff into one of the engine projects
+							ProjectFile EngineProject = EngineProjects.First(x => x.ProjectFilePath.GetFileNameWithoutAnyExtensions() == EngineProjectFileNameBase);
 
-						// Engine config files
-						if (bIncludeConfigFiles)
-						{
-							AddEngineConfigFiles(EngineProject);
-							if (bIncludeEnginePrograms)
+							// Engine config files
+							if (bIncludeConfigFiles)
 							{
-								AddUnrealHeaderToolConfigFiles(EngineProject);
-								AddUBTConfigFilesToEngineProject(EngineProject);
+								AddEngineConfigFiles(EngineProject);
+								if (bIncludeEnginePrograms)
+								{
+									AddUnrealHeaderToolConfigFiles(EngineProject);
+									AddUBTConfigFilesToEngineProject(EngineProject);
+								}
 							}
-						}
 
-						// Engine Extras files
-						AddEngineExtrasFiles(EngineProject);
+							// Engine Extras files
+							AddEngineExtrasFiles(EngineProject);
 
-						// Platform Extension files
-						AddEngineExtensionFiles(EngineProject);
+							// Platform Extension files
+							AddEngineExtensionFiles(EngineProject);
 
-						// Engine localization files
-						if (bIncludeLocalizationFiles)
-						{
-							AddEngineLocalizationFiles(EngineProject);
-						}
+							// Engine localization files
+							if (bIncludeLocalizationFiles)
+							{
+								AddEngineLocalizationFiles(EngineProject);
+							}
 
-						// Engine template files
-						if (bIncludeTemplateFiles)
-						{
-							AddEngineTemplateFiles(EngineProject);
-						}
+							// Engine template files
+							if (bIncludeTemplateFiles)
+							{
+								AddEngineTemplateFiles(EngineProject);
+							}
 
-						if (bIncludeShaderSource)
-						{
-							Logger.LogDebug("Adding shader source code...");
+							if (bIncludeShaderSource)
+							{
+								Logger.LogDebug("Adding shader source code...");
 
-							// Find shader source files and generate stub project
-							AddEngineShaderSource(EngineProject);
-						}
+								// Find shader source files and generate stub project
+								AddEngineShaderSource(EngineProject);
+							}
 
-						if (bIncludeBuildSystemFiles)
-						{
-							Logger.LogDebug("Adding build system files...");
+							if (bIncludeBuildSystemFiles)
+							{
+								Logger.LogDebug("Adding build system files...");
 
-							AddEngineBuildFiles(EngineProject);
-						}
+								AddEngineBuildFiles(EngineProject);
+							}
 
-						if (bIncludeDocumentation)
-						{
-							AddEngineDocumentation(EngineProject, Logger);
-						}
+							if (bIncludeDocumentation)
+							{
+								AddEngineDocumentation(EngineProject, Logger);
+							}
 
-						List<Tuple<ProjectFile, string>>? NewProjectFiles = EngineProject.WriteDebugProjectFiles(SupportedPlatforms, SupportedConfigurations, PlatformProjectGenerators, Logger);
+							List<Tuple<ProjectFile, string>>? NewProjectFiles = EngineProject.WriteDebugProjectFiles(SupportedPlatforms, SupportedConfigurations, PlatformProjectGenerators, Logger);
 
-						if (NewProjectFiles != null)
-						{
-							DebugProjectFiles.AddRange(NewProjectFiles);
+							if (NewProjectFiles != null)
+							{
+								DebugProjectFiles.AddRange(NewProjectFiles);
+							}
 						}
 					}
-				}
+				});
+
+				GameTask.Wait();
+				EngineTask.Wait();
 
 				foreach (ProjectFile CurModProject in ModProjects)
 				{
@@ -1753,8 +1763,13 @@ namespace UnrealBuildTool
 		/// </summary>
 		protected void AddAllGameProjects(List<ProjectFile> GameProjects)
 		{
-			HashSet<DirectoryReference> UniqueGameProjectDirectories = new HashSet<DirectoryReference>();
+			Dictionary<DirectoryReference, ProjectFile> UniqueGameProjectDirectories = new(GameProjects.Count);
 			foreach (ProjectFile GameProject in GameProjects)
+			{
+				UniqueGameProjectDirectories.TryAdd(GameProject.BaseDir, GameProject);
+			}
+
+			Parallel2.ForEach(GameProjects, (GameProject) =>
 			{
 				DirectoryReference GameProjectDirectory = GameProject.BaseDir;
 
@@ -1765,47 +1780,50 @@ namespace UnrealBuildTool
 					GameProject.AddFilesToProject(SourceFileSearch.FindFiles(ExtensionDir, SearchSubdirectories: false), GameProjectDirectory);
 				}
 
-				if (UniqueGameProjectDirectories.Add(GameProjectDirectory))
+				if (UniqueGameProjectDirectories[GameProjectDirectory] != GameProject)
 				{
-					// @todo projectfiles: We have engine localization files, but should we also add GAME localization files?
+					return;
+				}
 
-					// Game restricted source files, since they won't be added via a module FileReference
-					foreach (DirectoryReference GameRestrictedSourceDirectory in Unreal.GetExtensionDirs(GameProjectDirectory, "Source", bIncludePlatformDirectories: false, bIncludeBaseDirectory: false))
-					{
-						GameProject.AddFilesToProject(SourceFileSearch.FindFiles(GameRestrictedSourceDirectory), GameRestrictedSourceDirectory);
-					}
+				// @todo projectfiles: We have engine localization files, but should we also add GAME localization files?
 
-					// Game config files
-					if (bIncludeConfigFiles)
-					{
-						foreach (DirectoryReference GameConfigDirectory in Unreal.GetExtensionDirs(GameProjectDirectory, "Config"))
-						{
-							GameProject.AddFilesToProject(SourceFileSearch.FindFiles(GameConfigDirectory), GameProjectDirectory);
-						}
-					}
+				// Game restricted source files, since they won't be added via a module FileReference
+				foreach (DirectoryReference GameRestrictedSourceDirectory in Unreal.GetExtensionDirs(GameProjectDirectory, "Source", bIncludePlatformDirectories: false, bIncludeBaseDirectory: false))
+				{
+					GameProject.AddFilesToProject(SourceFileSearch.FindFiles(GameRestrictedSourceDirectory), GameRestrictedSourceDirectory);
+				}
 
-					// Game build files
-					if (bIncludeBuildSystemFiles)
+				// Game config files
+				if (bIncludeConfigFiles)
+				{
+					foreach (DirectoryReference GameConfigDirectory in Unreal.GetExtensionDirs(GameProjectDirectory, "Config"))
 					{
-						foreach (DirectoryReference GameBuildDirectory in Unreal.GetExtensionDirs(GameProjectDirectory, "Build"))
-						{
-							List<string> SubdirectoryNamesToExclude = new List<string>();
-							SubdirectoryNamesToExclude.Add("Receipts");
-							SubdirectoryNamesToExclude.Add("Scripts");
-							SubdirectoryNamesToExclude.Add("FileOpenOrder");
-							SubdirectoryNamesToExclude.Add("PipelineCaches");
-							SubdirectoryNamesToExclude.Add("symbols");
-
-							GameProject.AddFilesToProject(SourceFileSearch.FindFiles(GameBuildDirectory, SubdirectoryNamesToExclude), GameProjectDirectory);
-						}
+						GameProject.AddFilesToProject(SourceFileSearch.FindFiles(GameConfigDirectory), GameProjectDirectory);
 					}
+				}
 
-					foreach (DirectoryReference GameShaderDirectory in Unreal.GetExtensionDirs(GameProjectDirectory, "Shaders"))
+				// Game build files
+				if (bIncludeBuildSystemFiles)
+				{
+					foreach (DirectoryReference GameBuildDirectory in Unreal.GetExtensionDirs(GameProjectDirectory, "Build"))
 					{
-						GameProject.AddFilesToProject(SourceFileSearch.FindFiles(GameShaderDirectory), GameProjectDirectory);
+						List<string> SubdirectoryNamesToExclude = [
+							"Receipts",
+							"Scripts",
+							"FileOpenOrder",
+							"PipelineCaches",
+							"symbols",
+						];
+
+						GameProject.AddFilesToProject(SourceFileSearch.FindFiles(GameBuildDirectory, SubdirectoryNamesToExclude), GameProjectDirectory);
 					}
 				}
-			}
+
+				foreach (DirectoryReference GameShaderDirectory in Unreal.GetExtensionDirs(GameProjectDirectory, "Shaders"))
+				{
+					GameProject.AddFilesToProject(SourceFileSearch.FindFiles(GameShaderDirectory), GameProjectDirectory);
+				}
+			});
 		}
 
 		/// Adds all engine localization text files to the specified project
@@ -1846,20 +1864,25 @@ namespace UnrealBuildTool
 		/// Adds additional files from the platform extensions folder
 		protected virtual void AddEngineExtensionFiles(ProjectFile EngineProject)
 		{
+			List<string> SubdirectoryNamesToExclude = [
+				"AutomationTool", //automation files are added separately to the AutomationTool project
+				"Binaries",
+				"Content",
+			];
+
 			// @todo: this will add the same files to the solution (like the UBT source files that also get added to UnrealBuildTool project).
 			// not sure of a good filtering method here
-			foreach (DirectoryReference ExtensionDir in Unreal.GetExtensionDirs(Unreal.EngineDirectory))
+			Parallel2.ForEach(Unreal.GetExtensionDirs(Unreal.EngineDirectory), (ExtensionDir) =>
 			{
 				if (ExtensionDir != Unreal.EngineDirectory)
 				{
-					List<string> SubdirectoryNamesToExclude = new List<string>();
-					SubdirectoryNamesToExclude.Add("AutomationTool"); //automation files are added separately to the AutomationTool project
-					SubdirectoryNamesToExclude.Add("Binaries");
-					SubdirectoryNamesToExclude.Add("Content");
-
-					EngineProject.AddFilesToProject(SourceFileSearch.FindFiles(ExtensionDir, SubdirectoryNamesToExclude), Unreal.EngineDirectory);
+					List<FileReference> Files = SourceFileSearch.FindFiles(ExtensionDir, SubdirectoryNamesToExclude);
+					lock (EngineProject)
+					{
+						EngineProject.AddFilesToProject(Files, Unreal.EngineDirectory);
+					}
 				}
-			}
+			});
 		}
 
 		/// Adds UnrealHeaderTool config files to the specified project
@@ -2752,30 +2775,44 @@ namespace UnrealBuildTool
 		private void AddPluginFilesToProject(FileReference PluginFileName, DirectoryReference BaseFolder, ProjectFile ProjectFile)
 		{
 			// Add the .uplugin file
-			ProjectFile.AddFileToProject(PluginFileName, BaseFolder);
+			lock (ProjectFile)
+			{
+				ProjectFile.AddFileToProject(PluginFileName, BaseFolder);
+			}
+
+			DirectoryItem PluginDir = DirectoryItem.GetItemByDirectoryReference(PluginFileName.Directory);
 
 			// Add plugin config files if we have any
 			if (bIncludeConfigFiles)
 			{
-				DirectoryReference PluginConfigFolder = DirectoryReference.Combine(PluginFileName.Directory, "Config");
-				if (DirectoryReference.Exists(PluginConfigFolder))
+				if (PluginDir.TryGetDirectory("Config", out DirectoryItem? PluginConfigDirectory))
 				{
-					ProjectFile.AddFilesToProject(SourceFileSearch.FindFiles(PluginConfigFolder), BaseFolder);
+					List<FileReference> files = SourceFileSearch.FindFiles(PluginConfigDirectory);
+					lock (ProjectFile)
+					{
+						ProjectFile.AddFilesToProject(files, BaseFolder);
+					}
 				}
 			}
 
 			// Add plugin "resource" files if we have any
-			DirectoryReference PluginResourcesFolder = DirectoryReference.Combine(PluginFileName.Directory, "Resources");
-			if (DirectoryReference.Exists(PluginResourcesFolder))
+			if (PluginDir.TryGetDirectory("Resources", out DirectoryItem? PluginResourcesFolder))
 			{
-				ProjectFile.AddFilesToProject(SourceFileSearch.FindFiles(PluginResourcesFolder), BaseFolder);
+				List<FileReference> files = SourceFileSearch.FindFiles(PluginResourcesFolder);
+				lock (ProjectFile)
+				{
+					ProjectFile.AddFilesToProject(files, BaseFolder);
+				}
 			}
 
 			// Add plugin shader files if we have any
-			DirectoryReference PluginShadersFolder = DirectoryReference.Combine(PluginFileName.Directory, "Shaders");
-			if (DirectoryReference.Exists(PluginShadersFolder))
+			if (PluginDir.TryGetDirectory("Shaders", out DirectoryItem? PluginShadersFolder))
 			{
-				ProjectFile.AddFilesToProject(SourceFileSearch.FindFiles(PluginShadersFolder), BaseFolder);
+				lock (ProjectFile)
+				{
+					List<FileReference> files = SourceFileSearch.FindFiles(PluginShadersFolder);
+					ProjectFile.AddFilesToProject(files, BaseFolder);
+				}
 			}
 		}
 
@@ -2952,6 +2989,8 @@ namespace UnrealBuildTool
 			// Keep a cache of the default editor target for each project so that we don't have to interrogate the configs for each target
 			Dictionary<string, string?> DefaultProjectEditorTargetCache = new Dictionary<string, string?>();
 
+			RulesCompiler.CreateEngineRulesAssembly(false, false, false, Logger);
+
 			List<FileReference> UnprocessedGameProjects = new(AllGames);
 			foreach (FileReference TargetFilePath in AllTargetFiles.Except(AllSubTargetFiles))
 			{
@@ -3125,7 +3164,7 @@ namespace UnrealBuildTool
 					ProjectFile.IsForeignProject = CheckProjectFile != null && !NativeProjects.IsNativeProject(CheckProjectFile, Logger);
 					ProjectFile.IsGeneratedProject = true;
 					ProjectFile.IsStubProject = Unreal.IsProjectInstalled();
-					ProjectFile.IsHybridContentOnlyProject = CheckProjectFile != null && NativeProjects.IsHybridContentOnlyProject(CheckProjectFile, out _, Logger);
+					ProjectFile.IsHybridContentOnlyProject = CheckProjectFile != null && NativeProjects.IsHybridContentOnlyProject(CheckProjectFile, out string? _, Logger);
 					if (TargetRulesObject.bBuildInSolutionByDefault.HasValue)
 					{
 						ProjectFile.ShouldBuildByDefaultForSolutionTargets = TargetRulesObject.bBuildInSolutionByDefault.Value;
@@ -3195,7 +3234,7 @@ namespace UnrealBuildTool
 							ProjectFilePath: ProjectFilePath,
 							UnrealProjectFilePath: CheckProjectFile,
 							SupportedPlatforms: TargetRulesObject.GetSupportedPlatforms().Where(
-													x => UEBuildPlatform.TryGetBuildPlatform(x, out _) &&
+													x => UEBuildPlatform.TryGetBuildPlatform(x, out UEBuildPlatform? _) &&
 													(TargetRulesObject.LinkType != TargetLinkType.Modular || !UEBuildPlatform.PlatformRequiresMonolithicBuilds(x, TargetRulesObject.Configuration))
 													).ToArray(),
 							CreateRulesDelegate: (Platform, Configuration) => RulesAssembly.CreateTargetRules(TargetName, Platform, Configuration, null, CheckProjectFile, new CommandLineArguments(GetTargetArguments(Arguments)), Logger),
