Bump Android Gradle Plugin to 8.6.1 from PR #13689
Update Java source/target version
Additional gate on WRITE_EXTERNAL_STORAGE if Meta Quest enabled

#android

diff --git a/Engine/Source/Programs/UnrealBuildTool/Platform/Android/UEDeployAndroid.cs b/Engine/Source/Programs/UnrealBuildTool/Platform/Android/UEDeployAndroid.cs
--- a/Engine/Source/Programs/UnrealBuildTool/Platform/Android/UEDeployAndroid.cs
+++ b/Engine/Source/Programs/UnrealBuildTool/Platform/Android/UEDeployAndroid.cs
@@ -25,7 +25,7 @@ namespace UnrealBuildTool
 		private const string BUNDLETOOL_JAR = "bundletool-all-1.18.1.jar";
 
 		// classpath of default android build tools gradle plugin
-		private const string ANDROID_TOOLS_BUILD_GRADLE_VERSION = "com.android.tools.build:gradle:8.5.2";
+		private const string ANDROID_TOOLS_BUILD_GRADLE_VERSION = "com.android.tools.build:gradle:8.6.1";
 
 		// default NDK version if not set
 		private const string DEFAULT_NDK_VERSION = "27.2.12479018";
@@ -3167,7 +3167,7 @@ namespace UnrealBuildTool
 
 			// Figure out the required startup permissions if targetting devices supporting runtime permissions
 			string StartupPermissions = "";
-			if (TargetSDKVersion >= 23)
+			if (TargetSDKVersion >= 23 && !bPackageForMetaQuest)
 			{
 				if (Configuration != "Shipping" || !bUseExternalFilesDir)
 				{
@@ -6795,6 +6795,10 @@ import java.util.Collection;
 				BuildGradleContent.AppendLine("\tndkPath NDK_PATH");
 				BuildGradleContent.AppendLine("\tcompileSdk COMPILE_SDK_VERSION.toInteger()");
 				BuildGradleContent.AppendLine("\tbuildToolsVersion BUILD_TOOLS_VERSION");
+				BuildGradleContent.AppendLine("\tcompileOptions {");
+				BuildGradleContent.AppendLine("\t\tsourceCompatibility JavaVersion.VERSION_17");
+				BuildGradleContent.AppendLine("\t\ttargetCompatibility JavaVersion.VERSION_17");
+				BuildGradleContent.AppendLine("\t}");
 				BuildGradleContent.AppendLine("\tdefaultConfig {");
 
 				// Try to get the SDK target from the AndroidManifest.xml
