Enable desugar for older Android support

#android

diff --git a/Engine/Source/Programs/UnrealBuildTool/Platform/Android/UEDeployAndroid.cs b/Engine/Source/Programs/UnrealBuildTool/Platform/Android/UEDeployAndroid.cs
--- a/Engine/Source/Programs/UnrealBuildTool/Platform/Android/UEDeployAndroid.cs
+++ b/Engine/Source/Programs/UnrealBuildTool/Platform/Android/UEDeployAndroid.cs
@@ -6798,8 +6798,10 @@ import java.util.Collection;
 				BuildGradleContent.AppendLine("\tcompileOptions {");
 				BuildGradleContent.AppendLine("\t\tsourceCompatibility JavaVersion.VERSION_17");
 				BuildGradleContent.AppendLine("\t\ttargetCompatibility JavaVersion.VERSION_17");
+				BuildGradleContent.AppendLine("\t\tcoreLibraryDesugaringEnabled true");
 				BuildGradleContent.AppendLine("\t}");
 				BuildGradleContent.AppendLine("\tdefaultConfig {");
+				BuildGradleContent.AppendLine("\t\tmultiDexEnabled true");
 
 				// Try to get the SDK target from the AndroidManifest.xml
 				string VersionCode = "";
@@ -6881,6 +6883,7 @@ import java.util.Collection;
 				}
 				BuildGradleContent.AppendLine("\t}");
 				BuildGradleContent.AppendLine("}");
+				BuildGradleContent.AppendLine("dependencies.coreLibraryDesugaring \"com.android.tools:desugar_jdk_libs:2.1.5\"");
 
 				string AdditionsGradleFilename = Path.Combine(LibDir, "additions.gradle");
 				if (File.Exists(AdditionsGradleFilename))
