Generate gradle project that wraps around original gradle project that is suitable for use in Android Studio

#android

diff --git a/Engine/Build/AndroidStudio/.idea/runConfigurations/Build_and_Deploy.xml b/Engine/Build/AndroidStudio/.idea/runConfigurations/Build_and_Deploy.xml
new file mode 100644
--- /dev/null
+++ b/Engine/Build/AndroidStudio/.idea/runConfigurations/Build_and_Deploy.xml
@@ -0,0 +1,6 @@
+<component name="ProjectRunConfigurationManager">
+  <configuration default="false" name="Build and Deploy" type="AndroidTestRunConfigurationType" factoryName="Android Instrumented Tests" activateToolWindowBeforeRun="false">
+    <module name="UE" />
+    <method v="2" />
+  </configuration>
+</component>
\ No newline at end of file
diff --git a/Engine/Build/AndroidStudio/.idea/runConfigurations/Client.xml b/Engine/Build/AndroidStudio/.idea/runConfigurations/Client.xml
new file mode 100644
--- /dev/null
+++ b/Engine/Build/AndroidStudio/.idea/runConfigurations/Client.xml
@@ -0,0 +1,9 @@
+<component name="ProjectRunConfigurationManager">
+  <configuration default="false" name="Client" type="AndroidRunConfigurationType" factoryName="Android App">
+    <module name="app.app" />
+    <option name="DEBUGGER_TYPE" value="Hybrid" />
+    <Hybrid>${IDEA_RUN_CONFIGURATION_SYMBOL_PATHS}
+      <startup_commands startup_commands_attr="command script import &quot;${PY_VISUALIZER_PATH}&quot;" />
+    </Hybrid>
+  </configuration>
+</component>
\ No newline at end of file
diff --git a/Engine/Source/Programs/UnrealBuildTool/Modes/GenerateProjectFilesMode.cs b/Engine/Source/Programs/UnrealBuildTool/Modes/GenerateProjectFilesMode.cs
--- a/Engine/Source/Programs/UnrealBuildTool/Modes/GenerateProjectFilesMode.cs
+++ b/Engine/Source/Programs/UnrealBuildTool/Modes/GenerateProjectFilesMode.cs
@@ -37,6 +37,7 @@ namespace UnrealBuildTool
 		[CommandLine("-VSWorkspace", Value = nameof(ProjectFileFormat.VisualStudioWorkspace))]
 		[CommandLine("-CLion", Value = nameof(ProjectFileFormat.CLion))]
 		[CommandLine("-Rider", Value = nameof(ProjectFileFormat.Rider))]
+		[CommandLine("-AndroidStudio", Value = nameof(ProjectFileFormat.AndroidStudio))]
 #if __VPROJECT_AVAILABLE__
 		[CommandLine("-VProject", Value = nameof(ProjectFileFormat.VProject))]
 #endif
@@ -236,6 +237,9 @@ namespace UnrealBuildTool
 					case ProjectFileFormat.Rider:
 						Generator = new RiderProjectFileGenerator(ProjectFile, Arguments);
 						break;
+					case ProjectFileFormat.AndroidStudio:
+						Generator = new AndroidStudioFileGenerator(ProjectFile);
+						break;
 #if __VPROJECT_AVAILABLE__
 					case ProjectFileFormat.VProject:
 						Generator = new VProjectFileGenerator(ProjectFile);
diff --git a/Engine/Source/Programs/UnrealBuildTool/Platform/Android/UEDeployAndroid.cs b/Engine/Source/Programs/UnrealBuildTool/Platform/Android/UEDeployAndroid.cs
--- a/Engine/Source/Programs/UnrealBuildTool/Platform/Android/UEDeployAndroid.cs
+++ b/Engine/Source/Programs/UnrealBuildTool/Platform/Android/UEDeployAndroid.cs
@@ -108,6 +108,12 @@ namespace UnrealBuildTool
 		[CommandLine("-FastIterate", Value = "true")]
 		public bool bFastIterate = false;
 		bool IAndroidDeploy.bFastIterate => bFastIterate;
+		
+		/// <summary>
+		/// Only generates gradle 
+		/// </summary>
+		[CommandLine("-OnlyGradleGeneration", Value = "true")]
+		public bool bOnlyGradleGeneration = false;
 
 		public UEDeployAndroid(FileReference? InProjectFile, bool InForcePackageData, ILogger InLogger)
 			: base(InLogger)
@@ -532,7 +538,7 @@ namespace UnrealBuildTool
 			return true;
 		}
 
-		private static void CopyFileDirectory(string SourceDir, string DestDir, Dictionary<string, string>? Replacements = null, string[]? Excludes = null)
+		public static void CopyFileDirectory(string SourceDir, string DestDir, Dictionary<string, string>? Replacements = null, string[]? Excludes = null)
 		{
 			if (!Directory.Exists(SourceDir))
 			{
@@ -615,7 +621,7 @@ namespace UnrealBuildTool
 			}
 		}
 
-		private static void DeleteDirectory(string InPath, ILogger Logger, string SubDirectoryToKeep = "")
+		public static void DeleteDirectory(string InPath, ILogger Logger, string SubDirectoryToKeep = "")
 		{
 			// skip the dir we want to
 			if (String.Equals(Path.GetFileName(InPath), SubDirectoryToKeep, StringComparison.CurrentCultureIgnoreCase))
@@ -4852,7 +4858,7 @@ dependencies.implementation('androidx.tracing:tracing:1.2.0')
 			Tuple<UnrealArch, string>[] BuildList;
 
 			bool bRequiresOBB = RequiresOBB(bDisallowPackagingDataInApk, ObbFileLocation);
-			if (!bBuildSettingsMatch)
+			if (!bBuildSettingsMatch || bOnlyGradleGeneration)
 			{
 				// LINQ query is subsequently used multiple times, so must copy results
 				BuildList = (from Arch in Architectures.Architectures
@@ -4923,13 +4929,13 @@ dependencies.implementation('androidx.tracing:tracing:1.2.0')
 
 			// Now we have to spin over all the arch/gpu combinations to make sure they all match
 			int BuildListComboTotal = BuildList.Length;
-			if (BuildListComboTotal == 0)
+			if (BuildListComboTotal == 0 && !bOnlyGradleGeneration)
 			{
 				Logger.LogInformation("Output .apk file(s) are up to date (dependencies and build settings are up to date)");
 				return;
 			}
 
-			if (BypassGradlePackaging && !bEnableBundle && !bSkipGradleBuild && !bCreateFromScratch)
+			if (BypassGradlePackaging && !bEnableBundle && !bSkipGradleBuild && !bOnlyGradleGeneration && !bCreateFromScratch)
 			{
 				Logger.LogInformation("Attemping BypassGradlePackaging");
 				int BuildListComboRemaining = BuildListComboTotal;
@@ -5316,7 +5322,7 @@ dependencies.implementation('androidx.tracing:tracing:1.2.0')
 				{
 					SourceSOName = AndroidToolChain.InlineArchName(OutputPath, Arch);
 					string SourceSONameStripped = Path.Combine(Path.GetDirectoryName(SourceSOName)!, Path.GetFileNameWithoutExtension(SourceSOName) + "-stripped" + Path.GetExtension(SourceSOName));
-					if (!File.Exists(SourceSONameStripped))
+					if (!bOnlyGradleGeneration && !File.Exists(SourceSONameStripped))
 					{
 						Logger.LogWarning("Did not find compiled .so [{FinalSONameStripped}]", SourceSONameStripped);
 					}
@@ -5381,10 +5387,13 @@ popd
 					string SOPushScriptName = Path.Combine(Path.GetDirectoryName(SourceSOName)!, "Push_" + Path.GetFileNameWithoutExtension(SourceSOName) + "_so.bat");
 					File.WriteAllText(SOPushScriptName, SOPushScript);
 
-					// Remove libUnreal.so from gradle project if any
-					string JniDir = UnrealBuildPath + "/jni/" + NDKArch;
-					string TempFinalSOName = JniDir + "/libUnreal.so";
-					SafeDeleteFile(TempFinalSOName);
+					if (!bOnlyGradleGeneration)
+					{
+						// Remove libUnreal.so from gradle project if any
+						string JniDir = UnrealBuildPath + "/jni/" + NDKArch;
+						string TempFinalSOName = JniDir + "/libUnreal.so";
+						SafeDeleteFile(TempFinalSOName);
+					}
 				}
 				else
 				{
@@ -5396,11 +5405,11 @@ popd
 					// As we are always making seperate APKs we need to put the architecture into the name
 					DestApkName = AndroidToolChain.InlineArchName(DestApkName, Arch);
 
-					if (!File.Exists(SourceSOName))
+					if (!bOnlyGradleGeneration && !File.Exists(SourceSOName))
 					{
 						throw new BuildException("Can't make an APK without the compiled .so [{0}]", SourceSOName);
 					}
-					if (!Directory.Exists(UnrealBuildPath + "/jni"))
+					if (!bOnlyGradleGeneration && !Directory.Exists(UnrealBuildPath + "/jni"))
 					{
 						throw new BuildException("Can't make an APK without the jni directory [{0}/jni]", UnrealBuildFilesPath);
 					}
@@ -5413,14 +5422,17 @@ popd
 					string LibsDir = Path.Combine(UnrealBuildPath, "libs");
 					DeleteDirectory(LibsDir, Logger);
 					Directory.CreateDirectory(LibsDir);
-
-					string JniDir = Path.Combine(UnrealBuildPath, "jni", NDKArch);
+					
+					if (!bOnlyGradleGeneration)
+					{
+						string JniDir = Path.Combine(UnrealBuildPath, "jni", NDKArch);
 						Directory.CreateDirectory(JniDir);
-					Utils.LinkOrCopyIfDifferent(Path.Combine(JniDir, "libUnreal.so"), SourceSOName);
-
-					string JniLibsDir = Path.Combine(UnrealBuildGradleMainPath, "jniLibs", NDKArch);
-					Directory.CreateDirectory(JniLibsDir);
-					Utils.LinkOrCopyIfDifferent(Path.Combine(JniLibsDir, "libUnreal.so"), SourceSOName);
+						Utils.LinkOrCopyIfDifferent(Path.Combine(JniDir, "libUnreal.so"), SourceSOName);
+	
+						string JniLibsDir = Path.Combine(UnrealBuildGradleMainPath, "jniLibs", NDKArch);
+						Directory.CreateDirectory(JniLibsDir);
+						Utils.LinkOrCopyIfDifferent(Path.Combine(JniLibsDir, "libUnreal.so"), SourceSOName);
+					}
 				}
 
 				// update GameActivity.java and GameApplication.java if out of date
@@ -5515,7 +5527,7 @@ popd
 
 				CheckSOInGradleFolderAre16KbCompatible(UnrealBuildGradleMainPath, Arch);
 
-				if (!bSkipGradleBuild)
+				if (!bSkipGradleBuild && !bOnlyGradleGeneration)
 				{
 					string GradleScriptPath = Path.Combine(UnrealBuildGradlePath, "gradlew");
 					if (!OperatingSystem.IsWindows())
@@ -5635,7 +5647,7 @@ popd
 				Ini.GetBool("/Script/AndroidRuntimeSettings.AndroidRuntimeSettings", "bSaveSymbols", out bool bSaveSymbols);
 				// Deprecated check of bBuildWithHiddenSymbolVisibility and forced bSaveSymbols = true. Remove in 5.8+
 				bSaveSymbols = true;
-				if (bBuildWithHiddenSymbolVisibility || bSaveSymbols)
+				if ((bBuildWithHiddenSymbolVisibility || bSaveSymbols) && !bOnlyGradleGeneration)
 				{
 					// Copy .so with symbols to 
 					int StoreVersion = GetStoreVersion(bEnableBundle ? null : Arch);
@@ -5649,7 +5661,7 @@ popd
 			}
 
 			// Deal with generating an App Bundle
-			if (bEnableBundle && !bSkipGradleBuild)
+			if (bEnableBundle && !bSkipGradleBuild && !bOnlyGradleGeneration)
 			{
 				bool bCombinedBundleOK = true;
 
diff --git a/Engine/Source/Programs/UnrealBuildTool/ProjectFiles/AndroidStudio/AndroidStudioProjectFileGenerator.cs b/Engine/Source/Programs/UnrealBuildTool/ProjectFiles/AndroidStudio/AndroidStudioProjectFileGenerator.cs
new file mode 100644
--- /dev/null
+++ b/Engine/Source/Programs/UnrealBuildTool/ProjectFiles/AndroidStudio/AndroidStudioProjectFileGenerator.cs
@@ -0,0 +1,82 @@
+// Copyright Epic Games, Inc. All Rights Reserved.
+
+using System.Collections.Generic;
+using System.IO;
+using System.Linq;
+using EpicGames.Core;
+using Microsoft.Extensions.Logging;
+using UnrealBuildBase;
+
+namespace UnrealBuildTool
+{
+	class AndroidStudioProjectFile : ProjectFile
+	{
+		public AndroidStudioProjectFile(FileReference InitFilePath, DirectoryReference BaseDir)
+			: base(InitFilePath, BaseDir)
+		{
+		}
+	}
+	
+	class AndroidStudioFileGenerator : ProjectFileGenerator
+	{
+		public override string ProjectFileExtension => ".gradle";
+
+		public AndroidStudioFileGenerator(FileReference? InOnlyGameProject)
+			: base(InOnlyGameProject)
+		{
+		}
+
+		public override bool ShouldGenerateIntelliSenseData() => false;
+
+		protected override void ConfigureProjectFileGeneration(string[] Arguments, ref bool IncludeAllPlatforms, ILogger Logger)
+		{
+			bIncludeEnginePrograms = false;
+			bGeneratingGameProjectFiles = true;
+			
+			base.ConfigureProjectFileGeneration(Arguments, ref IncludeAllPlatforms, Logger);
+		}
+		
+		public override void CleanProjectFiles(DirectoryReference InPrimaryProjectDirectory, string InPrimaryProjectName, DirectoryReference InIntermediateProjectFilesDirectory, ILogger Logger)
+		{
+			string ProjectDir = Path.Combine(InPrimaryProjectDirectory.FullName, $"{InPrimaryProjectName}.gradle");
+			
+			UEDeployAndroid.DeleteDirectory(ProjectDir, Logger);
+			Directory.Delete(ProjectDir);
+		}
+
+		protected override ProjectFile AllocateProjectFile(FileReference InitFilePath, DirectoryReference BaseDir)
+		{
+			return new AndroidStudioProjectFile(InitFilePath, BaseDir);
+		}
+
+		protected override bool WriteProjectFiles(PlatformProjectGeneratorCollection PlatformProjectGenerators, ILogger Logger)
+		{
+			TargetRules ClientTargetRules = DefaultProject!.ProjectTargets.First(Project => Project.TargetRules!.Type == TargetType.Client).TargetRules!;
+			UnrealArchitectures Architectures = VCProjectFileGenerator.GetAndroidProjectArchitectures(OnlyGameProject, true);
+			UnrealArch Architecture = Architectures.Architectures.First();
+
+			string ProjectDir = Path.Combine(PrimaryProjectPath.FullName, $"{PrimaryProjectName}.gradle");
+			string LLDBSymbolsLibsDir = Path.Combine(PrimaryProjectPath.FullName, "Intermediate", "Android", "LLDBSymbolsLibs");
+			
+			Dictionary<string, string> Replacements = new Dictionary<string, string>
+			{
+				{"${PY_VISUALIZER_PATH}", Path.Combine(Unreal.EngineDirectory.FullName, "Extras", "LLDBDataFormatters", "UEDataFormatters_2ByteChars.py")},
+				{"${IDEA_RUN_CONFIGURATION_SYMBOL_PATHS}", $"\n      <symbol_dirs symbol_path=\"{Path.Combine(LLDBSymbolsLibsDir, Architecture.ToString())}\" />"}
+			};
+			
+			UEDeployAndroid.CopyFileDirectory(Path.Combine(Unreal.EngineDirectory.FullName, "Build", "Android", "Java", "gradle"), ProjectDir, Replacements, [".idea", "runConfigurations", "app"]);
+			UEDeployAndroid.CopyFileDirectory(Path.Combine(Unreal.EngineDirectory.FullName, "Build", "AndroidStudio"), ProjectDir, Replacements);
+			{
+				string GradlePropertiesContent = $"org.gradle.jvmargs=-XX:MaxHeapSize=4096m -Xmx9216m\nENGINELOCATION={Unreal.EngineDirectory.FullName.Replace('\\', '/')}\nPROJECT_NAME={PrimaryProjectName}\nTARGET_NAME={ClientTargetRules.Name}\nCONFIGURATION={ClientTargetRules.Configuration}\nUNREAL_ARCH={Architecture}\n";
+				WriteFileIfChanged(Path.Combine(ProjectDir, "gradle.properties"), GradlePropertiesContent, Logger);
+			}
+
+			return true;
+		}
+
+		protected override bool WritePrimaryProjectFile(ProjectFile? UBTProject, PlatformProjectGeneratorCollection PlatformProjectGenerators, ILogger Logger)
+		{
+			return true;
+		}
+	}
+}
diff --git a/Engine/Source/Programs/UnrealBuildTool/ProjectFiles/ProjectFileGenerator.cs b/Engine/Source/Programs/UnrealBuildTool/ProjectFiles/ProjectFileGenerator.cs
--- a/Engine/Source/Programs/UnrealBuildTool/ProjectFiles/ProjectFileGenerator.cs
+++ b/Engine/Source/Programs/UnrealBuildTool/ProjectFiles/ProjectFileGenerator.cs
@@ -151,7 +151,8 @@ namespace UnrealBuildTool
 		VisualStudioCode,
 		VisualStudioMac,
 		CLion,
-		Rider
+		Rider,
+		AndroidStudio
 #if __VPROJECT_AVAILABLE__
 					, VProject
 #endif
diff --git a/Engine/Source/Programs/UnrealBuildTool/ProjectFiles/VisualStudio/VCProjectFileGenerator.cs b/Engine/Source/Programs/UnrealBuildTool/ProjectFiles/VisualStudio/VCProjectFileGenerator.cs
--- a/Engine/Source/Programs/UnrealBuildTool/ProjectFiles/VisualStudio/VCProjectFileGenerator.cs
+++ b/Engine/Source/Programs/UnrealBuildTool/ProjectFiles/VisualStudio/VCProjectFileGenerator.cs
@@ -262,7 +262,7 @@ namespace UnrealBuildTool
 		}
 
 		// parses project ini for Android to get architecture(s) enabled
-		private static UnrealArchitectures GetAndroidProjectArchitectures(FileReference? ProjectFile, bool bGetAllSupported)
+		public static UnrealArchitectures GetAndroidProjectArchitectures(FileReference? ProjectFile, bool bGetAllSupported)
 		{
 			List<string> ActiveArches = new();
 
diff --git a/Engine/Source/Programs/UnrealBuildTool/UnrealBuildTool.cs b/Engine/Source/Programs/UnrealBuildTool/UnrealBuildTool.cs
--- a/Engine/Source/Programs/UnrealBuildTool/UnrealBuildTool.cs
+++ b/Engine/Source/Programs/UnrealBuildTool/UnrealBuildTool.cs
@@ -212,6 +212,7 @@ namespace UnrealBuildTool
 			[CommandLine("-VSMac", Value = "GenerateProjectFiles", Description = "Generate project files for Visual Studio Mac")]
 			[CommandLine("-CLion", Value = "GenerateProjectFiles", Description = "Generate project files for CLion")]
 			[CommandLine("-Rider", Value = "GenerateProjectFiles", Description = "Generate project files for Rider")]
+			[CommandLine("-AndroidStudio", Value = "GenerateProjectFiles", Description = "Generate project files for Android Studio")]
 #if __VPROJECT_AVAILABLE__
 			[CommandLine("-VProject", Value = "GenerateProjectFiles")]
 #endif
