[UBT]
* Commented out very spammy log trace entry

diff --git a/Engine/Source/Programs/UnrealBuildTool/Configuration/CppCompileWarnings.cs b/Engine/Source/Programs/UnrealBuildTool/Configuration/CppCompileWarnings.cs
--- a/Engine/Source/Programs/UnrealBuildTool/Configuration/CppCompileWarnings.cs
+++ b/Engine/Source/Programs/UnrealBuildTool/Configuration/CppCompileWarnings.cs
@@ -1606,7 +1606,8 @@ namespace UnrealBuildTool
 
 				if (defaultWarningLevel != WarningLevel.Default)
 				{
-					logger?.LogTrace("Obtained default value for ({Property}) via the attribute ({AttributeType}).", mi.Name, item.Key.Name);
+					// This is very spammy. Only comment out when debugging
+					//logger?.LogTrace("Obtained default value for ({Property}) via the attribute ({AttributeType}).", mi.Name, item.Key.Name);
 					break;
 				}
 			}
