[UBT]
* Changed so MSBuildProjectFile is lazily creating ProjectGuid since we want it to be resolved in a parallel step (it reads existing file and fetch the guid from there)

diff --git a/Engine/Source/Programs/UnrealBuildTool/ProjectFiles/VisualStudio/VCProject.cs b/Engine/Source/Programs/UnrealBuildTool/ProjectFiles/VisualStudio/VCProject.cs
--- a/Engine/Source/Programs/UnrealBuildTool/ProjectFiles/VisualStudio/VCProject.cs
+++ b/Engine/Source/Programs/UnrealBuildTool/ProjectFiles/VisualStudio/VCProject.cs
@@ -105,32 +105,38 @@ namespace UnrealBuildTool
 		{
 			// Each project gets its own GUID.  This is stored in the project file and referenced in the solution file.
 
-			// First, check to see if we have an existing file on disk.  If we do, then we'll try to preserve the
-			// GUID by loading it from the existing file.
-			if (FileReference.Exists(ProjectFilePath))
+			_projectGuid = new Lazy<Guid>(() =>
 			{
-				try
+				Guid projectGuid = Guid.Empty;
+				// First, check to see if we have an existing file on disk.  If we do, then we'll try to preserve the
+				// GUID by loading it from the existing file.
+				if (FileReference.Exists(ProjectFilePath))
 				{
-					LoadGUIDFromExistingProject();
+					try
+					{
+						projectGuid = InternalLoadGUIDFromExistingProject();
+					}
+					catch (Exception)
+					{
+						// Failed to find GUID, so just create a new one
+						return VCProjectFileGenerator.MakeMd5Guid(PathNamespaceGuid, ProjectFilePath.FullName);
+					}
 				}
-				catch (Exception)
+
+				if (projectGuid == Guid.Empty)
 				{
-					// Failed to find GUID, so just create a new one
-					ProjectGUID = VCProjectFileGenerator.MakeMd5Guid(PathNamespaceGuid, ProjectFilePath.FullName);
+					// Generate a brand new GUID
+					projectGuid = VCProjectFileGenerator.MakeMd5Guid(PathNamespaceGuid, ProjectFilePath.FullName);
 				}
-			}
 
-			if (ProjectGUID == Guid.Empty)
-			{
-				// Generate a brand new GUID
-				ProjectGUID = VCProjectFileGenerator.MakeMd5Guid(PathNamespaceGuid, ProjectFilePath.FullName);
-			}
+				return projectGuid;
+			});
 		}
 
 		/// <summary>
 		/// Attempts to load the project's GUID from an existing project file on disk
 		/// </summary>
-		public override void LoadGUIDFromExistingProject()
+		public Guid InternalLoadGUIDFromExistingProject()
 		{
 			// Only load GUIDs if we're in project generation mode.  Regular builds don't need GUIDs for anything.
 			if (ProjectFileGenerator.bGenerateProjectFiles)
@@ -151,9 +157,19 @@ namespace UnrealBuildTool
 						continue;
 					}
 
-					ProjectGUID = Guid.ParseExact(Element.InnerText.Trim("{}".ToCharArray()), "D");
+					return Guid.ParseExact(Element.InnerText.Trim("{}".ToCharArray()), "D");
 				}
 			}
+
+			return Guid.Empty;
+		}
+
+		public override void LoadGUIDFromExistingProject()
+		{
+			if (!FileReference.Exists(ProjectFilePath))
+			{
+				throw new Exception("Can't find project file " + ProjectFilePath);
+			}
 		}
 
 		/// <summary>
@@ -267,11 +283,9 @@ namespace UnrealBuildTool
 		/// <summary>
 		/// GUID for this Visual C++ project file
 		/// </summary>
-		public Guid ProjectGUID
-		{
-			get;
-			protected set;
-		}
+		public Guid ProjectGUID => _projectGuid.Value;
+
+		Lazy<Guid> _projectGuid;
 	}
 
 	class VCProjectFile : MSBuildProjectFile
@@ -1420,7 +1434,7 @@ namespace UnrealBuildTool
 				StringBuilder[] TmpVCProjectFileContent = new StringBuilder[ProjectConfigAndTargetCombinations.Count];
 				StringBuilder[] TmpVCUserFileContent = new StringBuilder[ProjectConfigAndTargetCombinations.Count];
 
-				Parallel.For(0, ProjectConfigAndTargetCombinations.Count, Index =>
+				Parallel2.For(0, ProjectConfigAndTargetCombinations.Count, Index =>
 				{
 					TmpVCProjectFileContent[Index] = new StringBuilder();
 					TmpVCUserFileContent[Index] = new StringBuilder();
