[UBT]
* TargetRules - Cached SupportedConfigurations per type in concurrent dictionary

diff --git a/Engine/Source/Programs/UnrealBuildTool/Configuration/TargetRules.cs b/Engine/Source/Programs/UnrealBuildTool/Configuration/TargetRules.cs
--- a/Engine/Source/Programs/UnrealBuildTool/Configuration/TargetRules.cs
+++ b/Engine/Source/Programs/UnrealBuildTool/Configuration/TargetRules.cs
@@ -3388,14 +3388,11 @@ namespace UnrealBuildTool
 		/// <returns>Array of configurations that the target supports</returns>
 		internal UnrealTargetConfiguration[] GetSupportedConfigurations()
 		{
-			// Otherwise take the SupportedConfigurationsAttribute from the first type in the inheritance chain that supports it
-			for (Type? currentType = GetType(); currentType != null; currentType = currentType.BaseType)
+
+			UnrealTargetConfiguration[]? configs = GetSupportedConfigurationFromAttribute(GetType());
+			if (configs != null)
 			{
-				object[] attributes = currentType.GetCustomAttributes(typeof(SupportedConfigurationsAttribute), false);
-				if (attributes.Length > 0)
-				{
-					return attributes.OfType<SupportedConfigurationsAttribute>().SelectMany(x => x.Configurations).Distinct().ToArray();
-				}
+				return configs;
 			}
 
 			// Otherwise, get the default for the target type
@@ -3409,23 +3406,6 @@ namespace UnrealBuildTool
 			}
 		}
 
-		static readonly ConcurrentDictionary<Type, Lazy<List<PropertyInfo>>> s_subPropsCache = new();
-
-		static List<PropertyInfo> GetSubObjectProperties(Type t) => s_subPropsCache.GetOrAdd(t, static type => new Lazy<List<PropertyInfo>>(() =>
-		{
-			var props = type.GetProperties(BindingFlags.Instance | BindingFlags.Public);
-			var list = new List<PropertyInfo>();
-			foreach (PropertyInfo pi in props)
-			{
-				if (pi.IsDefined(typeof(ConfigSubObjectAttribute), inherit: false))
-				{
-					list.Add(pi);
-				}
-			}
-			return list;
-		})).Value;
-
-
 		/// <summary>
 		/// Finds all the subobjects which can be configured by command line options and config files
 		/// </summary>
@@ -3605,5 +3585,52 @@ namespace UnrealBuildTool
 		{
 			return UEBuildPlatform.IsPlatformAvailable(TargetPlatform);
 		}
+
+		/// <summary>
+		/// Cached function for getting support configuration from type attributes
+		/// </summary>
+		/// <param name="t"></param>
+		/// <returns></returns>
+		static UnrealTargetConfiguration[]? GetSupportedConfigurationFromAttribute(Type t) => s_supportedConfigurations.GetOrAdd(t, static type => new Lazy<UnrealTargetConfiguration[]?>(() =>
+		{
+			object[] attributes = type.GetCustomAttributes(typeof(SupportedConfigurationsAttribute), false);
+			if (attributes.Length > 0)
+			{
+				return attributes.OfType<SupportedConfigurationsAttribute>().SelectMany(x => x.Configurations).Distinct().ToArray();
+			}
+			if (type.BaseType != null)
+			{
+				return GetSupportedConfigurationFromAttribute(type.BaseType);
+			}
+			return null;
+		})).Value;
+		static readonly ConcurrentDictionary<Type, Lazy<UnrealTargetConfiguration[]?>> s_supportedConfigurations = new();
+
+		/// <summary>
+		/// Cached function for getting properties with ConfigSubObject attribute
+		/// </summary>
+		/// <param name="t"></param>
+		/// <returns></returns>
+		static List<PropertyInfo> GetSubObjectProperties(Type t) => s_subPropsCache.GetOrAdd(t, static type => new Lazy<List<PropertyInfo>>(() =>
+		{
+			var list = new List<PropertyInfo>();
+
+			if (type.BaseType != null)
+			{
+				list.AddRange(GetSubObjectProperties(type.BaseType));
+			}
+
+			var props = type.GetProperties(BindingFlags.Instance | BindingFlags.Public | BindingFlags.DeclaredOnly);
+			foreach (PropertyInfo pi in props)
+			{
+				if (pi.IsDefined(typeof(ConfigSubObjectAttribute), inherit: false))
+				{
+					list.Add(pi);
+				}
+			}
+			return list;
+		})).Value;
+		static readonly ConcurrentDictionary<Type, Lazy<List<PropertyInfo>>> s_subPropsCache = new();
+
 	}
 }
