[UBT]
* Removed timeline scope event that was not useful

diff --git a/Engine/Source/Programs/UnrealBuildTool/System/Plugins.cs b/Engine/Source/Programs/UnrealBuildTool/System/Plugins.cs
--- a/Engine/Source/Programs/UnrealBuildTool/System/Plugins.cs
+++ b/Engine/Source/Programs/UnrealBuildTool/System/Plugins.cs
@@ -342,7 +342,6 @@ namespace UnrealBuildTool
 		/// <returns>Sequence of the found PluginInfo object.</returns>
 		public static IReadOnlyList<PluginInfo> ReadEnginePlugins(DirectoryReference EngineDirectory)
 		{
-			using var _ = Timeline.ScopeEvent("ReadEnginePlugins"); 
 			return ReadPluginsFromDirectory(EngineDirectory, "Plugins", PluginType.Engine);
 		}
 
