[UBT]
* Moved code in lambda in ReadHierarchy into its own method InternalReadHierarchy because it seems like superluminal is messing up the perf info

diff --git a/Engine/Source/Programs/UnrealBuildTool/System/ConfigCache.cs b/Engine/Source/Programs/UnrealBuildTool/System/ConfigCache.cs
--- a/Engine/Source/Programs/UnrealBuildTool/System/ConfigCache.cs
+++ b/Engine/Source/Programs/UnrealBuildTool/System/ConfigCache.cs
@@ -369,46 +369,49 @@ namespace UnrealBuildTool
 
 			// Try to get the cached hierarchy with this key
 			Logger.LogInformation($"Requested Hierarchy: {Type},{ProjectDir},{Platform},{CustomConfig}");
-			return HierarchyKeyToHierarchy.GetOrAdd(Key, new Lazy<ConfigHierarchy>(() =>
+			return HierarchyKeyToHierarchy.GetOrAdd(Key, new Lazy<ConfigHierarchy>(() => InternalReadHierarchy(Logger, Type, ProjectDir, Platform, CustomConfig, HotfixDir, IncludePluginsForTargetType, OverrideStrings))).Value;
+		}
+
+		static ConfigHierarchy InternalReadHierarchy(ILogger Logger, ConfigHierarchyType Type, DirectoryReference? ProjectDir, UnrealTargetPlatform Platform, string CustomConfig
+			, DirectoryReference? HotfixDir, TargetType? IncludePluginsForTargetType, List<string> OverrideStrings)
+		{
+			// Find all the input files
+			List<ConfigFile> Files = new List<ConfigFile>();
+			IEnumerable<FileReference> IniFileNames = ConfigHierarchy.EnumerateConfigFileLocations(Type, ProjectDir, Platform, CustomConfig, IncludePluginsForTargetType);
+			List<FileReference> ReadFiles = [];
+			foreach (FileReference IniFileName in IniFileNames)
 			{
-				// Find all the input files
-				List<ConfigFile> Files = new List<ConfigFile>();
-				IEnumerable<FileReference> IniFileNames = ConfigHierarchy.EnumerateConfigFileLocations(Type, ProjectDir, Platform, CustomConfig, IncludePluginsForTargetType);
-				List<FileReference> ReadFiles = [];
-				foreach (FileReference IniFileName in IniFileNames)
+				ConfigFile? File;
+				Logger.LogInformation($"Trying to load: {IniFileName}");
+				if (TryReadFile(IniFileName, out File))
 				{
-					ConfigFile? File;
-					Logger.LogInformation($"Trying to load: {IniFileName}");
-					if (TryReadFile(IniFileName, out File))
-					{
-						Logger.LogInformation($"    Found!");
-						Files.Add(File);
-						ReadFiles.Add(IniFileName);
-					}
+					Logger.LogInformation($"    Found!");
+					Files.Add(File);
+					ReadFiles.Add(IniFileName);
+				}
 
-					if (HotfixDir != null)
-					{
-						FileReference HotfixFile = FileReference.Combine(HotfixDir, IniFileName.GetFileName());
+				if (HotfixDir != null)
+				{
+					FileReference HotfixFile = FileReference.Combine(HotfixDir, IniFileName.GetFileName());
 
-						if (TryReadFile(HotfixFile, out File))
-						{
-							Logger.LogInformation($"    Hotfixed with '{HotfixFile}'");
-							Files.Add(File);
-							ReadFiles.Add(HotfixFile);
-						}
+					if (TryReadFile(HotfixFile, out File))
+					{
+						Logger.LogInformation($"    Hotfixed with '{HotfixFile}'");
+						Files.Add(File);
+						ReadFiles.Add(HotfixFile);
 					}
 				}
+			}
 
-				foreach (string OverrideString in OverrideStrings)
-				{
-					ConfigFile OverrideFile = new ConfigFile(OverrideString);
-					Files.Add(OverrideFile);
-					Logger.LogInformation($"Trying to load override: {OverrideFile}");
-				}
+			foreach (string OverrideString in OverrideStrings)
+			{
+				ConfigFile OverrideFile = new ConfigFile(OverrideString);
+				Files.Add(OverrideFile);
+				Logger.LogInformation($"Trying to load override: {OverrideFile}");
+			}
 
-				// Create the hierarchy
-				return new ConfigHierarchy(Files, ReadFiles);
-			})).Value;
+			// Create the hierarchy
+			return new ConfigHierarchy(Files, ReadFiles);
 		}
 
 		/// <summary>
