[UBT]
* Optimized solution generation by making more things run in parallel. AddProjectsForAllTargets and AddProjectsForAllModules now mostly run in parallel
* Changed so MSBuildProjectFile is lazily creating ProjectGuid since we want it to be resolved in a parallel step (it reads existing file and fetch the guid from there)

diff --git a/Engine/Source/Programs/UnrealBuildTool/ProjectFiles/ProjectFileGenerator.cs b/Engine/Source/Programs/UnrealBuildTool/ProjectFiles/ProjectFileGenerator.cs
--- a/Engine/Source/Programs/UnrealBuildTool/ProjectFiles/ProjectFileGenerator.cs
+++ b/Engine/Source/Programs/UnrealBuildTool/ProjectFiles/ProjectFileGenerator.cs
@@ -1162,7 +1162,7 @@ namespace UnrealBuildTool
 						{
 							FileReference ProjectFilePath = FileReference.Combine(IntermediateProjectFilesPath, EngineProjectFileNameBase + ProjectFileExtension);
 
-							ProjectFile EngineProject = FindOrAddProject(ProjectFilePath, Unreal.EngineDirectory, true, out _);
+							ProjectFile EngineProject = FindOrAddProject(ProjectFilePath, Unreal.EngineDirectory);
 
 							EngineProject.IsForeignProject = false;
 							EngineProject.IsGeneratedProject = true;
@@ -2650,10 +2650,20 @@ namespace UnrealBuildTool
 			}
 
 			List<ProjectDescriptor> AllGameDescriptors = AllGames.ConvertAll(ProjectDescriptor.FromFile);
-			HashSet<ProjectFile> ProjectsWithPlugins = new HashSet<ProjectFile>();
-			foreach (FileReference CurModuleFile in AllModuleFiles)
+			HashSet<ProjectFile> ProjectsWithPlugins = new();
+			ConcurrentBag<(FileReference, DirectoryReference, ProjectFile)> PluginsToAdd = new();
+
+			// Store Source directory reference because DirectoryReference.Combine is (surprisingly) expensive
+			List<(FileReference, DirectoryReference, ProjectDescriptor, DirectoryReference)> AllGamesZipped = new(AllGames.Count);
+			foreach ((FileReference Game, ProjectDescriptor GameDescriptor) in AllGames.Zip(AllGameDescriptors))
 			{
-				Logger.LogDebug("AddProjectsForAllModules {File}", CurModuleFile);
+				DirectoryReference GameDir = Game.Directory; // Game.Directory is quite expensive so cache it here for FindProjectsForModule loop
+				AllGamesZipped.Add((Game, GameDir, GameDescriptor, DirectoryReference.Combine(GameDir, "Source")));
+			}
+
+			Parallel2.ForEach(AllModuleFiles, (CurModuleFile) =>
+			{
+				using var _ = Timeline.ScopeEvent(CurModuleFile.GetFileName());
 
 				// The module's "base directory" is simply the directory where its xxx.Build.cs file is stored.  We'll always
 				// harvest source files for this module in this base directory directory and all of its sub-directories.
@@ -2671,24 +2681,27 @@ namespace UnrealBuildTool
 					{
 						// We were asked to exclude engine modules from the generated projects
 						WantProjectFileForModule = false;
-						continue;
+						return;
 					}
 				}
 
 				if (!IncludeCppSource)
 				{
-					continue;
+					return;
 				}
 
 				if (WantProjectFileForModule)
 				{
 					DirectoryReference BaseFolder;
-					List<ProjectFile> ProjectFiles = FindProjectsForModule(CurModuleFile, AllGames, AllGameDescriptors, ProgramProjects, ModProjects, ModuleToAdditionalPlugin, out BaseFolder)!;
+					List<ProjectFile> ProjectFiles = FindProjectsForModule(CurModuleFile, AllGamesZipped, ProgramProjects, ModProjects, ModuleToAdditionalPlugin, out BaseFolder)!;
 
 					// Update our module map
 					if (ProjectFiles.Count == 1)
 					{
-						ModuleToEditorProjectFileMap[CurModuleFile] = ProjectFiles[0];
+						lock (ModuleToEditorProjectFileMap)
+						{
+							ModuleToEditorProjectFileMap[CurModuleFile] = ProjectFiles[0];
+						}
 					}
 					else
 					{
@@ -2697,7 +2710,10 @@ namespace UnrealBuildTool
 						ProjectFile? EditorProjectFile = ProjectFiles.FirstOrDefault(x => x.ProjectTargets.Any(x => x.TargetRules!.Type == TargetType.Game));
 						if (EditorProjectFile != null)
 						{
-							ModuleToEditorProjectFileMap[CurModuleFile] = EditorProjectFile;
+							lock (ModuleToEditorProjectFileMap)
+							{
+								ModuleToEditorProjectFileMap[CurModuleFile] = EditorProjectFile;
+							}
 						}
 						else
 						{
@@ -2715,38 +2731,60 @@ namespace UnrealBuildTool
 					// remove any target files, they are technically not part of the module and are explicitly added when the project is created
 					FoundFiles.RemoveAll(f => f.FullName.EndsWith(".Target.cs"));
 
-					if (ModuleToAdditionalPlugin.ContainsKey(CurModuleFile))
+					lock (ModuleToAdditionalPlugin)
 					{
-						ModuleToAdditionalPlugin[CurModuleFile].ModuleToSourceFiles[CurModuleFile] = FoundFiles;
-						continue;
+						if (ModuleToAdditionalPlugin.ContainsKey(CurModuleFile))
+						{
+							ModuleToAdditionalPlugin[CurModuleFile].ModuleToSourceFiles[CurModuleFile] = FoundFiles;
+							return;
+						}
 					}
 
 					foreach (ProjectFile aProjectFile in ProjectFiles)
 					{
-						if (!aProjectFile.Contains(CurModuleFile))
+						lock (aProjectFile)
 						{
-							// Skip if already added
-							aProjectFile.AddFilesToProject(FoundFiles, BaseFolder);
+							if (!aProjectFile.Contains(CurModuleFile))
+							{
+								// Skip if already added
+								aProjectFile.AddFilesToProject(FoundFiles, BaseFolder);
+							}
 						}
+
+						bool added;
 						// Check if there's a plugin directory here
-						if (!ProjectsWithPlugins.Contains(aProjectFile))
+						lock (ProjectsWithPlugins)
 						{
-							foreach (DirectoryReference PluginFolder in Unreal.GetExtensionDirs(BaseFolder, "Plugins"))
+							added = ProjectsWithPlugins.Add(aProjectFile);
+						}
+						if (added)
+						{
+							List<DirectoryReference> plugins = Unreal.GetExtensionDirs(BaseFolder, "Plugins");
+							if (plugins.Count > 0)
 							{
-								// Add all the plugin files for this project
-								foreach (FileReference PluginFileName in PluginsBase.EnumeratePlugins(PluginFolder))
+								foreach (DirectoryReference PluginFolder in plugins)
 								{
-									if (!ModProjects.Any(x => x.BaseDir == PluginFileName.Directory))
+									// Add all the plugin files for this project
+									foreach (FileReference PluginFileName in PluginsBase.EnumeratePlugins(PluginFolder))
 									{
-										AddPluginFilesToProject(PluginFileName, BaseFolder, aProjectFile);
+										if (!ModProjects.Any(x => x.BaseDir == PluginFileName.Directory))
+										{
+											PluginsToAdd.Add((PluginFileName, BaseFolder, aProjectFile));
+										}
 									}
 								}
 							}
-
-							ProjectsWithPlugins.Add(aProjectFile);
 						}
 					}
 				}
+			});
+
+			{
+				using var _ = Timeline.ScopeEvent("AddPlugins");
+				Parallel2.ForEach(PluginsToAdd, (Rec) =>
+				{
+					AddPluginFilesToProject(Rec.Item1, Rec.Item2, Rec.Item3);
+				});
 			}
 
 			if (IncludeCppSource)
@@ -2809,9 +2847,9 @@ namespace UnrealBuildTool
 			// Add plugin shader files if we have any
 			if (PluginDir.TryGetDirectory("Shaders", out DirectoryItem? PluginShadersFolder))
 			{
+				List<FileReference> files = SourceFileSearch.FindFiles(PluginShadersFolder);
 				lock (ProjectFile)
 				{
-					List<FileReference> files = SourceFileSearch.FindFiles(PluginShadersFolder);
 					ProjectFile.AddFilesToProject(files, BaseFolder);
 				}
 			}
@@ -2821,10 +2859,10 @@ namespace UnrealBuildTool
 		{
 			// Setup a project file entry for this module's project.  Remember, some projects may host multiple modules!
 			FileReference ProjectFileName = FileReference.Combine(IntermediateProjectFilesPath, InProjectFileNameBase + ProjectFileExtension);
-			return FindOrAddProject(ProjectFileName, InBaseFolder, IncludeInGeneratedProjects: true, bAlreadyExisted: out _);
+			return FindOrAddProject(ProjectFileName, InBaseFolder);
 		}
 
-		private List<ProjectFile> FindProjectsForModule(FileReference CurModuleFile, List<FileReference> AllGames, List<ProjectDescriptor> AllGameDescriptors, Dictionary<FileReference, ProjectFile> ProgramProjects, List<ProjectFile> ModProjects, Dictionary<FileReference, AdditionalPluginData> ModuleToAdditionalPlugin, out DirectoryReference BaseFolder)
+		private List<ProjectFile> FindProjectsForModule(FileReference CurModuleFile, List<(FileReference, DirectoryReference, ProjectDescriptor, DirectoryReference)> AllGames, Dictionary<FileReference, ProjectFile> ProgramProjects, List<ProjectFile> ModProjects, Dictionary<FileReference, AdditionalPluginData> ModuleToAdditionalPlugin, out DirectoryReference BaseFolder)
 		{
 			// Starting at the base directory of the module find a project which has the same directory as base, walking up the directory hierarchy until a match is found
 
@@ -2835,10 +2873,10 @@ namespace UnrealBuildTool
 			while (!Path.IsRootDirectory())
 			{
 				// Figure out which game project this target belongs to
-				foreach ((FileReference Game, ProjectDescriptor GameDescriptor) in AllGames.Zip(AllGameDescriptors))
+				foreach ((FileReference Game, DirectoryReference GameDir, ProjectDescriptor GameDescriptor, DirectoryReference SourceDir) in AllGames)
 				{
 					// the source and the actual game directory are conceptually the same
-					if (Path == Game.Directory || Path == DirectoryReference.Combine(Game.Directory, "Source"))
+					if (Path == GameDir || Path == SourceDir)
 					{
 						FileReference ProjectInfo = Game;
 						BaseFolder = ProjectInfo.Directory;
@@ -2855,7 +2893,7 @@ namespace UnrealBuildTool
 							ModuleDescriptor? CurModuleDescriptor = GameDescriptor.Modules?.FirstOrDefault(x => x.Name == CurModuleFile.GetFileNameWithoutAnyExtensions());
 
 							// find the project that the module is under, and has a TargetType target (useful with bMakeProjectPerTarget)
-							foreach (KeyValuePair<FileReference, ProjectFile> Pair in ProjectFileMap)
+							foreach (Lazy<ProjectFile> Pair in ProjectFileMap.Values)
 							{
 								FileReference TargetFile = ((ProjectTarget)Pair.Value.ProjectTargets[0]).TargetFilePath;
 								bool bIsTemporaryTarget = TargetFile.ContainsName("Intermediate", 0);
@@ -2990,43 +3028,75 @@ namespace UnrealBuildTool
 			// Keep a cache of the default editor target for each project so that we don't have to interrogate the configs for each target
 			Dictionary<string, string?> DefaultProjectEditorTargetCache = new Dictionary<string, string?>();
 
-			RulesCompiler.CreateEngineRulesAssembly(false, false, false, Logger);
+			RulesAssembly EngineRulesAssembly = RulesCompiler.CreateEngineRulesAssembly(false, false, false, Logger);
+
+			List<FileReference> AllTargetsExceptSubTargets = AllTargetFiles.Except(AllSubTargetFiles).ToList();
+
+			Dictionary<FileReference, (
+				TargetRules TargetRulesObject,
+				string ProjectFileNameBase,
+				DirectoryReference? GameFolder,
+				string TargetName,
+				bool IsProgramTarget,
+				bool IsEngineTarget,
+				RulesAssembly RulesAssembly,
+				FileReference? CheckProjectFile,
+				bool IsHybridContentOnlyProject)> TargetInfos = [];
 
 			List<FileReference> UnprocessedGameProjects = new(AllGames);
-			foreach (FileReference TargetFilePath in AllTargetFiles.Except(AllSubTargetFiles))
+
+			// We will create the projects in two steps because CreateProjectRulesAssemly is very expensive and trigger c# compiler etc.
+			// .. so first part need to run in parallel. Second part must be in order so TargetInfos dictionary will store info created in first step and used in second step
+
+			// First step
+			Parallel2.ForEach(AllTargetsExceptSubTargets, (TargetFilePath) =>
 			{
-				string TargetName = TargetFilePath.GetFileNameWithoutAnyExtensions();       // Remove both ".cs" and ".Target"
+				using var _ = Timeline.ScopeEvent(TargetFilePath.GetFileNameWithoutAnyExtensions());
 
-				// Check to see if this is an Engine target.  That is, the target is located under the "Engine" folder
+				TargetRules TargetRulesObject;
+				string ProjectFileNameBase;
+				DirectoryReference? GameFolder = null;
+				string TargetName;
+				bool IsProgramTarget = false;
 				bool IsEngineTarget = false;
-				bool WantProjectFileForTarget = true;
-				bool ForceProgramInProject = false;
-				if (TargetFilePath.IsUnderDirectory(Unreal.EngineDirectory))
+				FileReference? CheckProjectFile;
+
 				{
-					// This is an engine target
-					IsEngineTarget = true;
+					TargetName = TargetFilePath.GetFileNameWithoutAnyExtensions();       // Remove both ".cs" and ".Target"
 
-					// if we are generating projects for a single game project, then we want the target, independent of bIncludeEngine flags
-					if (DoesProgramMatchOnlyGameProject(TargetFilePath))
-					{
-						WantProjectFileForTarget = true;
-						ForceProgramInProject = true;
-					}
-					else if (Unreal.GetExtensionDirs(Unreal.EngineDirectory, "Source/Programs").Any(x => TargetFilePath.IsUnderDirectory(x)))
+					// Check to see if this is an Engine target.  That is, the target is located under the "Engine" folder
+					bool WantProjectFileForTarget = true;
+					bool ForceProgramInProject = false;
+
+					if (TargetFilePath.IsUnderDirectory(Unreal.EngineDirectory))
 					{
-						WantProjectFileForTarget = bIncludeEnginePrograms;
+						// This is an engine target
+						IsEngineTarget = true;
+
+						// if we are generating projects for a single game project, then we want the target, independent of bIncludeEngine flags
+						if (DoesProgramMatchOnlyGameProject(TargetFilePath))
+						{
+							WantProjectFileForTarget = true;
+							ForceProgramInProject = true;
+						}
+						else if (Unreal.GetExtensionDirs(Unreal.EngineDirectory, "Source/Programs").Any(x => TargetFilePath.IsUnderDirectory(x)))
+						{
+							WantProjectFileForTarget = bIncludeEnginePrograms;
+						}
+						else if (Unreal.GetExtensionDirs(Unreal.EngineDirectory, "Source").Any(x => TargetFilePath.IsUnderDirectory(x)))
+						{
+							WantProjectFileForTarget = bIncludeEngineSource;
+						}
 					}
-					else if (Unreal.GetExtensionDirs(Unreal.EngineDirectory, "Source").Any(x => TargetFilePath.IsUnderDirectory(x)))
+
+					if (!WantProjectFileForTarget)
 					{
-						WantProjectFileForTarget = bIncludeEngineSource;
+						return;
 					}
-				}
 
-				if (WantProjectFileForTarget)
-				{
-					RulesAssembly RulesAssembly;
+					RulesAssembly RulesAssembly = EngineRulesAssembly;
 
-					FileReference? CheckProjectFile = AllGames.FirstOrDefault(x => TargetFilePath.IsUnderDirectory(x.Directory));
+					CheckProjectFile = AllGames.FirstOrDefault(x => TargetFilePath.IsUnderDirectory(x.Directory));
 					// if it wasn't found, it may be a program, and the target file isn't under the uproject, so look up by name
 					if (CheckProjectFile == null && TargetFilePath.ContainsName("Programs", 0))
 					{
@@ -3036,16 +3106,17 @@ namespace UnrealBuildTool
 
 					if (CheckProjectFile == null)
 					{
-						RulesAssembly = RulesCompiler.CreateEngineRulesAssembly(false, false, false, Logger);
-
 						// Record the Engine assembly, and any parent assemblies (varies e.g. Rules, ProgramRules, MarketplaceRules)
-						if (!RulesAssemblies.ContainsKey(RulesAssembly))
+						lock (RulesAssemblies)
 						{
-							RulesAssembly? NextAssembly = RulesAssembly;
-							do
+							if (!RulesAssemblies.ContainsKey(RulesAssembly))
 							{
-								RulesAssemblies.Add(NextAssembly, Unreal.EngineDirectory);
-							} while ((NextAssembly = NextAssembly.Parent) != null);
+								RulesAssembly? NextAssembly = RulesAssembly;
+								do
+								{
+									RulesAssemblies.Add(NextAssembly, Unreal.EngineDirectory);
+								} while ((NextAssembly = NextAssembly.Parent) != null);
+							}
 						}
 					}
 					else
@@ -3053,22 +3124,21 @@ namespace UnrealBuildTool
 						RulesAssembly = RulesCompiler.CreateProjectRulesAssembly(CheckProjectFile, false, false, false, Logger);
 						// Recording of game rule assemblies happens after BaseFolder has been computed
 
-						UnprocessedGameProjects.Remove(CheckProjectFile);
+						lock (UnprocessedGameProjects)
+						{
+							UnprocessedGameProjects.Remove(CheckProjectFile);
+						}
 					}
 
 					// Create target rules for all of the platforms and configuration combinations that we want to enable support for.
 					// Just use the current platform as we only need to recover the target type and both should be supported for all targets...
-					TargetRules TargetRulesObject = RulesAssembly.CreateTargetRules(TargetName, BuildHostPlatform.Current.Platform, UnrealTargetConfiguration.Development, null, CheckProjectFile, new CommandLineArguments(GetTargetArguments(Arguments)), Logger);
-
-					bool IsProgramTarget = false;
+					TargetRulesObject = RulesAssembly.CreateTargetRules(TargetName, BuildHostPlatform.Current.Platform, UnrealTargetConfiguration.Development, null, CheckProjectFile, new CommandLineArguments(GetTargetArguments(Arguments)), Logger, false, TargetRulesValidationOptions.ValidateTargetAndSDK, UnrealIntermediateEnvironment.GenerateProjectFiles);
 
-					DirectoryReference? GameFolder = null;
-					string ProjectFileNameBase;
 					if (TargetRulesObject.Type == TargetType.Program || TargetRulesObject.bGenerateProgramProject)
 					{
 						if (!ForceProgramInProject && !bIncludeEnginePrograms && IsEngineTarget)
 						{
-							continue;
+							return;
 						}
 						IsProgramTarget = true;
 						ProjectFileNameBase = TargetName;
@@ -3089,16 +3159,46 @@ namespace UnrealBuildTool
 						ProjectFileNameBase = ProjectInfo.GetFileNameWithoutExtension();
 					}
 
-					FileReference ProjectFilePath;
-					if (bMakeProjectPerTarget)
+					bool IsHybridContentOnlyProject = CheckProjectFile != null && NativeProjects.IsHybridContentOnlyProject(CheckProjectFile, out string? _, Logger);
+
+
+					lock (TargetInfos)
 					{
-						ProjectFilePath = GetProjectLocation(TargetRulesObject.Name);
+						TargetInfos.Add(TargetFilePath, (TargetRulesObject, ProjectFileNameBase, GameFolder, TargetName, IsProgramTarget, IsEngineTarget, RulesAssembly, CheckProjectFile, IsHybridContentOnlyProject));
 					}
-					else
+				}
+			});
+
+			// Second step
+			foreach (FileReference TargetFilePath in AllTargetsExceptSubTargets)
+			{
+				(
+				TargetRules TargetRulesObject,
+				string ProjectFileNameBase,
+				DirectoryReference? GameFolder,
+				string TargetName,
+				bool IsProgramTarget,
+				bool IsEngineTarget,
+				RulesAssembly RulesAssembly,
+				FileReference? CheckProjectFile,
+				bool IsHybridContentOnlyProject
+				) = TargetInfos[TargetFilePath];
+
+				FileReference ProjectFilePath;
+				DirectoryReference BaseFolder;
+
+				if (bMakeProjectPerTarget)
+				{
+					ProjectFilePath = GetProjectLocation(TargetRulesObject.Name);
+				}
+				else
+				{
+					// Look at the project engine config to see if it has specified a default editor target
+					FileReference ProjectLocation = GetProjectLocation(ProjectFileNameBase);
+					string? DefaultEditorTarget = null;
+
+					lock (DefaultProjectEditorTargetCache)
 					{
-						// Look at the project engine config to see if it has specified a default editor target
-						FileReference ProjectLocation = GetProjectLocation(ProjectFileNameBase);
-						string? DefaultEditorTarget = null;
 						if (!DefaultProjectEditorTargetCache.TryGetValue(ProjectFileNameBase, out DefaultEditorTarget))
 						{
 							if (GameFolder != null)
@@ -3108,153 +3208,153 @@ namespace UnrealBuildTool
 
 							DefaultProjectEditorTargetCache.Add(ProjectFileNameBase, DefaultEditorTarget);
 						}
+					}
 
-						// Get the suffix to use for this project file. If we have multiple targets of the same type, we'll have to split them out into separate IDE project files.
-						string? GeneratedProjectName = TargetRulesObject.GeneratedProjectName;
-						if (GeneratedProjectName == null)
-						{
-							ProjectFile? ExistingProjectFile;
-							// We should create a separate project for targets which aren't the default for this target type.
-							// Note that we currently only support changing the default editor target and not any other types, so
-							// if we aren't an editor, we'll just fall back to previous behavior and assume the first one we encounter
-							// should be added to the main project
-							bool bIsDefaultTargetForType = (DefaultEditorTarget == null) || (TargetRulesObject.Type != TargetType.Editor) || (TargetRulesObject.Name == DefaultEditorTarget);
-							if (!bIsDefaultTargetForType || (ProjectFileMap.TryGetValue(ProjectLocation, out ExistingProjectFile) && ExistingProjectFile.ProjectTargets.Any(x => x.TargetRules!.Type == TargetRulesObject.Type)))
-							{
-								GeneratedProjectName = TargetRulesObject.Name;
-							}
-							else
-							{
-								GeneratedProjectName = ProjectFileNameBase;
-							}
+					// Get the suffix to use for this project file. If we have multiple targets of the same type, we'll have to split them out into separate IDE project files.
+					string? GeneratedProjectName = TargetRulesObject.GeneratedProjectName;
+					if (GeneratedProjectName == null)
+					{
+						Lazy<ProjectFile>? ExistingProjectFile;
+						// We should create a separate project for targets which aren't the default for this target type.
+						// Note that we currently only support changing the default editor target and not any other types, so
+						// if we aren't an editor, we'll just fall back to previous behavior and assume the first one we encounter
+						// should be added to the main project
+						bool bIsDefaultTargetForType = (DefaultEditorTarget == null) || (TargetRulesObject.Type != TargetType.Editor) || (TargetRulesObject.Name == DefaultEditorTarget);
+						if (!bIsDefaultTargetForType || (ProjectFileMap.TryGetValue(ProjectLocation, out ExistingProjectFile) && ExistingProjectFile.Value.ProjectTargets.Any(x => x.TargetRules!.Type == TargetRulesObject.Type)))
+						{
+							GeneratedProjectName = TargetRulesObject.Name;
 						}
-
-						// @todo projectfiles: We should move all of the Target.cs files out of sub-folders to clean up the project directories a bit (e.g. GameUncooked folder)
-						ProjectFilePath = GetProjectLocation(GeneratedProjectName);
-						if (TargetRulesObject.Type == TargetType.Game || TargetRulesObject.Type == TargetType.Client || TargetRulesObject.Type == TargetType.Server)
-						{
-							// Allow platforms to generate stub projects here...
-							PlatformProjectGenerators.GenerateGameProjectStubs(
-								InGenerator: this,
-								InTargetName: TargetName,
-								InTargetFilepath: TargetFilePath.FullName,
-								InTargetRules: TargetRulesObject,
-								InPlatforms: SupportedPlatforms,
-								InConfigurations: SupportedConfigurations);
+						else
+						{
+							GeneratedProjectName = ProjectFileNameBase;
 						}
 					}
 
-					DirectoryReference BaseFolder;
-					if (IsProgramTarget)
+					// @todo projectfiles: We should move all of the Target.cs files out of sub-folders to clean up the project directories a bit (e.g. GameUncooked folder)
+					ProjectFilePath = GetProjectLocation(GeneratedProjectName);
+					if (TargetRulesObject.Type == TargetType.Game || TargetRulesObject.Type == TargetType.Client || TargetRulesObject.Type == TargetType.Server)
 					{
-						BaseFolder = TargetFilePath.Directory;
+						// Allow platforms to generate stub projects here...
+						PlatformProjectGenerators.GenerateGameProjectStubs(
+							InGenerator: this,
+							InTargetName: TargetName,
+							InTargetFilepath: TargetFilePath.FullName,
+							InTargetRules: TargetRulesObject,
+							InPlatforms: SupportedPlatforms,
+							InConfigurations: SupportedConfigurations);
 					}
-					else if (IsEngineTarget)
-					{
-						BaseFolder = Unreal.EngineDirectory;
-					}
-					else
-					{
-						BaseFolder = GameFolder!;
+				}
+
+				if (IsProgramTarget)
+				{
+					BaseFolder = TargetFilePath.Directory;
+				}
+				else if (IsEngineTarget)
+				{
+					BaseFolder = Unreal.EngineDirectory;
+				}
+				else
+				{
+					BaseFolder = GameFolder!;
 
+					lock (RulesAssemblies)
+					{
 						RulesAssemblies.TryAdd(RulesAssembly, BaseFolder);
 					}
+				}
 
-					bool bProjectAlreadyExisted;
-					ProjectFile ProjectFile = FindOrAddProject(ProjectFilePath, BaseFolder, IncludeInGeneratedProjects: true, bAlreadyExisted: out bProjectAlreadyExisted);
-					ProjectFile.IsForeignProject = CheckProjectFile != null && !NativeProjects.IsNativeProject(CheckProjectFile, Logger);
-					ProjectFile.IsGeneratedProject = true;
-					ProjectFile.IsStubProject = Unreal.IsProjectInstalled();
-					ProjectFile.IsHybridContentOnlyProject = CheckProjectFile != null && NativeProjects.IsHybridContentOnlyProject(CheckProjectFile, out string? _, Logger);
-					if (TargetRulesObject.bBuildInSolutionByDefault.HasValue)
+				ProjectFile ProjectFile = FindOrAddProject(ProjectFilePath, BaseFolder);
+				ProjectFile.IsForeignProject = CheckProjectFile != null && !NativeProjects.IsNativeProject(CheckProjectFile, Logger);
+				ProjectFile.IsGeneratedProject = true;
+				ProjectFile.IsStubProject = Unreal.IsProjectInstalled();
+				ProjectFile.IsHybridContentOnlyProject = IsHybridContentOnlyProject;
+				if (TargetRulesObject.bBuildInSolutionByDefault.HasValue)
+				{
+					ProjectFile.ShouldBuildByDefaultForSolutionTargets = TargetRulesObject.bBuildInSolutionByDefault.Value;
+				}
+
+				// Add the project to the right output list
+				if (IsProgramTarget)
+				{
+					ProgramProjects[TargetFilePath] = ProjectFile;
+				}
+				else if (IsEngineTarget)
+				{
+					// if we want only one project for all types, then remove any previous ones
+					if (bMakeProjectPerTarget == false)
 					{
-						ProjectFile.ShouldBuildByDefaultForSolutionTargets = TargetRulesObject.bBuildInSolutionByDefault.Value;
+						EngineProjects.Clear();
 					}
+					EngineProjects.Add(ProjectFile);
 
-					// Add the project to the right output list
-					if (IsProgramTarget)
+					if (Unreal.IsEngineInstalled())
 					{
-						ProgramProjects[TargetFilePath] = ProjectFile;
+						// Allow engine projects to be created but not built for Installed Engine builds
+						ProjectFile.IsForeignProject = false;
+						ProjectFile.IsGeneratedProject = true;
+						ProjectFile.IsStubProject = true;
 					}
-					else if (IsEngineTarget)
+				}
+				else
+				{
+					if (!GameProjects.Contains(ProjectFile))
 					{
-						// if we want only one project for all types, then remove any previous ones
-						if (bMakeProjectPerTarget == false)
-						{
-							EngineProjects.Clear();
-						}
-						EngineProjects.Add(ProjectFile);
+						GameProjects.Add(ProjectFile);
 
-						if (Unreal.IsEngineInstalled())
+						// Add the .uproject file for this game/template
+						FileReference UProjectFilePath = FileReference.Combine(BaseFolder, ProjectFileNameBase + ".uproject");
+						if (FileReference.Exists(UProjectFilePath))
 						{
-							// Allow engine projects to be created but not built for Installed Engine builds
-							ProjectFile.IsForeignProject = false;
-							ProjectFile.IsGeneratedProject = true;
-							ProjectFile.IsStubProject = true;
+							ProjectFile.AddFileToProject(UProjectFilePath, BaseFolder);
+							ProjectFileToUProjectFile.Add(ProjectFile, UProjectFilePath);
 						}
-					}
-					else
-					{
-						if (!bProjectAlreadyExisted)
+						else
 						{
-							GameProjects.Add(ProjectFile);
-
-							// Add the .uproject file for this game/template
-							FileReference UProjectFilePath = FileReference.Combine(BaseFolder, ProjectFileNameBase + ".uproject");
-							if (FileReference.Exists(UProjectFilePath))
-							{
-								ProjectFile.AddFileToProject(UProjectFilePath, BaseFolder);
-								ProjectFileToUProjectFile.Add(ProjectFile, UProjectFilePath);
-							}
-							else
-							{
-								throw new BuildException("Not expecting to find a game with no .uproject file.  File '{0}' doesn't exist", UProjectFilePath);
-							}
+							throw new BuildException("Not expecting to find a game with no .uproject file.  File '{0}' doesn't exist", UProjectFilePath);
 						}
 					}
+				}
 
-					foreach (Project ExistingProjectTarget in ProjectFile.ProjectTargets)
+				foreach (Project ExistingProjectTarget in ProjectFile.ProjectTargets)
+				{
+					if (ExistingProjectTarget.TargetRules!.Type == TargetRulesObject.Type)
 					{
-						if (ExistingProjectTarget.TargetRules!.Type == TargetRulesObject.Type)
-						{
-							throw new BuildException("Not expecting project {0} to already have a target rules of with configuration name {1} ({2}) while trying to add: {3}", ProjectFilePath, TargetRulesObject.Type.ToString(), ExistingProjectTarget.TargetRules.ToString(), TargetRulesObject.ToString());
-						}
+						throw new BuildException("Not expecting project {0} to already have a target rules of with configuration name {1} ({2}) while trying to add: {3}", ProjectFilePath, TargetRulesObject.Type.ToString(), ExistingProjectTarget.TargetRules.ToString(), TargetRulesObject.ToString());
+					}
 
-						// Not expecting to have both a game and a program in the same project.  These would alias because we share the project and solution configuration names (just because it makes sense to)
-						if ((ExistingProjectTarget.TargetRules.Type == TargetType.Game && TargetRulesObject.Type == TargetType.Program) ||
-							(ExistingProjectTarget.TargetRules.Type == TargetType.Program && TargetRulesObject.Type == TargetType.Game))
-						{
-							throw new BuildException("Not expecting project {0} to already have a Game/Program target ({1}) associated with it while trying to add: {2}", ProjectFilePath, ExistingProjectTarget.TargetRules.ToString(), TargetRulesObject.ToString());
-						}
+					// Not expecting to have both a game and a program in the same project.  These would alias because we share the project and solution configuration names (just because it makes sense to)
+					if ((ExistingProjectTarget.TargetRules.Type == TargetType.Game && TargetRulesObject.Type == TargetType.Program) ||
+						(ExistingProjectTarget.TargetRules.Type == TargetType.Program && TargetRulesObject.Type == TargetType.Game))
+					{
+						throw new BuildException("Not expecting project {0} to already have a Game/Program target ({1}) associated with it while trying to add: {2}", ProjectFilePath, ExistingProjectTarget.TargetRules.ToString(), TargetRulesObject.ToString());
 					}
+				}
 
-					ProjectTarget ProjectTarget = new ProjectTarget
-						(
-							TargetRules: TargetRulesObject,
-							TargetFilePath: TargetFilePath,
-							ProjectFilePath: ProjectFilePath,
-							UnrealProjectFilePath: CheckProjectFile,
-							SupportedPlatforms: TargetRulesObject.GetSupportedPlatforms().Where(
-													x => UEBuildPlatform.TryGetBuildPlatform(x, out UEBuildPlatform? _) &&
-													(TargetRulesObject.LinkType != TargetLinkType.Modular || !UEBuildPlatform.PlatformRequiresMonolithicBuilds(x, TargetRulesObject.Configuration))
-													).ToArray(),
-							CreateRulesDelegate: (Platform, Configuration) => RulesAssembly.CreateTargetRules(TargetName, Platform, Configuration, null, CheckProjectFile, new CommandLineArguments(GetTargetArguments(Arguments)), Logger),
-							CreateRulesDelegateWithArch: (Platform, Configuration, Arch) => RulesAssembly.CreateTargetRules(TargetName, Platform, Configuration, new(Arch), CheckProjectFile, new CommandLineArguments(GetTargetArguments(Arguments)), Logger)
-						);
+				ProjectTarget ProjectTarget = new ProjectTarget
+					(
+						TargetRules: TargetRulesObject,
+						TargetFilePath: TargetFilePath,
+						ProjectFilePath: ProjectFilePath,
+						UnrealProjectFilePath: CheckProjectFile,
+						SupportedPlatforms: TargetRulesObject.GetSupportedPlatforms().Where(
+												x => UEBuildPlatform.TryGetBuildPlatform(x, out UEBuildPlatform? _) &&
+												(TargetRulesObject.LinkType != TargetLinkType.Modular || !UEBuildPlatform.PlatformRequiresMonolithicBuilds(x, TargetRulesObject.Configuration))
+												).ToArray(),
+						CreateRulesDelegate: (Platform, Configuration) => RulesAssembly.CreateTargetRules(TargetName, Platform, Configuration, null, CheckProjectFile, new CommandLineArguments(GetTargetArguments(Arguments)), Logger, false, TargetRulesValidationOptions.ValidateNothing, UnrealIntermediateEnvironment.GenerateProjectFiles),
+						CreateRulesDelegateWithArch: (Platform, Configuration, Arch) => RulesAssembly.CreateTargetRules(TargetName, Platform, Configuration, new(Arch), CheckProjectFile, new CommandLineArguments(GetTargetArguments(Arguments)), Logger, false, TargetRulesValidationOptions.ValidateNothing, UnrealIntermediateEnvironment.GenerateProjectFiles)
+					);
 
-					ProjectFile.ProjectTargets.Add(ProjectTarget);
+				ProjectFile.ProjectTargets.Add(ProjectTarget);
 
-					// Make sure the *.Target.cs file is in the project.
-					ProjectFile.AddFileToProject(TargetFilePath, BaseFolder);
+				// Make sure the *.Target.cs file is in the project.
+				ProjectFile.AddFileToProject(TargetFilePath, BaseFolder);
 
-					// Add all matching *_<platform>.Target.cs to the same folder
-					if (AllSubTargetFilesPerTarget.ContainsKey(TargetName))
-					{
-						ProjectFile.AddFilesToProject(AllSubTargetFilesPerTarget[TargetName], BaseFolder);
-					}
-
-					Logger.LogDebug("Generating target {Target} for {Project}", TargetRulesObject.Type.ToString(), ProjectFilePath);
+				// Add all matching *_<platform>.Target.cs to the same folder
+				if (AllSubTargetFilesPerTarget.ContainsKey(TargetName))
+				{
+					ProjectFile.AddFilesToProject(AllSubTargetFilesPerTarget[TargetName], BaseFolder);
 				}
+				Logger.LogDebug("Generating target {Target} for {Project}", TargetRulesObject.Type.ToString(), ProjectFilePath);
 			}
 
 			// if we allow content only projects, we assume anything that wasn't processed above is content-only, because it would have had a target 
@@ -3280,11 +3380,11 @@ namespace UnrealBuildTool
 								// don't append Game suffix to the ContentOnly game target names, so the default case has the target/product named after the BP project name
 								// @note if this changes, see ApplePlatform.MakeContentOnlyTargetName
 								string TargetTypeSuffix = /*EngineTarget.TargetRules!.Type == TargetType.Game ? "" :*/ EngineTarget.TargetRules!.Type.ToString();
-								ProjectFile = FindOrAddProject(GetProjectLocation($"{ProjectName}{TargetTypeSuffix}"), ContentOnlyGameProject.Directory, IncludeInGeneratedProjects: true, bAlreadyExisted: out _);
+								ProjectFile = FindOrAddProject(GetProjectLocation($"{ProjectName}{TargetTypeSuffix}"), ContentOnlyGameProject.Directory);
 							}
 							else
 							{
-								ProjectFile ??= FindOrAddProject(GetProjectLocation(ProjectName), ContentOnlyGameProject.Directory, IncludeInGeneratedProjects: true, bAlreadyExisted: out _);
+								ProjectFile ??= FindOrAddProject(GetProjectLocation(ProjectName), ContentOnlyGameProject.Directory);
 							}
 
 							ProjectFile.IsForeignProject = true;
@@ -3318,8 +3418,7 @@ namespace UnrealBuildTool
 					{
 						FileReference ModProjectFilePath = FileReference.Combine(PluginInfo.Directory, "Mods", PluginInfo.Name + ProjectFileExtension);
 
-						bool bProjectAlreadyExisted;
-						ProjectFile ModProjectFile = FindOrAddProject(ModProjectFilePath, PluginInfo.Directory, IncludeInGeneratedProjects: true, bAlreadyExisted: out bProjectAlreadyExisted);
+						ProjectFile ModProjectFile = FindOrAddProject(ModProjectFilePath, PluginInfo.Directory);
 						ModProjectFile.IsForeignProject = GameProject.IsForeignProject;
 						ModProjectFile.IsGeneratedProject = true;
 						ModProjectFile.IsStubProject = false;
@@ -3438,35 +3537,40 @@ namespace UnrealBuildTool
 		/// </summary>
 		/// <param name="FilePath">Full path to the project file</param>
 		/// <param name="BaseDir">Base directory for files in this project</param>
-		/// <param name="IncludeInGeneratedProjects">True if this project should be included in the set of generated projects.  Only matters when actually generating project files.</param>
-		/// <param name="bAlreadyExisted">True if we already had this project file</param>
 		/// <returns>Object that represents this project file in Unreal Build Tool</returns>
-		public ProjectFile FindOrAddProject(FileReference FilePath, DirectoryReference BaseDir, bool IncludeInGeneratedProjects, out bool bAlreadyExisted)
+		public ProjectFile FindOrAddProject(FileReference FilePath, DirectoryReference BaseDir)
 		{
 			if (FilePath == null)
 			{
 				throw new BuildException("Not valid to call FindOrAddProject() with an empty file path!");
 			}
 
+			bool added = false;
+			Lazy<ProjectFile>? file;
+
 			// Do we already have this project?
-			ProjectFile? ExistingProjectFile;
-			if (ProjectFileMap.TryGetValue(FilePath, out ExistingProjectFile))
+			lock (ProjectFileMap)
 			{
-				bAlreadyExisted = true;
-				return ExistingProjectFile;
+				if (!ProjectFileMap.TryGetValue(FilePath, out file))
+				{
+					// Add a new project file for the specified path
+					file = new Lazy<ProjectFile>(AllocateProjectFile(FilePath, BaseDir));
+					ProjectFileMap[FilePath] = file;
+					added = true;
+				}
 			}
 
-			// Add a new project file for the specified path
-			ProjectFile NewProjectFile = AllocateProjectFile(FilePath, BaseDir);
-			ProjectFileMap[FilePath] = NewProjectFile;
+			ProjectFile f = file.Value;
 
-			if (IncludeInGeneratedProjects)
+			if (added)
 			{
-				GeneratedProjectFiles.Add(NewProjectFile);
+				lock (GeneratedProjectFiles)
+				{
+					GeneratedProjectFiles.Add(f);
+				}
 			}
 
-			bAlreadyExisted = false;
-			return NewProjectFile;
+			return f;
 		}
 
 		/// <summary>
@@ -3835,7 +3939,7 @@ namespace UnrealBuildTool
 		/// Map of project file names to their project files.  This includes every single project file in memory or otherwise that
 		/// we know about so far.  Note that when generating project files, this map may even include project files that we won't
 		/// be including in the generated projects.
-		protected readonly Dictionary<FileReference, ProjectFile> ProjectFileMap = new Dictionary<FileReference, ProjectFile>();
+		protected readonly Dictionary<FileReference, Lazy<ProjectFile>> ProjectFileMap = new();// Dictionary<FileReference, ProjectFile>();
 
 		/// List of project files that we'll be generating
 		protected readonly List<ProjectFile> GeneratedProjectFiles = new List<ProjectFile>();
diff --git a/Engine/Source/Programs/UnrealBuildTool/ProjectFiles/Rider/RiderProjectFileGenerator.cs b/Engine/Source/Programs/UnrealBuildTool/ProjectFiles/Rider/RiderProjectFileGenerator.cs
--- a/Engine/Source/Programs/UnrealBuildTool/ProjectFiles/Rider/RiderProjectFileGenerator.cs
+++ b/Engine/Source/Programs/UnrealBuildTool/ProjectFiles/Rider/RiderProjectFileGenerator.cs
@@ -279,9 +279,9 @@ namespace UnrealBuildTool
 					string? GeneratedProjectName = TargetRulesObject.GeneratedProjectName;
 					if (GeneratedProjectName == null)
 					{
-						ProjectFile? ExistingProjectFile;
+						Lazy<ProjectFile>? ExistingProjectFile;
 						if (ProjectFileMap.TryGetValue(GetRiderProjectLocation(ProjectFileNameBase), out ExistingProjectFile) &&
-							ExistingProjectFile.ProjectTargets.Any(x => x.TargetRules!.Type == TargetRulesObject.Type))
+							ExistingProjectFile.Value.ProjectTargets.Any(x => x.TargetRules!.Type == TargetRulesObject.Type))
 						{
 							GeneratedProjectName = TargetRulesObject.Name;
 						}
@@ -319,9 +319,7 @@ namespace UnrealBuildTool
 						BaseFolder = GameFolder!;
 					}
 
-					bool bProjectAlreadyExisted;
-					ProjectFile ProjectFile = FindOrAddProject(ProjectFilePath, BaseFolder,
-						true, out bProjectAlreadyExisted);
+					ProjectFile ProjectFile = FindOrAddProject(ProjectFilePath, BaseFolder);
 					ProjectFile.IsForeignProject =
 						CheckProjectFile != null && !NativeProjects.IsNativeProject(CheckProjectFile, Logger);
 					ProjectFile.IsGeneratedProject = true;
@@ -350,7 +348,7 @@ namespace UnrealBuildTool
 					}
 					else
 					{
-						if (!bProjectAlreadyExisted)
+						if (!GameProjects.Contains(ProjectFile))
 						{
 							GameProjects.Add(ProjectFile);
 
@@ -509,8 +507,7 @@ namespace UnrealBuildTool
 				FileReference ProjectFilePath = FileReference.Combine(IntermediateProjectFilesPath,
 					EngineProjectFileNameBase + ProjectFileExtension);
 
-				bool bAlreadyExisted;
-				EngineProject = FindOrAddProject(ProjectFilePath, Unreal.EngineDirectory, true, out bAlreadyExisted);
+				EngineProject = FindOrAddProject(ProjectFilePath, Unreal.EngineDirectory);
 
 				EngineProject.IsForeignProject = false;
 				EngineProject.IsGeneratedProject = true;
