Fix to pathing issue in root Visual Studio project calling UAT for Android deployment

#android

diff --git a/Engine/Source/Programs/UnrealBuildTool/Platform/Android/AndroidProjectGenerator.cs b/Engine/Source/Programs/UnrealBuildTool/Platform/Android/AndroidProjectGenerator.cs
--- a/Engine/Source/Programs/UnrealBuildTool/Platform/Android/AndroidProjectGenerator.cs
+++ b/Engine/Source/Programs/UnrealBuildTool/Platform/Android/AndroidProjectGenerator.cs
@@ -188,7 +188,7 @@ namespace UnrealBuildTool
 			}
 		}
 
-		public override string GetVisualStudioUserFileStrings(VisualStudioUserFileSettings VCUserFileSettings, VSSettings InVSSettings, string InConditionString, TargetRules InTargetRules, FileReference TargetRulesPath, FileReference ProjectFilePath, FileReference? NMakeOutputPath, string ProjectName, string? ForeignUProjectPath)
+		public override string GetVisualStudioUserFileStrings(VisualStudioUserFileSettings VCUserFileSettings, VSSettings InVSSettings, string InConditionString, TargetRules InTargetRules, FileReference TargetRulesPath, FileReference ProjectFilePath, FileReference? NMakeOutputPath, string ProjectName, string UProjectPath, string? ForeignUProjectPath)
 		{
 			if (AGDEInstalled
 				&& (InVSSettings.Platform == UnrealTargetPlatform.Android)
@@ -216,7 +216,7 @@ namespace UnrealBuildTool
 					// so regardless of what value it has, Split(':')[1] will never throw an exception.
 					string GetTargetDeviceSerial = "$([System.String]::Concat($(AndroidDebugTarget), \"::\").Split(':')[1])";
 					Out.AppendLine(
-						$"    <AndroidPostApkInstallCommands>CALL \"{UatPath}\" BuildCookRun -project=\"$(SolutionDir){ProjectName}.uproject\" -FastIterate -FromMsBuild -client -skipbuild -skipcook -skipstage -deploy -config={InVSSettings.Configuration} -platform={InVSSettings.Platform} -clientarchitecture={ShortAbi} -device=\"{GetTargetDeviceSerial}\";$(AndroidPostApkInstallCommands)</AndroidPostApkInstallCommands>");
+						$"    <AndroidPostApkInstallCommands>CALL \"{UatPath}\" BuildCookRun -project=\"{UProjectPath.Replace("\"", "")}\" -FastIterate -FromMsBuild -client -skipbuild -skipcook -skipstage -deploy -config={InVSSettings.Configuration} -platform={InVSSettings.Platform} -clientarchitecture={ShortAbi} -device=\"{GetTargetDeviceSerial}\";$(AndroidPostApkInstallCommands)</AndroidPostApkInstallCommands>");
 
 					// Ensure the following properties are up to date even if .vcxproj.user already exists and are in correct order between each other.
 					VCUserFileSettings.PatchProperty("AndroidDebugTarget", true);
diff --git a/Engine/Source/Programs/UnrealBuildTool/Platform/Windows/WindowsProjectGenerator.cs b/Engine/Source/Programs/UnrealBuildTool/Platform/Windows/WindowsProjectGenerator.cs
--- a/Engine/Source/Programs/UnrealBuildTool/Platform/Windows/WindowsProjectGenerator.cs
+++ b/Engine/Source/Programs/UnrealBuildTool/Platform/Windows/WindowsProjectGenerator.cs
@@ -112,7 +112,7 @@ namespace UnrealBuildTool
 		}
 
 		/// <inheritdoc/>
-		public override string GetVisualStudioUserFileStrings(VisualStudioUserFileSettings VCUserFileSettings, VSSettings InVSSettings, string InConditionString, TargetRules InTargetRules, FileReference TargetRulesPath, FileReference ProjectFilePath, FileReference? NMakeOutputPath, string ProjectName, string? ForeignUProjectPath)
+		public override string GetVisualStudioUserFileStrings(VisualStudioUserFileSettings VCUserFileSettings, VSSettings InVSSettings, string InConditionString, TargetRules InTargetRules, FileReference TargetRulesPath, FileReference ProjectFilePath, FileReference? NMakeOutputPath, string ProjectName, string UProjectPath, string? ForeignUProjectPath)
 		{
 			StringBuilder VCUserFileContent = new StringBuilder();
 
diff --git a/Engine/Source/Programs/UnrealBuildTool/ProjectFiles/PlatformProjectGenerator.cs b/Engine/Source/Programs/UnrealBuildTool/ProjectFiles/PlatformProjectGenerator.cs
--- a/Engine/Source/Programs/UnrealBuildTool/ProjectFiles/PlatformProjectGenerator.cs
+++ b/Engine/Source/Programs/UnrealBuildTool/ProjectFiles/PlatformProjectGenerator.cs
@@ -317,9 +317,10 @@ namespace UnrealBuildTool
 		/// <param name="ProjectFilePath">The project file path</param>
 		/// <param name="NMakeOutputPath">Output path for NMake</param>
 		/// <param name="ProjectName">The name of the project</param>
+		/// <param name="UProjectPath">Path to the project .uproject file</param>
 		/// <param name="ForeignUProjectPath">Path to foreign .uproject file, if any</param>
 		/// <returns>The string to append to the user file</returns>
-		public virtual string GetVisualStudioUserFileStrings(VisualStudioUserFileSettings VCUserFileSettings, VSSettings InVSSettings, string InConditionString, TargetRules InTargetRules, FileReference TargetRulesPath, FileReference ProjectFilePath, FileReference? NMakeOutputPath, string ProjectName, string? ForeignUProjectPath)
+		public virtual string GetVisualStudioUserFileStrings(VisualStudioUserFileSettings VCUserFileSettings, VSSettings InVSSettings, string InConditionString, TargetRules InTargetRules, FileReference TargetRulesPath, FileReference ProjectFilePath, FileReference? NMakeOutputPath, string ProjectName, string UProjectPath, string? ForeignUProjectPath)
 		{
 			return String.Empty;
 		}
diff --git a/Engine/Source/Programs/UnrealBuildTool/ProjectFiles/VisualStudio/VCProject.cs b/Engine/Source/Programs/UnrealBuildTool/ProjectFiles/VisualStudio/VCProject.cs
--- a/Engine/Source/Programs/UnrealBuildTool/ProjectFiles/VisualStudio/VCProject.cs
+++ b/Engine/Source/Programs/UnrealBuildTool/ProjectFiles/VisualStudio/VCProject.cs
@@ -2561,7 +2561,7 @@ namespace UnrealBuildTool
 					if (ProjGenerator != null)
 					{
 						string? ForeignUProjectPath = (IsForeignProject && !String.IsNullOrEmpty(UProjectPath)) ? UProjectPath : null;
-						VCUserFileContent.Append(ProjGenerator.GetVisualStudioUserFileStrings(VCUserFileSettings, new(Combination.Platform!.Value, Configuration, ProjectFileFormat, Combination.Architecture), ConditionString, TargetRulesObject, Combination.ProjectTarget.TargetFilePath, ProjectFilePath, NMakePath, ProjectName, ForeignUProjectPath));
+						VCUserFileContent.Append(ProjGenerator.GetVisualStudioUserFileStrings(VCUserFileSettings, new(Combination.Platform!.Value, Configuration, ProjectFileFormat, Combination.Architecture), ConditionString, TargetRulesObject, Combination.ProjectTarget.TargetFilePath, ProjectFilePath, NMakePath, ProjectName, UProjectPath, ForeignUProjectPath));
 					}
 				}
 			}
