Rename gameApplicationInterfaceAdditions to gameApplicationImplementsAdditions and add UPL comment documentation entry

#rnx

diff --git a/Engine/Build/Android/Java/src/com/epicgames/unreal/GameApplication.java.template b/Engine/Build/Android/Java/src/com/epicgames/unreal/GameApplication.java.template
--- a/Engine/Build/Android/Java/src/com/epicgames/unreal/GameApplication.java.template
+++ b/Engine/Build/Android/Java/src/com/epicgames/unreal/GameApplication.java.template
@@ -16,7 +16,7 @@ import java.util.concurrent.locks.ReentrantLock;
 import com.epicgames.unreal.network.NetworkChangedManager;
 //$${gameApplicationImportAdditions}$$
 
-public class GameApplication extends Application implements LifecycleObserver //$${gameApplicationInterfaceAdditions}$$
+public class GameApplication extends Application implements LifecycleObserver //$${gameApplicationImplementsAdditions}$$
 {
 	public static final boolean IS_SHIPPING_CONFIG = $${gameApplicationIsShippingConfig}$$;
 	
diff --git a/Engine/Plugins/Runtime/AndroidFileServer/Source/AndroidFileServer/AndroidFileServer_UPL.xml b/Engine/Plugins/Runtime/AndroidFileServer/Source/AndroidFileServer/AndroidFileServer_UPL.xml
--- a/Engine/Plugins/Runtime/AndroidFileServer/Source/AndroidFileServer/AndroidFileServer_UPL.xml
+++ b/Engine/Plugins/Runtime/AndroidFileServer/Source/AndroidFileServer/AndroidFileServer_UPL.xml
@@ -364,9 +364,9 @@ tasks.whenTaskAdded { task ->
 		</if>
 	</settingsGradleAdditions>
 
-	<gameApplicationInterfaceAdditions>
+	<gameApplicationImplementsAdditions>
 		<insert>, IAndroidFileServerSetup </insert>
-	</gameApplicationInterfaceAdditions>
+	</gameApplicationImplementsAdditions>
 
 	<gameApplicationImportAdditions>
 		<if condition="bEmbedAFS">
diff --git a/Engine/Source/Programs/UnrealBuildTool/Platform/Android/UEDeployAndroid.cs b/Engine/Source/Programs/UnrealBuildTool/Platform/Android/UEDeployAndroid.cs
--- a/Engine/Source/Programs/UnrealBuildTool/Platform/Android/UEDeployAndroid.cs
+++ b/Engine/Source/Programs/UnrealBuildTool/Platform/Android/UEDeployAndroid.cs
@@ -6570,7 +6570,7 @@ import java.util.Collection;
 
 			Dictionary<string, string> Replacements = new Dictionary<string, string>{
 				{ "//$${gameApplicationImportAdditions}$$", UPL!.ProcessPluginNode(NDKArch, "gameApplicationImportAdditions", GameApplicationImportAdditionsDefault)},
-				{ "//$${gameApplicationInterfaceAdditions}$$", UPL!.ProcessPluginNode(NDKArch, "gameApplicationInterfaceAdditions", "")},
+				{ "//$${gameApplicationImplementsAdditions}$$", UPL!.ProcessPluginNode(NDKArch, "gameApplicationImplementsAdditions", "")},
 				{ "//$${gameApplicationOnCreateAdditions}$$", UPL.ProcessPluginNode(NDKArch, "gameApplicationOnCreateAdditions", "")},
 				{ "//$${gameApplicationAttachBaseContextAdditions}$$", UPL.ProcessPluginNode(NDKArch, "gameApplicationAttachBaseContextAdditions", "")},
 				{ "//$${gameApplicationOnLowMemoryAdditions}$$", UPL.ProcessPluginNode(NDKArch, "gameApplicationOnLowMemoryAdditions", "")},
diff --git a/Engine/Source/Programs/UnrealBuildTool/System/UnrealPluginLanguage.cs b/Engine/Source/Programs/UnrealBuildTool/System/UnrealPluginLanguage.cs
--- a/Engine/Source/Programs/UnrealBuildTool/System/UnrealPluginLanguage.cs
+++ b/Engine/Source/Programs/UnrealBuildTool/System/UnrealPluginLanguage.cs
@@ -434,6 +434,9 @@ namespace UnrealBuildTool
 	 * 	<!-- optional additions to GameActivity onActivityResult in GameActivity.java -->
 	 * 	<gameActivityOnActivityResultAdditions>	</gameActivityOnActivityResultAdditions>
 	 * 	
+	 * 	<!-- optional additional interfaces that GameApplication can implement in GameApplication.java (must start with a comma and list interfaces using commas) -->
+	 * 	<gameApplicationImplementsAdditions> </gameApplicationImplementsAdditions>
+	 * 	
 	 * 	<!-- optional libraries to load in GameActivity.java before libUnreal.so -->
 	 * 	<soLoadLibrary>	</soLoadLibrary>
 	 * 	
