[UBT]
* Added bloom filter to front disallowed platforms logic. The test is insanely expensive but almost always doesn't hit, so a bloom filter to filter out paths that does not contain any platforms is a huge win

diff --git a/Engine/Source/Programs/UnrealBuildTool/System/RulesAssembly.cs b/Engine/Source/Programs/UnrealBuildTool/System/RulesAssembly.cs
--- a/Engine/Source/Programs/UnrealBuildTool/System/RulesAssembly.cs
+++ b/Engine/Source/Programs/UnrealBuildTool/System/RulesAssembly.cs
@@ -3,9 +3,11 @@
 using System;
 using System.Collections.Generic;
 using System.Diagnostics.CodeAnalysis;
+using System.IO;
 using System.Linq;
 using System.Reflection;
 using System.Runtime.CompilerServices;
+using System.Text;
 using EpicGames.Core;
 using Microsoft.Extensions.Logging;
 using UnrealBuildBase;
@@ -751,14 +753,51 @@ namespace UnrealBuildTool
 				// figure out what platforms/groups aren't allowed with this opted in list
 				List<string> DisallowedPlatformsAndGroups = Utils.MakeListOfUnsupportedPlatforms(Rules.OptedInModulePlatforms.ToList(), false, Logger);
 
+				// This is really slow (disallow list can contain 20 items and there can be 1000 plugins with long paths.. so search space is huuuuge)
+				// Use a bloom filter to minimize the expensive testing
+				BloomFilter bloomFilter = new();
+				bloomFilter.AddRange(DisallowedPlatformsAndGroups);
+
+				DirectoryReference? ProjectDir = Rules.ProjectFile?.Directory;
+
 				// look in all plugins' paths to see if any disallowed
-				IEnumerable<PluginInfo> DisallowedPlugins = EnumeratePlugins().Where(x => x.ChoiceVersion != null).Select(x => x.ChoiceVersion!).Where(Plugin =>
-					Plugin.File.ContainsAnyNames(DisallowedPlatformsAndGroups, Unreal.EngineDirectory) ||
-					(Rules.ProjectFile != null && Plugin.File.ContainsAnyNames(DisallowedPlatformsAndGroups, Rules.ProjectFile.Directory)));
-				// log out the plugins we are disabling
-				DisallowedPlugins.ToList().ForEach(x => Logger.LogDebug("Disallowing non-opted-in platform plugin {PluginFile}", x.File));
+				List<PluginInfo> DisallowedPluginsList = [];
+
+				foreach (PluginSet set in EnumeratePlugins())
+				{
+					if (set.ChoiceVersion == null)
+					{
+						continue;
+					}
+
+					PluginInfo info = set.ChoiceVersion;
+
+					bool mightContain = false;
+					if (info.File.IsUnderDirectory(Unreal.EngineDirectory))
+					{
+						var span = info.File.FullName.AsSpan(Unreal.EngineDirectory.FullName.Length + 1);
+						if (bloomFilter.MightContainSubString(span, Path.DirectorySeparatorChar))
+							mightContain = true;
+					}
+
+					if (!mightContain && ProjectDir != null && info.File.IsUnderDirectory(ProjectDir))
+					{
+						var span = info.File.FullName.AsSpan(ProjectDir.FullName.Length + 1);
+						if (bloomFilter.MightContainSubString(span, Path.DirectorySeparatorChar))
+							mightContain = true;
+					}
+
+					if (mightContain)
+					{
+						if (info.File.ContainsAnyNames(DisallowedPlatformsAndGroups, Unreal.EngineDirectory) ||
+						(ProjectDir != null && info.File.ContainsAnyNames(DisallowedPlatformsAndGroups, ProjectDir)))
+							DisallowedPluginsList.Add(info);
+					}
+				}
+
+				DisallowedPluginsList.ForEach(x => Logger.LogDebug("Disallowing non-opted-in platform plugin {PluginFile}", x.File));
 				// and, disable these plugins
-				Rules.DisablePlugins.AddRange(DisallowedPlugins.Select(x => x.Name));
+				Rules.DisablePlugins.AddRange(DisallowedPluginsList.Select(x => x.Name));
 			}
 
 			// Allow the platform to finalize the settings
