[UBT]
* Fixed bug in project generation code where it tried to read target info for a target that was not created

diff --git a/Engine/Source/Programs/UnrealBuildTool/ProjectFiles/ProjectFileGenerator.cs b/Engine/Source/Programs/UnrealBuildTool/ProjectFiles/ProjectFileGenerator.cs
--- a/Engine/Source/Programs/UnrealBuildTool/ProjectFiles/ProjectFileGenerator.cs
+++ b/Engine/Source/Programs/UnrealBuildTool/ProjectFiles/ProjectFileGenerator.cs
@@ -3172,6 +3172,11 @@ namespace UnrealBuildTool
 			// Second step
 			foreach (FileReference TargetFilePath in AllTargetsExceptSubTargets)
 			{
+				if (!TargetInfos.ContainsKey(TargetFilePath))
+				{
+					continue;
+				}
+
 				(
 				TargetRules TargetRulesObject,
 				string ProjectFileNameBase,
