ICX 2025.2 update & fixes

diff --git a/Engine/Config/Windows/Windows_SDK.json b/Engine/Config/Windows/Windows_SDK.json
--- a/Engine/Config/Windows/Windows_SDK.json
+++ b/Engine/Config/Windows/Windows_SDK.json
@@ -43,9 +43,9 @@
 	"MinimumVisualCppVersion": "14.38.33130",
 
 	"PreferredIntelOneApiVersions": [
-		"2025.1.0-2025.9999"
+		"2025.2.0-2025.2.9999"
 	],
-	"MinimumIntelOneApiVersion": "2025.1.0",
+	"MinimumIntelOneApiVersion": "2025.2.0",
 
 	"//04.1": "Microsoft.VisualStudio.Component.VC.14.44.17.14.x86.x64 -- Only LTSC versions of MSVC should be suggested",
 	"//04.2": "Microsoft.VisualStudio.Component.VC.14.44.17.14.ATL -- Match version of MSVC",
diff --git a/Engine/Source/Developer/DerivedDataCache/Private/FileSystemCacheStore.cpp b/Engine/Source/Developer/DerivedDataCache/Private/FileSystemCacheStore.cpp
--- a/Engine/Source/Developer/DerivedDataCache/Private/FileSystemCacheStore.cpp
+++ b/Engine/Source/Developer/DerivedDataCache/Private/FileSystemCacheStore.cpp
@@ -304,6 +304,11 @@ struct FFileSystemCacheStoreMaintainer::FRoot
 	bool bScannedRoot = false;
 };
 
+// Temp workaround for icx compiler crash
+#if defined(__INTEL_LLVM_COMPILER)
+#pragma optimize("", off)
+#endif // #if defined(__INTEL_LLVM_COMPILER)
+// End temp workaround for icx compiler crash
 struct FFileSystemCacheStoreMaintainer::FLegacyRoot
 {
 	inline explicit FLegacyRoot(FRandomStream& Stream)
@@ -319,6 +324,11 @@ struct FFileSystemCacheStoreMaintainer::FLegacyRoot
 	TStaticBitArray<10> ExistsLevel2[10][10];
 	uint32 DirectoryScanCount = 0;
 };
+// Temp workaround for icx compiler crash
+#if defined(__INTEL_LLVM_COMPILER)
+#pragma optimize("", on)
+#endif // #if defined(__INTEL_LLVM_COMPILER)
+// End temp workaround for icx compiler crash
 
 FFileSystemCacheStoreMaintainer::FFileSystemCacheStoreMaintainer(
 	const FFileSystemCacheStoreMaintainerParams& InParams,
diff --git a/Engine/Source/Programs/UnrealBuildTool/Configuration/CppCompileWarnings.cs b/Engine/Source/Programs/UnrealBuildTool/Configuration/CppCompileWarnings.cs
--- a/Engine/Source/Programs/UnrealBuildTool/Configuration/CppCompileWarnings.cs
+++ b/Engine/Source/Programs/UnrealBuildTool/Configuration/CppCompileWarnings.cs
@@ -1135,16 +1135,17 @@ namespace UnrealBuildTool
 		private WarningLevel _formatWarningLevel;
 
 		/// <summary>
-		/// Indicates what warning/error level to treat instances of implicit conversion between floating-point types of different precision.
+		/// Indicates what warning/error level to treat instances of implicit int conversion.
+		/// Clang - https://clang.llvm.org/docs/DiagnosticsReference.html#wimplicit-int-conversion
 		/// </summary>
-		[WarningsIntelCompiler(["-Wno-implicit-float-size-conversion"], ["-Wimplicit-float-size-conversion", "-Wno-error=implicit-float-size-conversion"], ["-Wimplicit-float-size-conversion"])]
+		[WarningsIntelCompiler(["-Wno-implicit-int-conversion"], ["-Wimplicit-int-conversion", "-Wno-error=implicit-int-conversion"], ["-Wimplicit-int-conversion"])]
 		[BasicWarningLevelDefault(WarningLevel.Off)]
-		public WarningLevel ImplicitFloatSizeConversionWarningLevel
+		public WarningLevel ImplicitIntConversionWarningLevel
 		{
-			get => ResolveWarning(ParentCppCompileWarnings?.ImplicitFloatSizeConversionWarningLevel, _implicitFloatSizeConversionWarningLevel);
-			set => _implicitFloatSizeConversionWarningLevel = value;
+			get => ResolveWarning(ParentCppCompileWarnings?.ImplicitIntConversionWarningLevel, _implicitIntConversionWarningLevel);
+			set => _implicitIntConversionWarningLevel = value;
 		}
-		private WarningLevel _implicitFloatSizeConversionWarningLevel;
+		private WarningLevel _implicitIntConversionWarningLevel;
 
 		/// <summary>
 		/// Indicates what warning/error level to treat instances where a command line argument is provided, but not used for anything.
diff --git a/Engine/Source/Programs/UnrealBuildTool/Configuration/ReadOnlyCppCompileWarnings.cs b/Engine/Source/Programs/UnrealBuildTool/Configuration/ReadOnlyCppCompileWarnings.cs
--- a/Engine/Source/Programs/UnrealBuildTool/Configuration/ReadOnlyCppCompileWarnings.cs
+++ b/Engine/Source/Programs/UnrealBuildTool/Configuration/ReadOnlyCppCompileWarnings.cs
@@ -124,7 +124,7 @@ namespace UnrealBuildTool
 		#region -- Intel Compiler Disables --
 
 		public WarningLevel FormatWarningLevel => _inner.FormatWarningLevel;
-		public WarningLevel ImplicitFloatSizeConversionWarningLevel => _inner.ImplicitFloatSizeConversionWarningLevel;
+		public WarningLevel ImplicitIntConversionWarningLevel => _inner.ImplicitIntConversionWarningLevel;
 		public WarningLevel SingleBitfieldConstantConversionWarningLevel => _inner.SingleBitfieldConstantConversionWarningLevel;
 		public WarningLevel UnusedCommandLineArgumentWarningLevel => _inner.UnusedCommandLineArgumentWarningLevel;
 		public WarningLevel CommentWarningLevel => _inner.CommentWarningLevel;
