[UBT][DefaultBuildSettings] Default Build Settings Updates for Targets - V6 - engine

diff --git a/Engine/Source/Programs/UnrealBuildTool/Configuration/TargetRules.cs b/Engine/Source/Programs/UnrealBuildTool/Configuration/TargetRules.cs
--- a/Engine/Source/Programs/UnrealBuildTool/Configuration/TargetRules.cs
+++ b/Engine/Source/Programs/UnrealBuildTool/Configuration/TargetRules.cs
@@ -152,12 +152,19 @@ namespace UnrealBuildTool
 		/// </summary>
 		V5,
 
+		/// <summary>
+		/// New defaults for 5.7:
+		/// * DanglingWarningLevel = WarningLevel.Warning
+		/// * UndefinedIdentifierWarningLevel = WarningLevel.Error
+		/// </summary>
+		V6,
+
 		// *** When adding new entries here, be sure to update GameProjectUtils::GetDefaultBuildSettingsVersion() to ensure that new projects are created correctly. ***
 
 		/// <summary>
 		/// Always use the defaults for the current engine version. Note that this may cause compatibility issues when upgrading.
 		/// </summary>
-		Latest = V5,
+		Latest = V6,
 	}
 
 	/// <summary>
