[UBT][DefaultBuildSettings] Build Settings V6 to be used for wundef and dangling.

diff --git a/Engine/Source/Programs/UnrealBuildTool/Configuration/CppCompileWarnings.cs b/Engine/Source/Programs/UnrealBuildTool/Configuration/CppCompileWarnings.cs
--- a/Engine/Source/Programs/UnrealBuildTool/Configuration/CppCompileWarnings.cs
+++ b/Engine/Source/Programs/UnrealBuildTool/Configuration/CppCompileWarnings.cs
@@ -146,7 +146,8 @@ namespace UnrealBuildTool
 		[CommandLine("-UndefinedIdentifierWarningLevel=")]
 		[UndefinedIdentifierWarningsVCToolChain(null, ["/w44668"], ["/we4668"])]
 		[WarningsClangToolChain(null, ["-Wundef", "-Wno-error=undef"], ["-Wundef"])]
-		[BasicWarningLevelDefault(WarningLevel.Off)]
+		[VersionWarningLevelDefault(WarningLevel.Off, BuildSettingsVersion.V1, BuildSettingsVersion.V5, InitializationContext.Any)]
+		[VersionWarningLevelDefault(WarningLevel.Error, BuildSettingsVersion.V6, BuildSettingsVersion.Latest, InitializationContext.Any)]
 		public WarningLevel UndefinedIdentifierWarningLevel
 		{
 			get => ResolveWarning(ParentCppCompileWarnings?.UndefinedIdentifierWarningLevel, _undefinedIdentifierWarningLevel);
@@ -1045,7 +1046,8 @@ namespace UnrealBuildTool
 		///		https://clang.llvm.org/docs/DiagnosticsReference.html#wdangling
 		/// </summary>	
 		[WarningsVCClang(["-Wno-dangling"], ["-Wdangling", "-Wno-error=dangling"], ["-Wdangling"])]
-		[BasicWarningLevelDefault(WarningLevel.Off)]
+		[VersionWarningLevelDefault(WarningLevel.Off, BuildSettingsVersion.V1, BuildSettingsVersion.V5, InitializationContext.Any)]
+		[VersionWarningLevelDefault(WarningLevel.Warning, BuildSettingsVersion.V6, BuildSettingsVersion.Latest, InitializationContext.Any)]
 		public WarningLevel DanglingWarningLevel
 		{
 			get => ResolveWarning(ParentCppCompileWarnings?.DanglingWarningLevel, _danglingWarningLevel);
diff --git a/Engine/Source/Programs/UnrealBuildTool/Configuration/TargetRules.cs b/Engine/Source/Programs/UnrealBuildTool/Configuration/TargetRules.cs
--- a/Engine/Source/Programs/UnrealBuildTool/Configuration/TargetRules.cs
+++ b/Engine/Source/Programs/UnrealBuildTool/Configuration/TargetRules.cs
@@ -3515,6 +3515,12 @@ namespace UnrealBuildTool
 					modifiedSettings.Add(Tuple.Create(String.Format("{0} = true", nameof(bValidateFormatStrings)), "Enables compile-time validation of strings+args passed to UE_LOG. (Previously: false)."));
 				}
 
+				if (BuildSettingsVersion.V6 <= latestVersion && DefaultBuildSettings < BuildSettingsVersion.V6)
+				{
+					modifiedSettings.Add(Tuple.Create(String.Format("{0} = WarningLevel.Warning", nameof(CppCompileWarningSettings.DanglingWarningLevel)), "Enables compile-time validation of dangling pointers. (Previously: Off)."));
+					modifiedSettings.Add(Tuple.Create(String.Format("{0} = WarningLevel.Error", nameof(CppCompileWarningSettings.UndefinedIdentifierWarningLevel)), "Enables compile-time validation of undefined identifiers in conditional expressions. (Previously: Off)."));
+				}
+
 				if (modifiedSettings.Count > 0)
 				{
 					string formatString = String.Format("[Upgrade]     {{0,-{0}}}   => {{1}}", modifiedSettings.Max(x => x.Item1.Length));
