Catch of DirectoryNotFoundException during Clean of Android Studio Project

diff --git a/Engine/Source/Programs/UnrealBuildTool/ProjectFiles/AndroidStudio/AndroidStudioProjectFileGenerator.cs b/Engine/Source/Programs/UnrealBuildTool/ProjectFiles/AndroidStudio/AndroidStudioProjectFileGenerator.cs
--- a/Engine/Source/Programs/UnrealBuildTool/ProjectFiles/AndroidStudio/AndroidStudioProjectFileGenerator.cs
+++ b/Engine/Source/Programs/UnrealBuildTool/ProjectFiles/AndroidStudio/AndroidStudioProjectFileGenerator.cs
@@ -41,7 +41,13 @@ namespace UnrealBuildTool
 			string ProjectDir = Path.Combine(InPrimaryProjectDirectory.FullName, $"{InPrimaryProjectName}.gradle");
 			
 			UEDeployAndroid.DeleteDirectory(ProjectDir, Logger);
-			Directory.Delete(ProjectDir);
+			try
+			{
+				Directory.Delete(ProjectDir);
+			}
+			catch (DirectoryNotFoundException)
+			{
+			}
 		}
 
 		protected override ProjectFile AllocateProjectFile(FileReference InitFilePath, DirectoryReference BaseDir)
