UnrealBuildTool: Don't warn on the root module directory being an include path for plugins, as many that only contain  a single module are set up like this

#rnx

diff --git a/Engine/Source/Programs/UnrealBuildTool/Configuration/UEBuildPlatform.cs b/Engine/Source/Programs/UnrealBuildTool/Configuration/UEBuildPlatform.cs
--- a/Engine/Source/Programs/UnrealBuildTool/Configuration/UEBuildPlatform.cs
+++ b/Engine/Source/Programs/UnrealBuildTool/Configuration/UEBuildPlatform.cs
@@ -920,7 +920,8 @@ namespace UnrealBuildTool
 				{
 					if (Path.IsUnderDirectory(Module.ModuleDirectory))
 					{
-						if (Path == Module.ModuleDirectory)
+						// Warn on the root module directory being added, except for plugins as there are some plugins with a single module directly in the Source directory
+						if (Path == Module.ModuleDirectory && (Module.Rules.Plugin == null || Module.ModuleDirectory.ParentDirectory != Module.Rules.Plugin.Directory))
 						{
 							LoggerFunc("Module '{Name}' is adding root directory to PublicIncludePaths. This is not allowed.", Module.Name);
 						}
