UnrealBuildTool: /IGNORE:importeddllmain with clang linker

#rnx

diff --git a/Engine/Source/Programs/UnrealBuildTool/Platform/Windows/VCToolChain.cs b/Engine/Source/Programs/UnrealBuildTool/Platform/Windows/VCToolChain.cs
--- a/Engine/Source/Programs/UnrealBuildTool/Platform/Windows/VCToolChain.cs
+++ b/Engine/Source/Programs/UnrealBuildTool/Platform/Windows/VCToolChain.cs
@@ -3054,6 +3054,11 @@ namespace UnrealBuildTool
 				Arguments.Add("/IGNORE:4221");
 			}
 
+			if (Target.WindowsPlatform.Compiler.IsClang() && Target.WindowsPlatform.bAllowClangLinker)
+			{
+				Arguments.Add("/IGNORE:importeddllmain");
+			}
+
 			if (!bIsBuildingLibraryOrImportLibrary)
 			{
 				// Delay-load these DLLs.
