UnrealBuildTool: Use job object to track peak memory usage of each process

#rnx

diff --git a/Engine/Source/Programs/Shared/EpicGames.Core/ManagedProcess.cs b/Engine/Source/Programs/Shared/EpicGames.Core/ManagedProcess.cs
--- a/Engine/Source/Programs/Shared/EpicGames.Core/ManagedProcess.cs
+++ b/Engine/Source/Programs/Shared/EpicGames.Core/ManagedProcess.cs
@@ -104,6 +104,9 @@ namespace EpicGames.Core
 		[DllImport("kernel32.dll", SetLastError = true)]
 		static extern bool QueryInformationJobObject(SafeFileHandle hJob, int JobObjectInformationClass, out JOBOBJECT_BASIC_ACCOUNTING_INFORMATION lpJobObjectInformation, int cbJobObjectInformationLength, out int lpReturnLength);
 
+		[DllImport("kernel32.dll", SetLastError = true)]
+		static extern bool QueryInformationJobObject(SafeFileHandle hJob, int JobObjectInformationClass, out JOBOBJECT_EXTENDED_LIMIT_INFORMATION lpJobObjectInformation, int cbJobObjectInformationLength, out int lpReturnLength);
+
 		[DllImport("kernel32.dll", SetLastError = true)]
 		static extern int AssignProcessToJobObject(SafeFileHandle? hJob, IntPtr hProcess);
 
@@ -193,20 +196,20 @@ namespace EpicGames.Core
 		/// <summary>
 		/// Returns the total CPU time usage for this job.
 		/// </summary>
-		public TimeSpan TotalProcessorTime
-		{
-			get
-			{
-				if (SupportsJobObjects)
-				{
-					if (QueryInformationJobObject(JobHandle!, JobObjectBasicAccountingInformation, out JOBOBJECT_BASIC_ACCOUNTING_INFORMATION accountingInformation, Marshal.SizeOf(typeof(JOBOBJECT_BASIC_ACCOUNTING_INFORMATION)), out _))
-					{
-						return new TimeSpan((long)accountingInformation.TotalUserTime + (long)accountingInformation.TotalKernelTime);
-					}
-				}
-				return new TimeSpan();
-			}
-		}
+		public TimeSpan TotalProcessorTime =>
+			SupportsJobObjects && QueryInformationJobObject(JobHandle!, JobObjectBasicAccountingInformation, out JOBOBJECT_BASIC_ACCOUNTING_INFORMATION accountingInformation, Marshal.SizeOf(typeof(JOBOBJECT_BASIC_ACCOUNTING_INFORMATION)), out _)
+				? new TimeSpan((long)accountingInformation.TotalUserTime + (long)accountingInformation.TotalKernelTime)
+				: new TimeSpan();
+
+		public long PeakProcessMemoryUsed =>
+			SupportsJobObjects && QueryInformationJobObject(JobHandle!, JobObjectExtendedLimitInformation, out JOBOBJECT_EXTENDED_LIMIT_INFORMATION accountingInformation, Marshal.SizeOf(typeof(JOBOBJECT_EXTENDED_LIMIT_INFORMATION)), out _)
+				? (long)accountingInformation.PeakProcessMemoryUsed
+				: 0;
+
+		public long PeakJobMemoryUsed =>
+			SupportsJobObjects && QueryInformationJobObject(JobHandle!, JobObjectExtendedLimitInformation, out JOBOBJECT_EXTENDED_LIMIT_INFORMATION accountingInformation, Marshal.SizeOf(typeof(JOBOBJECT_EXTENDED_LIMIT_INFORMATION)), out _)
+				? (long)accountingInformation.PeakJobMemoryUsed
+				: 0;
 
 		/// <summary>
 		/// Dispose of the process group
@@ -1237,20 +1240,9 @@ namespace EpicGames.Core
 		/// <summary>
 		/// Total amount of time spent by the CPU
 		/// </summary>
-		public TimeSpan TotalProcessorTime
-		{
-			get
-			{
-				if (_accountingProcessGroup != null)
-				{
-					return _accountingProcessGroup.TotalProcessorTime;
-				}
-				else
-				{
-					return new TimeSpan();
-				}
-			}
-		}
+		public TimeSpan TotalProcessorTime => _accountingProcessGroup != null ? _accountingProcessGroup.TotalProcessorTime : new TimeSpan();
+
+		public long PeakMemoryUsed => _accountingProcessGroup != null ? _accountingProcessGroup.PeakJobMemoryUsed : 0;
 
 		/// <summary>
 		/// Block until the process outputs a line of text, or terminates.
diff --git a/Engine/Source/Programs/UnrealBuildTool/Executors/ActionLogger.cs b/Engine/Source/Programs/UnrealBuildTool/Executors/ActionLogger.cs
--- a/Engine/Source/Programs/UnrealBuildTool/Executors/ActionLogger.cs
+++ b/Engine/Source/Programs/UnrealBuildTool/Executors/ActionLogger.cs
@@ -17,9 +17,10 @@ namespace UnrealBuildTool
 	/// <param name="ExitCode">Process return code.  Zero is assumed to be success and all other values an error.</param>
 	/// <param name="ExecutionTime">Wall time</param>
 	/// <param name="ProcessorTime">CPU time</param>
+	/// <param name="PeakMemoryUsed">Peak Memory Used</param>
 	/// <param name="Action">Action</param>
 	/// <param name="AdditionalDescription">Additional description of action</param>
-	record ExecuteResults(List<string> LogLines, int ExitCode, TimeSpan ExecutionTime, TimeSpan ProcessorTime, LinkedAction Action, string? AdditionalDescription = null);
+	record ExecuteResults(List<string> LogLines, int ExitCode, TimeSpan ExecutionTime, TimeSpan ProcessorTime, long PeakMemoryUsed, LinkedAction Action, string? AdditionalDescription = null);
 
 	class ActionLogger
 	{
@@ -146,19 +147,12 @@ namespace UnrealBuildTool
 		/// <param name="totalActions">Number of total actions</param>
 		public void LogAction(LinkedAction action, ExecuteResults? executeTaskResult, int totalActions)
 		{
-			List<string>? logLines = null;
-			int exitCode = Int32.MaxValue;
-			TimeSpan executionTime = TimeSpan.Zero;
-			TimeSpan processorTime = TimeSpan.Zero;
-			string? additionalDescription = null;
-			if (executeTaskResult != null)
-			{
-				logLines = executeTaskResult.LogLines;
-				exitCode = executeTaskResult.ExitCode;
-				executionTime = executeTaskResult.ExecutionTime;
-				processorTime = executeTaskResult.ProcessorTime;
-				additionalDescription = executeTaskResult.AdditionalDescription;
-			}
+			List<string>? logLines = executeTaskResult?.LogLines ?? null;
+			int exitCode = executeTaskResult?.ExitCode ?? Int32.MaxValue;
+			TimeSpan executionTime = executeTaskResult?.ExecutionTime ?? TimeSpan.Zero;
+			TimeSpan processorTime = executeTaskResult?.ProcessorTime ?? TimeSpan.Zero;
+			long peakMemory = executeTaskResult?.PeakMemoryUsed ?? 0;
+			string? additionalDescription = executeTaskResult?.AdditionalDescription;
 
 			// Write it to the log
 			string description = String.Empty;
@@ -203,10 +197,18 @@ namespace UnrealBuildTool
 
 				if (ShowPerActionCompilationTimes)
 				{
-					if (processorTime.Ticks > 0)
+					if (processorTime.Ticks > 0 && peakMemory > 0)
+					{
+						compilationTimes = $" (Wall: {executionTime.TotalSeconds:0.00}s CPU: {processorTime.TotalSeconds:0.00}s Mem: {StringUtils.FormatBytesString(peakMemory)})";
+					}
+					else if (processorTime.Ticks > 0)
 					{
 						compilationTimes = $" (Wall: {executionTime.TotalSeconds:0.00}s CPU: {processorTime.TotalSeconds:0.00}s)";
 					}
+					else if (peakMemory > 0)
+					{
+						compilationTimes = $" (Wall: {executionTime.TotalSeconds:0.00}s Mem: {StringUtils.FormatBytesString(peakMemory)})";
+					}
 					else
 					{
 						compilationTimes = $" (Wall: {executionTime.TotalSeconds:0.00}s)";
@@ -312,14 +314,14 @@ namespace UnrealBuildTool
 			_flushToolOutput();
 
 			LogLevel LogLevel = success ? LogLevel.Information : LogLevel.Debug;
+			Logger.Log(LogLevel, "");
 			if (ShowCPUUtilization)
 			{
 				lock (_cpuUtilization)
 				{
 					if (_cpuUtilization.Count > 0)
 					{
-						Logger.Log(LogLevel, "");
-						Logger.Log(LogLevel, "Average CPU Utilization: {CPUPercentage}%", (int)(_cpuUtilization.Average()));
+						Logger.Log(LogLevel, "Average CPU Utilization : {CPUPercentage:0.00}%", _cpuUtilization.Average());
 					}
 				}
 			}
@@ -329,10 +331,20 @@ namespace UnrealBuildTool
 				return;
 			}
 
-			Logger.Log(LogLevel, "");
-			if (ProcessGroup.TotalProcessorTime.Ticks > 0)
+			if (ProcessGroup.TotalProcessorTime.Ticks + ProcessGroup.PeakProcessMemoryUsed + ProcessGroup.PeakJobMemoryUsed > 0)
 			{
-				Logger.Log(LogLevel, "Total CPU Time: {TotalSeconds} s", ProcessGroup.TotalProcessorTime.TotalSeconds);
+				if (ProcessGroup.TotalProcessorTime.Ticks > 0)
+				{
+					Logger.Log(LogLevel, "Total Processor Time    : {TotalSeconds:0.00}s", ProcessGroup.TotalProcessorTime.TotalSeconds);
+				}
+				if (ProcessGroup.PeakProcessMemoryUsed > 0)
+				{
+					Logger.Log(LogLevel, "Peak Process Memory Used: {PeakProcessMemory}", StringUtils.FormatBytesString(ProcessGroup.PeakProcessMemoryUsed));
+				}
+				if (ProcessGroup.PeakJobMemoryUsed > 0)
+				{
+					Logger.Log(LogLevel, "Peak Total Memory Used  : {PeakTotalMemory}", StringUtils.FormatBytesString(ProcessGroup.PeakJobMemoryUsed));
+				}
 				Logger.Log(LogLevel, "");
 			}
 
@@ -351,13 +363,21 @@ namespace UnrealBuildTool
 					IExternalAction Action = Result.Action.Inner;
 
 					string Description = $"{(Action.CommandDescription ?? Action.CommandPath.GetFileNameWithoutExtension())} {Action.StatusDescription}".Trim();
-					if (Result.ProcessorTime.Ticks > 0)
+					if (Result.ProcessorTime.Ticks > 0 && Result.PeakMemoryUsed > 0)
+					{
+						Logger.Log(LogLevel, "{Description} [ Wall Time {ExecutionTime:0.00}s / CPU Time {ProcessorTime:0.00}s / Mem {Memory} ]", Description, Result.ExecutionTime.TotalSeconds, Result.ProcessorTime.TotalSeconds, StringUtils.FormatBytesString(Result.PeakMemoryUsed));
+					}
+					else if (Result.ProcessorTime.Ticks > 0)
+					{
+						Logger.Log(LogLevel, "{Description} [ Wall Time {ExecutionTime:0.00}s / CPU Time {ProcessorTime:0.00}s ]", Description, Result.ExecutionTime.TotalSeconds, Result.ProcessorTime.TotalSeconds);
+					}
+					else if (Result.PeakMemoryUsed > 0)
 					{
-						Logger.Log(LogLevel, "{Description} [ Wall Time {ExecutionTime:0.00} s / CPU Time {ProcessorTime:0.00} s ]", Description, Result.ExecutionTime.TotalSeconds, Result.ProcessorTime.TotalSeconds);
+						Logger.Log(LogLevel, "{Description} [ Time {ExecutionTime:0.00}s / Mem {Memory} ]", Description, Result.ExecutionTime.TotalSeconds, StringUtils.FormatBytesString(Result.PeakMemoryUsed));
 					}
 					else
 					{
-						Logger.Log(LogLevel, "{Description} [ Time {ExecutionTime:0.00} s ]", Description, Result.ExecutionTime.TotalSeconds);
+						Logger.Log(LogLevel, "{Description} [ Time {ExecutionTime:0.00}s ]", Description, Result.ExecutionTime.TotalSeconds);
 					}
 				}
 				Logger.Log(LogLevel, "");
diff --git a/Engine/Source/Programs/UnrealBuildTool/Executors/Experimental/FASTBuild.cs b/Engine/Source/Programs/UnrealBuildTool/Executors/Experimental/FASTBuild.cs
--- a/Engine/Source/Programs/UnrealBuildTool/Executors/Experimental/FASTBuild.cs
+++ b/Engine/Source/Programs/UnrealBuildTool/Executors/Experimental/FASTBuild.cs
@@ -501,7 +501,7 @@ namespace UnrealBuildTool
 				}
 
 				bool result = ExecuteBffFile(FASTBuildFilePath, Logger);
-				telemetryEvent = new TelemetryExecutorEvent(Name, startTimeUTC, result, Actions.Count(), -1, -1, 0, 0, DateTime.UtcNow);
+				telemetryEvent = new TelemetryExecutorEvent(Name, startTimeUTC, result, Actions.Count(), -1, -1, 0, 0, DateTime.UtcNow, 0, 0);
 				return result;
 			}
 
diff --git a/Engine/Source/Programs/UnrealBuildTool/Executors/ImmediateActionQueue.cs b/Engine/Source/Programs/UnrealBuildTool/Executors/ImmediateActionQueue.cs
--- a/Engine/Source/Programs/UnrealBuildTool/Executors/ImmediateActionQueue.cs
+++ b/Engine/Source/Programs/UnrealBuildTool/Executors/ImmediateActionQueue.cs
@@ -329,7 +329,7 @@ namespace UnrealBuildTool
 						ActionArtifactResult result = await _actionArtifactCache!.CompleteActionFromCacheAsync(action, CancellationToken);
 						if (result.Success)
 						{
-							ExecuteResults results = new(result.LogLines, result.ExitCode, TimeSpan.Zero, TimeSpan.Zero, action, "[Cache]");
+							ExecuteResults results = new(result.LogLines, result.ExitCode, TimeSpan.Zero, TimeSpan.Zero, 0, action, "[Cache]");
 							Interlocked.Increment(ref _cacheHitActions);
 							OnActionCompleted(action, results.ExitCode == 0, results);
 							OnArtifactsRead?.Invoke(action);
@@ -628,13 +628,13 @@ namespace UnrealBuildTool
 				}
 				else
 				{
-					ExecuteResults results = new(new List<string>(), Int32.MaxValue, TimeSpan.Zero, TimeSpan.Zero, action);
+					ExecuteResults results = new(new List<string>(), Int32.MaxValue, TimeSpan.Zero, TimeSpan.Zero, 0, action);
 					OnActionCompleted(action, false, results);
 				}
 			}
 			else if (ex is OperationCanceledException)
 			{
-				ExecuteResults results = new(new List<string>(), Int32.MaxValue, TimeSpan.Zero, TimeSpan.Zero, action);
+				ExecuteResults results = new(new List<string>(), Int32.MaxValue, TimeSpan.Zero, TimeSpan.Zero, 0, action);
 				OnActionCompleted(action, false, results);
 			}
 			else
@@ -644,7 +644,7 @@ namespace UnrealBuildTool
 					ExceptionUtils.FormatException(ex),
 					ExceptionUtils.FormatExceptionDetails(ex),
 				};
-				ExecuteResults results = new(text, Int32.MaxValue, TimeSpan.Zero, TimeSpan.Zero, action);
+				ExecuteResults results = new(text, Int32.MaxValue, TimeSpan.Zero, TimeSpan.Zero, 0, action);
 				OnActionCompleted(action, false, results);
 			}
 		}
diff --git a/Engine/Source/Programs/UnrealBuildTool/Executors/ParallelExecutor.cs b/Engine/Source/Programs/UnrealBuildTool/Executors/ParallelExecutor.cs
--- a/Engine/Source/Programs/UnrealBuildTool/Executors/ParallelExecutor.cs
+++ b/Engine/Source/Programs/UnrealBuildTool/Executors/ParallelExecutor.cs
@@ -203,7 +203,7 @@ namespace UnrealBuildTool
 				result = await queue.RunTillDone();
 
 				queue.GetActionResultCounts(out int totalActions, out int succeededActions, out int failedActions, out int cacheHitActions, out int cacheMissActions);
-				telemetryEvent = new TelemetryExecutorEvent(Name, startTimeUTC, result, totalActions, succeededActions, failedActions, cacheHitActions, cacheMissActions, DateTime.UtcNow);
+				telemetryEvent = new TelemetryExecutorEvent(Name, startTimeUTC, result, totalActions, succeededActions, failedActions, cacheHitActions, cacheMissActions, DateTime.UtcNow, queue.ProcessGroup.PeakProcessMemoryUsed, queue.ProcessGroup.PeakJobMemoryUsed);
 			}
 			else
 			{
@@ -216,7 +216,7 @@ namespace UnrealBuildTool
 				result = await queue.RunTillDone();
 
 				queue.GetActionResultCounts(out int totalActions, out int succeededActions, out int failedActions, out int cacheHitActions, out int cacheMissActions);
-				telemetryEvent = new TelemetryExecutorEvent(Name, startTimeUTC, result, totalActions, succeededActions, failedActions, cacheHitActions, cacheMissActions, DateTime.UtcNow);
+				telemetryEvent = new TelemetryExecutorEvent(Name, startTimeUTC, result, totalActions, succeededActions, failedActions, cacheHitActions, cacheMissActions, DateTime.UtcNow, queue.ProcessGroup.PeakProcessMemoryUsed, queue.ProcessGroup.PeakJobMemoryUsed);
 			}
 
 			return result;
@@ -248,6 +248,7 @@ namespace UnrealBuildTool
 			int ExitCode = Process.ExitCode;
 			TimeSpan ProcessorTime = Process.TotalProcessorTime;
 			TimeSpan ExecutionTime = Process.ExitTime - Process.StartTime;
+			long PeakMemoryUsed = Process.PeakMemoryUsed;
 
 			if (ExitCode == 0 && Action.bForceWarningsAsError && LogLines.Any(x => x.Contains(": warning: ", StringComparison.OrdinalIgnoreCase)))
 			{
@@ -255,7 +256,7 @@ namespace UnrealBuildTool
 				LogLines = [.. LogLines.Select(x => x.Replace(": warning: ", ": error: ", StringComparison.OrdinalIgnoreCase))];
 			}
 
-			return new ExecuteResults(LogLines, ExitCode, ExecutionTime, ProcessorTime, Action, AdditionalDescription);
+			return new ExecuteResults(LogLines, ExitCode, ExecutionTime, ProcessorTime, PeakMemoryUsed, Action, AdditionalDescription);
 		}
 	}
 
diff --git a/Engine/Source/Programs/UnrealBuildTool/Executors/SNDBS.cs b/Engine/Source/Programs/UnrealBuildTool/Executors/SNDBS.cs
--- a/Engine/Source/Programs/UnrealBuildTool/Executors/SNDBS.cs
+++ b/Engine/Source/Programs/UnrealBuildTool/Executors/SNDBS.cs
@@ -396,7 +396,7 @@ namespace UnrealBuildTool
 					LocalProcess.WaitForExit();
 					bool result = LocalProcess.ExitCode == 0;
 
-					telemetryEvent = new TelemetryExecutorEvent(Name, startTimeUTC, result, NumActions, -1, -1, 0, 0, DateTime.UtcNow);
+					telemetryEvent = new TelemetryExecutorEvent(Name, startTimeUTC, result, NumActions, -1, -1, 0, 0, DateTime.UtcNow, 0, 0);
 					return result;
 				}
 				catch (Exception Ex)
diff --git a/Engine/Source/Programs/UnrealBuildTool/Executors/UnrealBuildAccelerator/UBAExecutor.cs b/Engine/Source/Programs/UnrealBuildTool/Executors/UnrealBuildAccelerator/UBAExecutor.cs
--- a/Engine/Source/Programs/UnrealBuildTool/Executors/UnrealBuildAccelerator/UBAExecutor.cs
+++ b/Engine/Source/Programs/UnrealBuildTool/Executors/UnrealBuildAccelerator/UBAExecutor.cs
@@ -779,6 +779,7 @@ namespace UnrealBuildTool
 
 				TimeSpan processorTime = TimeSpan.Zero;
 				TimeSpan executionTime = TimeSpan.Zero;
+				long peakMemory = 0;
 				if (executionType != ProcessExecutionType.Cache)
 				{
 					processorTime = info.TotalProcessorTime;
@@ -826,7 +827,7 @@ namespace UnrealBuildTool
 					}
 				}
 
-				actionLogger.AddActionToLog(new(logLines, exitCode, executionTime, processorTime, action, additionalDescription));
+				actionLogger.AddActionToLog(new(logLines, exitCode, executionTime, processorTime, peakMemory, action, additionalDescription));
 
 				return ProcessFinishedResponse.None;
 			});
@@ -908,7 +909,8 @@ namespace UnrealBuildTool
 						!UBAConfig.bDisableRemote ? _agentCoordinators.Count : 0, !UBAConfig.bDisableRemote ? _successfulCoordinatorConnections : 0, !UBAConfig.bDisableRemote ? _failedCoordinatorConnections : 0,
 						!UBAConfig.bDisableRemote ? _ubaDurationWaitingForRemote : TimeSpan.Zero,
 						!UBAConfig.bDisableRemote || _agentCoordinators.Count == 0 ? "Local" : _remoteConnectionMode,
-						DateTime.UtcNow);
+						DateTime.UtcNow,
+						0, 0);
 				}
 				return res;
 			}
diff --git a/Engine/Source/Programs/UnrealBuildTool/Executors/XGE.cs b/Engine/Source/Programs/UnrealBuildTool/Executors/XGE.cs
--- a/Engine/Source/Programs/UnrealBuildTool/Executors/XGE.cs
+++ b/Engine/Source/Programs/UnrealBuildTool/Executors/XGE.cs
@@ -781,7 +781,7 @@ namespace UnrealBuildTool
 				// Run through the standard XGE executor
 				bool result = ExecuteTaskFile(TaskFilePath, EventHandlerWrapper, NumActions, Logger);
 
-				telemetryEvent = new TelemetryExecutorEvent(Name, startTimeUTC, result, NumActions, -1, -1, 0, 0, DateTime.UtcNow);
+				telemetryEvent = new TelemetryExecutorEvent(Name, startTimeUTC, result, NumActions, -1, -1, 0, 0, DateTime.UtcNow, 0, 0);
 
 				return result;
 			}
diff --git a/Engine/Source/Programs/UnrealBuildTool/System/Telemetry.cs b/Engine/Source/Programs/UnrealBuildTool/System/Telemetry.cs
--- a/Engine/Source/Programs/UnrealBuildTool/System/Telemetry.cs
+++ b/Engine/Source/Programs/UnrealBuildTool/System/Telemetry.cs
@@ -117,6 +117,9 @@ namespace UnrealBuildTool
 		public DateTime StartUTC { get; }
 		[JsonConverter(typeof(DurationDataRouterEventConverter))]
 		public TimeSpan Duration { get; }
+		public long PeakProcessMemory { get; }
+		public long PeakTotalMemory { get; }
+
 		public bool Result { get; }
 		public int TotalActions { get; }
 		public int SucceededActions { get; }
@@ -130,7 +133,8 @@ namespace UnrealBuildTool
 		public double SkipRate => SkippedActions / (double)TotalActions;
 		public double CacheRate => CacheHitActions + CacheMissActions > 0 ? CacheHitActions / (double)CacheCheckActions : 0.0;
 
-		public TelemetryExecutorEvent(string executor, DateTime startUTC, bool result, int totalActions, int succeededActions, int failedActions, int cacheHitActions, int cacheMissActions, DateTime timestamp)
+		public TelemetryExecutorEvent(string executor, DateTime startUTC, bool result, int totalActions, int succeededActions, int failedActions, int cacheHitActions, int cacheMissActions, DateTime timestamp,
+			long peakProcessMemory, long peakTotalMemory)
 			: base(timestamp)
 		{
 			Executor = executor;
@@ -142,6 +146,8 @@ namespace UnrealBuildTool
 			FailedActions = failedActions;
 			CacheHitActions = cacheHitActions;
 			CacheMissActions = cacheMissActions;
+			PeakProcessMemory = peakProcessMemory;
+			PeakTotalMemory = peakTotalMemory;
 		}
 	}
 
@@ -173,8 +179,8 @@ namespace UnrealBuildTool
 			int retriedLocalActions, int retriedDisabledActions,
 			int totalCoordinators, int succeededCoordinators, int failedCoordinators,
 			TimeSpan durationWaitingForRemote, string connectionMode,
-			DateTime timestamp)
-			: base(executor, startUTC, result, totalActions, succeededActions, failedActions, cacheHitActions, cacheMissActions, timestamp)
+			DateTime timestamp, long peakProcessMemory, long peakTotalMemory)
+			: base(executor, startUTC, result, totalActions, succeededActions, failedActions, cacheHitActions, cacheMissActions, timestamp, peakProcessMemory, peakTotalMemory)
 		{
 			LocalActions = localActions;
 			RemoteActions = remoteActions;
