UnrealBuildTool: Fix EngineIncludeOrderVersion Enum.GetValues with the wrong type

#rnx

diff --git a/Engine/Source/Programs/UnrealBuildTool/Configuration/TargetRules.cs b/Engine/Source/Programs/UnrealBuildTool/Configuration/TargetRules.cs
--- a/Engine/Source/Programs/UnrealBuildTool/Configuration/TargetRules.cs
+++ b/Engine/Source/Programs/UnrealBuildTool/Configuration/TargetRules.cs
@@ -3216,7 +3216,7 @@ namespace UnrealBuildTool
 			{
 				if (String.IsNullOrEmpty(target.Version.BuildVersionString))
 				{
-					BuildVersion = String.Format("{0}-CL-{1}", target.Version.BranchName, target.Version.Changelist);
+					BuildVersion = $"{target.Version.BranchName}-CL-{target.Version.Changelist}";
 				}
 				else
 				{
@@ -3237,7 +3237,7 @@ namespace UnrealBuildTool
 			EncryptionAndSigning.CryptoSettings cryptoSettings = EncryptionAndSigning.ParseCryptoSettings(cryptoSettingsDir, Platform, Logger);
 			if (cryptoSettings.IsAnyEncryptionEnabled())
 			{
-				ProjectDefinitions.Add(String.Format("IMPLEMENT_ENCRYPTION_KEY_REGISTRATION()=UE_REGISTER_ENCRYPTION_KEY({0})", FormatHexBytes(cryptoSettings.EncryptionKey!.Key!)));
+				ProjectDefinitions.Add($"IMPLEMENT_ENCRYPTION_KEY_REGISTRATION()=UE_REGISTER_ENCRYPTION_KEY({FormatHexBytes(cryptoSettings.EncryptionKey!.Key!)})");
 			}
 			else
 			{
@@ -3246,7 +3246,7 @@ namespace UnrealBuildTool
 
 			if (cryptoSettings.IsPakSigningEnabled())
 			{
-				ProjectDefinitions.Add(String.Format("IMPLEMENT_SIGNING_KEY_REGISTRATION()=UE_REGISTER_SIGNING_KEY(UE_LIST_ARGUMENT({0}), UE_LIST_ARGUMENT({1}))", FormatHexBytes(cryptoSettings.SigningKey!.PublicKey.Exponent!), FormatHexBytes(cryptoSettings.SigningKey.PublicKey.Modulus!)));
+				ProjectDefinitions.Add($"IMPLEMENT_SIGNING_KEY_REGISTRATION()=UE_REGISTER_SIGNING_KEY(UE_LIST_ARGUMENT({FormatHexBytes(cryptoSettings.SigningKey!.PublicKey.Exponent!)}), UE_LIST_ARGUMENT({FormatHexBytes(cryptoSettings.SigningKey.PublicKey.Modulus!)}))");
 			}
 			else
 			{
@@ -3261,7 +3261,7 @@ namespace UnrealBuildTool
 		/// <returns>List of hexadecimal bytes</returns>
 		private static string FormatHexBytes(byte[] data)
 		{
-			return String.Join(",", data.Select(x => String.Format("0x{0:X2}", x)));
+			return String.Join(",", data.Select(x => $"0x{x:X2}"));
 		}
 
 		/// <summary>
@@ -3409,7 +3409,7 @@ namespace UnrealBuildTool
 			}
 			else
 			{
-				return ((UnrealTargetConfiguration[])Enum.GetValues(typeof(UnrealTargetConfiguration))).Where(x => x != UnrealTargetConfiguration.Unknown).ToArray();
+				return [.. Enum.GetValues<UnrealTargetConfiguration>().Where(x => x != UnrealTargetConfiguration.Unknown)];
 			}
 		}
 
@@ -3477,7 +3477,7 @@ namespace UnrealBuildTool
 		{
 			// Resolve BuildSettingsVersion.Latest to the version it's assigned to
 			BuildSettingsVersion latestVersion = BuildSettingsVersion.Latest;
-			foreach (BuildSettingsVersion value in Enum.GetValues(typeof(BuildSettingsVersion)))
+			foreach (BuildSettingsVersion value in Enum.GetValues<BuildSettingsVersion>())
 			{
 				if ((int)value == (int)BuildSettingsVersion.Latest)
 				{
@@ -3494,33 +3494,33 @@ namespace UnrealBuildTool
 				List<Tuple<string, string>> modifiedSettings = [];
 				if (BuildSettingsVersion.V2 <= latestVersion && DefaultBuildSettings < BuildSettingsVersion.V2)
 				{
-					modifiedSettings.Add(Tuple.Create(String.Format("{0} = false", nameof(bLegacyPublicIncludePaths)), "Omits subfolders from public include paths to reduce compiler command line length. (Previously: true)."));
-					modifiedSettings.Add(Tuple.Create(String.Format("{0} = WarningLevel.Error", nameof(CppCompileWarningSettings.ShadowVariableWarningLevel)), "Treats shadowed variable warnings as errors. (Previously: WarningLevel.Warning)."));
-					modifiedSettings.Add(Tuple.Create(String.Format("{0} = PCHUsageMode.UseExplicitOrSharedPCHs", nameof(ModuleRules.PCHUsage)), "Set in build.cs files to enables IWYU-style PCH model. See https://docs.unrealengine.com/en-US/Programming/BuildTools/UnrealBuildTool/IWYU/index.html. (Previously: PCHUsageMode.UseSharedPCHs)."));
+					modifiedSettings.Add(Tuple.Create($"{nameof(bLegacyPublicIncludePaths)} = false", "Omits subfolders from public include paths to reduce compiler command line length. (Previously: true)."));
+					modifiedSettings.Add(Tuple.Create($"{nameof(CppCompileWarningSettings.ShadowVariableWarningLevel)} = WarningLevel.Error", "Treats shadowed variable warnings as errors. (Previously: WarningLevel.Warning)."));
+					modifiedSettings.Add(Tuple.Create($"{nameof(ModuleRules.PCHUsage)} = PCHUsageMode.UseExplicitOrSharedPCHs", "Set in build.cs files to enables IWYU-style PCH model. See https://docs.unrealengine.com/en-US/Programming/BuildTools/UnrealBuildTool/IWYU/index.html. (Previously: PCHUsageMode.UseSharedPCHs)."));
 				}
 
 				if (BuildSettingsVersion.V3 <= latestVersion && DefaultBuildSettings < BuildSettingsVersion.V3)
 				{
-					modifiedSettings.Add(Tuple.Create(String.Format("{0} = false", nameof(bLegacyParentIncludePaths)), "Omits module parent folders from include paths to reduce compiler command line length. (Previously: true)."));
+					modifiedSettings.Add(Tuple.Create($"{nameof(bLegacyParentIncludePaths)} = false", "Omits module parent folders from include paths to reduce compiler command line length. (Previously: true)."));
 				}
 
 				if (BuildSettingsVersion.V4 <= latestVersion && DefaultBuildSettings < BuildSettingsVersion.V4)
 				{
-					modifiedSettings.Add(Tuple.Create(String.Format("{0} = CppStandardVersion.Default", nameof(CppStandard)), "Updates C++ Standard to C++20 (Previously: CppStandardVersion.Cpp17)."));
-					modifiedSettings.Add(Tuple.Create(String.Format("{0}.{1} = true", nameof(WindowsPlatform), nameof(WindowsPlatform.bStrictConformanceMode)), "Updates MSVC strict conformance mode to true (Previously: false)."));
+					modifiedSettings.Add(Tuple.Create($"{nameof(CppStandard)} = CppStandardVersion.Default", "Updates C++ Standard to C++20 (Previously: CppStandardVersion.Cpp17)."));
+					modifiedSettings.Add(Tuple.Create($"{nameof(WindowsPlatform)}.{nameof(WindowsPlatform.bStrictConformanceMode)} = true", "Updates MSVC strict conformance mode to true (Previously: false)."));
 				}
 
 				if (BuildSettingsVersion.V5 <= latestVersion && DefaultBuildSettings < BuildSettingsVersion.V5)
 				{
-					modifiedSettings.Add(Tuple.Create(String.Format("{0} = true", nameof(bValidateFormatStrings)), "Enables compile-time validation of strings+args passed to UE_LOG. (Previously: false)."));
+					modifiedSettings.Add(Tuple.Create($"{nameof(bValidateFormatStrings)} = true", "Enables compile-time validation of strings+args passed to UE_LOG. (Previously: false)."));
 				}
 
 				if (modifiedSettings.Count > 0)
 				{
-					string formatString = String.Format("[Upgrade]     {{0,-{0}}}   => {{1}}", modifiedSettings.Max(x => x.Item1.Length));
+					string formatString = $"[Upgrade]     {{0,-{modifiedSettings.Max(x => x.Item1.Length)}}}   => {{1}}";
 					diagnostics.AddRange(modifiedSettings.Select(modifiedSetting => String.Format(formatString, modifiedSetting.Item1, modifiedSetting.Item2)));
 				}
-				diagnostics.Add(String.Format("[Upgrade] Suppress this message by setting 'DefaultBuildSettings = BuildSettingsVersion.{0};' in {1}, and explicitly overriding settings that differ from the new defaults.", latestVersion, File!.GetFileName()));
+				diagnostics.Add($"[Upgrade] Suppress this message by setting 'DefaultBuildSettings = BuildSettingsVersion.{latestVersion};' in {File!.GetFileName()}, and explicitly overriding settings that differ from the new defaults.");
 				diagnostics.Add("[Upgrade]");
 			}
 
@@ -3528,7 +3528,7 @@ namespace UnrealBuildTool
 			{
 				// Resolve EngineIncludeOrderVersion.Latest to the version it's assigned to
 				EngineIncludeOrderVersion latestEngineIncludeOrder = EngineIncludeOrderVersion.Latest;
-				foreach (EngineIncludeOrderVersion value in Enum.GetValues(typeof(BuildSettingsVersion)))
+				foreach (EngineIncludeOrderVersion value in Enum.GetValues<EngineIncludeOrderVersion>())
 				{
 					if ((int)value == (int)EngineIncludeOrderVersion.Latest)
 					{
@@ -3559,7 +3559,7 @@ namespace UnrealBuildTool
 			{
 				// Resolve CppStandardVersion.Default to the version it's assigned to
 				CppStandardVersion defaultCppStandard = CppStandardVersion.Latest;
-				foreach (CppStandardVersion value in Enum.GetValues(typeof(CppStandardVersion)))
+				foreach (CppStandardVersion value in Enum.GetValues<CppStandardVersion>())
 				{
 					if ((int)value == (int)CppStandardVersion.Default)
 					{
