Add iOS app extensions as dependencies to project. Fix appex bundleId.

#rnx

#DisabledValidations Preflight, NonUnityAndStaticAnalysis, Json, PerfReportingServerJson, CheckCopyright, CoreRedirects, CheckUnacceptableWords, PackageDataValidator

diff --git a/Engine/Source/Programs/UnrealBuildTool/ProjectFiles/Xcode/XcodeProject.cs b/Engine/Source/Programs/UnrealBuildTool/ProjectFiles/Xcode/XcodeProject.cs
--- a/Engine/Source/Programs/UnrealBuildTool/ProjectFiles/Xcode/XcodeProject.cs
+++ b/Engine/Source/Programs/UnrealBuildTool/ProjectFiles/Xcode/XcodeProject.cs
@@ -1450,6 +1450,8 @@ namespace UnrealBuildTool.XcodeProjectXcconfig
 								XcodeSourceFile SourceFile = new XcodeSourceFile(FileRef, FileRef.Directory, XcodeProjectFileGenerator.MakeXcodeGuid());
 								Project.FileCollection.ProcessFile(SourceFile, true, false, "Source");
 							}
+							
+							Project.FileCollection.AddFileReference(ExtensionGuid, Extension, "explicitFileType", "wrapper.app-extension", "BUILT_PRODUCTS_DIR", "Products");
 
 							DirectoryReference FinalAppExtDir = DirectoryReference.Combine(ProjectDirectory, DI.Name);
 							EmbedAppExt.AddExtension(FinalAppExtDir, ExtensionGuid);
@@ -2335,7 +2337,7 @@ namespace UnrealBuildTool.XcodeProjectXcconfig
 
 				foreach (FileSystemInfo FSI in ExtensionPath.EnumerateFileSystemInfos())
 				{
-					if (FSI.Name.EndsWith(".h") || FSI.Name.StartsWith(".") || FSI.Name.EndsWith(".plist"))
+					if (FSI.Name.EndsWith(".h") || FSI.Name.StartsWith(".") || FSI.Name.EndsWith(".plist") || FSI.Name.EndsWith(".uappex"))
 					{
 						continue;
 					}
@@ -2629,7 +2631,7 @@ namespace UnrealBuildTool.XcodeProjectXcconfig
 			Xcconfig.AppendLine("");
 			Xcconfig.AppendLine("// Project settings");
 			Xcconfig.AppendLine($"TARGETED_DEVICE_FAMILY = {SupportedDevices}");
-			Xcconfig.AppendLine($"PRODUCT_BUNDLE_IDENTIFIER = $(UE_SIGNING_PREFIX).{UnrealTargetName}.{Name}");
+			Xcconfig.AppendLine($"PRODUCT_BUNDLE_IDENTIFIER = {BundleIdentifier}.{Name}");
 			Xcconfig.AppendLine($"{DeploymentTargetKey} = {DeploymentTarget}");
 
 			Xcconfig.AppendLine("");
