Allow Clang to gracefully build without PGO optimization if the PGO directory is missing

diff --git a/Engine/Source/Programs/UnrealBuildTool/Platform/Windows/VCToolChain.cs b/Engine/Source/Programs/UnrealBuildTool/Platform/Windows/VCToolChain.cs
--- a/Engine/Source/Programs/UnrealBuildTool/Platform/Windows/VCToolChain.cs
+++ b/Engine/Source/Programs/UnrealBuildTool/Platform/Windows/VCToolChain.cs
@@ -1026,23 +1026,26 @@ namespace UnrealBuildTool
 				else if (CompileEnvironment.bPGOOptimize)
 				{
 					// Use a merged profdata file.
-					FileReference ProfDataFilename = GetClangCompileProfDataFilename(CompileEnvironment);
-					Log.TraceInformationOnce($"Using PGO profile data \"{ProfDataFilename}\"");
-					if (Target.WindowsPlatform.bSampleBasedPGO)
+					if (Directory.Exists(CompileEnvironment.PGODirectory))
 					{
-						if (Target.WindowsPlatform.Compiler.IsIntel())
+						FileReference ProfDataFilename = GetClangCompileProfDataFilename(CompileEnvironment);
+						Log.TraceInformationOnce($"Using PGO profile data \"{ProfDataFilename}\"");
+						if (Target.WindowsPlatform.bSampleBasedPGO)
 						{
-							Arguments.Add($"-fprofile-sample-use=\"{ProfDataFilename}\"");
+							if (Target.WindowsPlatform.Compiler.IsIntel())
+							{
+								Arguments.Add($"-fprofile-sample-use=\"{ProfDataFilename}\"");
+							}
+							else
+							{
+								Arguments.Add($"/clang:-fprofile-sample-use=\"{ProfDataFilename}\"");
+							}
 						}
 						else
 						{
-							Arguments.Add($"/clang:-fprofile-sample-use=\"{ProfDataFilename}\"");
+							Arguments.Add($"-fprofile-use=\"{ProfDataFilename}\"");
 						}
 					}
-					else
-					{
-						Arguments.Add($"-fprofile-use=\"{ProfDataFilename}\"");
-					}
 				}
 			}
 
@@ -1985,7 +1988,10 @@ namespace UnrealBuildTool
 
 			if (EnvVars.Compiler.IsClang() && CompileEnvironment.bPGOOptimize)
 			{
-				BaseCompileAction.AdditionalPrerequisiteItems.Add(FileItem.GetItemByFileReference(GetClangCompileProfDataFilename(CompileEnvironment)));
+				if (Directory.Exists(CompileEnvironment.PGODirectory))
+				{
+					BaseCompileAction.AdditionalPrerequisiteItems.Add(FileItem.GetItemByFileReference(GetClangCompileProfDataFilename(CompileEnvironment)));
+				}
 			}
 
 			// Generate the timing info
@@ -3562,7 +3568,6 @@ namespace UnrealBuildTool
 					else
 					{
 						Logger.LogWarning("Unable to prepare Profile Guided Optimization (PGO) files. Using PGO Optimize build will be disabled.");
-						bPGOOptimize = false;
 					}
 				}
 				else if (bPGOProfile)
@@ -3621,21 +3626,28 @@ namespace UnrealBuildTool
 				{
 					if (Target.WindowsPlatform.bSampleBasedPGO)
 					{
-						Log.TraceInformationOnce("Enabling using Sample-Based Profile Guided Optimization (SPGO). Linking will take a while.");
+						if (Directory.Exists(LinkEnvironment.PGODirectory))
+						{
+							FileReference ProfDataFilename = GetClangLinkProfDataFilename(LinkEnvironment);
 
-						Arguments.Add($"/dwodir:\"{LinkEnvironment.PGODirectory!}\"");
+							Log.TraceInformationOnce("Enabling using Sample-Based Profile Guided Optimization (SPGO). Linking will take a while.");
 
-						FileReference ProfDataFilename = GetClangLinkProfDataFilename(LinkEnvironment);
+							Arguments.Add($"/dwodir:\"{LinkEnvironment.PGODirectory!}\"");
 
-						if (LinkEnvironment.bAllowLTCG && Target.WindowsPlatform.bAllowClangLinker)
-						{
-							// Use a merged profdata file.
-							Arguments.Add($"/lto-sample-profile:\"{ProfDataFilename}\"");
-						}
+							if (LinkEnvironment.bAllowLTCG && Target.WindowsPlatform.bAllowClangLinker)
+							{
+								// Use a merged profdata file.
+								Arguments.Add($"/lto-sample-profile:\"{ProfDataFilename}\"");
+							}
 
-						if (Target.WindowsPlatform.Compiler.IsIntel())
+							if (Target.WindowsPlatform.Compiler.IsIntel())
+							{
+								Arguments.Add($"-fprofile-sample-use=\"{ProfDataFilename}\"");
+							}
+						}
+						else
 						{
-							Arguments.Add($"-fprofile-sample-use=\"{ProfDataFilename}\"");
+							Logger.LogWarning("\"{PGODir}\" does not exist. Using PGO Optimize build will be disabled.", LinkEnvironment.PGODirectory);
 						}
 					}
 					else
