UnrealBuildTool: Enabling LTO generally explicit for Clang platforms and should not automatically be enabled when PGO is enabled

#rnx

diff --git a/Engine/Source/Programs/UnrealBuildTool/ToolChain/ClangToolChain.cs b/Engine/Source/Programs/UnrealBuildTool/ToolChain/ClangToolChain.cs
--- a/Engine/Source/Programs/UnrealBuildTool/ToolChain/ClangToolChain.cs
+++ b/Engine/Source/Programs/UnrealBuildTool/ToolChain/ClangToolChain.cs
@@ -264,7 +264,7 @@ namespace UnrealBuildTool
 		{
 			base.SetUpGlobalEnvironment(Target, GlobalCompileEnvironment, GlobalLinkEnvironment);
 
-			if (Target.bAllowLTCG || Target.bPGOOptimize || Target.bPGOProfile)
+			if (Target.bAllowLTCG)
 			{
 				Options |= ClangToolChainOptions.EnableLinkTimeOptimization;
 				if (Target.bPreferThinLTO)
@@ -735,7 +735,7 @@ namespace UnrealBuildTool
 		/// <param name="Arguments"></param>
 		protected virtual void GetCompileArguments_Optimizations(CppCompileEnvironment CompileEnvironment, List<string> Arguments)
 		{
-			bool bLTOEnabled = Options.HasFlag(ClangToolChainOptions.EnableLinkTimeOptimization) || CompileEnvironment.bAllowLTCG || CompileEnvironment.bPGOOptimize || CompileEnvironment.bPGOProfile;
+			bool bLTOEnabled = Options.HasFlag(ClangToolChainOptions.EnableLinkTimeOptimization);
 
 			// Link Time Optimization (LTO)
 			if (bLTOEnabled)
