Made UnrealGame/Client fallback to old deploy method from Visual Studio project. Made mixed build/deploy with -FastIterate set/unset more failsafe

#android

diff --git a/Engine/Source/Programs/AutomationTool/Android/AndroidPlatform.Automation.cs b/Engine/Source/Programs/AutomationTool/Android/AndroidPlatform.Automation.cs
--- a/Engine/Source/Programs/AutomationTool/Android/AndroidPlatform.Automation.cs
+++ b/Engine/Source/Programs/AutomationTool/Android/AndroidPlatform.Automation.cs
@@ -2582,6 +2582,12 @@ public class AndroidPlatform : Platform
 		}
 	}
 
+	private bool WasBuiltWithDontBundleLibrariesInAPK(string FinalSOName)
+	{
+		// The existence of the Push_..._so.bat file signifies a DontBundleLibrariesInAPK build
+		return File.Exists(Path.Combine(Path.GetDirectoryName(FinalSOName)!, "Push_" + Path.GetFileNameWithoutExtension(FinalSOName) + "_so.bat"));
+	}
+
 	private void DeployAndroidFileServer(ProjectParams Params, DeploymentContext SC, string AFSToken)
 	{
 		ConfigHierarchy Ini = ConfigCache.ReadHierarchy(ConfigHierarchyType.Engine, DirectoryReference.FromFile(Params.RawProjectPath), UnrealTargetPlatform.Android);
@@ -2870,7 +2876,7 @@ public class AndroidPlatform : Platform
 				client.FileWriteString("APK: " + APKLastUpdateTime + "\n", "^ext/APKFileStamp.txt");
 			}
 
-			if (bDontBundleLibrariesInAPK)
+			if (bDontBundleLibrariesInAPK && WasBuiltWithDontBundleLibrariesInAPK(FinalSOName))
 			{
 				string FinalSONameStripped = Path.Combine(Path.GetDirectoryName(FinalSOName), Path.GetFileNameWithoutExtension(FinalSOName) + "-stripped" + Path.GetExtension(FinalSOName));
 
@@ -3362,7 +3368,7 @@ public class AndroidPlatform : Platform
                 }
             }
 
-			if (bDontBundleLibrariesInAPK)
+			if (bDontBundleLibrariesInAPK && WasBuiltWithDontBundleLibrariesInAPK(FinalSOName))
 			{
 				string FinalSOPathStripped = Path.Combine(Path.GetDirectoryName(FinalSOName), Path.GetFileNameWithoutExtension(FinalSOName) + "-stripped" + Path.GetExtension(FinalSOName));
 				string FinalSOFileNameStripped = Path.GetFileName(FinalSOPathStripped);
diff --git a/Engine/Source/Programs/UnrealBuildTool/Platform/Android/AndroidProjectGenerator.cs b/Engine/Source/Programs/UnrealBuildTool/Platform/Android/AndroidProjectGenerator.cs
--- a/Engine/Source/Programs/UnrealBuildTool/Platform/Android/AndroidProjectGenerator.cs
+++ b/Engine/Source/Programs/UnrealBuildTool/Platform/Android/AndroidProjectGenerator.cs
@@ -214,9 +214,22 @@ namespace UnrealBuildTool
 					// MSBuild Property Functions allow to invoke limited C# expression from with-in MSBuild, see https://learn.microsoft.com/en-us/visualstudio/msbuild/property-functions?view=vs-2022
 					// They lack conditional statements, so this expression makes a branch-less variant by always appending "::" to AndroidDebugTarget,
 					// so regardless of what value it has, Split(':')[1] will never throw an exception.
-					string GetTargetDeviceSerial = "$([System.String]::Concat($(AndroidDebugTarget), \"::\").Split(':')[1])";
-					Out.AppendLine(
-						$"    <AndroidPostApkInstallCommands>CALL \"{UatPath}\" BuildCookRun -project=\"{UProjectPath.Replace("\"", "")}\" -FastIterate -FromMsBuild -client -skipbuild -skipcook -skipstage -deploy -config={InVSSettings.Configuration} -platform={InVSSettings.Platform} -clientarchitecture={ShortAbi} -device=\"{GetTargetDeviceSerial}\";$(AndroidPostApkInstallCommands)</AndroidPostApkInstallCommands>");
+					if (UProjectPath.Length > 0)
+					{
+						string GetTargetDeviceSerial = "$([System.String]::Concat($(AndroidDebugTarget), \"::\").Split(':')[1])";
+						Out.AppendLine(
+							$"    <AndroidPostApkInstallCommands>CALL \"{UatPath}\" BuildCookRun -project=\"{UProjectPath.Replace("\"", "")}\" -FastIterate -FromMsBuild -client -skipbuild -skipcook -skipstage -deploy -config={InVSSettings.Configuration} -platform={InVSSettings.Platform} -clientarchitecture={ShortAbi} -device=\"{GetTargetDeviceSerial}\";$(AndroidPostApkInstallCommands)</AndroidPostApkInstallCommands>");
+					}
+					else
+					{
+						// UAT deploy can only be called with a project file. Fallback to Push_..._so.bat script
+						string PushSOScript = Path.Combine(
+							Path.GetDirectoryName(NMakeOutputPath.FullName)!,
+							"Push_" + Path.GetFileNameWithoutExtension(NMakeOutputPath.FullName) + $"-{ShortAbi}_so.bat");
+						string GetTargetDeviceSerial = "$([System.String]::Concat($(AndroidDebugTarget), \"::\").Split(':')[1])";
+						Out.AppendLine(
+							$"    <AndroidPostApkInstallCommands>IF EXIST {PushSOScript} {PushSOScript} {GetTargetDeviceSerial};$(AndroidPostApkInstallCommands)</AndroidPostApkInstallCommands>");
+					}
 
 					// Ensure the following properties are up to date even if .vcxproj.user already exists and are in correct order between each other.
 					VCUserFileSettings.PatchProperty("AndroidDebugTarget", true);
