UnrealBuildTool: Remove out of date comment

#rnx

diff --git a/Engine/Source/Programs/UnrealBuildTool/ToolChain/ClangToolChain.cs b/Engine/Source/Programs/UnrealBuildTool/ToolChain/ClangToolChain.cs
--- a/Engine/Source/Programs/UnrealBuildTool/ToolChain/ClangToolChain.cs
+++ b/Engine/Source/Programs/UnrealBuildTool/ToolChain/ClangToolChain.cs
@@ -761,7 +761,6 @@ namespace UnrealBuildTool
 			}
 			else if (CompileEnvironment.bPGOProfile)
 			{
-				// Always enable LTO when generating PGO profile data.
 				Log.TraceInformationOnce("Enabling Profile Guided Instrumentation (PGI). Linking will take a while.");
 				Arguments.Add("-fprofile-generate");
 			}
