Enable ThinLTO for PGO when using Clang on Windows

diff --git a/Engine/Source/Programs/UnrealBuildTool/Platform/Windows/VCToolChain.cs b/Engine/Source/Programs/UnrealBuildTool/Platform/Windows/VCToolChain.cs
--- a/Engine/Source/Programs/UnrealBuildTool/Platform/Windows/VCToolChain.cs
+++ b/Engine/Source/Programs/UnrealBuildTool/Platform/Windows/VCToolChain.cs
@@ -989,7 +989,7 @@ namespace UnrealBuildTool
 
 			if (Target.WindowsPlatform.Compiler.IsClang())
 			{
-				if (CompileEnvironment.bAllowLTCG && Target.WindowsPlatform.bAllowClangLinker)
+				if (bEnableLTCG && Target.WindowsPlatform.bAllowClangLinker)
 				{
 					// Enable link-time code generation. Requires Clang linker.
 					Arguments.Add("-flto=thin");
@@ -3594,11 +3594,18 @@ namespace UnrealBuildTool
 			}
 			else // Clang
 			{
-				if (LinkEnvironment.bAllowLTCG)
+				if (LinkEnvironment.bAllowLTCG || bPGOProfile || bPGOOptimize)
 				{
 					if (Target.WindowsPlatform.bAllowClangLinker)
 					{
-						Log.TraceInformationOnce("Enabling Link-time optimization. Linking will take a while.");
+						if (LinkEnvironment.bPGOOptimize || LinkEnvironment.bPGOProfile)
+						{
+							Log.TraceInformationOnce("Enabling Link-time optimization as Profile Guided Optimization (PGO) was requested. Linking will take a while.");
+						}
+						else
+						{
+							Log.TraceInformationOnce("Enabling Link-time optimization. Linking will take a while.");
+						}
 
 						// lld should consider logical cores when determining how many threads to use
 						Arguments.Add("/opt:lldltojobs=all");
