Fix build logic for LTCG/LTO builds not applying to PGO builds for Windows

Submitting on behalf of Joe Kirchoff.

diff --git a/Engine/Source/Programs/UnrealBuildTool/Platform/Windows/VCToolChain.cs b/Engine/Source/Programs/UnrealBuildTool/Platform/Windows/VCToolChain.cs
--- a/Engine/Source/Programs/UnrealBuildTool/Platform/Windows/VCToolChain.cs
+++ b/Engine/Source/Programs/UnrealBuildTool/Platform/Windows/VCToolChain.cs
@@ -3342,7 +3342,7 @@ namespace UnrealBuildTool
 				LinkAction.bCanExecuteRemotely = true;
 			}
 
-			if (Target.WindowsPlatform.Compiler.IsClang() && Target.WindowsPlatform.bAllowClangLinker && LinkEnvironment.bAllowLTCG)
+			if (Target.WindowsPlatform.Compiler.IsClang() && Target.WindowsPlatform.bAllowClangLinker && (LinkEnvironment.bAllowLTCG || LinkEnvironment.bPGOOptimize || LinkEnvironment.bPGOProfile))
 			{
 				// Set the weight to number of logical cores as lld can max out the available cores
 				LinkAction.Weight = Utils.GetLogicalProcessorCount();
