UnrealBuildTool: Show collection content when a target violates the shared build environment

#rnx

diff --git a/Engine/Source/Programs/UnrealBuildTool/Configuration/TargetRules.cs b/Engine/Source/Programs/UnrealBuildTool/Configuration/TargetRules.cs
--- a/Engine/Source/Programs/UnrealBuildTool/Configuration/TargetRules.cs
+++ b/Engine/Source/Programs/UnrealBuildTool/Configuration/TargetRules.cs
@@ -2949,7 +2949,7 @@ namespace UnrealBuildTool
 					object? baseValue = field.GetValue(baseObjects[idx]);
 					if (!CheckValuesMatch(field.FieldType, thisValue, baseValue))
 					{
-						propNamesThatRequireUnique.Add(field.Name, (thisValue?.ToString(), baseValue?.ToString()));
+						propNamesThatRequireUnique.Add(field.Name, (ValueAsString(field.FieldType, thisValue), ValueAsString(field.FieldType, baseValue)));
 					}
 				}
 
@@ -2959,7 +2959,7 @@ namespace UnrealBuildTool
 					object? baseValue = property.GetValue(baseObjects[idx]);
 					if (!CheckValuesMatch(property.PropertyType, thisValue, baseValue))
 					{
-						propNamesThatRequireUnique.Add(property.Name, (thisValue?.ToString(), baseValue?.ToString()));
+						propNamesThatRequireUnique.Add(property.Name, (ValueAsString(property.PropertyType, thisValue), ValueAsString(property.PropertyType, baseValue)));
 					}
 				}
 			}
@@ -2974,7 +2974,7 @@ namespace UnrealBuildTool
 					object? baseValue = field.GetValue(baseObjects[idx]);
 					if (!CheckValuesMatch(field.FieldType, thisValue, baseValue))
 					{
-						propNamesThatRequireUnique.Add(field.Name, (thisValue?.ToString(), baseValue?.ToString()));
+						propNamesThatRequireUnique.Add(field.Name, (ValueAsString(field.FieldType, thisValue), ValueAsString(field.FieldType, baseValue)));
 					}
 				}
 
@@ -2984,7 +2984,7 @@ namespace UnrealBuildTool
 					object? baseValue = property.GetValue(baseObjects[idx]);
 					if (!CheckValuesMatch(property.PropertyType, thisValue, baseValue))
 					{
-						propNamesThatRequireUnique.Add(property.Name, (thisValue?.ToString(), baseValue?.ToString()));
+						propNamesThatRequireUnique.Add(property.Name, (ValueAsString(property.PropertyType, thisValue), ValueAsString(property.PropertyType, baseValue)));
 					}
 				}
 			}
@@ -3070,6 +3070,20 @@ namespace UnrealBuildTool
 			return bFieldsMatch;
 		}
 
+		static string ValueAsString(Type valueType, object? value)
+		{
+			if (value == null)
+			{
+				return "null";
+			}
+			else if (typeof(IEnumerable<string>).IsAssignableFrom(valueType))
+			{
+				IEnumerable<string> enumerable = (IEnumerable<string>)value;
+				return String.Join(',', (IEnumerable<string>)value);
+			}
+			return value?.ToString() ?? String.Empty;
+		}
+
 		/// <summary>
 		/// Create a TargetRules instance
 		/// </summary>
