Made MakeApk do a clean during packaging if EnableBundle is flipped

#android

diff --git a/Engine/Source/Programs/UnrealBuildTool/Platform/Android/UEDeployAndroid.cs b/Engine/Source/Programs/UnrealBuildTool/Platform/Android/UEDeployAndroid.cs
--- a/Engine/Source/Programs/UnrealBuildTool/Platform/Android/UEDeployAndroid.cs
+++ b/Engine/Source/Programs/UnrealBuildTool/Platform/Android/UEDeployAndroid.cs
@@ -2209,7 +2209,7 @@ namespace UnrealBuildTool
 		}
 
 		private string GetAllBuildSettings(AndroidToolChain ToolChain, UnrealPluginLanguage UPL, bool bForDistribution, bool bMakeSeparateApks, bool bPackageDataInsideApk,
-			bool bDisableVerifyOBBOnStartUp, bool bUseExternalFilesDir, bool bDontBundleLibrariesInAPK, bool bUsesZenStreaming, string TemplatesHashCode)
+			bool bDisableVerifyOBBOnStartUp, bool bUseExternalFilesDir, bool bEnableBundle, bool bDontBundleLibrariesInAPK, bool bUsesZenStreaming, string TemplatesHashCode)
 		{
 			// make the settings string - this will be char by char compared against last time
 			StringBuilder CurrentSettings = new StringBuilder();
@@ -2225,6 +2225,7 @@ namespace UnrealBuildTool
 			CurrentSettings.AppendLine(String.Format("bUseExternalFilesDir={0}", bUseExternalFilesDir));
 			CurrentSettings.AppendLine(String.Format("UPLHashCode={0}", UPLHashCode));
 			CurrentSettings.AppendLine(String.Format("TemplatesHashCode={0}", TemplatesHashCode));
+			CurrentSettings.AppendLine(String.Format("bEnableBundle={0}", bEnableBundle));
 			CurrentSettings.AppendLine(String.Format("bDontBundleLibrariesInAPK={0}", bDontBundleLibrariesInAPK));
 			CurrentSettings.AppendLine(String.Format("bUsesZenStreaming={0}", bUsesZenStreaming));
 
@@ -4589,6 +4590,21 @@ dependencies.implementation('androidx.tracing:tracing:1.2.0')
 				}
 			}
 
+			// Check if bEnableBundle changed
+			string BuildSettingsCacheFile = Path.Combine(IntermediateAndroidPath, "UEBuildSettings.txt");
+			if (!bCreateFromScratch)
+			{
+				if (File.Exists(BuildSettingsCacheFile))
+				{
+					const string Key = "bEnableBundle=";
+					string? Line = File.ReadAllLines(BuildSettingsCacheFile).FirstOrDefault(Line => Line.StartsWith(Key));
+					if (Line != null && !Line[Key.Length..].Equals(bEnableBundle ? "True" : "False"))
+					{
+						bCreateFromScratch = true;
+					}
+				}
+			}
+
 			if (bCreateFromScratch)
 			{
 				Logger.LogInformation("Cleaning {IntermediateAndroidPath}", IntermediateAndroidPath);
@@ -4749,8 +4765,7 @@ dependencies.implementation('androidx.tracing:tracing:1.2.0')
 
 			// check to see if any "meta information" is newer than last time we build
 			string TemplatesHashCode = GenerateTemplatesHashCode(EngineDirectory);
-			string CurrentBuildSettings = GetAllBuildSettings(ToolChain, UPL!, bForDistribution, bMakeSeparateApks, bPackageDataInsideApk, bDisableVerifyOBBOnStartUp, bUseExternalFilesDir, bDontBundleLibrariesInAPK, bUsesZenStreaming, TemplatesHashCode);
-			string BuildSettingsCacheFile = Path.Combine(IntermediateAndroidPath, "UEBuildSettings.txt");
+			string CurrentBuildSettings = GetAllBuildSettings(ToolChain, UPL!, bForDistribution, bMakeSeparateApks, bPackageDataInsideApk, bDisableVerifyOBBOnStartUp, bUseExternalFilesDir, bEnableBundle, bDontBundleLibrariesInAPK, bUsesZenStreaming, TemplatesHashCode);
 
 			// do we match previous build settings?
 			bool bBuildSettingsMatch = true;
