Experimental Visual Studio 2026 support
Remove defunct Visual Studio Mac project generator

diff --git a/Engine/Config/Windows/Windows_SDK.json b/Engine/Config/Windows/Windows_SDK.json
--- a/Engine/Config/Windows/Windows_SDK.json
+++ b/Engine/Config/Windows/Windows_SDK.json
@@ -74,7 +74,11 @@
 	"VisualStudio2022SuggestedComponents": [
 	],
 
+	"VisualStudio2026SuggestedComponents": [
+	],
+
 	"MinimumVisualStudio2022Version": "17.8",
+	"MinimumVisualStudio2026Version": "18.0",
 
 	"//05.1": "X64 -- 0xA00 = Win10",
 	"//05.2": "Arm64 -- 0xA00 = Win10",
diff --git a/Engine/Plugins/Developer/VisualStudioSourceCodeAccess/Source/VisualStudioSourceCodeAccess/Private/VisualStudioSourceCodeAccessModule.cpp b/Engine/Plugins/Developer/VisualStudioSourceCodeAccess/Source/VisualStudioSourceCodeAccess/Private/VisualStudioSourceCodeAccessModule.cpp
--- a/Engine/Plugins/Developer/VisualStudioSourceCodeAccess/Source/VisualStudioSourceCodeAccess/Private/VisualStudioSourceCodeAccessModule.cpp
+++ b/Engine/Plugins/Developer/VisualStudioSourceCodeAccess/Source/VisualStudioSourceCodeAccess/Private/VisualStudioSourceCodeAccessModule.cpp
@@ -27,6 +27,7 @@ void FVisualStudioSourceCodeAccessModule::StartupModule()
 	// Add all the explicit version wrappers. If one of these is selected, UBT will generate project files in the appropriate format. Editor behavior is still to detect which version to use
 	// from the solution on disk.
 	RegisterWrapper("VisualStudio2022", LOCTEXT("VisualStudio2022", "Visual Studio 2022"), LOCTEXT("UsingVisualStudio2022", "Open source code files in Visual Studio 2022"));
+	RegisterWrapper("VisualStudio2026", LOCTEXT("VisualStudio2026", "Visual Studio 2026 (experimental}"), LOCTEXT("UsingVisualStudio2026", "Open source code files in Visual Studio 2026"));
 
 	// Bind our source control provider to the editor
 	IModularFeatures::Get().RegisterModularFeature(TEXT("SourceCodeAccessor"), &VisualStudioSourceCodeAccessor.Get() );
diff --git a/Engine/Plugins/Developer/VisualStudioSourceCodeAccess/Source/VisualStudioSourceCodeAccess/Private/VisualStudioSourceCodeAccessor.cpp b/Engine/Plugins/Developer/VisualStudioSourceCodeAccess/Source/VisualStudioSourceCodeAccess/Private/VisualStudioSourceCodeAccessor.cpp
--- a/Engine/Plugins/Developer/VisualStudioSourceCodeAccess/Source/VisualStudioSourceCodeAccess/Private/VisualStudioSourceCodeAccessor.cpp
+++ b/Engine/Plugins/Developer/VisualStudioSourceCodeAccess/Source/VisualStudioSourceCodeAccess/Private/VisualStudioSourceCodeAccessor.cpp
@@ -117,7 +117,9 @@ static void OnModuleCompileStarted(bool bIsAsyncCompile)
 
 int32 GetVisualStudioVersionForCompiler()
 {
-#if _MSC_VER >= 1930
+#if _MSC_VER >= 1950
+	return 18; // Visual Studio 2026
+#elif _MSC_VER >= 1930
 	return 17; // Visual Studio 2022
 #elif defined(__clang__)
 	return 0;
diff --git a/Engine/Source/Developer/SourceCodeAccess/Private/SourceCodeAccessModule.cpp b/Engine/Source/Developer/SourceCodeAccess/Private/SourceCodeAccessModule.cpp
--- a/Engine/Source/Developer/SourceCodeAccess/Private/SourceCodeAccessModule.cpp
+++ b/Engine/Source/Developer/SourceCodeAccess/Private/SourceCodeAccessModule.cpp
@@ -78,7 +78,7 @@ bool FSourceCodeAccessModule::CanCompileSourceCode() const
 {
 #if PLATFORM_WINDOWS
 	// Need to have Visual Studio installed to compile on Windows, regardless of chosen IDE
-	return IsSourceCodeAccessorAvailable("VisualStudio2022");
+	return IsSourceCodeAccessorAvailable("VisualStudio2022") || IsSourceCodeAccessorAvailable("VisualStudio2026");
 #else
 	// Default behavior
 	return CanAccessSourceCode();
diff --git a/Engine/Source/Developer/Windows/WindowsTargetPlatformSettings/Classes/WindowsTargetSettings.h b/Engine/Source/Developer/Windows/WindowsTargetPlatformSettings/Classes/WindowsTargetSettings.h
--- a/Engine/Source/Developer/Windows/WindowsTargetPlatformSettings/Classes/WindowsTargetSettings.h
+++ b/Engine/Source/Developer/Windows/WindowsTargetPlatformSettings/Classes/WindowsTargetSettings.h
@@ -21,6 +21,7 @@ enum class ECompilerVersion : uint8
 	VisualStudio2017 = 2 UMETA(DisplayName = "Visual Studio 2017 (deprecated)"),
 	VisualStudio2019 = 3 UMETA(DisplayName = "Visual Studio 2019 (deprecated)"),
 	VisualStudio2022 = 4 UMETA(DisplayName = "Visual Studio 2022"),
+	VisualStudio2026 = 5 UMETA(DisplayName = "Visual Studio 2026 (experimental)"),
 };
 
 UENUM()
diff --git a/Engine/Source/Programs/UnrealBuildTool/Executors/Experimental/FASTBuild.cs b/Engine/Source/Programs/UnrealBuildTool/Executors/Experimental/FASTBuild.cs
--- a/Engine/Source/Programs/UnrealBuildTool/Executors/Experimental/FASTBuild.cs
+++ b/Engine/Source/Programs/UnrealBuildTool/Executors/Experimental/FASTBuild.cs
@@ -923,6 +923,7 @@ namespace UnrealBuildTool
 				switch (VCEnv.Compiler)
 				{
 					case WindowsCompiler.VisualStudio2022:
+					case WindowsCompiler.VisualStudio2026:
 						// For now we are working with the 140 version, might need to change to 141 or 150 depending on the version of the Toolchain you chose
 						// to install
 						platformVersionNumber = "140";
diff --git a/Engine/Source/Programs/UnrealBuildTool/Modes/GenerateProjectFilesMode.cs b/Engine/Source/Programs/UnrealBuildTool/Modes/GenerateProjectFilesMode.cs
--- a/Engine/Source/Programs/UnrealBuildTool/Modes/GenerateProjectFilesMode.cs
+++ b/Engine/Source/Programs/UnrealBuildTool/Modes/GenerateProjectFilesMode.cs
@@ -33,7 +33,6 @@ namespace UnrealBuildTool
 		[CommandLine("-XCodeProjectFiles", Value = nameof(ProjectFileFormat.XCode))]
 		[CommandLine("-EddieProjectFiles", Value = nameof(ProjectFileFormat.Eddie))]
 		[CommandLine("-VSCode", Value = nameof(ProjectFileFormat.VisualStudioCode))]
-		[CommandLine("-VSMac", Value = nameof(ProjectFileFormat.VisualStudioMac))]
 		[CommandLine("-VSWorkspace", Value = nameof(ProjectFileFormat.VisualStudioWorkspace))]
 		[CommandLine("-CLion", Value = nameof(ProjectFileFormat.CLion))]
 		[CommandLine("-Rider", Value = nameof(ProjectFileFormat.Rider))]
@@ -216,6 +215,9 @@ namespace UnrealBuildTool
 					case ProjectFileFormat.VisualStudio2022:
 						Generator = new VCProjectFileGenerator(ProjectFile, VCProjectFileFormat.VisualStudio2022, Arguments);
 						break;
+					case ProjectFileFormat.VisualStudio2026:
+						Generator = new VCProjectFileGenerator(ProjectFile, VCProjectFileFormat.VisualStudio2026, Arguments);
+						break;
 					case ProjectFileFormat.XCode:
 						Generator = new XcodeProjectFileGenerator(ProjectFile, Arguments);
 						break;
@@ -231,9 +233,6 @@ namespace UnrealBuildTool
 					case ProjectFileFormat.CLion:
 						Generator = new CLionGenerator(ProjectFile);
 						break;
-					case ProjectFileFormat.VisualStudioMac:
-						Generator = new VCMacProjectFileGenerator(ProjectFile, Arguments);
-						break;
 					case ProjectFileFormat.Rider:
 						Generator = new RiderProjectFileGenerator(ProjectFile, Arguments);
 						break;
diff --git a/Engine/Source/Programs/UnrealBuildTool/Platform/Windows/MicrosoftPlatformSDK.VisualStudio.cs b/Engine/Source/Programs/UnrealBuildTool/Platform/Windows/MicrosoftPlatformSDK.VisualStudio.cs
--- a/Engine/Source/Programs/UnrealBuildTool/Platform/Windows/MicrosoftPlatformSDK.VisualStudio.cs
+++ b/Engine/Source/Programs/UnrealBuildTool/Platform/Windows/MicrosoftPlatformSDK.VisualStudio.cs
@@ -25,10 +25,20 @@ namespace UnrealBuildTool
 		static readonly string[] s_visualStudio2022SuggestedComponents = SDK.GetStringArrayFromConfig("VisualStudio2022SuggestedComponents");
 
 		/// <summary>
-		/// The minimum supported Visual Studio IDE version
+		/// Additional set of components that should be suggested to be installed for Visual Studio 2026.
+		/// </summary>
+		static readonly string[] s_visualStudio2026SuggestedComponents = SDK.GetStringArrayFromConfig("VisualStudio2026SuggestedComponents");
+
+		/// <summary>
+		/// The minimum supported Visual Studio 2022 IDE version
 		/// </summary>
 		static VersionNumber MinimumVisualStudio2022Version => SDK.GetRequiredVersionNumberFromConfig("MinimumVisualStudio2022Version");
 
+		/// <summary>
+		/// The minimum supported Visual Studio 2026 IDE version
+		/// </summary>
+		static VersionNumber MinimumVisualStudio2026Version => SDK.GetRequiredVersionNumberFromConfig("MinimumVisualStudio2026Version");
+
 		/// <summary>
 		/// Returns the list of suggested of components that should be suggested to be installed for Visual Studio.
 		/// Used to generate a .vsconfig file which will prompt Visual Studio to ask the user to install these components.
@@ -38,13 +48,16 @@ namespace UnrealBuildTool
 			bool platformLinuxValid = (InstalledPlatformInfo.IsValidPlatform(UnrealTargetPlatform.Linux) && UEBuildPlatform.IsPlatformAvailable(UnrealTargetPlatform.Linux))
 				|| (InstalledPlatformInfo.IsValidPlatform(UnrealTargetPlatform.LinuxArm64) && UEBuildPlatform.IsPlatformAvailable(UnrealTargetPlatform.LinuxArm64));
 
-			SortedSet<string> components = new(s_visualStudioSuggestedComponents);
+			SortedSet<string> components = [.. s_visualStudioSuggestedComponents];
 
 			switch (format)
 			{
 				case VCProjectFileFormat.VisualStudio2022:
 					components.UnionWith(s_visualStudio2022SuggestedComponents);
 					break;
+				case VCProjectFileFormat.VisualStudio2026:
+					components.UnionWith(s_visualStudio2026SuggestedComponents);
+					break;
 				default:
 					throw new BuildException("Unsupported Visual Studio version {0}", format);
 			}
diff --git a/Engine/Source/Programs/UnrealBuildTool/Platform/Windows/MicrosoftPlatformSDK.cs b/Engine/Source/Programs/UnrealBuildTool/Platform/Windows/MicrosoftPlatformSDK.cs
--- a/Engine/Source/Programs/UnrealBuildTool/Platform/Windows/MicrosoftPlatformSDK.cs
+++ b/Engine/Source/Programs/UnrealBuildTool/Platform/Windows/MicrosoftPlatformSDK.cs
@@ -621,6 +621,8 @@ namespace UnrealBuildTool
 			{
 				case WindowsCompiler.VisualStudio2022:
 					return "Visual Studio 2022";
+				case WindowsCompiler.VisualStudio2026:
+					return "Visual Studio 2026";
 				default:
 					return Compiler.ToString();
 			}
@@ -753,8 +755,9 @@ namespace UnrealBuildTool
 							AddClangToolChain(Compiler, new DirectoryReference(LlvmPath), ToolChains, IsAutoSdk: false, Logger);
 						}
 
-						// Check for installations bundled with Visual Studio 2022
-						foreach (VisualStudioInstallation Installation in FindVisualStudioInstallations(WindowsCompiler.VisualStudio2022, Logger))
+						// Check for installations bundled with Visual Studio
+						foreach (VisualStudioInstallation Installation in FindVisualStudioInstallations(WindowsCompiler.VisualStudio2022, Logger)
+							.Concat(FindVisualStudioInstallations(WindowsCompiler.VisualStudio2026, Logger)))
 						{
 							AddClangToolChain(Compiler, DirectoryReference.Combine(Installation.BaseDir, "VC", "Tools", "Llvm"), ToolChains, IsAutoSdk: false, Logger);
 							AddClangToolChain(Compiler, DirectoryReference.Combine(Installation.BaseDir, "VC", "Tools", "Llvm", "x64"), ToolChains, IsAutoSdk: false, Logger);
@@ -869,6 +872,7 @@ namespace UnrealBuildTool
 							switch (Compiler)
 							{
 								case WindowsCompiler.VisualStudio2022: VSDir = "VS2022"; break;
+								case WindowsCompiler.VisualStudio2026: VSDir = "VS2026"; break;
 							}
 
 							if (!String.IsNullOrEmpty(VSDir))
@@ -986,7 +990,11 @@ namespace UnrealBuildTool
 						int MajorVersion = Version.GetComponent(0);
 
 						WindowsCompiler Compiler;
-						if (Version >= MicrosoftPlatformSDK.MinimumVisualStudio2022Version) // Treat any newer versions as 2022, until we have an explicit enum for them
+						if (Version >= MicrosoftPlatformSDK.MinimumVisualStudio2026Version) // Treat any newer versions as 2026, until we have an explicit enum for them
+						{
+							Compiler = WindowsCompiler.VisualStudio2026;
+						}
+						else if (Version >= MicrosoftPlatformSDK.MinimumVisualStudio2022Version)
 						{
 							Compiler = WindowsCompiler.VisualStudio2022;
 						}
@@ -1440,6 +1448,7 @@ namespace UnrealBuildTool
 			// Get the Visual Studio install directory
 #pragma warning disable CS0618 // Type or member is obsolete
 			IEnumerable<DirectoryReference> InstallDirs = FindVisualStudioInstallations(WindowsCompiler.VisualStudio2022, Logger)
+				.Concat(FindVisualStudioInstallations(WindowsCompiler.VisualStudio2026, Logger))
 				.Concat(FindVisualStudioInstallations(WindowsCompiler.VisualStudioUnsupported, Logger))
 				.Select(x => x.BaseDir);
 #pragma warning restore CS0618 // Type or member is obsolete
@@ -1516,6 +1525,7 @@ namespace UnrealBuildTool
 					switch (Compiler)
 					{
 						case WindowsCompiler.VisualStudio2022: VSDir = "VS2022"; break;
+						case WindowsCompiler.VisualStudio2026: VSDir = "VS2026"; break;
 					}
 
 					if (!String.IsNullOrEmpty(VSDir))
diff --git a/Engine/Source/Programs/UnrealBuildTool/Platform/Windows/UEBuildWindows.cs b/Engine/Source/Programs/UnrealBuildTool/Platform/Windows/UEBuildWindows.cs
--- a/Engine/Source/Programs/UnrealBuildTool/Platform/Windows/UEBuildWindows.cs
+++ b/Engine/Source/Programs/UnrealBuildTool/Platform/Windows/UEBuildWindows.cs
@@ -1,5 +1,7 @@
 // Copyright Epic Games, Inc. All Rights Reserved.
 
+using EpicGames.Core;
+using Microsoft.Extensions.Logging;
 using System;
 using System.Collections.Generic;
 using System.Diagnostics;
@@ -8,8 +10,6 @@ using System.IO;
 using System.Linq;
 using System.Runtime.Versioning;
 using System.Text;
-using EpicGames.Core;
-using Microsoft.Extensions.Logging;
 using UnrealBuildBase;
 
 namespace UnrealBuildTool
@@ -147,6 +147,11 @@ namespace UnrealBuildTool
 		/// </summary>
 		VisualStudio2022,
 
+		/// <summary>
+		/// Visual Studio 2026 (Visual C++ 19.0)
+		/// </summary>
+		VisualStudio2026,
+
 		/// <summary>
 		/// Unsupported Visual Studio
 		/// Must be the last entry in WindowsCompiler enum and should only be used in limited circumstances
@@ -287,6 +292,7 @@ namespace UnrealBuildTool
 		[ConfigFile(ConfigHierarchyType.Engine, "/Script/WindowsTargetPlatform.WindowsTargetSettings", "Compiler")]
 		[XmlConfigFile(Category = "WindowsPlatform")]
 		[CommandLine("-2022", Value = nameof(WindowsCompiler.VisualStudio2022))]
+		[CommandLine("-2026", Value = nameof(WindowsCompiler.VisualStudio2022))]
 		[CommandLine("-Compiler=")]
 		public WindowsCompiler Compiler { get; set; } = WindowsCompiler.Default;
 
@@ -866,7 +872,8 @@ namespace UnrealBuildTool
 				case WindowsCompiler.ClangCustom:
 				case WindowsCompiler.Intel:
 				case WindowsCompiler.VisualStudio2022:
-					return "2015"; // VS2022 is backwards compatible with VS2015 compiler
+				case WindowsCompiler.VisualStudio2026:
+					return "2015"; // VS2022+ is backwards compatible with VS2015 compiler
 
 				default:
 					throw new BuildException("Unexpected WindowsCompiler version for GetVisualStudioCompilerVersionName().  Either not using a Visual Studio compiler or switch block needs to be updated");
@@ -1532,6 +1539,10 @@ namespace UnrealBuildTool
 					{
 						return WindowsCompiler.VisualStudio2022;
 					}
+					else if (Format == ProjectFileFormat.VisualStudio2026)
+					{
+						return WindowsCompiler.VisualStudio2026;
+					}
 				}
 			}
 
@@ -1544,6 +1555,10 @@ namespace UnrealBuildTool
 				{
 					return WindowsCompiler.VisualStudio2022;
 				}
+				else if (ProjectFormat == VCProjectFileFormat.VisualStudio2026)
+				{
+					return WindowsCompiler.VisualStudio2026;
+				}
 			}
 
 			// Check the editor settings too
@@ -1554,6 +1569,10 @@ namespace UnrealBuildTool
 				{
 					return WindowsCompiler.VisualStudio2022;
 				}
+				else if (PreferredAccessor == ProjectFileFormat.VisualStudio2026)
+				{
+					return WindowsCompiler.VisualStudio2026;
+				}
 			}
 
 			// Second, default based on what's installed, test for 2022 first
@@ -1561,6 +1580,10 @@ namespace UnrealBuildTool
 			{
 				return WindowsCompiler.VisualStudio2022;
 			}
+			else if (MicrosoftPlatformSDK.HasValidCompiler(WindowsCompiler.VisualStudio2026, Architecture, Logger))
+			{
+				return WindowsCompiler.VisualStudio2026;
+			}
 
 			if (!bSkipWarning && !MicrosoftPlatformSDK.IgnoreToolchainErrors)
 			{
@@ -1586,6 +1609,14 @@ namespace UnrealBuildTool
 						MinimumVsVersionStr,
 						PreferredComponentStr);
 				}
+				else if (TryGetVSInstallDirs(WindowsCompiler.VisualStudio2026, Logger) != null)
+				{
+					Logger.LogWarning("Visual Studio 2026 is installed, but is out of date or missing a valid C++ toolchain (minimum version {MinVersion}, preferred version {PreferredVersion}). Please update Visual Studio 2026 to {MinimumVsVersionStr} or later and verify that the \"{Component}\" component is selected in the Visual Studio 2026 installation options.",
+						MinimumMsvcVersionStr,
+						PreferredMsvcVersionStr,
+						MinimumVsVersionStr,
+						PreferredComponentStr);
+				}
 				else
 				{
 					Logger.LogWarning("No valid Visual C++ toolchain was found (minimum version {MinVersion}, preferred version {PreferredVersion}). Please download and install Visual Studio 2022 {MinimumVsVersionStr} or later and verify that the \"{Component}\" component is selected in the Visual Studio 2022 installation options.",
diff --git a/Engine/Source/Programs/UnrealBuildTool/Platform/Windows/VCEnvironment.cs b/Engine/Source/Programs/UnrealBuildTool/Platform/Windows/VCEnvironment.cs
--- a/Engine/Source/Programs/UnrealBuildTool/Platform/Windows/VCEnvironment.cs
+++ b/Engine/Source/Programs/UnrealBuildTool/Platform/Windows/VCEnvironment.cs
@@ -516,7 +516,7 @@ namespace UnrealBuildTool
 
 				if (!WindowsPlatform.TryGetToolChainDir(ToolChain, ToolchainVersion, Architecture, Logger, out SelectedToolChainVersion, out SelectedToolChainDir, out SelectedRedistDir))
 				{
-					throw new BuildException("{0} must be installed in order to build this target.", WindowsPlatform.GetCompilerName(WindowsCompiler.VisualStudio2022));
+					throw new BuildException("{0} or newer must be installed in order to build this target.", WindowsPlatform.GetCompilerName(WindowsCompiler.VisualStudio2022));
 				}
 
 				VersionNumber? PreferredSelectedToolChainVersion = SelectedToolChainVersion;
@@ -531,7 +531,7 @@ namespace UnrealBuildTool
 					{
 						if (!WindowsPlatform.TryGetToolChainDir(ToolChain, $"<{SelectedToolChainVersion.GetComponent(0)}.{SelectedToolChainVersion.GetComponent(1)}", Architecture, Logger, out SelectedToolChainVersion, out SelectedToolChainDir, out SelectedRedistDir))
 						{
-							throw new BuildException("{0} compiler that supports Clang {1} must be installed in order to build this target.", WindowsPlatform.GetCompilerName(WindowsCompiler.VisualStudio2022));
+							throw new BuildException("{0} compiler or newer that supports Clang {1} must be installed in order to build this target.", WindowsPlatform.GetCompilerName(WindowsCompiler.VisualStudio2022));
 						}
 					}
 				}
diff --git a/Engine/Source/Programs/UnrealBuildTool/ProjectFiles/PlatformProjectGenerator.cs b/Engine/Source/Programs/UnrealBuildTool/ProjectFiles/PlatformProjectGenerator.cs
--- a/Engine/Source/Programs/UnrealBuildTool/ProjectFiles/PlatformProjectGenerator.cs
+++ b/Engine/Source/Programs/UnrealBuildTool/ProjectFiles/PlatformProjectGenerator.cs
@@ -90,6 +90,10 @@ namespace UnrealBuildTool
 						VSCompiler = WindowsCompiler.VisualStudio2022;
 						VCVersion = "v170";
 						break;
+					case VCProjectFileFormat.VisualStudio2026:
+						VSCompiler = WindowsCompiler.VisualStudio2026;
+						VCVersion = "v180";
+						break;
 					default:
 						// Unknown VS Version
 						return false;
diff --git a/Engine/Source/Programs/UnrealBuildTool/ProjectFiles/ProjectFileGenerator.cs b/Engine/Source/Programs/UnrealBuildTool/ProjectFiles/ProjectFileGenerator.cs
--- a/Engine/Source/Programs/UnrealBuildTool/ProjectFiles/ProjectFileGenerator.cs
+++ b/Engine/Source/Programs/UnrealBuildTool/ProjectFiles/ProjectFileGenerator.cs
@@ -145,11 +145,11 @@ namespace UnrealBuildTool
 		CodeLite,
 		VisualStudio,
 		VisualStudio2022,
+		VisualStudio2026,
 		VisualStudioWorkspace,
 		XCode,
 		Eddie,
 		VisualStudioCode,
-		VisualStudioMac,
 		CLion,
 		Rider,
 		AndroidStudio
diff --git a/Engine/Source/Programs/UnrealBuildTool/ProjectFiles/VisualStudio/VCProjectFileGenerator.cs b/Engine/Source/Programs/UnrealBuildTool/ProjectFiles/VisualStudio/VCProjectFileGenerator.cs
--- a/Engine/Source/Programs/UnrealBuildTool/ProjectFiles/VisualStudio/VCProjectFileGenerator.cs
+++ b/Engine/Source/Programs/UnrealBuildTool/ProjectFiles/VisualStudio/VCProjectFileGenerator.cs
@@ -18,6 +18,7 @@ namespace UnrealBuildTool
 	{
 		Default,          // Default to the best installed version, but allow SDKs to override
 		VisualStudio2022,
+		VisualStudio2026,
 	}
 
 	class VCProjectFileSettings
@@ -224,34 +225,33 @@ namespace UnrealBuildTool
 		/// "4.0", "12.0", or "14.0", etc...
 		public static string GetProjectFileToolVersionString(VCProjectFileFormat ProjectFileFormat)
 		{
-			switch (ProjectFileFormat)
+			return ProjectFileFormat switch
 			{
-				case VCProjectFileFormat.VisualStudio2022:
-					return "17.0";
-			}
-			return String.Empty;
+				VCProjectFileFormat.VisualStudio2022 => "17.0",
+				VCProjectFileFormat.VisualStudio2026 => "18.0",
+				_ => String.Empty,
+			};
 		}
 
 		/// for instance: <PlatformToolset>v110</PlatformToolset>
 		public static string GetProjectFilePlatformToolsetVersionString(VCProjectFileFormat ProjectFileFormat)
 		{
-			switch (ProjectFileFormat)
+			return ProjectFileFormat switch
 			{
-				case VCProjectFileFormat.VisualStudio2022:
-					return "v143";
-
-			}
-			return String.Empty;
+				VCProjectFileFormat.VisualStudio2022 => "v143",
+				VCProjectFileFormat.VisualStudio2026 => "v145",
+				_ => String.Empty,
+			};
 		}
 
 		public static WindowsCompiler GetCompilerForIntellisense(VCProjectFileFormat ProjectFileFormat)
 		{
-			switch (ProjectFileFormat)
+			return ProjectFileFormat switch
 			{
-				case VCProjectFileFormat.VisualStudio2022:
-					return WindowsCompiler.VisualStudio2022;
-			}
-			return WindowsCompiler.VisualStudio2022;
+				VCProjectFileFormat.VisualStudio2022 => WindowsCompiler.VisualStudio2022,
+				VCProjectFileFormat.VisualStudio2026 => WindowsCompiler.VisualStudio2026,
+				_ => WindowsCompiler.VisualStudio2022,
+			};
 		}
 
 		public static void AppendPlatformToolsetProperty(StringBuilder VCProjectFileContent, VCProjectFileFormat ProjectFileFormat)
@@ -334,6 +334,14 @@ namespace UnrealBuildTool
 					Settings.ProjectFileFormat = VCProjectFileFormat.Default;
 				}
 			}
+			else if (Settings.ProjectFileFormat == VCProjectFileFormat.VisualStudio2026)
+			{
+				if (!WindowsPlatform.HasCompiler(WindowsCompiler.VisualStudio2026, UnrealArch.X64, Logger))
+				{
+					Logger.LogWarning("Visual Studio C++ 2026 installation not found - ignoring preferred project file format.");
+					Settings.ProjectFileFormat = VCProjectFileFormat.Default;
+				}
+			}
 
 			// Certain platforms override the project file format because their debugger add-ins may not yet support the latest
 			// version of Visual Studio.  This is their chance to override that.
@@ -352,6 +360,11 @@ namespace UnrealBuildTool
 						Format = VCProjectFileFormat.VisualStudio2022;
 						break;
 					}
+					else if (Installation.Compiler == WindowsCompiler.VisualStudio2026)
+					{
+						Format = VCProjectFileFormat.VisualStudio2026;
+						break;
+					}
 				}
 				Settings.ProjectFileFormat = Format;
 
@@ -692,6 +705,14 @@ namespace UnrealBuildTool
 				VCSolutionFileContent.AppendLine("VisualStudioVersion = 17.0.31314.256");
 				VCSolutionFileContent.AppendLine("MinimumVisualStudioVersion = 10.0.40219.1");
 			}
+			else if (Settings.ProjectFileFormat == VCProjectFileFormat.VisualStudio2026)
+			{
+				VCSolutionFileContent.AppendLine();
+				VCSolutionFileContent.AppendLine("Microsoft Visual Studio Solution File, Format Version 12.00");
+				VCSolutionFileContent.AppendLine("# Visual Studio Version 18");
+				VCSolutionFileContent.AppendLine("VisualStudioVersion = 18.0.11005.162");
+				VCSolutionFileContent.AppendLine("MinimumVisualStudioVersion = 10.0.40219.1");
+			}
 			else
 			{
 				throw new BuildException("Unexpected ProjectFileFormat");
@@ -953,15 +974,12 @@ namespace UnrealBuildTool
 			if (bSuccess && Settings.bWriteSolutionOptionFile && OperatingSystem.IsWindows())
 			{
 				// Figure out the filename for the SUO file. VS will automatically import the options from earlier versions if necessary.
-				FileReference SolutionOptionsFileName;
-				switch (Settings.ProjectFileFormat)
+				FileReference SolutionOptionsFileName = Settings.ProjectFileFormat switch
 				{
-					case VCProjectFileFormat.VisualStudio2022:
-						SolutionOptionsFileName = FileReference.Combine(PrimaryProjectPath, ".vs", Path.GetFileNameWithoutExtension(SolutionFileName), "v17", ".suo");
-						break;
-					default:
-						throw new BuildException("Unsupported Visual Studio version");
-				}
+					VCProjectFileFormat.VisualStudio2022 => FileReference.Combine(PrimaryProjectPath, ".vs", Path.GetFileNameWithoutExtension(SolutionFileName), "v17", ".suo"),
+					VCProjectFileFormat.VisualStudio2026 => FileReference.Combine(PrimaryProjectPath, ".vs", Path.GetFileNameWithoutExtension(SolutionFileName), "v18", ".suo"),
+					_ => throw new BuildException("Unsupported Visual Studio version"),
+				};
 
 				// Check it doesn't exist before overwriting it. Since these files store the user's preferences, it'd be bad form to overwrite them.
 				if (!FileReference.Exists(SolutionOptionsFileName))
diff --git a/Engine/Source/Programs/UnrealBuildTool/ProjectFiles/VisualStudioMac/VisualStudioMacProjectFileGenerator.cs b/Engine/Source/Programs/UnrealBuildTool/ProjectFiles/VisualStudioMac/VisualStudioMacProjectFileGenerator.cs
deleted file mode 100644
diff --git a/Engine/Source/Programs/UnrealBuildTool/ProjectFiles/VisualStudioWorkspace/VSWorkspaceProjectFileGenerator.cs b/Engine/Source/Programs/UnrealBuildTool/ProjectFiles/VisualStudioWorkspace/VSWorkspaceProjectFileGenerator.cs
--- a/Engine/Source/Programs/UnrealBuildTool/ProjectFiles/VisualStudioWorkspace/VSWorkspaceProjectFileGenerator.cs
+++ b/Engine/Source/Programs/UnrealBuildTool/ProjectFiles/VisualStudioWorkspace/VSWorkspaceProjectFileGenerator.cs
@@ -431,6 +431,11 @@ namespace UnrealBuildTool
 					Format = VCProjectFileFormat.VisualStudio2022;
 					break;
 				}
+				else if (Installation.Compiler == WindowsCompiler.VisualStudio2026)
+				{
+					Format = VCProjectFileFormat.VisualStudio2026;
+					break;
+				}
 			}
 
 			return Format;
diff --git a/Engine/Source/Programs/UnrealBuildTool/System/BuildHostPlatform.cs b/Engine/Source/Programs/UnrealBuildTool/System/BuildHostPlatform.cs
--- a/Engine/Source/Programs/UnrealBuildTool/System/BuildHostPlatform.cs
+++ b/Engine/Source/Programs/UnrealBuildTool/System/BuildHostPlatform.cs
@@ -156,7 +156,7 @@ namespace UnrealBuildTool
 				}
 				catch { }
 			}
-			return Result.ToArray();
+			return [.. Result];
 		}
 
 		/// <summary>
@@ -193,7 +193,7 @@ namespace UnrealBuildTool
 					Result.Add(Info);
 				}
 			}
-			return Result.ToArray();
+			return [.. Result];
 		}
 
 		/// <summary>
@@ -220,7 +220,7 @@ namespace UnrealBuildTool
 				}
 			}
 			catch { }
-			return Modules.ToArray();
+			return [.. Modules];
 		}
 
 		/// <summary>
@@ -324,7 +324,7 @@ namespace UnrealBuildTool
 				Proc.WaitForExit();
 			}
 			catch { }
-			return Result.ToArray();
+			return [.. Result];
 		}
 
 		/// <summary>
@@ -359,7 +359,7 @@ namespace UnrealBuildTool
 				ProcessVMMapOutput(Proc, Modules);
 			}
 			catch { }
-			return Modules.ToArray();
+			return [.. Modules];
 		}
 		private void ProcessVMMapOutput(Process Proc, HashSet<string> Modules)
 		{
@@ -383,7 +383,6 @@ namespace UnrealBuildTool
 		internal override IEnumerable<ProjectFileFormat> GetDefaultProjectFileFormats()
 		{
 			yield return ProjectFileFormat.XCode;
-			yield return ProjectFileFormat.VisualStudioMac;
 #if __VPROJECT_AVAILABLE__
 			yield return ProjectFileFormat.VProject;
 #endif
@@ -408,7 +407,7 @@ namespace UnrealBuildTool
 		public override ProcessInfo[] GetProcesses()
 		{
 			// @TODO: Implement for Linux
-			return new List<ProcessInfo>().ToArray();
+			return [];
 		}
 
 		/// <summary>
@@ -421,7 +420,7 @@ namespace UnrealBuildTool
 		public override string[] GetProcessModules(int PID, string Filename)
 		{
 			// @TODO: Implement for Linux
-			return new List<string>().ToArray();
+			return [];
 		}
 
 		internal override IEnumerable<ProjectFileFormat> GetDefaultProjectFileFormats()
