UnrealBuildTool: Fix a few minor issues \ typos with last change

#rnx

diff --git a/Engine/Source/Programs/UnrealBuildTool/Modes/GenerateProjectFilesMode.cs b/Engine/Source/Programs/UnrealBuildTool/Modes/GenerateProjectFilesMode.cs
--- a/Engine/Source/Programs/UnrealBuildTool/Modes/GenerateProjectFilesMode.cs
+++ b/Engine/Source/Programs/UnrealBuildTool/Modes/GenerateProjectFilesMode.cs
@@ -245,7 +245,7 @@ namespace UnrealBuildTool
 						break;
 #endif
 					default:
-						throw new BuildException("Unhandled project file type '{0}", ProjectFileFormat);
+						throw new BuildException("Unhandled project file type '{0}'", ProjectFileFormat);
 				}
 				// remember if we only wanted a single target (similar to -game -project, except usable with progarms without uprojects)
 				Generator.SingleTargetName = SingleTargetName;
diff --git a/Engine/Source/Programs/UnrealBuildTool/Platform/Windows/UEBuildWindows.cs b/Engine/Source/Programs/UnrealBuildTool/Platform/Windows/UEBuildWindows.cs
--- a/Engine/Source/Programs/UnrealBuildTool/Platform/Windows/UEBuildWindows.cs
+++ b/Engine/Source/Programs/UnrealBuildTool/Platform/Windows/UEBuildWindows.cs
@@ -292,7 +292,7 @@ namespace UnrealBuildTool
 		[ConfigFile(ConfigHierarchyType.Engine, "/Script/WindowsTargetPlatform.WindowsTargetSettings", "Compiler")]
 		[XmlConfigFile(Category = "WindowsPlatform")]
 		[CommandLine("-2022", Value = nameof(WindowsCompiler.VisualStudio2022))]
-		[CommandLine("-2026", Value = nameof(WindowsCompiler.VisualStudio2022))]
+		[CommandLine("-2026", Value = nameof(WindowsCompiler.VisualStudio2026))]
 		[CommandLine("-Compiler=")]
 		public WindowsCompiler Compiler { get; set; } = WindowsCompiler.Default;
 
diff --git a/Engine/Source/Programs/UnrealBuildTool/ProjectFiles/ProjectFileGenerator.cs b/Engine/Source/Programs/UnrealBuildTool/ProjectFiles/ProjectFileGenerator.cs
--- a/Engine/Source/Programs/UnrealBuildTool/ProjectFiles/ProjectFileGenerator.cs
+++ b/Engine/Source/Programs/UnrealBuildTool/ProjectFiles/ProjectFileGenerator.cs
@@ -910,6 +910,11 @@ namespace UnrealBuildTool
 					Format = ProjectFileFormat.VisualStudio2022;
 					return true;
 				}
+				else if (PreferredAccessor == "visualstudio2026")
+				{
+					Format = ProjectFileFormat.VisualStudio2026;
+					return true;
+				}
 			}
 
 			Format = ProjectFileFormat.VisualStudio;
