UnrealBuildTool: Check for nullable object in GitSourceFileWorkingSet

#rnx

diff --git a/Engine/Source/Programs/UnrealBuildTool/System/SourceFileWorkingSet.cs b/Engine/Source/Programs/UnrealBuildTool/System/SourceFileWorkingSet.cs
--- a/Engine/Source/Programs/UnrealBuildTool/System/SourceFileWorkingSet.cs
+++ b/Engine/Source/Programs/UnrealBuildTool/System/SourceFileWorkingSet.cs
@@ -108,25 +108,30 @@ namespace UnrealBuildTool
 		public GitSourceFileWorkingSet(string GitPath, DirectoryReference RootDir, GitSourceFileWorkingSet? Inner, ILogger Logger)
 		{
 			this.RootDir = RootDir;
-			Files = new HashSet<FileReference>();
-			Directories = new List<DirectoryReference>();
-			ErrorOutput = new List<string>();
+			Files = [];
+			Directories = [];
+			ErrorOutput = [];
 			this.Inner = Inner;
 			this.Logger = Logger;
 
 			Logger.LogInformation("Using 'git status' to determine working set for adaptive non-unity build ({RootDir}).", RootDir);
 
-			BackgroundProcess = new Process();
-			BackgroundProcess.StartInfo.FileName = GitPath;
-			BackgroundProcess.StartInfo.Arguments = "--no-optional-locks status --porcelain";
-			BackgroundProcess.StartInfo.WorkingDirectory = RootDir.FullName;
-			BackgroundProcess.StartInfo.RedirectStandardOutput = true;
-			BackgroundProcess.StartInfo.RedirectStandardError = true;
-			BackgroundProcess.StartInfo.UseShellExecute = false;
-			BackgroundProcess.ErrorDataReceived += ErrorDataReceived;
-			BackgroundProcess.OutputDataReceived += OutputDataReceived;
+			BackgroundProcess = new()
+			{
+				StartInfo = new()
+				{
+					FileName = GitPath,
+					Arguments = "--no-optional-locks status --porcelain",
+					WorkingDirectory = RootDir.FullName,
+					RedirectStandardOutput = true,
+					RedirectStandardError = true,
+					UseShellExecute = false
+				}
+			};
 			try
 			{
+				BackgroundProcess.ErrorDataReceived += ErrorDataReceived;
+				BackgroundProcess.OutputDataReceived += OutputDataReceived;
 				BackgroundProcess.Start();
 				BackgroundProcess.BeginErrorReadLine();
 				BackgroundProcess.BeginOutputReadLine();
@@ -145,11 +150,11 @@ namespace UnrealBuildTool
 		{
 			if (BackgroundProcess != null)
 			{
-				if (!BackgroundProcess.HasExited)
+				if (BackgroundProcess?.HasExited == false)
 				{
 					try
 					{
-						BackgroundProcess.Kill();
+						BackgroundProcess?.Kill();
 					}
 					catch
 					{
@@ -166,23 +171,23 @@ namespace UnrealBuildTool
 		{
 			if (BackgroundProcess != null)
 			{
-				if (!BackgroundProcess.WaitForExit(500))
+				if (BackgroundProcess?.WaitForExit(500) == false)
 				{
 					Logger.LogInformation("Waiting for 'git status' command to complete");
 				}
-				if (!BackgroundProcess.WaitForExit(15000))
+				if (BackgroundProcess?.WaitForExit(15000) == false)
 				{
 					Logger.LogInformation("Terminating git child process due to timeout");
 					try
 					{
-						BackgroundProcess.Kill();
+						BackgroundProcess?.Kill();
 					}
 					catch
 					{
 					}
 				}
-				BackgroundProcess.WaitForExit();
-				BackgroundProcess.Dispose();
+				BackgroundProcess?.WaitForExit();
+				BackgroundProcess?.Dispose();
 				BackgroundProcess = null;
 			}
 		}
