[UBT][Editor] Graceful messaging for Target upgrade requirements on editor path.
- Emit the target build settings version & build environment to target file for editor runtime consumption.
- Consume target file metadata to improve error messaging for Target versions.

diff --git a/Engine/Source/Developer/DesktopPlatform/Private/TargetReceipt.cpp b/Engine/Source/Developer/DesktopPlatform/Private/TargetReceipt.cpp
--- a/Engine/Source/Developer/DesktopPlatform/Private/TargetReceipt.cpp
+++ b/Engine/Source/Developer/DesktopPlatform/Private/TargetReceipt.cpp
@@ -61,6 +61,14 @@ bool FTargetReceipt::Read(const FString& FileName, bool bExpandVariables)
 	{
 		return false;
 	}
+	if (!Object->TryGetStringField(TEXT("BuildSettingsVersion"), BuildSettingsVersion))
+	{
+		// This was added at 5.7; no-op on failure.
+	}
+	if (!Object->TryGetStringField(TEXT("TargetBuildEnvironment"), TargetBuildEnvironment))
+	{
+		// This was added at 5.7; no-op on failure.
+	}
 
 	// Read the configuration
 	FString ConfigurationString;
diff --git a/Engine/Source/Developer/DesktopPlatform/Public/TargetReceipt.h b/Engine/Source/Developer/DesktopPlatform/Public/TargetReceipt.h
--- a/Engine/Source/Developer/DesktopPlatform/Public/TargetReceipt.h
+++ b/Engine/Source/Developer/DesktopPlatform/Public/TargetReceipt.h
@@ -69,6 +69,12 @@ struct FTargetReceipt
 	/** Which platform the target is compiled for */
 	FString Architecture;
 
+	/** What build settings version the target is using */
+	FString BuildSettingsVersion;
+
+	/** What build environment the target is using */
+	FString TargetBuildEnvironment;
+
 	/** Which configuration this target is compiled in */
 	EBuildConfiguration Configuration;
 
diff --git a/Engine/Source/Programs/UnrealBuildTool/Configuration/UEBuildTarget.cs b/Engine/Source/Programs/UnrealBuildTool/Configuration/UEBuildTarget.cs
--- a/Engine/Source/Programs/UnrealBuildTool/Configuration/UEBuildTarget.cs
+++ b/Engine/Source/Programs/UnrealBuildTool/Configuration/UEBuildTarget.cs
@@ -1321,6 +1321,10 @@ namespace UnrealBuildTool
 				catch (Exception)
 				{
 					RulesObject.PrintBuildSettingsInfoWarnings();
+
+					// Prepare and write a partial receipt for downstream systems to use in related messaging about shared environment invalidations.
+					PrepareAndWritePartialReceipt(Descriptor, RulesObject, RulesAssembly, Logger);
+
 					throw;
 				}
 			}
@@ -1390,7 +1394,6 @@ namespace UnrealBuildTool
 			{
 				Target.PreBuildSetup(Logger);
 			}
-
 			return Target;
 		}
 
@@ -2317,7 +2320,12 @@ namespace UnrealBuildTool
 			}
 
 			// Create the receipt
-			TargetReceipt Receipt = new TargetReceipt(ProjectFile, TargetName, TargetType, Platform, Configuration, Version, Architectures, Rules.IsTestTarget);
+			TargetReceipt Receipt = new TargetReceipt(ProjectFile, TargetName, TargetType, Platform, Configuration, Version, Architectures, Rules.IsTestTarget, Rules.DefaultBuildSettings, Rules.BuildEnvironment);
+
+			if (BuildProducts.Count == 0 && RuntimeDependencies.Count == 0)
+			{
+				return Receipt;
+			}
 
 			if (!Rules.bShouldCompileAsDLL)
 			{
@@ -2453,6 +2461,42 @@ namespace UnrealBuildTool
 			return Receipt;
 		}
 
+		/// <summary>
+		/// Writes out a partial <see cref="TargetReceipt"/>. This contains a minimal set of properties.
+		/// </summary>
+		/// <param name="Descriptor">The target descriptor to use in target receipt construction.</param>
+		/// <param name="RulesObject">The target rules to use in target receipt construction.</param>
+		/// <param name="RulesAssembly">The rules assembly to use in target receipt construction.</param>
+		/// <param name="Logger">Logger for error and info emission.</param>
+		/// <remarks>This should be used to emit minimal target data in failure circumstances.</remarks>
+		static TargetReceipt? PrepareAndWritePartialReceipt(TargetDescriptor Descriptor, TargetRules RulesObject, RulesAssembly RulesAssembly, ILogger Logger)
+		{
+			try
+			{
+				UEBuildTarget incompleteTarget = new UEBuildTarget(Descriptor, new ReadOnlyTargetRules(RulesObject), RulesAssembly, Logger);
+				UEToolChain targetToolChain = incompleteTarget.CreateToolchain(incompleteTarget.Platform, Logger);
+				TargetReceipt incompleteTargetReceipt = incompleteTarget.PrepareReceipt(targetToolChain, [], []);
+
+				if (incompleteTargetReceipt != null)
+				{
+					FileReference receipt = incompleteTarget.ReceiptFileName;
+					DirectoryReference.CreateDirectory(receipt.Directory);
+					incompleteTargetReceipt.Write(receipt);
+
+					Logger.LogInformation("Wrote partial receipt to {ReceiptPath}", receipt);
+				}
+
+				return incompleteTargetReceipt;
+			}
+			catch (Exception innerEx)
+			{
+				Logger.LogWarning("Failed to write partial receipt for target {Descriptor}: {Message}", Descriptor.Name, innerEx.Message);
+				Logger.LogDebug("{InnerException}", innerEx);
+				
+				return null;
+			}
+		}
+
 		/// <summary>
 		/// Gathers dependency modules for given binaries list.
 		/// </summary>
diff --git a/Engine/Source/Programs/UnrealBuildTool/System/TargetReceipt.cs b/Engine/Source/Programs/UnrealBuildTool/System/TargetReceipt.cs
--- a/Engine/Source/Programs/UnrealBuildTool/System/TargetReceipt.cs
+++ b/Engine/Source/Programs/UnrealBuildTool/System/TargetReceipt.cs
@@ -272,6 +272,16 @@ namespace UnrealBuildTool
 		/// </summary>
 		public string TargetName;
 
+		/// <summary>
+		/// The build settings version used for this target.
+		/// </summary>
+		public BuildSettingsVersion BuildSettingsVersion;
+
+		/// <summary>
+		/// The build environment used for this target.
+		/// </summary>
+		public TargetBuildEnvironment TargetBuildEnvironment;
+
 		/// <summary>
 		/// Which platform the target is compiled for
 		/// </summary>
@@ -348,11 +358,15 @@ namespace UnrealBuildTool
 		/// <param name="InVersion">Version information for the target</param>
 		/// <param name="InArchitectures">Architecture information for the target</param>
 		/// <param name="InIsTestTarget">Whether it's a target for low level tests</param>
-		public TargetReceipt(FileReference? InProjectFile, string InTargetName, TargetType InTargetType, UnrealTargetPlatform InPlatform, UnrealTargetConfiguration InConfiguration, BuildVersion InVersion, UnrealArchitectures InArchitectures, bool InIsTestTarget)
+		/// <param name="InBuildSettingsVersion">Build settings version of the target.</param>
+		/// <param name="InTargetBuildEnvironment">Build environment used to produce the target.</param>
+		public TargetReceipt(FileReference? InProjectFile, string InTargetName, TargetType InTargetType, UnrealTargetPlatform InPlatform, UnrealTargetConfiguration InConfiguration, BuildVersion InVersion, UnrealArchitectures InArchitectures, bool InIsTestTarget, BuildSettingsVersion InBuildSettingsVersion, TargetBuildEnvironment InTargetBuildEnvironment)
 		{
 			ProjectFile = InProjectFile;
 			ProjectDir = DirectoryReference.FromFile(InProjectFile);
 			TargetName = InTargetName;
+			BuildSettingsVersion = InBuildSettingsVersion;
+			TargetBuildEnvironment = InTargetBuildEnvironment;
 			Platform = InPlatform;
 			Configuration = InConfiguration;
 			TargetType = InTargetType;
@@ -579,8 +593,22 @@ namespace UnrealBuildTool
 			bool IsTestTarget;
 			RawObject.TryGetBoolField("IsTestTarget", out IsTestTarget);
 
+			BuildSettingsVersion BuildSettingsVersion;
+			if(!RawObject.TryGetEnumField<BuildSettingsVersion>("BuildSettingsVersion", out BuildSettingsVersion))
+			{
+				// @todo this doesn't necessailry match how it was compiled - should this be an error case?
+				BuildSettingsVersion = BuildSettingsVersion.Latest;
+			}
+
+			TargetBuildEnvironment TargetBuildEnvironment;
+			if (!RawObject.TryGetEnumField<TargetBuildEnvironment>("TargetBuildEnvironment", out TargetBuildEnvironment))
+			{
+				// @todo this doesn't necessailry match how it was compiled - should this be an error case?
+				TargetBuildEnvironment = TargetBuildEnvironment.Unique;
+			}
+
 			// Create the receipt
-			TargetReceipt Receipt = new TargetReceipt(ProjectFile, TargetName, TargetType, Platform, Configuration, Version, Architectures, IsTestTarget);
+			TargetReceipt Receipt = new TargetReceipt(ProjectFile, TargetName, TargetType, Platform, Configuration, Version, Architectures, IsTestTarget, BuildSettingsVersion, TargetBuildEnvironment);
 
 			// Get the project directory
 			DirectoryReference? ProjectDir = Receipt.ProjectDir;
@@ -751,6 +779,8 @@ namespace UnrealBuildTool
 				Writer.WriteValue("TargetName", TargetName);
 				Writer.WriteValue("Platform", Platform.ToString());
 				Writer.WriteValue("Configuration", Configuration.ToString());
+				Writer.WriteValue("BuildSettingsVersion", BuildSettingsVersion.ToString());
+				Writer.WriteValue("TargetBuildEnvironment", TargetBuildEnvironment.ToString());
 				Writer.WriteValue("TargetType", TargetType.ToString());
 				Writer.WriteValue("IsTestTarget", IsTestTarget);
 				Writer.WriteValue("Architecture", Architectures.ToString());
diff --git a/Engine/Source/Runtime/Launch/Private/LaunchEngineLoop.cpp b/Engine/Source/Runtime/Launch/Private/LaunchEngineLoop.cpp
--- a/Engine/Source/Runtime/Launch/Private/LaunchEngineLoop.cpp
+++ b/Engine/Source/Runtime/Launch/Private/LaunchEngineLoop.cpp
@@ -6317,6 +6317,52 @@ void CleanUpPath(FString& InPath)
 	FPaths::RemoveDuplicateSlashes(InPath);
 }
 
+bool DoesCurrentTargetRequireUpgrade(const FString& EditorTargetFileName, FString& InTargetName)
+{
+	FTargetReceipt Receipt;
+
+	if (!FPaths::FileExists(EditorTargetFileName))
+	{
+		return false;
+	}
+
+	// We do not return on Receipt.Read returning false, as Target upgrade requirements occur before a full target receipt can be emitted.
+	// Only partial target receipts exist, of which the most important properties are BuildSettingsVersion && TargetBuildEnvironment.
+	Receipt.Read(EditorTargetFileName);
+
+	// Handle backwards compatibility before we added these properties to the receipt at 5.7.
+	if (Receipt.BuildSettingsVersion.IsEmpty() || Receipt.TargetBuildEnvironment.IsEmpty())
+	{
+		UE_LOG(
+			LogInit,
+			Log,
+			TEXT("Unable to discern build settings version & build environment used for target - target may require upgrade. Target file: %s"),
+			*EditorTargetFileName
+		);
+		return false;
+	}
+
+	InTargetName = Receipt.TargetName;
+	const bool bIsShared = (Receipt.TargetBuildEnvironment == TEXT("Shared"));
+	const bool bIsUpToDate = (Receipt.BuildSettingsVersion == TEXT("V6"));
+
+	if (!bIsUpToDate && bIsShared)
+	{
+		UE_LOG(
+			LogInit,
+			Warning,
+			TEXT("Current target requires upgrade to V6, as it is not a unique build environment. File: %s, BuildSettingsVersion: %s, TargetBuildEnvironment: %s"),
+			*EditorTargetFileName,
+			*Receipt.BuildSettingsVersion,
+			*Receipt.TargetBuildEnvironment
+		);
+
+		return true;
+	}
+
+	return false;
+}
+
 bool DoesCurrentExecutableMatch(const FString& EditorTargetFileName)
 {
 	FTargetReceipt Receipt;
@@ -6568,7 +6614,33 @@ bool FEngineLoop::AppInit( )
 		}
 		if(!bNeedCompile)
 		{
-			// Check if any of the project or plugin modules are out of date, and the user wants to compile them.
+			// Check if any of the target, project or plugin modules are out of date, and the user wants to compile them.
+			FString targetName;
+			bool TargetRequiresUpgrade = DoesCurrentTargetRequireUpgrade(EditorTargetFileName, targetName);
+
+			if (TargetRequiresUpgrade)
+			{
+				FString TargetRequiresUpgradeMessage = FString::Printf(
+					TEXT("The last build used an outdated BuildSettingsVersion, or the target still requires upgrading."
+						"\n\nPlease upgrade %s.Target.cs:"
+						"\n\n  Update the build settings version (Installed Engine Build & Source Build):"
+						"\n    DefaultBuildSettings = BuildSettingsVersion.V6;"
+						"\n\n  OR"
+						"\n\n  Declare a unique build environment (Source Build only):"
+						"\n    BuildEnvironment = TargetBuildEnvironment.Unique;"
+						"\n\nIf you have already updated the target, would you like to rebuild & open the project?"),
+					*targetName
+				);
+
+				if (FPlatformMisc::MessageBoxExt(EAppMsgType::YesNo, *TargetRequiresUpgradeMessage, TEXT("Target Upgrade Required")) == EAppReturnType::No)
+				{
+					return false;
+				}
+				
+				// Rebuilding should update the target receipt file, and prevent subsequent checks.
+				bNeedCompile = true;
+			}
+
 			TArray<FString> IncompatibleFiles;
 			ProjectManager.CheckModuleCompatibility(IncompatibleFiles);
 
