UnrealBuildTool: Tweak preferred Visual Studio version

#rnx

diff --git a/Engine/Source/Programs/UnrealBuildTool/Platform/Windows/MicrosoftPlatformSDK.cs b/Engine/Source/Programs/UnrealBuildTool/Platform/Windows/MicrosoftPlatformSDK.cs
--- a/Engine/Source/Programs/UnrealBuildTool/Platform/Windows/MicrosoftPlatformSDK.cs
+++ b/Engine/Source/Programs/UnrealBuildTool/Platform/Windows/MicrosoftPlatformSDK.cs
@@ -1017,11 +1017,11 @@ namespace UnrealBuildTool
 						Logger.LogDebug("Found {VisualStudio} installation: {BaseDir} (Product={ProductId}, Version={Version})", Compiler, BaseDir, ProductId, Version);
 					}
 
-					Installations = Installations.OrderByDescending(x => x.Compiler)
+					Installations = [.. Installations.OrderBy(x => x.Compiler)
 						.ThenBy(x => x.bCommunity)
 						.ThenBy(x => x.ReleaseChannel)
 						.ThenByDescending(x => x.Version)
-						.ToList();
+					];
 				}
 				catch (Exception Ex)
 				{
diff --git a/Engine/Source/Programs/UnrealBuildTool/ProjectFiles/VisualStudio/VCProjectFileGenerator.cs b/Engine/Source/Programs/UnrealBuildTool/ProjectFiles/VisualStudio/VCProjectFileGenerator.cs
--- a/Engine/Source/Programs/UnrealBuildTool/ProjectFiles/VisualStudio/VCProjectFileGenerator.cs
+++ b/Engine/Source/Programs/UnrealBuildTool/ProjectFiles/VisualStudio/VCProjectFileGenerator.cs
@@ -349,7 +349,8 @@ namespace UnrealBuildTool
 			if (Settings.ProjectFileFormat == VCProjectFileFormat.Default)
 			{
 				// Enumerate all the valid installations. This list is already sorted by preference.
-				List<VisualStudioInstallation> Installations = MicrosoftPlatformSDK.FindVisualStudioInstallations(Logger).Where(x => WindowsPlatform.HasCompiler(x.Compiler, UnrealArch.X64, Logger)).ToList();
+				IEnumerable<VisualStudioInstallation> Installations = MicrosoftPlatformSDK.FindVisualStudioInstallations(Logger)
+					.Where(x => WindowsPlatform.HasCompiler(x.Compiler, UnrealArch.X64, Logger));
 
 				// Get the corresponding project file format
 				VCProjectFileFormat Format = VCProjectFileFormat.Default;
