Don't allow any Clang linker jobs to overlap and not just ThinLTO jobs

diff --git a/Engine/Source/Programs/UnrealBuildTool/Platform/Windows/VCToolChain.cs b/Engine/Source/Programs/UnrealBuildTool/Platform/Windows/VCToolChain.cs
--- a/Engine/Source/Programs/UnrealBuildTool/Platform/Windows/VCToolChain.cs
+++ b/Engine/Source/Programs/UnrealBuildTool/Platform/Windows/VCToolChain.cs
@@ -3342,7 +3342,7 @@ namespace UnrealBuildTool
 				LinkAction.bCanExecuteRemotely = true;
 			}
 
-			if (Target.WindowsPlatform.Compiler.IsClang() && Target.WindowsPlatform.bAllowClangLinker && (LinkEnvironment.bAllowLTCG || LinkEnvironment.bPGOOptimize || LinkEnvironment.bPGOProfile))
+			if (Target.WindowsPlatform.Compiler.IsClang() && Target.WindowsPlatform.bAllowClangLinker)
 			{
 				// Set the weight to number of logical cores as lld can max out the available cores
 				LinkAction.Weight = Utils.GetLogicalProcessorCount();
