Restrict Clang linker serialization to monolithic builds to keep editor links parallel

diff --git a/Engine/Source/Programs/UnrealBuildTool/Platform/Windows/VCToolChain.cs b/Engine/Source/Programs/UnrealBuildTool/Platform/Windows/VCToolChain.cs
--- a/Engine/Source/Programs/UnrealBuildTool/Platform/Windows/VCToolChain.cs
+++ b/Engine/Source/Programs/UnrealBuildTool/Platform/Windows/VCToolChain.cs
@@ -3342,7 +3342,9 @@ namespace UnrealBuildTool
 				LinkAction.bCanExecuteRemotely = true;
 			}
 
-			if (Target.WindowsPlatform.Compiler.IsClang() && Target.WindowsPlatform.bAllowClangLinker)
+			if (Target.WindowsPlatform.Compiler.IsClang() &&
+				Target.WindowsPlatform.bAllowClangLinker &&
+				(Target.LinkType == TargetLinkType.Monolithic))
 			{
 				// Set the weight to number of logical cores as lld can max out the available cores
 				LinkAction.Weight = Utils.GetLogicalProcessorCount();
