UnrealBuildTool: Support for AVX10.x (via Intel)

diff --git a/Engine/Source/Programs/UnrealBuildTool/Platform/Windows/VCToolChain.cs b/Engine/Source/Programs/UnrealBuildTool/Platform/Windows/VCToolChain.cs
--- a/Engine/Source/Programs/UnrealBuildTool/Platform/Windows/VCToolChain.cs
+++ b/Engine/Source/Programs/UnrealBuildTool/Platform/Windows/VCToolChain.cs
@@ -1054,49 +1054,90 @@ namespace UnrealBuildTool
 			//
 			if (CompileEnvironment.Architecture == UnrealArch.X64 && CompileEnvironment.MinCpuArchX64 != MinimumCpuArchitectureX64.None)
 			{
+				MinimumCpuArchitectureX64 MinCpuArchX64 = CompileEnvironment.MinCpuArchX64;
 				// Define /arch:AVX[2,512] for the current compilation unit.  Machines without AVX support will crash on any SSE/AVX instructions if they run this compilation unit.
 				if (Target.WindowsPlatform.Compiler.IsMSVC())
 				{
-					Arguments.Add($"/arch:{CompileEnvironment.MinCpuArchX64}");
+					// Downgrade AVX if unsupported by MSVC version
+					if (MinCpuArchX64 >= MinimumCpuArchitectureX64.AVX10_2 && EnvVars.CompilerVersion < new VersionNumber(14, 50))
+					{
+						MinCpuArchX64 = MinimumCpuArchitectureX64.AVX10_1;
+					}
+					if (MinCpuArchX64 >= MinimumCpuArchitectureX64.AVX10_1 && EnvVars.CompilerVersion < new VersionNumber(14, 42))
+					{
+						MinCpuArchX64 = MinimumCpuArchitectureX64.AVX512;
+					}
+
+					Arguments.Add($"/arch:{MinCpuArchX64}".Replace('_', '.'));
 				}
 				else if (Target.WindowsPlatform.Compiler.IsClang())
 				{
-					if (CompileEnvironment.MinCpuArchX64 >= MinimumCpuArchitectureX64.AVX)
+					VersionNumber ClangVersion = Target.WindowsPlatform.Compiler.IsIntel() ? MicrosoftPlatformSDK.GetClangVersionForIntelCompiler(EnvVars.CompilerPath) : EnvVars.CompilerVersion;
+					// Downgrade AVX if unsupported by Clang version
+					if (MinCpuArchX64 >= MinimumCpuArchitectureX64.AVX10_2 && ClangVersion < new VersionNumber(20, 1))
+					{
+						MinCpuArchX64 = MinimumCpuArchitectureX64.AVX10_1;
+					}
+					if (MinCpuArchX64 >= MinimumCpuArchitectureX64.AVX10_1 && ClangVersion < new VersionNumber(18, 1))
+					{
+						MinCpuArchX64 = MinimumCpuArchitectureX64.AVX512;
+					}
+
+					if (MinCpuArchX64 >= MinimumCpuArchitectureX64.AVX)
 					{
 						// Apparently MSVC enables (a subset?) of BMI (bit manipulation instructions) when /arch:AVX is set. Some code relies on this, so mirror it by enabling BMI1
 						Arguments.Add("-mavx");
 						Arguments.Add("-mbmi");
 					}
-					if (CompileEnvironment.MinCpuArchX64 >= MinimumCpuArchitectureX64.AVX2)
+					if (MinCpuArchX64 >= MinimumCpuArchitectureX64.AVX2)
 					{
 						Arguments.Add("-mavx2");
 					}
-					if (CompileEnvironment.MinCpuArchX64 >= MinimumCpuArchitectureX64.AVX512)
+					if (MinCpuArchX64 >= MinimumCpuArchitectureX64.AVX512)
 					{
 						// Match MSVC which says (https://learn.microsoft.com/en-us/cpp/build/reference/arch-x64?view=msvc-170):
 						// > The __AVX512F__, __AVX512CD__, __AVX512BW__, __AVX512DQ__ and __AVX512VL__ preprocessor symbols are defined when the /arch:AVX512 compiler option is specified
-						Arguments.Add("-mavx512f");
-						Arguments.Add("-mavx512cd");
-						Arguments.Add("-mavx512bw");
-						Arguments.Add("-mavx512dq");
-						Arguments.Add("-mavx512vl");
+						if (MinCpuArchX64 < MinimumCpuArchitectureX64.AVX10_1)
+						{
+							Arguments.Add("-mavx512f");
+							Arguments.Add("-mavx512cd");
+							Arguments.Add("-mavx512bw");
+							Arguments.Add("-mavx512dq");
+							Arguments.Add("-mavx512vl");
+						}
+					}
+					if (MinCpuArchX64 >= MinimumCpuArchitectureX64.AVX10_1)
+					{
+						Arguments.Add("-mavx10.1");
+					}
+					if (MinCpuArchX64 >= MinimumCpuArchitectureX64.AVX10_2)
+					{
+						Arguments.Add("-mavx10.2");
 					}
 				}
 
 				// AVX available implies sse4 and sse2 available.
 				// Inform Unreal code that we have sse2, sse4, and AVX, both available to compile and available to run
-				// By setting the ALWAYS_HAS defines, we we direct Unreal code to skip cpuid checks to verify that the running hardware supports sse/avx.
+				// By setting the ALWAYS_HAS defines, we direct Unreal code to skip cpuid checks to verify that the running hardware supports sse/avx.
 				AddDefinition(Arguments, "PLATFORM_ENABLE_VECTORINTRINSICS=1");
 				AddDefinition(Arguments, "PLATFORM_MAYBE_HAS_AVX=1");
 				AddDefinition(Arguments, "PLATFORM_ALWAYS_HAS_AVX=1");
-				if (CompileEnvironment.MinCpuArchX64 >= MinimumCpuArchitectureX64.AVX2)
+				if (MinCpuArchX64 >= MinimumCpuArchitectureX64.AVX2)
 				{
 					AddDefinition(Arguments, "PLATFORM_ALWAYS_HAS_AVX_2=1");
 				}
-				if (CompileEnvironment.MinCpuArchX64 >= MinimumCpuArchitectureX64.AVX512)
+				if (MinCpuArchX64 >= MinimumCpuArchitectureX64.AVX512)
 				{
 					AddDefinition(Arguments, "PLATFORM_ALWAYS_HAS_AVX_512=1");
 				}
+				if (MinCpuArchX64 >= MinimumCpuArchitectureX64.AVX10_1)
+				{
+					AddDefinition(Arguments, "PLATFORM_ALWAYS_HAS_AVX_10_1=1");
+				}
+				if (MinCpuArchX64 >= MinimumCpuArchitectureX64.AVX10_2)
+				{
+					AddDefinition(Arguments, "PLATFORM_ALWAYS_HAS_AVX_10_2=1");
+				}
 			}
 
 			if (CompileEnvironment.Architecture == UnrealArch.X64 && CompileEnvironment.MinCpuArchX64 == MinimumCpuArchitectureX64.None && Target.WindowsPlatform.Compiler.IsIntel())
diff --git a/Engine/Source/Programs/UnrealBuildTool/System/CppCompileEnvironment.cs b/Engine/Source/Programs/UnrealBuildTool/System/CppCompileEnvironment.cs
--- a/Engine/Source/Programs/UnrealBuildTool/System/CppCompileEnvironment.cs
+++ b/Engine/Source/Programs/UnrealBuildTool/System/CppCompileEnvironment.cs
@@ -137,6 +137,16 @@ namespace UnrealBuildTool
 		/// </summary>
 		AVX512,
 
+		/// <summary>
+		/// Enables the use of Intel Advanced Vector Extensions 10 version 1 instructions
+		/// </summary>
+		AVX10_1,
+
+		/// <summary>
+		/// Enables the use of Intel Advanced Vector Extensions 10 version 2 instructions
+		/// </summary>
+		AVX10_2,
+
 		/// <summary>
 		/// Use the default minimum architecture
 		/// </summary>
diff --git a/Engine/Source/Programs/UnrealBuildTool/ToolChain/ClangToolChain.cs b/Engine/Source/Programs/UnrealBuildTool/ToolChain/ClangToolChain.cs
--- a/Engine/Source/Programs/UnrealBuildTool/ToolChain/ClangToolChain.cs
+++ b/Engine/Source/Programs/UnrealBuildTool/ToolChain/ClangToolChain.cs
@@ -785,13 +785,24 @@ namespace UnrealBuildTool
 			if (CompileEnvironment.Architecture == UnrealArch.X64 && CompileEnvironment.MinCpuArchX64 != MinimumCpuArchitectureX64.None)
 			{
 				// The binary created will be targeting AVX instructions. Machines without AVX support will crash on any AVX instructions if they run this compilation unit.
+				var MinCpuArchX64 = CompileEnvironment.MinCpuArchX64;
+				VersionNumber ClangVersion = new(Info.ClangVersion.Major, Info.ClangVersion.Minor, Info.ClangVersion.Build);
+				// Downgrade AVX if unsupported by Clang version
+				if (ClangVersion < new VersionNumber(20, 1) && MinCpuArchX64 >= MinimumCpuArchitectureX64.AVX10_2)
+				{
+					MinCpuArchX64 = MinimumCpuArchitectureX64.AVX10_1;
+				}
+				if (ClangVersion < new VersionNumber(18, 1) && MinCpuArchX64 >= MinimumCpuArchitectureX64.AVX10_1)
+				{
+					MinCpuArchX64 = MinimumCpuArchitectureX64.AVX512;
+				}
 
 				// AVX available implies sse4 and sse2 available.
 				// Inform Unreal code that we have sse2, sse4, and AVX, both available to compile and available to run
 				// By setting the ALWAYS_HAS defines, we we direct Unreal code to skip cpuid checks to verify that the running hardware supports sse/avx.
 				Arguments.Add("-DPLATFORM_ENABLE_VECTORINTRINSICS=1");
 
-				if (CompileEnvironment.MinCpuArchX64 >= MinimumCpuArchitectureX64.AVX)
+				if (MinCpuArchX64 >= MinimumCpuArchitectureX64.AVX)
 				{
 					// Apparently MSVC enables (a subset?) of BMI (bit manipulation instructions) when /arch:AVX is set. Some code relies on this, so mirror it by enabling BMI1
 					Arguments.Add("-mavx");
@@ -802,23 +813,39 @@ namespace UnrealBuildTool
 					Arguments.Add("-DPLATFORM_ALWAYS_HAS_AVX=1");
 				}
 
-				if (CompileEnvironment.MinCpuArchX64 >= MinimumCpuArchitectureX64.AVX2)
+				if (MinCpuArchX64 >= MinimumCpuArchitectureX64.AVX2)
 				{
 					Arguments.Add("-mavx2");
 					Arguments.Add("-DPLATFORM_ALWAYS_HAS_AVX_2=1");
 				}
 
-				if (CompileEnvironment.MinCpuArchX64 >= MinimumCpuArchitectureX64.AVX512)
+				if (MinCpuArchX64 >= MinimumCpuArchitectureX64.AVX512)
 				{
 					// Match MSVC which says (https://learn.microsoft.com/en-us/cpp/build/reference/arch-x64?view=msvc-170):
 					// > The __AVX512F__, __AVX512CD__, __AVX512BW__, __AVX512DQ__ and __AVX512VL__ preprocessor symbols are defined when the /arch:AVX512 compiler option is specified
-					Arguments.Add("-mavx512f");
-					Arguments.Add("-mavx512cd");
-					Arguments.Add("-mavx512bw");
-					Arguments.Add("-mavx512dq");
-					Arguments.Add("-mavx512vl");
+					if (MinCpuArchX64 < MinimumCpuArchitectureX64.AVX10_1)
+					{
+						Arguments.Add("-mavx512f");
+						Arguments.Add("-mavx512cd");
+						Arguments.Add("-mavx512bw");
+						Arguments.Add("-mavx512dq");
+						Arguments.Add("-mavx512vl");
+					}
+
 					Arguments.Add("-DPLATFORM_ALWAYS_HAS_AVX_512=1");
 				}
+
+				if (MinCpuArchX64 >= MinimumCpuArchitectureX64.AVX10_1)
+				{
+					Arguments.Add("-DPLATFORM_ALWAYS_HAS_AVX_10_1=1");
+					Arguments.Add("-mavx10.1");
+				}
+
+				if (MinCpuArchX64 >= MinimumCpuArchitectureX64.AVX10_2)
+				{
+					Arguments.Add("-DPLATFORM_ALWAYS_HAS_AVX_10_2=1");
+					Arguments.Add("-mavx10.2");
+				}
 			}
 
 			// Treat chars as signed on arm64
