UnrealBuildTool: Clean up some collection initialization in ModuleRules

#rnx

diff --git a/Engine/Source/Programs/UnrealBuildTool/Configuration/ModuleRules.cs b/Engine/Source/Programs/UnrealBuildTool/Configuration/ModuleRules.cs
--- a/Engine/Source/Programs/UnrealBuildTool/Configuration/ModuleRules.cs
+++ b/Engine/Source/Programs/UnrealBuildTool/Configuration/ModuleRules.cs
@@ -362,7 +362,7 @@ namespace UnrealBuildTool
 			/// <summary>
 			/// Inner list of runtime dependencies
 			/// </summary>
-			readonly List<RuntimeDependency> _inner = new();
+			readonly List<RuntimeDependency> _inner = [];
 
 			/// <summary>
 			/// Readonly access of inner list of runtime dependencies
@@ -416,7 +416,7 @@ namespace UnrealBuildTool
 			/// <summary>
 			/// Inner list of runtime dependencies
 			/// </summary>
-			readonly List<ReceiptProperty> _inner = new();
+			readonly List<ReceiptProperty> _inner = [];
 
 			/// <summary>
 			/// Readonly access of inner list of runtime dependencies
@@ -639,7 +639,7 @@ namespace UnrealBuildTool
 		/// <summary>
 		/// Additional directories that contribute to this module but are not based on a subclass (NotForLicensees, etc)
 		/// </summary>
-		private List<DirectoryReference> AdditionalModuleDirectories = new();
+		private List<DirectoryReference> AdditionalModuleDirectories { get; } = [];
 
 		/// <summary>
 		/// The rules assembly to use when searching for modules
@@ -785,7 +785,7 @@ namespace UnrealBuildTool
 		/// <summary>
 		/// Header files that should be force included for every source file in this module.
 		/// </summary>
-		public List<string> ForceIncludeFiles { get; } = new List<string>();
+		public List<string> ForceIncludeFiles { get; } = [];
 
 		/// <summary>
 		/// Explicit private PCH for this module. Implies that this module will not use a shared PCH.
@@ -848,12 +848,12 @@ namespace UnrealBuildTool
 		/// <summary>
 		/// List of files that will be generated for this module. Entries in list are (FileName, FileContent)
 		/// </summary>
-		public Dictionary<string, IEnumerable<string>> FilesToGenerate = [];
+		public Dictionary<string, IEnumerable<string>> FilesToGenerate { get; } = [];
 
 		/// <summary>
 		/// If this contains a platform and the project has had it's SDK version overriden from default, this module will be compiled as a project module, instead of shared engine module
 		/// </summary>
-		public List<UnrealTargetPlatform> SDKVersionRelevantPlatforms { get; } = new();
+		public List<UnrealTargetPlatform> SDKVersionRelevantPlatforms { get; } = [];
 
 		/// <summary>
 		/// Emits compilation errors for incorrect UE_LOG format strings.
@@ -986,7 +986,7 @@ namespace UnrealBuildTool
 		/// The static analyzer rulesets that should be used to filter warnings. This is only supported for MSVC.
 		/// See https://learn.microsoft.com/en-us/cpp/code-quality/using-rule-sets-to-specify-the-cpp-rules-to-run
 		/// </summary>
-		public HashSet<FileReference> StaticAnalyzerRulesets = new();
+		public HashSet<FileReference> StaticAnalyzerRulesets = [];
 
 		/// <summary>
 		/// The static analyzer checkers that should be enabled rather than the defaults. This is only supported for Clang.
@@ -996,7 +996,7 @@ namespace UnrealBuildTool
 		///    'clang -Xclang -analyzer-checker-help-alpha' 
 		/// for the list of experimental checkers.
 		/// </summary>
-		public HashSet<string> StaticAnalyzerCheckers { get; } = new();
+		public HashSet<string> StaticAnalyzerCheckers { get; } = [];
 
 		/// <summary>
 		/// The static analyzer default checkers that should be disabled. Unused if StaticAnalyzerCheckers is populated. This is only supported for Clang.
@@ -1006,7 +1006,7 @@ namespace UnrealBuildTool
 		///    'clang -Xclang -analyzer-checker-help-alpha' 
 		/// for the list of experimental checkers.
 		/// </summary>
-		public HashSet<string> StaticAnalyzerDisabledCheckers { get; } = new() {
+		public HashSet<string> StaticAnalyzerDisabledCheckers { get; } = [
 			// Work in progress, remove disabled checkers as all warnings are fixed
 			"core.BitwiseShift",
 			"core.CallAndMessage",
@@ -1024,7 +1024,7 @@ namespace UnrealBuildTool
 
 			// Needs evaluation
 			"deadcode.DeadStores",
-		};
+		];
 
 		/// <summary>
 		/// The static analyzer non-default checkers that should be enabled. Unused if StaticAnalyzerCheckers is populated. This is only supported for Clang.
@@ -1034,12 +1034,12 @@ namespace UnrealBuildTool
 		///    'clang -Xclang -analyzer-checker-help-alpha' 
 		/// for the list of experimental checkers.
 		/// </summary>
-		public HashSet<string> StaticAnalyzerAdditionalCheckers { get; } = new();
+		public HashSet<string> StaticAnalyzerAdditionalCheckers { get; } = [];
 
 		/// <summary>
 		/// The PVS Studio analysis errors that should be disabled.
 		/// </summary>
-		public HashSet<string> StaticAnalyzerPVSDisabledErrors { get; } = new();
+		public HashSet<string> StaticAnalyzerPVSDisabledErrors { get; } = [];
 
 		private bool? bUseUnityOverride;
 		/// <summary>
@@ -1060,7 +1060,7 @@ namespace UnrealBuildTool
 		/// <summary>
 		/// Source file types that will be placed into unity files when enabled.
 		/// </summary>
-		public UnityFileType UnityFileTypes { get; set;  } = UnityFileType.Cpp;
+		public UnityFileType UnityFileTypes { get; set; } = UnityFileType.Cpp;
 
 		/// <summary>
 		/// The number of source files in this module before unity build will be activated for that module.  If set to
@@ -1103,12 +1103,12 @@ namespace UnrealBuildTool
 		/// <summary>
 		/// List of folders which are allowed to be referenced when compiling this binary, without propagating restricted folder names
 		/// </summary>
-		public List<string> AllowedRestrictedFolders = new();
+		public List<string> AllowedRestrictedFolders = [];
 
 		/// <summary>
 		/// Set of aliased restricted folder references
 		/// </summary>
-		public Dictionary<string, string> AliasRestrictedFolders = new();
+		public Dictionary<string, string> AliasRestrictedFolders = [];
 
 		/// <summary>
 		/// Allows "include what you use" to modify the source code when run. bEnforceIWYU must be true for this variable to matter.
@@ -1152,7 +1152,7 @@ namespace UnrealBuildTool
 		/// that uses transactionally-disabled code in places that are transactionally-enabled.
 		/// </summary>
 		public bool bInjectAutoRTFMAttributeDisables { get; set; }
-		
+
 		/// <summary>
 		/// If true, autogenerated functions of USTRUCTS are always exported
 		/// </summary>
@@ -1167,88 +1167,88 @@ namespace UnrealBuildTool
 		/// If true, this module can contain UE types (i.e. UCLASS) in namespaces
 		/// </summary>
 		public bool bAllowUETypesInNamespaces = false;
-  
+
 		/// <summary>
 		/// If true, minimizes the amount of header files included in .generated.h files, which may require more headers from CoreUObject to be explicitly included.
 		/// </summary>
 		public bool bMinimizeGeneratedIncludes = false;
-		
+
 		/// <summary>
 		/// List of modules names (no path needed) with header files that our module's public headers needs access to, but we don't need to "import" or link against.
 		/// </summary>
-		public List<string> PublicIncludePathModuleNames = new();
+		public List<string> PublicIncludePathModuleNames = [];
 
 		/// <summary>
 		/// List of public dependency module names (no path needed) (automatically does the private/public include). These are modules that are required by our public source files.
 		/// </summary>
-		public List<string> PublicDependencyModuleNames = new();
+		public List<string> PublicDependencyModuleNames = [];
 
 		/// <summary>
 		/// List of modules name (no path needed) with header files that our module's private code files needs access to, but we don't need to "import" or link against.
 		/// </summary>
-		public List<string> PrivateIncludePathModuleNames = new();
+		public List<string> PrivateIncludePathModuleNames = [];
 
 		/// <summary>
 		/// List of private dependency module names.  These are modules that our private code depends on but nothing in our public
 		/// include files depend on.
 		/// </summary>
-		public List<string> PrivateDependencyModuleNames = new();
+		public List<string> PrivateDependencyModuleNames = [];
 
 		/// <summary>
 		/// Only for legacy reason, should not be used in new code. List of module dependencies that should be treated as circular references.  This modules must have already been added to
 		/// either the public or private dependent module list.
 		/// </summary>
-		public List<string> CircularlyReferencedDependentModules = new();
+		public List<string> CircularlyReferencedDependentModules = [];
 
 		/// <summary>
 		/// List of system/library include paths - typically used for External (third party) modules.  These are public stable header file directories that are not checked when resolving header dependencies.
 		/// </summary>
-		public List<string> PublicSystemIncludePaths = new();
+		public List<string> PublicSystemIncludePaths = [];
 
 		/// <summary>
 		/// (This setting is currently not need as we discover all files from the 'Public' folder) List of all paths to include files that are exposed to other modules
 		/// </summary>
-		public List<string> PublicIncludePaths = new();
+		public List<string> PublicIncludePaths = [];
 
 		/// <summary>
 		/// (This setting is currently not need as we discover all files from the 'Internal' folder) List of all paths to include files that are exposed to other internal modules
 		/// </summary>
-		public List<string> InternalIncludePaths = new();
+		public List<string> InternalIncludePaths = [];
 
 		/// <summary>
 		/// List of all paths to this module's internal include files, not exposed to other modules (at least one include to the 'Private' path, more if we want to avoid relative paths)
 		/// </summary>
-		public List<string> PrivateIncludePaths = new();
+		public List<string> PrivateIncludePaths = [];
 
 		/// <summary>
 		/// List of system library paths (directory of .lib files) - for External (third party) modules please use the PublicAdditionalLibaries instead
 		/// </summary>
-		public List<string> PublicSystemLibraryPaths = new();
+		public List<string> PublicSystemLibraryPaths = [];
 
 		/// <summary>
 		/// List of search paths for libraries at runtime (eg. .so files)
 		/// </summary>
-		public List<string> PrivateRuntimeLibraryPaths = new();
+		public List<string> PrivateRuntimeLibraryPaths = [];
 
 		/// <summary>
 		/// List of search paths for libraries at runtime (eg. .so files)
 		/// </summary>
-		public List<string> PublicRuntimeLibraryPaths = new();
+		public List<string> PublicRuntimeLibraryPaths = [];
 
 		/// <summary>
 		/// List of additional libraries (names of the .lib files including extension) - typically used for External (third party) modules
 		/// </summary>
-		public List<string> PublicAdditionalLibraries = new();
+		public List<string> PublicAdditionalLibraries = [];
 
 		/// <summary>
 		/// List of additional debug visualizers (.natvis, .natstepfilter, and .natjmc) exposed to other modules - typically used for External (third party) modules
 		/// </summary>
-		public List<string> PublicDebugVisualizerPaths = new();
+		public List<string> PublicDebugVisualizerPaths = [];
 
 		/// <summary>
 		/// Per-architecture lists of dependencies for linking to ignore (useful when building for multiple architectures, and a lib only is needed for one architecture), it's up to the Toolchain to use this
 		/// </summary>
-		public Dictionary<string, List<UnrealArch>> DependenciesToSkipPerArchitecture = new();
+		public Dictionary<string, List<UnrealArch>> DependenciesToSkipPerArchitecture = [];
 
 		/// <summary>
 		/// Returns the directory of where the passed in module name lives.
@@ -1265,57 +1265,57 @@ namespace UnrealBuildTool
 		/// <summary>
 		/// List of additional pre-build libraries (names of the .lib files including extension) - typically used for additional targets which are still built, but using either TargetRules.PreBuildSteps or TargetRules.PreBuildTargets.
 		/// </summary>
-		public List<string> PublicPreBuildLibraries = new();
+		public List<string> PublicPreBuildLibraries = [];
 
 		/// <summary>
 		/// List of system libraries to use - these are typically referenced via name and then found via the system paths. If you need to reference a .lib file use the PublicAdditionalLibraries instead
 		/// </summary>
-		public List<string> PublicSystemLibraries = new();
+		public List<string> PublicSystemLibraries = [];
 
 		/// <summary>
 		/// List of XCode frameworks (iOS and MacOS)
 		/// </summary>
-		public List<string> PublicFrameworks = new();
+		public List<string> PublicFrameworks = [];
 
 		/// <summary>
 		/// List of weak frameworks (for OS version transitions)
 		/// </summary>
-		public List<string> PublicWeakFrameworks = new();
+		public List<string> PublicWeakFrameworks = [];
 
 		/// <summary>
 		/// List of addition frameworks - typically used for External (third party) modules on Mac and iOS
 		/// </summary>
-		public List<Framework> PublicAdditionalFrameworks = new();
+		public List<Framework> PublicAdditionalFrameworks = [];
 
 		/// <summary>
 		/// List of addition resources that should be copied to the app bundle for Mac or iOS
 		/// </summary>
-		public List<BundleResource> AdditionalBundleResources = new();
+		public List<BundleResource> AdditionalBundleResources = [];
 
 		/// <summary>
 		/// List of type libraries that we need to generate headers for (Windows only)
 		/// </summary>
-		public List<TypeLibrary> TypeLibraries = new();
+		public List<TypeLibrary> TypeLibraries = [];
 
 		/// <summary>
 		/// List of delay load DLLs - typically used for External (third party) modules
 		/// </summary>
-		public List<string> PublicDelayLoadDLLs = new();
+		public List<string> PublicDelayLoadDLLs = [];
 
 		/// <summary>
 		/// Private compiler definitions for this module
 		/// </summary>
-		public List<string> PrivateDefinitions = new();
+		public List<string> PrivateDefinitions = [];
 
 		/// <summary>
 		/// Public compiler definitions for this module
 		/// </summary>
-		public List<string> PublicDefinitions = new();
+		public List<string> PublicDefinitions = [];
 
 		/// <summary>
 		/// AutoRTFM external open -> closed mapping files
 		/// </summary>
-		public List<string> AutoRTFMExternalMappingFiles = new();
+		public List<string> AutoRTFMExternalMappingFiles = [];
 
 		/// <summary>
 		/// Append (or create)
@@ -1341,7 +1341,7 @@ namespace UnrealBuildTool
 		/// <summary>
 		/// Addition modules this module may require at run-time 
 		/// </summary>
-		public List<string> DynamicallyLoadedModuleNames = new();
+		public List<string> DynamicallyLoadedModuleNames = [];
 
 		/// <summary>
 		/// List of files which this module depends on at runtime. These files will be staged along with the target.
@@ -1361,7 +1361,7 @@ namespace UnrealBuildTool
 		/// <summary>
 		/// External files which invalidate the makefile if modified. Relative paths are resolved relative to the .build.cs file.
 		/// </summary>
-		public List<string> ExternalDependencies = new();
+		public List<string> ExternalDependencies = [];
 
 		/// <summary>
 		/// Subclass rules files which invalidate the makefile if modified.
@@ -1435,7 +1435,7 @@ namespace UnrealBuildTool
 		/// The subdirectory name is appended to the generated code directory to form a new directory
 		/// that headers are generated inside.
 		/// </summary>
-		public List<(string, Action<ILogger, DirectoryReference>)> GenerateHeaderFuncs = new();
+		public List<(string, Action<ILogger, DirectoryReference>)> GenerateHeaderFuncs = [];
 
 		/// <summary>
 		///  Control visibility of symbols
@@ -1489,7 +1489,7 @@ namespace UnrealBuildTool
 		/// A List of files to be used as a LLVM FPass-Plugins for all clang compilation units of the current module.
 		///    'clang -fpass-plugin=[FPassPlugin]' 
 		/// </summary>
-		public List<string> FPassPlugins = new();
+		public List<string> FPassPlugins = [];
 
 #nullable disable
 		/// <summary>
@@ -1546,11 +1546,10 @@ namespace UnrealBuildTool
 			if (Target.bUseVerseVM)
 			{
 				PublicDependencyModuleNames.AddRange(
-					new string[]
-					{
+					[
 						"libpas",
 						"CorePreciseFP",
-					}
+					]
 				);
 			}
 		}
@@ -1666,19 +1665,19 @@ namespace UnrealBuildTool
 		public void SetupModulePhysicsSupport(ReadOnlyTargetRules target)
 		{
 			PublicIncludePathModuleNames.AddRange(
-					new string[] {
+				[
 					"Chaos",
-					}
-				);
+				]
+			);
 
 			PublicDependencyModuleNames.AddRange(
-				new string[] {
+				[
 					"PhysicsCore",
 					"ChaosCore",
 					"Chaos",
 					"NNE",
-					}
-				);
+				]
+			);
 
 			PublicDefinitions.Add("WITH_CLOTH_COLLISION_DETECTION=1");
 
@@ -1778,7 +1777,7 @@ namespace UnrealBuildTool
 		public IEnumerable<string> GetModulesInGroup(string moduleGroup, bool bOnlyValid = true)
 		{
 			// figure out what platforms/groups aren't allowed with this opted in list
-			List<string>? disallowedPlatformsAndGroups = (bOnlyValid && Target.OptedInModulePlatforms != null) ? Utils.MakeListOfUnsupportedPlatforms(Target.OptedInModulePlatforms.ToList(), false, Logger) : null;
+			List<string>? disallowedPlatformsAndGroups = (bOnlyValid && Target.OptedInModulePlatforms != null) ? Utils.MakeListOfUnsupportedPlatforms([.. Target.OptedInModulePlatforms], false, Logger) : null;
 
 			foreach (Type moduleType in RulesAssembly.GetTypes()
 				.Where(t => t.IsSubclassOf(typeof(ModuleRules)) && t.IsDefined(typeof(ModuleGroupsAttribute), true))
@@ -1848,12 +1847,7 @@ namespace UnrealBuildTool
 		/// Returns the directories for all subclasses of this module, as well as any additional directories specified by the rules
 		/// </summary>
 		/// <returns>List of directories, or null if none were added</returns>
-		public DirectoryReference[] GetAllModuleDirectories()
-		{
-			List<DirectoryReference> allDirectories = new List<DirectoryReference>(DirectoriesForModuleSubClasses.Values);
-			allDirectories.AddRange(AdditionalModuleDirectories);
-			return allDirectories.ToArray();
-		}
+		public DirectoryReference[] GetAllModuleDirectories() => [.. DirectoriesForModuleSubClasses.Values, .. AdditionalModuleDirectories];
 
 		/// <summary>
 		/// Adds an additional module directory, if it exists (useful for NotForLicensees/NoRedist)
