[UBT]
* Fixed race condition in platform project generation by adding lock around code

diff --git a/Engine/Source/Programs/UnrealBuildTool/ProjectFiles/PlatformProjectGenerator.cs b/Engine/Source/Programs/UnrealBuildTool/ProjectFiles/PlatformProjectGenerator.cs
--- a/Engine/Source/Programs/UnrealBuildTool/ProjectFiles/PlatformProjectGenerator.cs
+++ b/Engine/Source/Programs/UnrealBuildTool/ProjectFiles/PlatformProjectGenerator.cs
@@ -295,14 +295,17 @@ namespace UnrealBuildTool
 			/// <param name="bPreserveExistingValue">If true, the value of the property will be taken from existing user file, but relative order will be taken from newly generated file.</param>
 			public void PatchProperty(string PropertyName, bool bPreserveExistingValue = false)
 			{
-				if (!PropertiesToPatchContainer.Contains(PropertyName))
+				lock (PropertiesToPatchContainer)
 				{
-					PropertiesToPatchContainer.Add(PropertyName);
-				}
+					if (!PropertiesToPatchContainer.Contains(PropertyName))
+					{
+						PropertiesToPatchContainer.Add(PropertyName);
+					}
 
-				if (bPreserveExistingValue && !PropertiesToPatchOrderButPreserveValueContainer.Contains(PropertyName))
-				{
-					PropertiesToPatchOrderButPreserveValueContainer.Add(PropertyName);
+					if (bPreserveExistingValue && !PropertiesToPatchOrderButPreserveValueContainer.Contains(PropertyName))
+					{
+						PropertiesToPatchOrderButPreserveValueContainer.Add(PropertyName);
+					}
 				}
 			}
 
