UnrealBuildTool: Automatically bump PdbPageSize for merge module editor builds

#rnx

diff --git a/Engine/Source/Programs/UnrealBuildTool/Platform/Windows/UEBuildWindows.cs b/Engine/Source/Programs/UnrealBuildTool/Platform/Windows/UEBuildWindows.cs
--- a/Engine/Source/Programs/UnrealBuildTool/Platform/Windows/UEBuildWindows.cs
+++ b/Engine/Source/Programs/UnrealBuildTool/Platform/Windows/UEBuildWindows.cs
@@ -1365,10 +1365,10 @@ namespace UnrealBuildTool
 				throw new BuildException("bAdaptiveUnityEnablesEditAndContinue requires bAdaptiveUnityDisablesPCH or bAdaptiveUnityCreatesDedicatedPCH");
 			}
 
-			// for monolithic editor builds, add the PDBPAGESIZE option, (VS 16.11, VC toolchain 14.29.30133), but the pdb will be too large without this
+			// for monolithic or merged module editor builds, add the PDBPAGESIZE option, (VS 16.11, VC toolchain 14.29.30133), but the pdb will be too large without this
 			// some monolithic game builds could be too large as well, but they can be added in a .Target.cs with:
 			// TargetRules.WindowsPlatform.PdbPageSize = 8192;
-			if (!Target.WindowsPlatform.PdbPageSize.HasValue && Target.LinkType == TargetLinkType.Monolithic && Target.Type == TargetType.Editor)
+			if (!Target.WindowsPlatform.PdbPageSize.HasValue && (Target.LinkType == TargetLinkType.Monolithic || Target.bMergeModules) && Target.Type == TargetType.Editor)
 			{
 				Target.WindowsPlatform.PdbPageSize = 16384;
 			}
@@ -1381,8 +1381,8 @@ namespace UnrealBuildTool
 
 			if (Target.WindowsPlatform.bDynamicDebugging)
 			{
-				// Monolithic targets always require increasing the pdb page size with Dynamic Debugging or linking may fail
-				if (Target.LinkType == TargetLinkType.Monolithic && (Target.WindowsPlatform.PdbPageSize ?? 0) < 16384)
+				// Monolithic or merged module targets always require increasing the pdb page size with Dynamic Debugging or linking may fail
+				if ((Target.LinkType == TargetLinkType.Monolithic || Target.bMergeModules) && (Target.WindowsPlatform.PdbPageSize ?? 0) < 16384)
 				{
 					Target.WindowsPlatform.PdbPageSize = 16384;
 				}
@@ -1390,10 +1390,7 @@ namespace UnrealBuildTool
 				// '/DYNAMICDEOPT' combined with '/OPT:ICF' has a degraded debugging experience;
 				if (Target.WindowsPlatform.bMergeIdenticalCOMDATs)
 				{
-					if (Target.WindowsPlatform.bMergeIdenticalCOMDATs)
-					{
-						Target.WindowsPlatform.bMergeIdenticalCOMDATs = false;
-					}
+					Target.WindowsPlatform.bMergeIdenticalCOMDATs = false;
 				}
 			}
 
