UnrealBuildTool: No need to write PVS-Studio.lic into intermediate, it is not used

#rnx

diff --git a/Engine/Source/Programs/UnrealBuildTool/Platform/Windows/PVSToolChain.cs b/Engine/Source/Programs/UnrealBuildTool/Platform/Windows/PVSToolChain.cs
--- a/Engine/Source/Programs/UnrealBuildTool/Platform/Windows/PVSToolChain.cs
+++ b/Engine/Source/Programs/UnrealBuildTool/Platform/Windows/PVSToolChain.cs
@@ -585,12 +585,6 @@ namespace UnrealBuildTool
 				{
 					throw new BuildException("All PVS-Studio analysis modes are disabled.");
 				}
-
-				if (!String.IsNullOrEmpty(ApplicationSettings.UserName) && !String.IsNullOrEmpty(ApplicationSettings.SerialNumber))
-				{
-					LicenseFile = FileReference.Combine(Unreal.EngineDirectory, "Intermediate", "PVS", "PVS-Studio.lic");
-					Utils.WriteFileIfChanged(LicenseFile, String.Format("{0}\n{1}\n", ApplicationSettings.UserName, ApplicationSettings.SerialNumber), Logger);
-				}
 			}
 			else
 			{
