Made undecorated Android config in Visual Studio use project preferred Architecture.

diff --git a/Engine/Source/Programs/UnrealBuildTool/ProjectFiles/VisualStudio/VCProject.cs b/Engine/Source/Programs/UnrealBuildTool/ProjectFiles/VisualStudio/VCProject.cs
--- a/Engine/Source/Programs/UnrealBuildTool/ProjectFiles/VisualStudio/VCProject.cs
+++ b/Engine/Source/Programs/UnrealBuildTool/ProjectFiles/VisualStudio/VCProject.cs
@@ -433,11 +433,18 @@ namespace UnrealBuildTool
 							throw new BuildException("Project config and target combinations has not been populated.");
 						}
 
-						ProjectConfigAndTargetCombination? Combination = ProjectConfigAndTargetCombinations.FirstOrDefault(Combination =>
-							Combination.Platform == ProjectPlatform &&
-							Combination.Configuration == ProjectConfiguration &&
-							Combination.ProjectTarget == MatchingProjectTarget &&
-							Combination.Architecture == Architecture);
+						ProjectConfigAndTargetCombination? Combination = ProjectConfigAndTargetCombinations.FirstOrDefault(
+							ProjectPlatform == UnrealTargetPlatform.Android && Architecture == null ?
+								// Match first preferred Architecture if none is provided
+								Combination =>
+									Combination.Platform == ProjectPlatform &&
+									Combination.Configuration == ProjectConfiguration &&
+									Combination.ProjectTarget == MatchingProjectTarget :
+								Combination =>
+									Combination.Platform == ProjectPlatform &&
+									Combination.Configuration == ProjectConfiguration &&
+									Combination.ProjectTarget == MatchingProjectTarget &&
+									Combination.Architecture == Architecture);
 
 						if (Combination == null)
 						{
@@ -751,15 +758,14 @@ namespace UnrealBuildTool
 									}
 									else
 									{
-										if (BuildPlatform.ArchitectureConfig.Mode == UnrealArchitectureMode.SingleTargetLinkSeparately)
+										IEnumerable<UnrealArch> Archs = Architectures.Architectures;
+										if (BuildPlatform.ArchitectureConfig.Mode == UnrealArchitectureMode.SingleTargetLinkSeparately && BuildPlatform.Platform == UnrealTargetPlatform.Android)
 										{
-											if (BuildPlatform.Platform == UnrealTargetPlatform.Android)
-											{
-												Architectures = BuildPlatform.ArchitectureConfig.AllSupportedArchitectures;
-											}
+											// Order project Architectures first, as they are preferred
+											Archs = Archs.Concat(BuildPlatform.ArchitectureConfig.AllSupportedArchitectures.Architectures.Except(Archs));
 										}
 
-										foreach (UnrealArch Arch in Architectures.Architectures)
+										foreach (UnrealArch Arch in Archs)
 										{
 											AddProjectAndTargetCombination(Arch);
 										}
diff --git a/Engine/Source/Programs/UnrealBuildTool/ProjectFiles/VisualStudio/VCProjectFileGenerator.cs b/Engine/Source/Programs/UnrealBuildTool/ProjectFiles/VisualStudio/VCProjectFileGenerator.cs
--- a/Engine/Source/Programs/UnrealBuildTool/ProjectFiles/VisualStudio/VCProjectFileGenerator.cs
+++ b/Engine/Source/Programs/UnrealBuildTool/ProjectFiles/VisualStudio/VCProjectFileGenerator.cs
@@ -1139,19 +1139,20 @@ namespace UnrealBuildTool
 						}
 						else
 						{
-							foreach (UnrealArch Arch in Architectures.Architectures)
+							if (BuildPlatform.ArchitectureConfig.Mode == UnrealArchitectureMode.SingleTargetLinkSeparately && BuildPlatform.Platform == UnrealTargetPlatform.Android)
 							{
-								AddSolutionConfig(Arch, OutSolutionConfigs);
-							}
+								AddSolutionConfig(null, OutSolutionConfigs);
 
-							if (BuildPlatform.ArchitectureConfig.Mode == UnrealArchitectureMode.SingleTargetLinkSeparately)
+								foreach (UnrealArch Arch in BuildPlatform.ArchitectureConfig.AllSupportedArchitectures.Architectures)
+								{
+									AddSolutionConfig(Arch, OutSolutionConfigs, true);
+								}
+							}
+							else
 							{
-								if (BuildPlatform.Platform == UnrealTargetPlatform.Android)
+								foreach (UnrealArch Arch in Architectures.Architectures)
 								{
-									foreach (UnrealArch Arch in BuildPlatform.ArchitectureConfig.AllSupportedArchitectures.Architectures)
-									{
-										AddSolutionConfig(Arch, OutSolutionConfigs, true);
-									}
+									AddSolutionConfig(Arch, OutSolutionConfigs);
 								}
 							}
 						}
