Improve the module merging strategy to allow for multiple root GFPs inside a single library

Sometimes we have many small root GFPs that are closely related to eachother, we don't want an entire DLL for each because we have platform limits on that.
Instead, this allows for a shared merged library rule to work without a matching dedicated merged library.
This is very helpful for when a very small root GFP has dependencies that would otherwise produce unmergeable modules.

#rnx

diff --git a/Engine/Source/Programs/UnrealBuildTool/Configuration/UEBuildTarget.cs b/Engine/Source/Programs/UnrealBuildTool/Configuration/UEBuildTarget.cs
--- a/Engine/Source/Programs/UnrealBuildTool/Configuration/UEBuildTarget.cs
+++ b/Engine/Source/Programs/UnrealBuildTool/Configuration/UEBuildTarget.cs
@@ -3790,8 +3790,29 @@ namespace UnrealBuildTool
 				Dictionary<string, HashSet<UEBuildPlugin>> pluginGroups = new();
 				Dictionary<string, HashSet<UEBuildPlugin>> pluginGroupDependencies = new();
 
-				// Create plugin groups with only the root plugin
+				// We want the ability to add entire root GFPs to a shared library without a separate dedicated library.
+				// If there's a sharing rule but no root library, we'll add a fake root here, and later on remap its contents to the shared one.
+				List<string> rootPluginsToMerge = new();
+				Dictionary<string, string> rootPluginRemappings = new();
 				foreach (string name in Rules.MergePlugins)
+				{
+					rootPluginsToMerge.Add(name);
+				}
+				foreach (KeyValuePair<string, IEnumerable<string>> sharedGroupRule in Rules.MergePluginsShared)
+				{
+					foreach (string name in sharedGroupRule.Value)
+					{
+						if (!Rules.MergePlugins.Contains(name))
+						{
+							rootPluginsToMerge.Add(name);
+							rootPluginRemappings.Add(name, sharedGroupRule.Key);
+						}
+					}
+				}
+
+				// When processing root GFPs, we add the plugin and direct dependencies to a plugin group.
+				// Groups will then be expanded with indirect dependencies.
+				foreach (string name in rootPluginsToMerge)
 				{
 					UEBuildPlugin? plugin = BuildPlugins.FirstOrDefault(x => x.Name == name);
 					if (plugin == null)
@@ -3879,6 +3900,16 @@ namespace UnrealBuildTool
 					}
 				}
 
+				// Process the remapping rules for fully shared root GFPs.
+				foreach (KeyValuePair<string, string> remapping in rootPluginRemappings)
+				{
+					string source = remapping.Key;
+					string destination = remapping.Value;
+
+					pluginGroups[destination].UnionWith(pluginGroups[source]);
+					pluginGroups.Remove(source);
+				}
+
 				// At this point, we have to exclude any code that falls under remaining unmergeable plugins, dependencies included.
 				// We need that because there is no way the executable can depend on a merged library, and this remaining chunk would always be loaded as its own library.
 				IEnumerable<UEBuildPlugin> ignoredPlugins = BuildPlugins.Where(x => Rules.MergePluginsIgnored.Contains(x.Name));
