UnrealBuildTool: Check both "Unreal Engine" and "UnrealEngine" for user config files

#rnx

diff --git a/Engine/Source/Programs/UnrealBuildTool/System/ConfigHierarchy.cs b/Engine/Source/Programs/UnrealBuildTool/System/ConfigHierarchy.cs
--- a/Engine/Source/Programs/UnrealBuildTool/System/ConfigHierarchy.cs
+++ b/Engine/Source/Programs/UnrealBuildTool/System/ConfigHierarchy.cs
@@ -976,10 +976,13 @@ namespace UnrealBuildTool
 			"{PROJECT}/Config/{PLATFORM}/Custom/{CUSTOMCONFIG}/{PLATFORM}{TYPE}.ini",
 			// AppSettings/.../System*.ini
 			"{APPSETTINGS}/Unreal Engine/Engine/Config/System{TYPE}.ini",
+			"{APPSETTINGS}/UnrealEngine/Engine/Config/System{TYPE}.ini",
 			// UserSettings/.../User*.ini
 			"{USERSETTINGS}/Unreal Engine/Engine/Config/User{TYPE}.ini",
+			"{USERSETTINGS}/UnrealEngine/Engine/Config/User{TYPE}.ini",
 			// UserDir/.../User*.ini
 			"{USER}/Unreal Engine/Engine/Config/User{TYPE}.ini",
+			"{USER}/UnrealEngine/Engine/Config/User{TYPE}.ini",
 			// Project/User*.ini
 			"{PROJECT}/Config/User{TYPE}.ini",
 		};
