UnrealBuildTool: Show peak process memory for UBA actions

#rnx
#FailedValidations Preflight

diff --git a/Engine/Source/Programs/Shared/EpicGames.UBA/IScheduler.cs b/Engine/Source/Programs/Shared/EpicGames.UBA/IScheduler.cs
--- a/Engine/Source/Programs/Shared/EpicGames.UBA/IScheduler.cs
+++ b/Engine/Source/Programs/Shared/EpicGames.UBA/IScheduler.cs
@@ -71,6 +71,11 @@ namespace EpicGames.UBA
 		/// </summary>
 		TimeSpan TotalWallTime { get; }
 
+		/// <summary>
+		/// Peak memory used, requires a job object so will only be non-zero for Windows hosts.
+		/// </summary>
+		long PeakMemoryUsed { get; }
+
 		/// <summary>
 		/// UserData that was provided in EnqueueProcess
 		/// </summary>
diff --git a/Engine/Source/Programs/Shared/EpicGames.UBA/Impl/SchedulerImpl.cs b/Engine/Source/Programs/Shared/EpicGames.UBA/Impl/SchedulerImpl.cs
--- a/Engine/Source/Programs/Shared/EpicGames.UBA/Impl/SchedulerImpl.cs
+++ b/Engine/Source/Programs/Shared/EpicGames.UBA/Impl/SchedulerImpl.cs
@@ -28,6 +28,9 @@ namespace EpicGames.UBA.Impl
 
 		[DllImport("UbaHost", CharSet = CharSet.Auto)]
 		static extern ulong ProcessHandle_GetTotalWallTime(nint process);
+
+		[DllImport("UbaHost", CharSet = CharSet.Auto)]
+		static extern ulong ProcessHandle_GetPeakMemory(nint process);
 		#endregion
 
 		internal ProcessFinishedInfoImpl(object userData, nint processHandle)
@@ -57,6 +60,7 @@ namespace EpicGames.UBA.Impl
 		public string? ExecutingHost => Marshal.PtrToStringAuto(ProcessHandle_GetExecutingHost(_processHandle));
 		public TimeSpan TotalProcessorTime => new((long)ProcessHandle_GetTotalProcessorTime(_processHandle));
 		public TimeSpan TotalWallTime => new((long)ProcessHandle_GetTotalWallTime(_processHandle));
+		public long PeakMemoryUsed => (long)ProcessHandle_GetPeakMemory(_processHandle);
 		public object UserData => _userData;
 		public nint ProcessHandle => _processHandle;
 	}
diff --git a/Engine/Source/Programs/UnrealBuildTool/Executors/UnrealBuildAccelerator/UBAExecutor.cs b/Engine/Source/Programs/UnrealBuildTool/Executors/UnrealBuildAccelerator/UBAExecutor.cs
--- a/Engine/Source/Programs/UnrealBuildTool/Executors/UnrealBuildAccelerator/UBAExecutor.cs
+++ b/Engine/Source/Programs/UnrealBuildTool/Executors/UnrealBuildAccelerator/UBAExecutor.cs
@@ -657,6 +657,8 @@ namespace UnrealBuildTool
 			int remoteProcessedActions = 0;
 			int localRetryActions = 0;
 			int forcedRetryActions = 0;
+			long peakProcessMemory = 0;
+			long peakTotalMemory = 0;
 			bool shouldCancel = false;
 
 			scheduler.SetProcessFinishedCallback((info) =>
@@ -784,8 +786,11 @@ namespace UnrealBuildTool
 				{
 					processorTime = info.TotalProcessorTime;
 					executionTime = info.TotalWallTime;
+					peakMemory = info.PeakMemoryUsed;
 				}
 
+				peakProcessMemory = Math.Max(peakProcessMemory, peakMemory);
+
 				logLines.RemoveAll(line =>
 					line.StartsWith("   Creating library ", StringComparison.Ordinal) ||
 					line.StartsWith("   Creating object ", StringComparison.Ordinal) ||
@@ -910,7 +915,7 @@ namespace UnrealBuildTool
 						!UBAConfig.bDisableRemote ? _ubaDurationWaitingForRemote : TimeSpan.Zero,
 						!UBAConfig.bDisableRemote || _agentCoordinators.Count == 0 ? "Local" : _remoteConnectionMode,
 						DateTime.UtcNow,
-						0, 0);
+						peakProcessMemory, peakTotalMemory);
 				}
 				return res;
 			}
