PR #13826: Add --prioritize-dependency-scanning SN-DBS build option

diff --git a/Engine/Source/Programs/UnrealBuildTool/Executors/SNDBS.cs b/Engine/Source/Programs/UnrealBuildTool/Executors/SNDBS.cs
--- a/Engine/Source/Programs/UnrealBuildTool/Executors/SNDBS.cs
+++ b/Engine/Source/Programs/UnrealBuildTool/Executors/SNDBS.cs
@@ -51,6 +51,12 @@ namespace UnrealBuildTool
 
 		private Dictionary<string, string> ActiveTemplates = BuiltInTemplates.ToDictionary(p => p.Key, p => p.Value);
 
+		/// <summary>
+		/// When enabled, dependency scanning is prioritized to avoid being slowed down by running local jobs.
+		/// </summary>
+		[XmlConfigFile(Category = "BuildConfiguration")]
+		bool bPrioritizeDependencyScanning = false;
+
 		/// <summary>
 		/// When enabled, SN-DBS will stop compiling targets after a compile error occurs.  Recommended, as it saves computing resources for others.
 		/// </summary>
@@ -329,6 +335,10 @@ namespace UnrealBuildTool
 			{
 				StartInfo.Arguments += " -k";
 			}
+			if (bPrioritizeDependencyScanning)
+			{
+				StartInfo.Arguments += " --prioritize-dependency-scanning";
+			}
 			return ExecuteProcessWithProgressMarkup(StartInfo, Actions.Count(), Logger);
 		}
 
