UnrealBuildTool: Don't merge module manifests for merged targets, as they are including stale module lists

#rnx

diff --git a/Engine/Source/Programs/UnrealBuildTool/Configuration/UEBuildTarget.cs b/Engine/Source/Programs/UnrealBuildTool/Configuration/UEBuildTarget.cs
--- a/Engine/Source/Programs/UnrealBuildTool/Configuration/UEBuildTarget.cs
+++ b/Engine/Source/Programs/UnrealBuildTool/Configuration/UEBuildTarget.cs
@@ -2932,7 +2932,7 @@ namespace UnrealBuildTool
 				// Create an action which to generate the module receipts
 				if (VersionFile == null)
 				{
-					WriteMetadataTargetInfo TargetInfo = new WriteMetadataTargetInfo(ProjectFile, null, null, ReceiptFileName, Receipt, FileNameToModuleManifest, FileNameToLoadOrderManifest);
+					WriteMetadataTargetInfo TargetInfo = new WriteMetadataTargetInfo(ProjectFile, null, null, ReceiptFileName, Receipt, FileNameToModuleManifest, FileNameToLoadOrderManifest, Rules.bMergeModules);
 					FileReference TargetInfoFile = FileReference.Combine(ProjectIntermediateDirectory, "TargetMetadata.dat");
 					CreateWriteMetadataAction(MakefileBuilder, ReceiptFileName.GetFileName(), TargetInfoFile, TargetInfo, Makefile.OutputItems);
 				}
@@ -2968,12 +2968,12 @@ namespace UnrealBuildTool
 						}
 					}
 
-					WriteMetadataTargetInfo EngineInfo = new WriteMetadataTargetInfo(null, VersionFile, Receipt.Version, null, null, EngineFileToManifest, null);
+					WriteMetadataTargetInfo EngineInfo = new WriteMetadataTargetInfo(null, VersionFile, Receipt.Version, null, null, EngineFileToManifest, null, Rules.bMergeModules);
 					FileReference EngineInfoFile = FileReference.Combine(ProjectIntermediateDirectory, "EngineMetadata.dat");
 					string EngineInfoName = TargetName.Equals(VersionFile.GetFileNameWithoutExtension()) ? VersionFile.GetFileName() : $"{VersionFile.GetFileName()} ({TargetName})";
 					CreateWriteMetadataAction(MakefileBuilder, EngineInfoName, EngineInfoFile, EngineInfo, EnginePrereqItems);
 
-					WriteMetadataTargetInfo TargetInfo = new WriteMetadataTargetInfo(ProjectFile, VersionFile, null, ReceiptFileName, Receipt, TargetFileToManifest, FileNameToLoadOrderManifest);
+					WriteMetadataTargetInfo TargetInfo = new WriteMetadataTargetInfo(ProjectFile, VersionFile, null, ReceiptFileName, Receipt, TargetFileToManifest, FileNameToLoadOrderManifest, Rules.bMergeModules);
 					FileReference TargetInfoFile = FileReference.Combine(ProjectIntermediateDirectory, "TargetMetadata.dat");
 					CreateWriteMetadataAction(MakefileBuilder, ReceiptFileName.GetFileName(), TargetInfoFile, TargetInfo, TargetPrereqItems);
 				}
diff --git a/Engine/Source/Programs/UnrealBuildTool/Modes/WriteMetadataMode.cs b/Engine/Source/Programs/UnrealBuildTool/Modes/WriteMetadataMode.cs
--- a/Engine/Source/Programs/UnrealBuildTool/Modes/WriteMetadataMode.cs
+++ b/Engine/Source/Programs/UnrealBuildTool/Modes/WriteMetadataMode.cs
@@ -14,64 +14,21 @@ namespace UnrealBuildTool
 	/// <summary>
 	/// Parameters for the WriteMetadata mode
 	/// </summary>
+	/// <param name="ProjectFile">The project file</param>
+	/// <param name="VersionFile">Output location for the version file</param>
+	/// <param name="Version">The new version to write. This should only be set on the engine step</param>
+	/// <param name="ReceiptFile">Output location for the target file</param>
+	/// <param name="Receipt">The new receipt to write. This should only be set on the target step</param>
+	/// <param name="FileToManifest">Map of module manifest filenames to their location on disk</param>
+	/// <param name="FileToLoadOrderManifest">Map of load order manifest filenames to their locations on disk (generally, at most one load oder manifest is expected)</param>
+	/// <param name="MergeModules">If this target uses merged modules. When true module manifests will not be merged with on disk version when buildids match</param>
 	[Serializable]
-	class WriteMetadataTargetInfo
+	record WriteMetadataTargetInfo(FileReference? ProjectFile, FileReference? VersionFile, BuildVersion? Version, FileReference? ReceiptFile, TargetReceipt? Receipt, Dictionary<FileReference, ModuleManifest> FileToManifest, Dictionary<FileReference, LoadOrderManifest>? FileToLoadOrderManifest, bool MergeModules)
 	{
-		/// <summary>
-		/// The project file
-		/// </summary>
-		public FileReference? ProjectFile;
-
-		/// <summary>
-		/// Output location for the version file
-		/// </summary>
-		public FileReference? VersionFile;
-
-		/// <summary>
-		/// The new version to write. This should only be set on the engine step.
-		/// </summary>
-		public BuildVersion? Version;
-
-		/// <summary>
-		/// Output location for the target file
-		/// </summary>
-		public FileReference? ReceiptFile;
-
-		/// <summary>
-		/// The new receipt to write. This should only be set on the target step.
-		/// </summary>
-		public TargetReceipt? Receipt;
-
-		/// <summary>
-		/// Map of module manifest filenames to their location on disk.
-		/// </summary>
-		public Dictionary<FileReference, ModuleManifest> FileToManifest;
-
 		/// <summary>
 		/// Map of load order manifest filenames to their locations on disk (generally, at most one load oder manifest is expected).
 		/// </summary>
-		public Dictionary<FileReference, LoadOrderManifest> FileToLoadOrderManifest;
-
-		/// <summary>
-		/// Constructor
-		/// </summary>
-		/// <param name="ProjectFile"></param>
-		/// <param name="VersionFile"></param>
-		/// <param name="Version"></param>
-		/// <param name="ReceiptFile"></param>
-		/// <param name="Receipt"></param>
-		/// <param name="FileToManifest"></param>
-		/// <param name="FileToLoadOrderManifest"></param>
-		public WriteMetadataTargetInfo(FileReference? ProjectFile, FileReference? VersionFile, BuildVersion? Version, FileReference? ReceiptFile, TargetReceipt? Receipt, Dictionary<FileReference, ModuleManifest> FileToManifest, Dictionary<FileReference, LoadOrderManifest>? FileToLoadOrderManifest)
-		{
-			this.ProjectFile = ProjectFile;
-			this.VersionFile = VersionFile;
-			this.Version = Version;
-			this.ReceiptFile = ReceiptFile;
-			this.Receipt = Receipt;
-			this.FileToManifest = FileToManifest;
-			this.FileToLoadOrderManifest = FileToLoadOrderManifest ?? new Dictionary<FileReference, LoadOrderManifest>();
-		}
+		public Dictionary<FileReference, LoadOrderManifest> FileToLoadOrderManifest { get; init; } = FileToLoadOrderManifest ?? [];
 	}
 
 	/// <summary>
@@ -142,13 +99,15 @@ namespace UnrealBuildTool
 				BuildId = Guid.NewGuid().ToString();
 			}
 
-			// Read all the existing manifests and merge them into the new ones if they have the same build id
-			foreach (KeyValuePair<FileReference, ModuleManifest> Pair in TargetInfo.FileToManifest)
+			// Read all the existing manifests and merge them into the new ones if they have the same build id and not merging modules
+			if (!TargetInfo.MergeModules)
 			{
-				ModuleManifest? SourceManifest;
-				if (TryReadManifest(Pair.Key, Logger, out SourceManifest) && SourceManifest.BuildId == BuildId)
+				foreach (KeyValuePair<FileReference, ModuleManifest> Pair in TargetInfo.FileToManifest)
 				{
-					MergeManifests(SourceManifest, Pair.Value);
+					if (TryReadManifest(Pair.Key, Logger, out ModuleManifest? SourceManifest) && SourceManifest.BuildId == BuildId)
+					{
+						MergeManifests(SourceManifest, Pair.Value);
+					}
 				}
 			}
 
