Restrict Clang linker serialization to build machines or LTO/PGO to avoid harming local iteration

diff --git a/Engine/Source/Programs/UnrealBuildTool/Platform/Windows/VCToolChain.cs b/Engine/Source/Programs/UnrealBuildTool/Platform/Windows/VCToolChain.cs
--- a/Engine/Source/Programs/UnrealBuildTool/Platform/Windows/VCToolChain.cs
+++ b/Engine/Source/Programs/UnrealBuildTool/Platform/Windows/VCToolChain.cs
@@ -3385,7 +3385,8 @@ namespace UnrealBuildTool
 
 			if (Target.WindowsPlatform.Compiler.IsClang() &&
 				Target.WindowsPlatform.bAllowClangLinker &&
-				(Target.LinkType == TargetLinkType.Monolithic))
+				(Target.LinkType == TargetLinkType.Monolithic) &&
+				(Unreal.IsBuildMachine() || LinkEnvironment.bAllowLTCG || LinkEnvironment.bPGOOptimize || LinkEnvironment.bPGOProfile))
 			{
 				// Set the weight to number of logical cores as lld can max out the available cores
 				LinkAction.Weight = Utils.GetLogicalProcessorCount();
