Fix platform extensions working if the parent plugin has an _ in it

diff --git a/Engine/Source/Programs/UnrealBuildTool/System/Plugins.cs b/Engine/Source/Programs/UnrealBuildTool/System/Plugins.cs
--- a/Engine/Source/Programs/UnrealBuildTool/System/Plugins.cs
+++ b/Engine/Source/Programs/UnrealBuildTool/System/Plugins.cs
@@ -420,10 +420,13 @@ namespace UnrealBuildTool
 			// find the parent
 			PluginInfo? Parent = null;
 
-			string[] Tokens = Filename.GetFileNameWithoutAnyExtensions().Split("_".ToCharArray());
-			if (Tokens.Length == 2)
+			string FilenameWithoutExtension = Filename.GetFileNameWithoutAnyExtensions();
+			int LastUnderscoreIndex = FilenameWithoutExtension.LastIndexOf('_');
+
+			string ParentPluginName;
+			if (LastUnderscoreIndex > 0 && LastUnderscoreIndex < FilenameWithoutExtension.Length - 1)
 			{
-				string ParentPluginName = Tokens[0];
+				ParentPluginName = FilenameWithoutExtension.Substring(0, LastUnderscoreIndex);
 				Parent = Plugins.FirstOrDefault(x => x.Name.Equals(ParentPluginName, StringComparison.OrdinalIgnoreCase) && x.LoadedFrom == Child.LoadedFrom);
 			}
 			else
@@ -434,11 +437,11 @@ namespace UnrealBuildTool
 			// did we find a parent plugin?
 			if (Parent == null)
 			{
-				throw new BuildException("Unable to find parent plugin {0} for platform extension plugin {1}. Make sure {0}.uplugin exists.", Tokens[0], Filename);
+				throw new BuildException("Unable to find parent plugin {0} for platform extension plugin {1}. Make sure {0}.uplugin exists.", ParentPluginName, Filename);
 			}
 
 			// validate child plugin file name
-			string PlatformName = Tokens[1];
+			string PlatformName = FilenameWithoutExtension.Substring(LastUnderscoreIndex + 1);
 			if (!IsValidChildPluginSuffix(PlatformName))
 			{
 				Log.TraceWarningTask(Filename, $"Ignoring child plugin: {Child.File.GetFileName()} - Unknown suffix \"{PlatformName}\": expected valid platform or group.");
