Fix Microsoft.CodeAnalysis.CSharp reference which caused a transitive reference to System.Collections.Immutable 9.0.0 rather than 8.0.0

#rnx

diff --git a/Engine/Build/BatchFiles/BuildUAT.bat b/Engine/Build/BatchFiles/BuildUAT.bat
--- a/Engine/Build/BatchFiles/BuildUAT.bat
+++ b/Engine/Build/BatchFiles/BuildUAT.bat
@@ -56,6 +56,7 @@ dotnet build Programs\AutomationTool\AutomationTool.csproj -c Development -p Aut
 if errorlevel 1 goto Error_UATCompileFailed
 
 echo Publishing AutomationTool...
+if exist ..\Binaries\DotNET\AutomationTool rmdir /s /q ..\Binaries\DotNET\AutomationTool >nul 2>&1
 if not exist ..\Binaries\DotNET\AutomationTool mkdir ..\Binaries\DotNET\AutomationTool >nul 2>&1
 dotnet publish Programs\AutomationTool\AutomationTool.csproj -c Development -p AutomationToolProjectOnly=true --output ..\Binaries\DotNET\AutomationTool --no-build -v %MSBUILD_LOGLEVEL%
 if errorlevel 1 goto Error_UATCompileFailed
diff --git a/Engine/Build/BatchFiles/BuildUAT.sh b/Engine/Build/BatchFiles/BuildUAT.sh
--- a/Engine/Build/BatchFiles/BuildUAT.sh
+++ b/Engine/Build/BatchFiles/BuildUAT.sh
@@ -67,6 +67,12 @@ if [ $PERFORM_REBUILD -eq 1 ]; then
   fi
 
   echo "Publishing AutomationTool..."
+  if [ -d ../Binaries/DotNET/AutomationTool ]; then
+    rm -rf ../Binaries/DotNET/AutomationTool >/dev/null
+  fi
+  if [ ! -d ../Binaries/DotNET/AutomationTool ]; then
+    mkdir -p ../Binaries/DotNET/AutomationTool >/dev/null
+  fi
   "$SCRIPT_DIR/DotnetRetry.sh" publish Programs/AutomationTool/AutomationTool.csproj -c Development --output ../Binaries/DotNET/AutomationTool --no-build -v "$MSBuild_Verbosity"
   if [ $? -ne 0 ]; then
     echo "AutomationTool publish failed"
diff --git a/Engine/Build/BatchFiles/BuildUBT.bat b/Engine/Build/BatchFiles/BuildUBT.bat
--- a/Engine/Build/BatchFiles/BuildUBT.bat
+++ b/Engine/Build/BatchFiles/BuildUBT.bat
@@ -52,6 +52,7 @@ dotnet build Programs\UnrealBuildTool\UnrealBuildTool.csproj -c Development -v %
 if errorlevel 1 goto Error_UBTCompileFailed
 
 echo Publishing UnrealBuildTool...
+if exist ..\Binaries\DotNET\UnrealBuildTool rmdir /s /q ..\Binaries\DotNET\UnrealBuildTool >nul 2>&1
 if not exist ..\Binaries\DotNET\UnrealBuildTool mkdir ..\Binaries\DotNET\UnrealBuildTool >nul 2>&1
 dotnet publish Programs\UnrealBuildTool\UnrealBuildTool.csproj -c Development --output ..\Binaries\DotNET\UnrealBuildTool --no-build -v %MSBUILD_LOGLEVEL%
 if errorlevel 1 goto Error_UBTCompileFailed
diff --git a/Engine/Build/BatchFiles/BuildUBT.sh b/Engine/Build/BatchFiles/BuildUBT.sh
--- a/Engine/Build/BatchFiles/BuildUBT.sh
+++ b/Engine/Build/BatchFiles/BuildUBT.sh
@@ -68,6 +68,12 @@ if [ $PERFORM_REBUILD -eq 1 ]; then
   fi
 
   echo "Publishing UnrealBuildTool..."
+  if [ -d ../Binaries/DotNET/UnrealBuildTool ]; then
+    rm -rf ../Binaries/DotNET/UnrealBuildTool >/dev/null
+  fi
+  if [ ! -d ../Binaries/DotNET/UnrealBuildTool ]; then
+    mkdir -p ../Binaries/DotNET/UnrealBuildTool >/dev/null
+  fi
   "$SCRIPT_DIR/DotnetRetry.sh" publish Programs/UnrealBuildTool/UnrealBuildTool.csproj -c Development --output ../Binaries/DotNET/UnrealBuildTool --no-build -v "$MSBuild_Verbosity"
   if [ $? -ne 0 ]; then
     echo "UnrealBuildTool publish failed"
diff --git a/Engine/Source/Programs/UnrealAndroidFileTool/UnrealAndroidFileTool.csproj b/Engine/Source/Programs/UnrealAndroidFileTool/UnrealAndroidFileTool.csproj
--- a/Engine/Source/Programs/UnrealAndroidFileTool/UnrealAndroidFileTool.csproj
+++ b/Engine/Source/Programs/UnrealAndroidFileTool/UnrealAndroidFileTool.csproj
@@ -71,7 +71,7 @@
   </ItemGroup>
 
   <ItemGroup>
-    <PackageReference Include="Microsoft.CodeAnalysis.CSharp" Version="4.14.0" />
+    <PackageReference Include="Microsoft.CodeAnalysis.CSharp" Version="4.13.0" />
     <PackageReference Include="Microsoft.Win32.Registry" Version="5.0.0" />
     <PackageReference Include="OpenTracing" Version="0.12.1" />
     <PackageReference Include="System.Management" Version="4.7.0" />
diff --git a/Engine/Source/Programs/UnrealBuildTool/UnrealBuildTool.csproj b/Engine/Source/Programs/UnrealBuildTool/UnrealBuildTool.csproj
--- a/Engine/Source/Programs/UnrealBuildTool/UnrealBuildTool.csproj
+++ b/Engine/Source/Programs/UnrealBuildTool/UnrealBuildTool.csproj
@@ -1,4 +1,4 @@
-<Project Sdk="Microsoft.NET.Sdk">
+﻿<Project Sdk="Microsoft.NET.Sdk">
   <Import Project="..\Shared\UnrealEngine.csproj.props" />
   
   <PropertyGroup>
@@ -94,7 +94,7 @@
     <ProjectReference Include="..\Shared\EpicGames.UHT\EpicGames.UHT.csproj" />
   </ItemGroup>
   <ItemGroup>
-    <PackageReference Include="Microsoft.CodeAnalysis.CSharp" Version="4.14.0" />
+    <PackageReference Include="Microsoft.CodeAnalysis.CSharp" Version="4.13.0" />
     <PackageReference Include="Microsoft.Extensions.Logging.Console" Version="8.0.1" />
     <PackageReference Include="Microsoft.VisualStudio.Setup.Configuration.Interop" Version="3.14.2075" />
     <PackageReference Include="Microsoft.Win32.Registry" Version="5.0.0" />
