Add source map commands to lldb setup for android debugging when using VFS

diff --git a/Engine/Source/Programs/UnrealBuildTool/Platform/Android/AndroidProjectGenerator.cs b/Engine/Source/Programs/UnrealBuildTool/Platform/Android/AndroidProjectGenerator.cs
--- a/Engine/Source/Programs/UnrealBuildTool/Platform/Android/AndroidProjectGenerator.cs
+++ b/Engine/Source/Programs/UnrealBuildTool/Platform/Android/AndroidProjectGenerator.cs
@@ -4,6 +4,7 @@ using System;
 using System.Collections.Generic;
 using System.Diagnostics;
 using System.IO;
+using System.Linq;
 using System.Text;
 using System.Threading.Tasks;
 using EpicGames.Core;
@@ -197,8 +198,12 @@ namespace UnrealBuildTool
 				StringBuilder Out = new StringBuilder();
 				Out.AppendLine("  <PropertyGroup " + InConditionString + ">");
 
-				string LldbFormatterImport = $"command script import \"" + Path.Combine(Unreal.EngineDirectory.FullName, "Extras", "LLDBDataFormatters", "UEDataFormatters_2ByteChars.py") + "\"";
-				Out.AppendLine($"    <AndroidLldbStartupCommands>{LldbFormatterImport};$(AndroidLldbStartupCommands)</AndroidLldbStartupCommands>");
+				Dictionary<DirectoryReference, DirectoryReference> SourceFileMap = ProjectFileGenerator.GenerateSourceMap(InTargetRules.ProjectFile);
+				Out.Append("    <AndroidLldbStartupCommands>");
+				Out.Append($"command script import \"{Path.Combine(Unreal.EngineDirectory.FullName, "Extras", "LLDBDataFormatters", "UEDataFormatters_2ByteChars.py")}\";");
+				Out.AppendJoin(';', SourceFileMap.Select(from_to => $"setting append target.source-map {from_to.Key} {from_to.Value}"));
+				Out.Append(SourceFileMap.Any() ? ";" : "");
+				Out.AppendLine($"$(AndroidLldbStartupCommands)</AndroidLldbStartupCommands>");
 				VCUserFileSettings.PatchProperty("AndroidLldbStartupCommands");
 
 				if (NMakeOutputPath != null)
diff --git a/Engine/Source/Programs/UnrealBuildTool/ProjectFiles/ProjectFileGenerator.cs b/Engine/Source/Programs/UnrealBuildTool/ProjectFiles/ProjectFileGenerator.cs
--- a/Engine/Source/Programs/UnrealBuildTool/ProjectFiles/ProjectFileGenerator.cs
+++ b/Engine/Source/Programs/UnrealBuildTool/ProjectFiles/ProjectFileGenerator.cs
@@ -861,6 +861,32 @@ namespace UnrealBuildTool
 			return ProjectFiles;
 		}
 
+		/// <summary>
+		/// Generate a mapping for source file directories when using VFS for debugger configuration 
+		/// </summary>
+		/// <param name="UProjectFile">Path to uproject file for the target we're generating debugger config for</param>
+		/// <returns>A mapping where keys are VFS paths and values are disk paths</returns>
+		public static Dictionary<DirectoryReference, DirectoryReference> GenerateSourceMap(FileReference? UProjectFile)
+		{
+			Dictionary<DirectoryReference, DirectoryReference> SourceFileMap = new();
+			if (Unreal.OriginalCompilationRootDirectory != Unreal.RootDirectory)
+			{
+				SourceFileMap.Add(Unreal.OriginalCompilationRootDirectory, Unreal.RootDirectory);
+			}
+			CppRootPaths RootPaths = new();
+			RootPaths[CppRootPathFolder.Root] = Unreal.RootDirectory;
+			if (UProjectFile != null)
+			{
+				RootPaths[CppRootPathFolder.Project] = UProjectFile.Directory;
+			}
+
+			foreach ((uint id, DirectoryReference vfs, DirectoryReference local) entry in RootPaths)
+			{
+				SourceFileMap.Add(entry.vfs, entry.local);
+			}
+			return SourceFileMap;
+		}
+
 		/// <summary>
 		/// Gets the user's preferred IDE from their editor settings
 		/// </summary>
diff --git a/Engine/Source/Programs/UnrealBuildTool/ProjectFiles/VisualStudioCode/VSCodeProjectFileGenerator.cs b/Engine/Source/Programs/UnrealBuildTool/ProjectFiles/VisualStudioCode/VSCodeProjectFileGenerator.cs
--- a/Engine/Source/Programs/UnrealBuildTool/ProjectFiles/VisualStudioCode/VSCodeProjectFileGenerator.cs
+++ b/Engine/Source/Programs/UnrealBuildTool/ProjectFiles/VisualStudioCode/VSCodeProjectFileGenerator.cs
@@ -1397,22 +1397,7 @@ namespace UnrealBuildTool
 						OutFile.AddField("type", "lldb");
 					}
 
-					Dictionary<DirectoryReference, DirectoryReference> SourceFileMap = new();
-					if (Unreal.OriginalCompilationRootDirectory != ProjectRoot)
-					{
-						SourceFileMap.Add(Unreal.OriginalCompilationRootDirectory, ProjectRoot);
-					}
-					CppRootPaths RootPaths = new();
-					RootPaths[CppRootPathFolder.Root] = Unreal.RootDirectory;
-					if (BuildProduct.UProjectFile != null)
-					{
-						RootPaths[CppRootPathFolder.Project] = BuildProduct.UProjectFile.Directory;
-					}
-					foreach ((uint id, DirectoryReference vfs, DirectoryReference local) entry in RootPaths)
-					{
-						SourceFileMap.Add(entry.vfs, entry.local);
-					}
-
+					Dictionary<DirectoryReference, DirectoryReference> SourceFileMap = GenerateSourceMap(BuildProduct.UProjectFile);
 					if (SourceFileMap.Any())
 					{
 						OutFile.BeginObject("sourceFileMap");
