Add UnreachableCode details to GetBuildSettingsInfo
#rnx

diff --git a/Engine/Source/Programs/UnrealBuildTool/Configuration/TargetRules.cs b/Engine/Source/Programs/UnrealBuildTool/Configuration/TargetRules.cs
--- a/Engine/Source/Programs/UnrealBuildTool/Configuration/TargetRules.cs
+++ b/Engine/Source/Programs/UnrealBuildTool/Configuration/TargetRules.cs
@@ -3556,6 +3556,7 @@ namespace UnrealBuildTool
 				if (BuildSettingsVersion.V6 <= latestVersion && DefaultBuildSettings < BuildSettingsVersion.V6)
 				{
 					modifiedSettings.Add(Tuple.Create(String.Format("{0} = WarningLevel.Error", nameof(CppCompileWarningSettings.UndefinedIdentifierWarningLevel)), "Enables compile-time validation of undefined identifiers in conditional expressions. (Previously: Off)."));
+					modifiedSettings.Add(Tuple.Create(String.Format("{0} = WarningLevel.Error", nameof(CppCompileWarningSettings.UnreachableCodeWarnings)), "Enables compile-time validation of unreachable code. (Previously: Error for MSVC, Off for Clang)."));
 				}
 
 				if (modifiedSettings.Count > 0)
