UnrealBuildTool: Remove support for /DEBUG:FASTLINK it is no longer recommended.
Always use /DEBUG:GHASH when linking with lld-link

#rnx

diff --git a/Engine/Source/Programs/UnrealBuildTool/Configuration/ReadOnlyTargetRules.cs b/Engine/Source/Programs/UnrealBuildTool/Configuration/ReadOnlyTargetRules.cs
--- a/Engine/Source/Programs/UnrealBuildTool/Configuration/ReadOnlyTargetRules.cs
+++ b/Engine/Source/Programs/UnrealBuildTool/Configuration/ReadOnlyTargetRules.cs
@@ -443,6 +443,7 @@ namespace UnrealBuildTool
 
 		public bool bBreakBuildOnLicenseViolation => Inner.bBreakBuildOnLicenseViolation;
 
+		[Obsolete("Deprecated in UE5.7 - No longer recommended")]
 		public bool? bUseFastPDBLinking => Inner.bUseFastPDBLinking;
 
 		public bool bCreateMapFile => Inner.bCreateMapFile;
diff --git a/Engine/Source/Programs/UnrealBuildTool/Configuration/TargetRules.cs b/Engine/Source/Programs/UnrealBuildTool/Configuration/TargetRules.cs
--- a/Engine/Source/Programs/UnrealBuildTool/Configuration/TargetRules.cs
+++ b/Engine/Source/Programs/UnrealBuildTool/Configuration/TargetRules.cs
@@ -2282,10 +2282,10 @@ namespace UnrealBuildTool
 		public bool bBreakBuildOnLicenseViolation { get; set; } = true;
 
 		/// <summary>
-		/// Whether to use the :FASTLINK option when building with /DEBUG to create local PDBs on Windows. Fast, but currently seems to have problems finding symbols in the debugger.
+		/// Whether to use the :FASTLINK option when building with /DEBUG to create local PDBs on Windows. No longer recommended.
 		/// </summary>
-		[CommandLine("-FastPDB")]
 		[XmlConfigFile(Category = "BuildConfiguration")]
+		[Obsolete("Deprecated in UE5.7 - No longer recommended")]
 		public bool? bUseFastPDBLinking { get; set; }
 
 		/// <summary>
diff --git a/Engine/Source/Programs/UnrealBuildTool/Configuration/UEBuildTarget.cs b/Engine/Source/Programs/UnrealBuildTool/Configuration/UEBuildTarget.cs
--- a/Engine/Source/Programs/UnrealBuildTool/Configuration/UEBuildTarget.cs
+++ b/Engine/Source/Programs/UnrealBuildTool/Configuration/UEBuildTarget.cs
@@ -5742,7 +5742,6 @@ namespace UnrealBuildTool
 			GlobalLinkEnvironment.bPGOOptimize = Rules.bPGOOptimize;
 			GlobalLinkEnvironment.bPGOProfile = Rules.bPGOProfile;
 			GlobalLinkEnvironment.bUseIncrementalLinking = Rules.bUseIncrementalLinking;
-			GlobalLinkEnvironment.bUseFastPDBLinking = Rules.bUseFastPDBLinking ?? false;
 			GlobalLinkEnvironment.bDeterministic = Rules.bDeterministic;
 			GlobalLinkEnvironment.bPrintTimingInfo = Rules.bPrintToolChainTimingInfo;
 			GlobalLinkEnvironment.bUsePIE = Rules.bEnablePIE;
diff --git a/Engine/Source/Programs/UnrealBuildTool/Platform/Windows/VCToolChain.cs b/Engine/Source/Programs/UnrealBuildTool/Platform/Windows/VCToolChain.cs
--- a/Engine/Source/Programs/UnrealBuildTool/Platform/Windows/VCToolChain.cs
+++ b/Engine/Source/Programs/UnrealBuildTool/Platform/Windows/VCToolChain.cs
@@ -1500,7 +1500,7 @@ namespace UnrealBuildTool
 				
 				if (Target.WindowsPlatform.bAllowClangLinker)
 				{
-					// Enable codeview ghash for faster lld links on Clang and Intel
+					// Enable codeview ghash for faster lld links with Clang and Intel
 					Arguments.Add("-Xclang -gcodeview-ghash");
 				}
 
@@ -1532,13 +1532,9 @@ namespace UnrealBuildTool
 			Arguments.Add("/TC");
 
 			// Level 0 warnings.  Needed for external C projects that produce warnings at higher warning levels.
-			// Foricbly remove "/W4" in order to avoid excessive warning that we know will happen.
+			// Forcibly remove "/W4" in order to avoid excessive warning that we know will happen.
 			// > Command line warning D9025 : overriding '/W4' with '/W0'
-			if (Arguments.Contains("/W4"))
-			{
-				Arguments.Remove("/W4");
-			}
-
+			Arguments.Remove("/W4");
 			Arguments.Add("/W0");
 
 			// Select C Standard version available
@@ -1663,20 +1659,13 @@ namespace UnrealBuildTool
 			else if (LinkEnvironment.bCreateDebugInfo || Target.WindowsPlatform.bEnableAddressSanitizer)
 			{
 				// Output debug info for the linked executable.
-				// Beginning in Visual Studio 2017 /DEBUG defaults to /DEBUG:FASTLINK for debug builds
-				Arguments.Add("/DEBUG:FULL");
-
-				if (LinkEnvironment.bCreateDebugInfo && LinkEnvironment.bUseFastPDBLinking)
+				if (Target.WindowsPlatform.Compiler.IsClang() && Target.WindowsPlatform.bAllowClangLinker)
 				{
-					// Allow partial PDBs for faster linking on Clang and Intel
-					if (Target.WindowsPlatform.Compiler.IsClang() && Target.WindowsPlatform.bAllowClangLinker)
-					{
-						Arguments[^1] = "/DEBUG:GHASH";
-					}
-					else
-					{
-						Arguments[^1] = "/DEBUG:FASTLINK";
-					}
+					Arguments.Add("/DEBUG:GHASH");
+				}
+				else
+				{
+					Arguments.Add("/DEBUG:FULL");
 				}
 			}
 
diff --git a/Engine/Source/Programs/UnrealBuildTool/System/LinkEnvironment.cs b/Engine/Source/Programs/UnrealBuildTool/System/LinkEnvironment.cs
--- a/Engine/Source/Programs/UnrealBuildTool/System/LinkEnvironment.cs
+++ b/Engine/Source/Programs/UnrealBuildTool/System/LinkEnvironment.cs
@@ -267,11 +267,6 @@ namespace UnrealBuildTool
 		/// </summary>
 		public bool bUsePDBFiles;
 
-		/// <summary>
-		/// Whether to use the :FASTLINK option when building with /DEBUG to create local PDBs
-		/// </summary>
-		public bool bUseFastPDBLinking;
-
 		/// <summary>
 		/// Use Position Independent Executable (PIE). Has an overhead cost
 		/// </summary>
@@ -431,7 +426,6 @@ namespace UnrealBuildTool
 			PGOMergedFilenamePrefix = Other.PGOMergedFilenamePrefix;
 			bCreateMapFile = Other.bCreateMapFile;
 			bUsePDBFiles = Other.bUsePDBFiles;
-			bUseFastPDBLinking = Other.bUseFastPDBLinking;
 			bUsePIE = Other.bUsePIE;
 			bIgnoreUnresolvedSymbols = Other.bIgnoreUnresolvedSymbols;
 			bDeterministic = Other.bDeterministic;
