UnrealBuildTool: Clean up bStrictPreprocessorConformance setting as it has been force enabled for C++20 which is now the minimum version

#rnx

diff --git a/Engine/Source/Programs/UnrealBuildTool/Platform/Windows/UEBuildWindows.cs b/Engine/Source/Programs/UnrealBuildTool/Platform/Windows/UEBuildWindows.cs
--- a/Engine/Source/Programs/UnrealBuildTool/Platform/Windows/UEBuildWindows.cs
+++ b/Engine/Source/Programs/UnrealBuildTool/Platform/Windows/UEBuildWindows.cs
@@ -583,7 +583,7 @@ namespace UnrealBuildTool
 		[RequiresUniqueBuildEnvironment]
 		[XmlConfigFile(Category = "WindowsPlatform")]
 		[CommandLine("-UpdatedCPPMacro")]
-		public bool bUpdatedCPPMacro = true;
+		public bool bUpdatedCPPMacro { get; set; } = true;
 
 		/// <summary>
 		/// Enables inline conformance (Remove unreferenced COMDAT) (/Zc:inline).
@@ -591,33 +591,34 @@ namespace UnrealBuildTool
 		[RequiresUniqueBuildEnvironment]
 		[XmlConfigFile(Category = "WindowsPlatform")]
 		[CommandLine("-StrictInline")]
-		public bool bStrictInlineConformance = false;
+		public bool bStrictInlineConformance { get; set; } = false;
 
 		/// <summary>
-		/// Enables new preprocessor conformance (/Zc:preprocessor). This is always enabled for C++20 modules.
+		/// Enables preprocessor conformance (/Zc:preprocessor).
 		/// </summary>
+		[RequiresUniqueBuildEnvironment]
 		[XmlConfigFile(Category = "WindowsPlatform")]
-		[CommandLine("-StrictPreprocessor")]
-		public bool bStrictPreprocessorConformance = false;
+		[CommandLine("-NoStrictPreprocessor", Value = "false")]
+		public bool bStrictPreprocessorConformance { get; set; } = true;
 
 		/// <summary>
-		/// Enables enum types conformance (/Zc:enumTypes) in VS2022 17.4 Preview 4.0+.
+		/// Enables enum types conformance (/Zc:enumTypes).
 		/// </summary>
 		[RequiresUniqueBuildEnvironment]
 		[XmlConfigFile(Category = "WindowsPlatform")]
 		[CommandLine("-StrictEnumTypes")]
-		public bool bStrictEnumTypesConformance = false;
+		public bool bStrictEnumTypesConformance { get; set; } = false;
 
 		/// <summary>
-		/// Enables enforcing standard C++ ODR violations (/Zc:checkGwOdr) in VS2022 17.5 Preview 2.0+
+		/// Enables enforcing standard C++ ODR violations (/Zc:checkGwOdr).
 		/// </summary>
 		[RequiresUniqueBuildEnvironment]
 		[XmlConfigFile(Category = "WindowsPlatform")]
 		[CommandLine("-StrictODR")]
-		public bool bStrictODRViolationConformance = false;
+		public bool bStrictODRViolationConformance { get; set; } = false;
 
 		/// <summary>
-		/// Volatile Metadata is enabled by default and improves x64 emulation on arm64, but may come at a small perfomance cost (/volatileMetadata-).
+		/// Volatile Metadata is enabled by default and improves x64 emulation on arm64, but may come at a small performance cost (/volatileMetadata-).
 		/// </summary>
 		[RequiresUniqueBuildEnvironment]
 		[XmlConfigFile(Category = "WindowsPlatform")]
diff --git a/Engine/Source/Programs/UnrealBuildTool/Platform/Windows/VCToolChain.cs b/Engine/Source/Programs/UnrealBuildTool/Platform/Windows/VCToolChain.cs
--- a/Engine/Source/Programs/UnrealBuildTool/Platform/Windows/VCToolChain.cs
+++ b/Engine/Source/Programs/UnrealBuildTool/Platform/Windows/VCToolChain.cs
@@ -1465,11 +1465,6 @@ namespace UnrealBuildTool
 
 			if (CompileEnvironment.CppStandard >= CppStandardVersion.Cpp20)
 			{
-				if (Target.WindowsPlatform.Compiler.IsMSVC() && !CompileEnvironment.AdditionalArguments.Contains("/Zc:preprocessor-", StringComparison.Ordinal))
-				{
-					Arguments.Add("/Zc:preprocessor");
-				}
-
 				// warning C5054: operator ___: deprecated between enumerations of different types
 				// re: http://www.open-std.org/jtc1/sc22/wg21/docs/papers/2018/p1120r0.html
 
