[UBT]
* Fixed problems with configuration manager setup when generating project files. Lots of configs had "Invalid".

diff --git a/Engine/Source/Programs/UnrealBuildTool/ProjectFiles/ProjectFileGenerator.cs b/Engine/Source/Programs/UnrealBuildTool/ProjectFiles/ProjectFileGenerator.cs
--- a/Engine/Source/Programs/UnrealBuildTool/ProjectFiles/ProjectFileGenerator.cs
+++ b/Engine/Source/Programs/UnrealBuildTool/ProjectFiles/ProjectFileGenerator.cs
@@ -2152,6 +2152,12 @@ namespace UnrealBuildTool
 			}
 			Folder.SubFolders.AddRange(SubFoldersToAdd);
 
+			// Sort subfolders to be deterministic
+			Folder.SubFolders.Sort(static (a, b) => a.FolderName.CompareTo(b.FolderName));
+
+			// Sort child projects as well while we're at it... to make .sln deterministic
+			Folder.ChildProjects.Sort(static (a, b) => a.ProjectFilePath.CompareTo(b.ProjectFilePath));
+
 			// After everything has been collapsed, do a bit of data validation
 			Validate(Folder, ParentPrimaryProjectFolderPath);
 		}
@@ -3603,6 +3609,7 @@ namespace UnrealBuildTool
 				lock (GeneratedProjectFiles)
 				{
 					GeneratedProjectFiles.Add(f);
+					AllProjectFilesModified = true;
 				}
 			}
 
@@ -3625,10 +3632,21 @@ namespace UnrealBuildTool
 		{
 			get
 			{
-				List<ProjectFile> CombinedList = new List<ProjectFile>();
-				CombinedList.AddRange(GeneratedProjectFiles);
-				CombinedList.AddRange(OtherProjectFiles);
-				return CombinedList;
+				//lock (GeneratedProjectFiles)
+				{
+					//lock (OtherProjectFiles)
+					{
+						if (AllProjectFilesModified)
+						{
+							AllProjectFilesCached.Clear();
+							AllProjectFilesCached.AddRange(GeneratedProjectFiles);
+							AllProjectFilesCached.AddRange(OtherProjectFiles);
+							AllProjectFilesCached.Sort(static (a, b) => a.ProjectFilePath.CompareTo(b.ProjectFilePath));
+							AllProjectFilesModified = false;
+						}
+						return AllProjectFilesCached;
+					}
+				}
 			}
 		}
 
@@ -3929,6 +3947,7 @@ namespace UnrealBuildTool
 			lock (OtherProjectFiles)
 			{
 				OtherProjectFiles.Add(InProject);
+				AllProjectFilesModified = true;
 			}
 		}
 
@@ -3990,6 +4009,9 @@ namespace UnrealBuildTool
 
 		/// List of top-level folders in the primary project file
 		protected PrimaryProjectFolder RootFolder;
+
+		private bool AllProjectFilesModified = true;
+		private readonly List<ProjectFile> AllProjectFilesCached = [];
 	}
 
 	/// <summary>
diff --git a/Engine/Source/Programs/UnrealBuildTool/ProjectFiles/VisualStudio/VCProject.cs b/Engine/Source/Programs/UnrealBuildTool/ProjectFiles/VisualStudio/VCProject.cs
--- a/Engine/Source/Programs/UnrealBuildTool/ProjectFiles/VisualStudio/VCProject.cs
+++ b/Engine/Source/Programs/UnrealBuildTool/ProjectFiles/VisualStudio/VCProject.cs
@@ -348,8 +348,7 @@ namespace UnrealBuildTool
 			bool bDeployByDefault = false;
 
 			// Programs are built in editor configurations (since the editor is like a desktop program too) and game configurations (since we omit the "game" qualification in the configuration name).
-			bool IsProgramProject = ProjectTargets[0].TargetRules != null && ProjectTargets[0].TargetRules!.Type == TargetType.Program;
-			if (!IsProgramProject || SolutionTarget == TargetType.Game || SolutionTarget == TargetType.Editor)
+			bool IsProgramProject = ProjectTargets.Any() && ProjectTargets[0].TargetRules != null && ProjectTargets[0].TargetRules!.Type == TargetType.Program;
 			{
 				// Get the target type we expect to find for this project
 				TargetType TargetConfigurationName = SolutionTarget;
@@ -357,6 +356,10 @@ namespace UnrealBuildTool
 				{
 					TargetConfigurationName = TargetType.Program;
 				}
+				else if (SolutionTarget == TargetType.Program)
+				{
+					TargetConfigurationName = TargetType.Game;
+				}
 
 				// Now, we want to find a target in this project that maps to the current solution config combination.  Only up to one target should
 				// and every solution config combination should map to at least one target in one project (otherwise we shouldn't have added it!).
@@ -687,12 +690,6 @@ namespace UnrealBuildTool
 							UEBuildPlatform? BuildPlatform;
 							if (UEBuildPlatform.TryGetBuildPlatform(Platform, out BuildPlatform) && (BuildPlatform.HasRequiredSDKsInstalled() == SDKStatus.Valid))
 							{
-								// Now go through all of the target types for this project
-								if (ProjectTargets.Count == 0)
-								{
-									throw new BuildException("Expecting at least one ProjectTarget to be associated with project '{0}' in the TargetProjects list ", ProjectFilePath);
-								}
-
 								foreach (ProjectTarget ProjectTarget in ProjectTargets.OfType<ProjectTarget>())
 								{
 									if (!IsValidProjectPlatformAndConfiguration(ProjectTarget, Platform, Configuration, Logger))
@@ -738,7 +735,7 @@ namespace UnrealBuildTool
 										}
 
 										TargetType TargetConfigurationType = ProjectTarget.TargetRules!.Type;
-										if (!bMakeProjectPerTarget && TargetConfigurationType != TargetType.Game)
+										if (!bMakeProjectPerTarget && TargetConfigurationType != TargetType.Game && TargetConfigurationType != TargetType.Program)
 										{
 											ProjectConfigurationName += "_" + TargetConfigurationType.ToString();
 										}
diff --git a/Engine/Source/Programs/UnrealBuildTool/ProjectFiles/VisualStudio/VCProjectFileGenerator.cs b/Engine/Source/Programs/UnrealBuildTool/ProjectFiles/VisualStudio/VCProjectFileGenerator.cs
--- a/Engine/Source/Programs/UnrealBuildTool/ProjectFiles/VisualStudio/VCProjectFileGenerator.cs
+++ b/Engine/Source/Programs/UnrealBuildTool/ProjectFiles/VisualStudio/VCProjectFileGenerator.cs
@@ -778,7 +778,6 @@ namespace UnrealBuildTool
 				}
 
 				// Project files
-				//List<MSBuildProjectFile> AllProjectFilesSorted = AllProjectFiles.OrderBy((ProjFile) => ProjFile.ProjectFilePath.GetFileNameWithoutExtension()).Cast<MSBuildProjectFile>().ToList();
 				foreach (MSBuildProjectFile CurProject in AllProjectFiles)
 				{
 					// Visual Studio uses different GUID types depending on the project type
@@ -924,8 +923,7 @@ namespace UnrealBuildTool
 						// filter)'s GUID, and Guid2 is the solution filter directory to parent the child project (or solution
 						// filter) to.  This sets up the hierarchical solution explorer tree for all solution folders and projects.
 
-						System.Action<StringBuilder /* VCSolutionFileContent */, List<PrimaryProjectFolder> /* Folders */ >? FolderProcessorFunction = null;
-						FolderProcessorFunction = (LocalVCSolutionFileContent, LocalPrimaryProjectFolders) =>
+						void FolderProcessorFunction(StringBuilder LocalVCSolutionFileContent, List<PrimaryProjectFolder> LocalPrimaryProjectFolders)
 							{
 								foreach (PrimaryProjectFolder CurFolder in LocalPrimaryProjectFolders)
 								{
@@ -951,7 +949,7 @@ namespace UnrealBuildTool
 									}
 
 									// Recurse into subfolders
-									FolderProcessorFunction!(LocalVCSolutionFileContent, CurFolder.SubFolders);
+									FolderProcessorFunction(LocalVCSolutionFileContent, CurFolder.SubFolders);
 								}
 							};
 						FolderProcessorFunction(VCSolutionFileContent, RootFolder.SubFolders);
@@ -1058,12 +1056,6 @@ namespace UnrealBuildTool
 							{
 								if (!CurProject.IsStubProject)
 								{
-									if (CurProject.ProjectTargets.Count == 0)
-									{
-										throw new BuildException("Expecting project '" + CurProject.ProjectFilePath +
-																 "' to have at least one ProjectTarget associated with it!");
-									}
-
 									// Figure out the set of valid target configuration names
 									foreach (ProjectTarget ProjectTarget in CurProject.ProjectTargets.OfType<ProjectTarget>())
 									{
