Fixed non Client configuration deployment in Visual Studio

#android

diff --git a/Engine/Source/Programs/UnrealBuildTool/Platform/Android/AndroidProjectGenerator.cs b/Engine/Source/Programs/UnrealBuildTool/Platform/Android/AndroidProjectGenerator.cs
--- a/Engine/Source/Programs/UnrealBuildTool/Platform/Android/AndroidProjectGenerator.cs
+++ b/Engine/Source/Programs/UnrealBuildTool/Platform/Android/AndroidProjectGenerator.cs
@@ -221,9 +221,10 @@ namespace UnrealBuildTool
 					// so regardless of what value it has, Split(':')[1] will never throw an exception.
 					if (UProjectPath.Length > 0)
 					{
+						string ClientParam = InTargetRules.Type == TargetRules.TargetType.Client ? "-client" : "";
 						string GetTargetDeviceSerial = "$([System.String]::Concat($(AndroidDebugTarget), \"::\").Split(':')[1])";
 						Out.AppendLine(
-							$"    <AndroidPostApkInstallCommands>CALL \"{UatPath}\" BuildCookRun -project=\"{UProjectPath.Replace("\"", "")}\" -FastIterate -FromMsBuild -client -skipbuild -skipcook -skipstage -deploy -config={InVSSettings.Configuration} -platform={InVSSettings.Platform} -clientarchitecture={ShortAbi} -device=\"{GetTargetDeviceSerial}\";$(AndroidPostApkInstallCommands)</AndroidPostApkInstallCommands>");
+							$"    <AndroidPostApkInstallCommands>CALL \"{UatPath}\" BuildCookRun -project=\"{UProjectPath.Replace("\"", "")}\" -FastIterate -FromMsBuild {ClientParam} -skipbuild -skipcook -skipstage -deploy -config={InVSSettings.Configuration} -platform={InVSSettings.Platform} -clientarchitecture={ShortAbi} -device=\"{GetTargetDeviceSerial}\";$(AndroidPostApkInstallCommands)</AndroidPostApkInstallCommands>");
 					}
 					else
 					{
