UnrealBuildTool: Ensure module intermediate Gen\ files are added as BuildProducts so they are added to the binary install

#rnx

diff --git a/Engine/Source/Programs/UnrealBuildTool/Configuration/UEBuildBinary.cs b/Engine/Source/Programs/UnrealBuildTool/Configuration/UEBuildBinary.cs
--- a/Engine/Source/Programs/UnrealBuildTool/Configuration/UEBuildBinary.cs
+++ b/Engine/Source/Programs/UnrealBuildTool/Configuration/UEBuildBinary.cs
@@ -497,47 +497,9 @@ namespace UnrealBuildTool
 		public void GetBuildProducts(ReadOnlyTargetRules Target, UEToolChain ToolChain, Dictionary<FileReference, BuildProductType> BuildProducts, bool bCreateDebugInfo)
 		{
 			// Add all the precompiled outputs
-			foreach (UEBuildModuleCPP Module in Modules.OfType<UEBuildModuleCPP>())
+			foreach (UEBuildModule Module in Modules)
 			{
-				if (Module.Rules.bPrecompile)
-				{
-					if (Module.GeneratedCodeDirectory != null && DirectoryReference.Exists(Module.GeneratedCodeDirectory))
-					{
-						foreach (FileReference GeneratedCodeFile in DirectoryReference.EnumerateFiles(Module.GeneratedCodeDirectory, "*", SearchOption.AllDirectories))
-						{
-							// Exclude timestamp files, since they're always updated and cause collisions between builds
-							if (!GeneratedCodeFile.GetFileName().Equals("Timestamp", StringComparison.OrdinalIgnoreCase) && !GeneratedCodeFile.HasExtension(".cpp"))
-							{
-								BuildProducts.Add(GeneratedCodeFile, BuildProductType.BuildResource);
-							}
-						}
-					}
-					foreach (FileItem NatvisSourceFile in Module.NatvisFiles)
-					{
-						FileItem? Item = ToolChain.LinkDebuggerVisualizer(NatvisSourceFile, Module.IntermediateDirectory);
-						if (Item != null)
-						{
-							BuildProducts.Add(Item.Location, BuildProductType.BuildResource);
-						}
-					}
-					if (Target.LinkType == TargetLinkType.Monolithic)
-					{
-						FileReference PrecompiledManifestLocation = Module.PrecompiledManifestLocation;
-						BuildProducts.Add(PrecompiledManifestLocation, BuildProductType.BuildResource);
-
-						if (FileReference.Exists(PrecompiledManifestLocation)) // May be single file compile; skipped modulePrecompiledManifestLocation{
-						{
-							PrecompiledManifest ModuleManifest = PrecompiledManifest.Read(PrecompiledManifestLocation);
-							foreach (FileReference OutputFile in ModuleManifest.OutputFiles)
-							{
-								if (!BuildProducts.ContainsKey(OutputFile))
-								{
-									BuildProducts.Add(OutputFile, BuildProductType.BuildResource);
-								}
-							}
-						}
-					}
-				}
+				Module.GetBuildProducts(ToolChain, BuildProducts);
 			}
 
 			// Add all the binary outputs
diff --git a/Engine/Source/Programs/UnrealBuildTool/Configuration/UEBuildModule.cs b/Engine/Source/Programs/UnrealBuildTool/Configuration/UEBuildModule.cs
--- a/Engine/Source/Programs/UnrealBuildTool/Configuration/UEBuildModule.cs
+++ b/Engine/Source/Programs/UnrealBuildTool/Configuration/UEBuildModule.cs
@@ -450,6 +450,34 @@ namespace UnrealBuildTool
 			return Rules.CircularlyReferencedDependentModules.Contains(Module.Name);
 		}
 
+		/// <summary>
+		/// Add precompile build products from this module to the target
+		/// </summary>
+		/// <param name="ToolChain"></param>
+		/// <param name="BuildProducts"></param>
+		public virtual void GetBuildProducts(UEToolChain ToolChain, Dictionary<FileReference, BuildProductType> BuildProducts)
+		{
+			if (Rules.bPrecompile)
+			{
+				foreach (FileItem NatvisSourceFile in NatvisFiles)
+				{
+					FileItem? Item = ToolChain.LinkDebuggerVisualizer(NatvisSourceFile, IntermediateDirectory);
+					if (Item != null)
+					{
+						BuildProducts.Add(Item.Location, BuildProductType.BuildResource);
+					}
+				}
+				foreach (string filename in Rules.FilesToGenerate.Keys)
+				{
+					FileReference generatedItem = FileReference.Combine(IntermediateDirectory, "Gen", filename);
+					if (generatedItem.HasExtension(".h"))
+					{
+						BuildProducts.Add(generatedItem, BuildProductType.BuildResource);
+					}
+				}
+			}
+		}
+
 		/// <summary>
 		/// Enumerates additional build products which may be produced by this module. Some platforms (eg. Mac, Linux) can link directly against .so/.dylibs, but they 
 		/// are also copied to the output folder by the toolchain.
diff --git a/Engine/Source/Programs/UnrealBuildTool/Configuration/UEBuildModuleCPP.cs b/Engine/Source/Programs/UnrealBuildTool/Configuration/UEBuildModuleCPP.cs
--- a/Engine/Source/Programs/UnrealBuildTool/Configuration/UEBuildModuleCPP.cs
+++ b/Engine/Source/Programs/UnrealBuildTool/Configuration/UEBuildModuleCPP.cs
@@ -118,6 +118,48 @@ namespace UnrealBuildTool
 		/// </summary>
 		public override DirectoryReference? VerseDirectory => AssociatedVerseDirectory;
 
+		/// <inheritdoc/>
+		public override void GetBuildProducts(UEToolChain ToolChain, Dictionary<FileReference, BuildProductType> BuildProducts)
+		{
+			base.GetBuildProducts(ToolChain, BuildProducts);
+
+			if (Rules.bPrecompile)
+			{
+				if (GeneratedCodeDirectory != null && DirectoryReference.Exists(GeneratedCodeDirectory))
+				{
+					foreach (FileReference GeneratedCodeFile in DirectoryReference.EnumerateFiles(GeneratedCodeDirectory, "*", SearchOption.AllDirectories))
+					{
+						// Exclude timestamp files, since they're always updated and cause collisions between builds
+						if (!GeneratedCodeFile.GetFileName().Equals("Timestamp", StringComparison.OrdinalIgnoreCase) && !GeneratedCodeFile.HasExtension(".cpp"))
+						{
+							BuildProducts.Add(GeneratedCodeFile, BuildProductType.BuildResource);
+						}
+					}
+				}
+				foreach (string filename in Rules.FilesToGenerate.Keys)
+				{
+					FileReference generatedItem = FileReference.Combine(IntermediateDirectory, "Gen", filename);
+					if (generatedItem.HasExtension(".h"))
+					{
+						BuildProducts.Add(generatedItem, BuildProductType.BuildResource);
+					}
+				}
+				if (Rules.Target.LinkType == TargetLinkType.Monolithic)
+				{
+					BuildProducts.Add(PrecompiledManifestLocation, BuildProductType.BuildResource);
+
+					if (FileReference.Exists(PrecompiledManifestLocation)) // May be single file compile; skipped modulePrecompiledManifestLocation{
+					{
+						PrecompiledManifest ModuleManifest = PrecompiledManifest.Read(PrecompiledManifestLocation);
+						foreach (FileReference OutputFile in ModuleManifest.OutputFiles)
+						{
+							BuildProducts.TryAdd(OutputFile, BuildProductType.BuildResource);
+						}
+					}
+				}
+			}
+		}
+
 		protected override void GetReferencedDirectories(HashSet<DirectoryReference> Directories)
 		{
 			base.GetReferencedDirectories(Directories);
@@ -1652,7 +1694,7 @@ namespace UnrealBuildTool
 		/// </summary>
 		/// <param name="ModuleCompileEnvironment">The module compile environment</param>
 		/// <param name="CompileEnvironment">The shared PCH compile environment</param>
-		/// <returns>True if the two compile enviroments are compatible</returns>
+		/// <returns>True if the two compile environments are compatible</returns>
 		internal static bool IsCompatibleForSharedPCH(CppCompileEnvironment ModuleCompileEnvironment, CppCompileEnvironment CompileEnvironment)
 		{
 			if (ModuleCompileEnvironment.bTreatAsEngineModule != CompileEnvironment.bTreatAsEngineModule)
@@ -2070,8 +2112,8 @@ namespace UnrealBuildTool
 						FileItem PrivateDefinitionsFileItem;
 						{
 							// Only add new definitions that are not already existing in the shared pch
-							List<string> NewDefinitions = new();
-							List<string> Undefintions = new();
+							List<string> NewDefinitions = [];
+							List<string> Undefinitions = [];
 							bool ModuleApiUndef = false;
 							bool ModuleNonAttributedUndef = false;
 							foreach (string Definition in CompileEnvironment.Definitions)
@@ -2087,19 +2129,19 @@ namespace UnrealBuildTool
 								if (!ModuleApiUndef && Definition.StartsWith(ModuleApiDefine, StringComparison.Ordinal))
 								{
 									ModuleApiUndef = true;
-									Undefintions.Add(ModuleApiDefine);
+									Undefinitions.Add(ModuleApiDefine);
 								}
 
 								// Remove the module _NON_ATTRIBUTED_API definition for cases where there are circular dependencies between the shared PCH module and modules using it
 								if (!ModuleNonAttributedUndef && Definition.StartsWith(ModuleNonAttributedDefine, StringComparison.Ordinal))
 								{
 									ModuleNonAttributedUndef = true;
-									Undefintions.Add(ModuleNonAttributedDefine);
+									Undefinitions.Add(ModuleNonAttributedDefine);
 								}
 							}
 
 							StringBuilder Writer = new();
-							WriteDefinitions($"Shared PCH Definitions for {Name}", new string[] { Instance.DefinitionsFile.Name }, NewDefinitions, Undefintions, Writer);
+							WriteDefinitions($"Shared PCH Definitions for {Name}", new string[] { Instance.DefinitionsFile.Name }, NewDefinitions, Undefinitions, Writer);
 							PrivateDefinitionsFileItem = Graph.CreateIntermediateTextFile(PrivateDefinitionsFile, Writer.ToString());
 						}
 
