UnrealBuildTool: Missed file that wasn't saved in last cl when I moved some code

#rnx

diff --git a/Engine/Source/Programs/UnrealBuildTool/Configuration/UEBuildModule.cs b/Engine/Source/Programs/UnrealBuildTool/Configuration/UEBuildModule.cs
--- a/Engine/Source/Programs/UnrealBuildTool/Configuration/UEBuildModule.cs
+++ b/Engine/Source/Programs/UnrealBuildTool/Configuration/UEBuildModule.cs
@@ -467,14 +467,6 @@ namespace UnrealBuildTool
 						BuildProducts.Add(Item.Location, BuildProductType.BuildResource);
 					}
 				}
-				foreach (string filename in Rules.FilesToGenerate.Keys)
-				{
-					FileReference generatedItem = FileReference.Combine(IntermediateDirectory, "Gen", filename);
-					if (generatedItem.HasExtension(".h"))
-					{
-						BuildProducts.Add(generatedItem, BuildProductType.BuildResource);
-					}
-				}
 			}
 		}
 
