UnrealBuildTool: Fix unhandled exception handling during static initialization

#rnx

diff --git a/Engine/Build/BatchFiles/RunUBT.bat b/Engine/Build/BatchFiles/RunUBT.bat
--- a/Engine/Build/BatchFiles/RunUBT.bat
+++ b/Engine/Build/BatchFiles/RunUBT.bat
@@ -27,7 +27,7 @@ if errorlevel 1 goto Error_UBTCompileFailed
 rem ## Run UnrealBuildTool to generate Visual Studio solution and project files
 rem ## NOTE: We also pass along any arguments here
 dotnet ..\Binaries\DotNET\UnrealBuildTool\UnrealBuildTool.dll %*
-if errorlevel 1 goto Error_UBTRunFailed
+if not %ERRORLEVEL% == 0 goto Error_UBTRunFailed
 
 rem ## Success!
 popd
diff --git a/Engine/Source/Programs/UnrealBuildTool/UnrealBuildTool.cs b/Engine/Source/Programs/UnrealBuildTool/UnrealBuildTool.cs
--- a/Engine/Source/Programs/UnrealBuildTool/UnrealBuildTool.cs
+++ b/Engine/Source/Programs/UnrealBuildTool/UnrealBuildTool.cs
@@ -52,19 +52,19 @@ namespace UnrealBuildTool
 		/// The full name of the Engine/Source directory
 		/// </summary>
 		[Obsolete("Replace with Unreal.EngineSourceDirectory")]
-		public static readonly DirectoryReference EngineSourceDirectory = Unreal.EngineSourceDirectory;
+		public static DirectoryReference EngineSourceDirectory => Unreal.EngineSourceDirectory;
 
 		/// <summary>
 		/// Writable engine directory. Uses the user's settings folder for installed builds.
 		/// </summary>
 		[Obsolete("Replace with Unreal.WritableEngineDirectory")]
-		public static DirectoryReference WritableEngineDirectory = Unreal.WritableEngineDirectory;
+		public static DirectoryReference WritableEngineDirectory => Unreal.WritableEngineDirectory;
 
 		/// <summary>
 		/// The engine programs directory
 		/// </summary>
 		[Obsolete("Replace with Unreal.EngineProgramSavedDirectory")]
-		public static DirectoryReference EngineProgramSavedDirectory = Unreal.EngineProgramSavedDirectory;
+		public static DirectoryReference EngineProgramSavedDirectory => Unreal.EngineProgramSavedDirectory;
 
 		/// <summary>
 		/// The original root directory that was used to compile the installed engine
