Reconcile usage of TargetRules.WindowsPlatform.bStrictInlineConformance with ModuleRules.bVcRemoveUnreferencedComdat
Clean up some obsolete properties

#rnx

diff --git a/Engine/Source/Programs/UnrealBuildTool/Configuration/TargetRules.Obsolete.cs b/Engine/Source/Programs/UnrealBuildTool/Configuration/TargetRules.Obsolete.cs
--- a/Engine/Source/Programs/UnrealBuildTool/Configuration/TargetRules.Obsolete.cs
+++ b/Engine/Source/Programs/UnrealBuildTool/Configuration/TargetRules.Obsolete.cs
@@ -71,16 +71,22 @@ namespace UnrealBuildTool
 		/// <summary>
 		/// Whether to merge module and generated unity files for faster compilation.
 		/// </summary>
-		[Obsolete("Deprecated in UE5.3; use DisableMergingModuleAndGeneratedFilesInUnityFiles instead")]
+		[Obsolete("Deprecated in UE5.3 - use DisableMergingModuleAndGeneratedFilesInUnityFiles instead")]
 		public bool bMergeModuleAndGeneratedUnityFiles { get; set; }
 
 		/// <summary>
 		/// Whether to add additional information to the unity files, such as '_of_X' in the file name. Not recommended.
 		/// </summary>
 		[XmlConfigFile(Category = "BuildConfiguration")]
-		[Obsolete("Deprecated in UE5.6")]
+		[Obsolete("Deprecated in UE5.6 - Caused unnecessary rebuilds when generated unity files change")]
 		public bool bDetailedUnityFiles { get; set; }
 
+		/// <summary>
+		/// Whether to use the :FASTLINK option when building with /DEBUG to create local PDBs on Windows. No longer recommended.
+		/// </summary>
+		[Obsolete("Deprecated in UE5.7 - No longer recommended")]
+		public bool? bUseFastPDBLinking { get; set; }
+
 		/// <summary>
 		/// Whether to globally disable debug info generation; Obsolete, please use TargetRules.DebugInfoMode instead
 		/// </summary>
@@ -92,6 +98,12 @@ namespace UnrealBuildTool
 			set => DebugInfo = value ? DebugInfo = DebugInfoMode.None : DebugInfo = DebugInfoMode.Full;
 		}
 
+		/// <summary>
+		/// Generate dependency files by preprocessing. Obsolete.
+		/// </summary>
+		[Obsolete("Deprecated in UE5.6 - bPreprocessDepends is no longer supported")]
+		public bool bPreprocessDepends { get; set; }
+
 		/// <summary>
 		/// Whether to enable support for C++20 coroutines
 		/// This option is provided to facilitate evaluation of the feature.
@@ -109,7 +121,7 @@ namespace UnrealBuildTool
 		/// <summary>
 		/// Disable supports for inlining gen.cpps
 		/// </summary>
-		[Obsolete("Deprecated in UE5.3; the engine code relies on inlining generated cpp files to compile")]
+		[Obsolete("Deprecated in UE5.3 - the engine code relies on inlining generated cpp files to compile")]
 		public bool bDisableInliningGenCpps { get; set; }
 
 		/// <summary>
diff --git a/Engine/Source/Programs/UnrealBuildTool/Configuration/TargetRules.cs b/Engine/Source/Programs/UnrealBuildTool/Configuration/TargetRules.cs
--- a/Engine/Source/Programs/UnrealBuildTool/Configuration/TargetRules.cs
+++ b/Engine/Source/Programs/UnrealBuildTool/Configuration/TargetRules.cs
@@ -154,8 +154,8 @@ namespace UnrealBuildTool
 
 		/// <summary>
 		/// New defaults for 5.7:
-		/// * DanglingWarningLevel = WarningLevel.Warning
-		/// * UndefinedIdentifierWarningLevel = WarningLevel.Error
+		/// * ModuleRules.CppCompileWarningSettings.DanglingWarningLevel = WarningLevel.Warning
+		/// * ModuleRules.CppCompileWarningSettings.UndefinedIdentifierWarningLevel = WarningLevel.Error
 		/// </summary>
 		V6,
 
@@ -2062,14 +2062,6 @@ namespace UnrealBuildTool
 		[CommandLine("-Preprocess")]
 		public bool bPreprocessOnly { get; set; }
 
-		/// <summary>
-		/// Generate dependency files by preprocessing. This is only recommended when distributing builds as it adds additional overhead.
-		/// </summary>
-		[CommandLine("-PreprocessDepends")]
-		[XmlConfigFile(Category = "BuildConfiguration")]
-		[Obsolete("bPreprocessDepends is no longer supported")]
-		public bool bPreprocessDepends { get; set; }
-
 		/// <summary>
 		/// Whether to generate assembly data while compiling this target. Works exclusively on MSVC for now.
 		/// </summary>
@@ -2281,13 +2273,6 @@ namespace UnrealBuildTool
 		[XmlConfigFile(Category = "BuildConfiguration")]
 		public bool bBreakBuildOnLicenseViolation { get; set; } = true;
 
-		/// <summary>
-		/// Whether to use the :FASTLINK option when building with /DEBUG to create local PDBs on Windows. No longer recommended.
-		/// </summary>
-		[XmlConfigFile(Category = "BuildConfiguration")]
-		[Obsolete("Deprecated in UE5.7 - No longer recommended")]
-		public bool? bUseFastPDBLinking { get; set; }
-
 		/// <summary>
 		/// Outputs a map file as part of the build.
 		/// </summary>
diff --git a/Engine/Source/Programs/UnrealBuildTool/Platform/Windows/UEBuildWindows.cs b/Engine/Source/Programs/UnrealBuildTool/Platform/Windows/UEBuildWindows.cs
--- a/Engine/Source/Programs/UnrealBuildTool/Platform/Windows/UEBuildWindows.cs
+++ b/Engine/Source/Programs/UnrealBuildTool/Platform/Windows/UEBuildWindows.cs
@@ -590,8 +590,8 @@ namespace UnrealBuildTool
 		/// </summary>
 		[RequiresUniqueBuildEnvironment]
 		[XmlConfigFile(Category = "WindowsPlatform")]
-		[CommandLine("-StrictInline")]
-		public bool bStrictInlineConformance { get; set; } = false;
+		[CommandLine("-NoStrictInline", Value = "false")]
+		public bool bStrictInlineConformance { get; set; } = true;
 
 		/// <summary>
 		/// Enables preprocessor conformance (/Zc:preprocessor).
diff --git a/Engine/Source/Programs/UnrealBuildTool/Platform/Windows/VCToolChain.cs b/Engine/Source/Programs/UnrealBuildTool/Platform/Windows/VCToolChain.cs
--- a/Engine/Source/Programs/UnrealBuildTool/Platform/Windows/VCToolChain.cs
+++ b/Engine/Source/Programs/UnrealBuildTool/Platform/Windows/VCToolChain.cs
@@ -458,17 +458,6 @@ namespace UnrealBuildTool
 				Arguments.Add("/d2ssa-cfg-question-");
 			}
 
-			if (CompileEnvironment.bVcRemoveUnreferencedComdat)
-			{
-				// Suppress generation of object code for unreferenced inline functions. Enabling this option is more standards compliant, and causes a big reduction
-				// in object file sizes (and link times) due to the amount of stuff we inline.
-				Arguments.Add("/Zc:inline");
-			}
-			else
-			{
-				Arguments.Add("/Zc:inline-");
-			}
-
 			// @todo clang: Clang on Windows doesn't respect "#pragma warning (error: ####)", and we're not passing "/WX", so warnings are not
 			// treated as errors when compiling on Windows using Clang right now.
 
@@ -699,45 +688,47 @@ namespace UnrealBuildTool
 			Arguments.Add("/wd4819");
 
 			// Disable Microsoft extensions on VS2017+ for improved standards compliance.
-			if (Target.WindowsPlatform.Compiler.IsMSVC())
+			if (Target.WindowsPlatform.bStrictConformanceMode)
 			{
-				if (Target.WindowsPlatform.bStrictConformanceMode)
-				{
-					// This define is needed to ensure that MSVC static analysis mode doesn't declare attributes that are incompatible with strict conformance mode
-					AddDefinition(Arguments, "SAL_NO_ATTRIBUTE_DECLARATIONS=1");
+				// This define is needed to ensure that MSVC static analysis mode doesn't declare attributes that are incompatible with strict conformance mode
+				AddDefinition(Arguments, "SAL_NO_ATTRIBUTE_DECLARATIONS=1");
 
-					Arguments.Add("/permissive-");
-					Arguments.Add("/Zc:strictStrings-"); // Have to disable strict const char* semantics due to Windows headers not being compliant.
-				}
-				else
+				Arguments.Add("/permissive-");
+				if (Target.WindowsPlatform.Compiler.IsMSVC())
 				{
-					Arguments.Add("/Zc:hiddenFriend");
+					Arguments.Add("/Zc:strictStrings-"); // Have to disable strict const char* semantics due to Windows headers not being compliant.
 				}
+			}
+			else if (Target.WindowsPlatform.Compiler.IsMSVC())
+			{
+				Arguments.Add("/Zc:hiddenFriend");
+			}
 
-				if (Target.WindowsPlatform.bUpdatedCPPMacro)
-				{
-					Arguments.Add("/Zc:__cplusplus");
-				}
+			if (Target.WindowsPlatform.bUpdatedCPPMacro)
+			{
+				Arguments.Add("/Zc:__cplusplus");
+			}
 
-				if (Target.WindowsPlatform.bStrictInlineConformance)
-				{
-					Arguments.Add("/Zc:inline");
-				}
+			if (CompileEnvironment.bVcRemoveUnreferencedComdat)
+			{
+				// Suppress generation of object code for unreferenced inline functions. Enabling this option is more standards compliant, and causes a big reduction
+				// in object file sizes (and link times) due to the amount of stuff we inline.
+				Arguments.Add("/Zc:inline");
+			}
 
-				if (Target.WindowsPlatform.bStrictPreprocessorConformance)
-				{
-					Arguments.Add("/Zc:preprocessor");
-				}
+			if (Target.WindowsPlatform.bStrictPreprocessorConformance && Target.WindowsPlatform.Compiler.IsMSVC())
+			{
+				Arguments.Add("/Zc:preprocessor");
+			}
 
-				if (Target.WindowsPlatform.bStrictEnumTypesConformance)
-				{
-					Arguments.Add("/Zc:enumTypes");
-				}
+			if (Target.WindowsPlatform.bStrictEnumTypesConformance && Target.WindowsPlatform.Compiler.IsMSVC())
+			{
+				Arguments.Add("/Zc:enumTypes");
+			}
 
-				if (Target.WindowsPlatform.bStrictODRViolationConformance && Target.WindowsPlatform.bOptimizeGlobalData)
-				{
-					Arguments.Add("/Zc:checkGwOdr");
-				}
+			if (Target.WindowsPlatform.bStrictODRViolationConformance && Target.WindowsPlatform.bOptimizeGlobalData && Target.WindowsPlatform.Compiler.IsMSVC())
+			{
+				Arguments.Add("/Zc:checkGwOdr");
 			}
 
 			// @todo HoloLens: UE is non-compliant when it comes to use of %s and %S
diff --git a/Engine/Source/Programs/UnrealBuildTool/ProjectFiles/VisualStudio/VCProject.cs b/Engine/Source/Programs/UnrealBuildTool/ProjectFiles/VisualStudio/VCProject.cs
--- a/Engine/Source/Programs/UnrealBuildTool/ProjectFiles/VisualStudio/VCProject.cs
+++ b/Engine/Source/Programs/UnrealBuildTool/ProjectFiles/VisualStudio/VCProject.cs
@@ -593,19 +593,21 @@ namespace UnrealBuildTool
 
 		private string GetConformanceCompileArguments(TargetRules Target)
 		{
-			if (Target.Platform.IsInGroup(UnrealPlatformGroup.Microsoft) && Target.WindowsPlatform.Compiler.IsMSVC())
+			if (Target.Platform.IsInGroup(UnrealPlatformGroup.Microsoft))
 			{
-				VersionNumber CompilerVersion = Target.WindowsPlatform.Environment?.CompilerVersion ?? new VersionNumber(0);
-				List<string> Arguments = new();
+				List<string> Arguments = [];
 				if (Target.WindowsPlatform.bStrictConformanceMode)
 				{
 					// This define is needed to ensure that MSVC static analysis mode doesn't declare attributes that are incompatible with strict conformance mode
 					Arguments.Add("/DSAL_NO_ATTRIBUTE_DECLARATIONS=1");
 
 					Arguments.Add("/permissive-");
-					Arguments.Add("/Zc:strictStrings-"); // Have to disable strict const char* semantics due to Windows headers not being compliant.
+					if (Target.WindowsPlatform.Compiler.IsMSVC())
+					{
+						Arguments.Add("/Zc:strictStrings-"); // Have to disable strict const char* semantics due to Windows headers not being compliant.
+					}
 				}
-				else
+				else if (Target.WindowsPlatform.Compiler.IsMSVC())
 				{
 					Arguments.Add("/Zc:hiddenFriend");
 				}
@@ -620,15 +622,20 @@ namespace UnrealBuildTool
 					Arguments.Add("/Zc:inline");
 				}
 
-				if (Target.WindowsPlatform.bStrictPreprocessorConformance)
+				if (Target.WindowsPlatform.bStrictPreprocessorConformance && Target.WindowsPlatform.Compiler.IsMSVC())
 				{
 					Arguments.Add("/Zc:preprocessor");
 				}
 
-				if (Target.WindowsPlatform.bStrictEnumTypesConformance)
+				if (Target.WindowsPlatform.bStrictEnumTypesConformance && Target.WindowsPlatform.Compiler.IsMSVC())
 				{
 					Arguments.Add("/Zc:enumTypes");
 				}
+
+				if (Target.WindowsPlatform.bStrictODRViolationConformance && Target.WindowsPlatform.bOptimizeGlobalData && Target.WindowsPlatform.Compiler.IsMSVC())
+				{
+					Arguments.Add("/Zc:checkGwOdr");
+				}
 				return String.Join(' ', Arguments);
 			}
 			return String.Empty;
