Support alternate app icon sets for iOS.
Two new ini IOSRuntimeSettings:
bAssetCatalogAllIcons : usually set to true if using alternate app icons
AssetCatalogAlternates : a string of alternate app icons, ie: "Variant1 Variant2"

Then need to add new AppIcon appiconset to Assets.xcassets matching the "AssetCatalogAlternates" name.

diff --git a/Engine/Source/Programs/UnrealBuildTool/ProjectFiles/Xcode/XcodeProject.cs b/Engine/Source/Programs/UnrealBuildTool/ProjectFiles/Xcode/XcodeProject.cs
--- a/Engine/Source/Programs/UnrealBuildTool/ProjectFiles/Xcode/XcodeProject.cs
+++ b/Engine/Source/Programs/UnrealBuildTool/ProjectFiles/Xcode/XcodeProject.cs
@@ -2147,9 +2147,20 @@ namespace UnrealBuildTool.XcodeProjectXcconfig
 			Xcconfig.AppendLine($"INFOPLIST_FILE = {PlistMetadata.XcodeProjectRelative}");
 			if (PlistMetadata.Mode == MetadataMode.UpdateTemplate)
 			{
+				bool bAssetCatalogAllIcons = false;
+				String? AssetCatalogAlternates = null;
+				PlatformIni.TryGetValue("/Script/IOSRuntimeSettings.IOSRuntimeSettings", "bAssetCatalogAllIcons", out bAssetCatalogAllIcons);
+				PlatformIni.TryGetValue("/Script/IOSRuntimeSettings.IOSRuntimeSettings", "AssetCatalogAlternates", out AssetCatalogAlternates);
+
 				// allow Xcode to generate the final plist file from our input, some INFOPLIST settings and other settings 
 				Xcconfig.AppendLine("GENERATE_INFOPLIST_FILE = YES");
 				Xcconfig.AppendLine($"ASSETCATALOG_COMPILER_APPICON_NAME = AppIcon");
+				String AllIcons = bAssetCatalogAllIcons ? "YES" : "NO";
+				Xcconfig.AppendLine($"ASSETCATALOG_COMPILER_INCLUDE_ALL_APPICON_ASSETS = {AllIcons}");
+				if (!String.IsNullOrEmpty(AssetCatalogAlternates))
+				{
+					Xcconfig.AppendLine($"ASSETCATALOG_COMPILER_ALTERNATE_APPICON_NAMES = {AssetCatalogAlternates}");
+				}
 				Xcconfig.AppendLine($"CURRENT_PROJECT_VERSION = $(UE_{Platform.ToString().ToUpper()}_BUILD_VERSION)");
 				Xcconfig.AppendLine($"MARKETING_VERSION = {MarketingVersion}");
 				Xcconfig.AppendLine($"INFOPLIST_KEY_LSApplicationCategoryType = {AppCategory}");
