UnrealBuildTool: libmmd.dll is required for Intel compiler
Intel compiler only supports x64

#rnx

diff --git a/Engine/Source/Programs/UnrealBuildTool/Platform/Windows/MicrosoftPlatformSDK.cs b/Engine/Source/Programs/UnrealBuildTool/Platform/Windows/MicrosoftPlatformSDK.cs
--- a/Engine/Source/Programs/UnrealBuildTool/Platform/Windows/MicrosoftPlatformSDK.cs
+++ b/Engine/Source/Programs/UnrealBuildTool/Platform/Windows/MicrosoftPlatformSDK.cs
@@ -1271,8 +1271,6 @@ namespace UnrealBuildTool
 
 				Logger.LogDebug("Found Intel OneAPI toolchain: {ToolChainDir} (Version={Version}, HostArchitecture={HostArchitecture}, Rank={Rank}, Error={Error})", ToolChainDir, Version, Is64Bit, Rank, Error != null);
 				ToolChains.Add(new ToolChainInstallation(Family, Rank, Version, HostArchitecture, WindowsCompilerChannel.Any, UnrealArch.X64, Error, ToolChainDir, null, IsAutoSdk));
-				ToolChains.Add(new ToolChainInstallation(Family, Rank, Version, HostArchitecture, WindowsCompilerChannel.Any, UnrealArch.Arm64, Error, ToolChainDir, null, IsAutoSdk));
-				ToolChains.Add(new ToolChainInstallation(Family, Rank, Version, HostArchitecture, WindowsCompilerChannel.Any, UnrealArch.Arm64ec, Error, ToolChainDir, null, IsAutoSdk));
 			}
 		}
 
diff --git a/Engine/Source/Programs/UnrealBuildTool/Platform/Windows/UEBuildWindows.cs b/Engine/Source/Programs/UnrealBuildTool/Platform/Windows/UEBuildWindows.cs
--- a/Engine/Source/Programs/UnrealBuildTool/Platform/Windows/UEBuildWindows.cs
+++ b/Engine/Source/Programs/UnrealBuildTool/Platform/Windows/UEBuildWindows.cs
@@ -1256,6 +1256,12 @@ namespace UnrealBuildTool
 				Target.bCompileCEF3 = false;
 			}
 
+			// Intel compiler only supports x64
+			if (Target.WindowsPlatform.Compiler.IsIntel() && !Target.Architecture.bIsX64)
+			{
+				Target.WindowsPlatform.Compiler = WindowsCompiler.Default;
+			}
+
 			if (Target.WindowsPlatform.bEnableInstrumentation) 
 			{
 				Target.WindowsPlatform.Compiler = WindowsCompiler.ClangInstrument;
@@ -1502,15 +1508,15 @@ namespace UnrealBuildTool
 				{
 					if (Compiler != PreferredCompilers.First())
 					{
-						Logger.Log(bSkipWarning || MicrosoftPlatformSDK.IgnoreToolchainErrors ? LogLevel.Debug : LogLevel.Warning, "{Target} preferred compiler {Preferred} not available, compiler set to {Compiler}", TargetRules?.Name ?? "Default", PreferredCompilers.First(), Compiler);
+						Logger.Log(bSkipWarning || MicrosoftPlatformSDK.IgnoreToolchainErrors ? LogLevel.Debug : LogLevel.Warning, "{Target} {Platform} {Architecture} preferred compiler {Preferred} not available, compiler set to {Compiler}", TargetRules?.Name ?? "Default", TargetRules?.Platform ?? UnrealTargetPlatform.Win64, Architecture, PreferredCompilers.First(), Compiler);
 					}
 					else
 					{
-						Logger.LogDebug("{Target} Compiler set to preferred compiler {Compiler}", TargetRules?.Name ?? "Default", Compiler);
+						Logger.LogDebug("{Target} {Platform} {Architecture} Compiler set to preferred compiler {Compiler}", TargetRules?.Name ?? "Default", TargetRules?.Platform ?? UnrealTargetPlatform.Win64, Architecture, Compiler);
 					}
 					return (WindowsCompiler)Compiler;
 				}
-				Logger.Log(bSkipWarning || MicrosoftPlatformSDK.IgnoreToolchainErrors ? LogLevel.Debug : LogLevel.Warning, "{Target} No preferred compilers available [{Preferred}]", TargetRules?.Name ?? "Default", String.Join(',', PreferredCompilers));
+				Logger.Log(bSkipWarning || MicrosoftPlatformSDK.IgnoreToolchainErrors ? LogLevel.Debug : LogLevel.Warning, "{Target} {Platform} {Architecture} No preferred compilers available [{Preferred}]", TargetRules?.Name ?? "Default", TargetRules?.Platform ?? UnrealTargetPlatform.Win64, Architecture, String.Join(',', PreferredCompilers));
 			}
 
 			return GetDefaultToolchain(TargetRules, Architecture, Logger, bSkipWarning);
diff --git a/Engine/Source/Programs/UnrealBuildTool/Platform/Windows/VCToolChain.cs b/Engine/Source/Programs/UnrealBuildTool/Platform/Windows/VCToolChain.cs
--- a/Engine/Source/Programs/UnrealBuildTool/Platform/Windows/VCToolChain.cs
+++ b/Engine/Source/Programs/UnrealBuildTool/Platform/Windows/VCToolChain.cs
@@ -2682,10 +2682,16 @@ namespace UnrealBuildTool
 				TargetFileToSourceFile[DstFile] = SrcFile;
 				RuntimeDependencies.Add(new RuntimeDependency(DstFile, StagedFileType.NonUFS));
 			}
-		}
-
 
+			if (Target.WindowsPlatform.Compiler.IsIntel())
+			{
+				FileReference SrcFile = FileReference.Combine(EnvVars.CompilerDir, "bin", "libmmd.dll");
+				FileReference DstFile = FileReference.Combine(ExeDir, "libmmd.dll");
 
+				TargetFileToSourceFile[DstFile] = SrcFile;
+				RuntimeDependencies.Add(new RuntimeDependency(DstFile, StagedFileType.NonUFS));
+			}
+		}
 
 		public virtual FileReference GetApplicationIcon(FileReference? ProjectFile)
 		{
