Reconcile usage of TargetRules.WindowsPlatform.bStrictInlineConformance with ModuleRules.bVcRemoveUnreferencedComdat
Clean up some obsolete properties

#rnx

diff --git a/Engine/Source/Programs/UnrealBuildTool/Configuration/ModuleRules.cs b/Engine/Source/Programs/UnrealBuildTool/Configuration/ModuleRules.cs
--- a/Engine/Source/Programs/UnrealBuildTool/Configuration/ModuleRules.cs
+++ b/Engine/Source/Programs/UnrealBuildTool/Configuration/ModuleRules.cs
@@ -925,7 +925,12 @@ namespace UnrealBuildTool
 		/// Whether to direct MSVC to remove unreferenced COMDAT functions and data.
 		/// </summary>
 		/// <seealso href="https://learn.microsoft.com/en-us/cpp/build/reference/zc-inline-remove-unreferenced-comdat">zc-inline-remove-unreferenced-comdat</seealso>
-		public bool bVcRemoveUnreferencedComdat { get; set; } = true;
+		public bool bVcRemoveUnreferencedComdat
+		{
+			get => bVcRemoveUnreferencedComdatPrivate ?? Target.WindowsPlatform.bStrictInlineConformance;
+			set => bVcRemoveUnreferencedComdatPrivate = value;
+		}
+		private bool? bVcRemoveUnreferencedComdatPrivate;
 
 		/// <summary>
 		/// Enable code coverage compilation/linking support.
