UnrealBuildTool: Update BuildSettingsVersion.V6 comments

#rnx

diff --git a/Engine/Source/Programs/UnrealBuildTool/Configuration/TargetRules.cs b/Engine/Source/Programs/UnrealBuildTool/Configuration/TargetRules.cs
--- a/Engine/Source/Programs/UnrealBuildTool/Configuration/TargetRules.cs
+++ b/Engine/Source/Programs/UnrealBuildTool/Configuration/TargetRules.cs
@@ -154,7 +154,7 @@ namespace UnrealBuildTool
 
 		/// <summary>
 		/// New defaults for 5.7:
-		/// * ModuleRules.CppCompileWarningSettings.DanglingWarningLevel = WarningLevel.Warning
+		/// * TargetRules.c
 		/// * ModuleRules.CppCompileWarningSettings.UndefinedIdentifierWarningLevel = WarningLevel.Error
 		/// </summary>
 		V6,
@@ -3539,6 +3539,7 @@ namespace UnrealBuildTool
 
 				if (BuildSettingsVersion.V6 <= latestVersion && DefaultBuildSettings < BuildSettingsVersion.V6)
 				{
+					modifiedSettings.Add(Tuple.Create($"{nameof(WindowsPlatform)}.{nameof(WindowsPlatform.bStrictInlineConformance)} = true", "Updates MSVC strict inline conformance mode to true (Previously: false)."));
 					modifiedSettings.Add(Tuple.Create(String.Format("{0} = WarningLevel.Error", nameof(CppCompileWarningSettings.UndefinedIdentifierWarningLevel)), "Enables compile-time validation of undefined identifiers in conditional expressions. (Previously: Off)."));
 				}
 
