[UBT]
* Added option to control if retry can be done natively or not
* Changed so compress mode allow retry but not natively since it is using virtual files

diff --git a/Engine/Source/Programs/UnrealBuildTool/Executors/UnrealBuildAccelerator/UBAExecutor.cs b/Engine/Source/Programs/UnrealBuildTool/Executors/UnrealBuildAccelerator/UBAExecutor.cs
--- a/Engine/Source/Programs/UnrealBuildTool/Executors/UnrealBuildAccelerator/UBAExecutor.cs
+++ b/Engine/Source/Programs/UnrealBuildTool/Executors/UnrealBuildAccelerator/UBAExecutor.cs
@@ -24,6 +24,9 @@ namespace UnrealBuildTool
 		public UnrealBuildAcceleratorConfig UBAConfig { get; init; } = new();
 
 		public string Crypto { get; private set; } = String.Empty;
+
+		public bool CanRetryNatively { get; set; } = true;
+
 		ISessionServer? _session;
 		object _sessionLock = new();
 		IEnumerable<UnrealBuildAcceleratorCacheConfig> _cacheConfigs = [];
@@ -743,7 +746,7 @@ namespace UnrealBuildTool
 						{
 							if (UBAConfig.bAllowRetry && (exitCode != 0 && UBAConfig.bForcedRetry || (exitCode >= 9000 && exitCode < 10000)))
 							{
-								if (UBAConfig.bStoreObjFilesCompressed)
+								if (!CanRetryNatively || UBAConfig.bStoreObjFilesCompressed)
 								{
 									_threadedLogger.LogInformation("{Description} {StatusDescription}: Exited with error code {ExitCode}. This action will retry", action.CommandDescription, action.StatusDescription, exitCode);
 									++forcedRetryActions;
