UnrealBuildTool: Allow creating new files under Engine\ when -NoEngineChanges is passed, which will allow new projects to create any necessary shared PCH without needing to close the editor

#rnx

diff --git a/Engine/Source/Programs/UnrealBuildTool/Modes/BuildMode.cs b/Engine/Source/Programs/UnrealBuildTool/Modes/BuildMode.cs
--- a/Engine/Source/Programs/UnrealBuildTool/Modes/BuildMode.cs
+++ b/Engine/Source/Programs/UnrealBuildTool/Modes/BuildMode.cs
@@ -727,18 +727,20 @@ namespace UnrealBuildTool
 					MergedActionsToExecute = LiveCodingActionsToExecute;
 				}
 
-				// Make sure we're not modifying any engine files
-				if ((Options & BuildOptions.NoEngineChanges) != 0)
+				// Make sure we're not modifying any existing engine files while allowing new engine files to be created
+				if (Options.HasFlag(BuildOptions.NoEngineChanges))
 				{
-					List<FileItem> EngineChanges = MergedActionsToExecute.SelectMany(x => x.ProducedItems).Where(x => x.Location.IsUnderDirectory(Unreal.EngineDirectory)).Distinct().OrderBy(x => x.FullName).ToList();
-					if (EngineChanges.Count > 0)
+					IEnumerable<FileItem> EngineChanges = MergedActionsToExecute.SelectMany(x => x.ProducedItems).Where(x => x.Location.IsUnderDirectory(Unreal.EngineDirectory) && x.Exists).Distinct().Order();
+					if (EngineChanges.Any())
 					{
-						StringBuilder Result = new StringBuilder("Building would modify the following engine files:\n");
+						StringBuilder Result = new("Building would modify the following existing engine files:");
+						Result.AppendLine();
 						foreach (FileItem EngineChange in EngineChanges)
 						{
-							Result.AppendFormat("\n{0}", EngineChange.FullName);
+							Result.AppendLine(EngineChange.FullName);
 						}
-						Result.Append("\n\nPlease rebuild from an IDE instead.");
+						Result.AppendLine();
+						Result.AppendLine("Please build from an IDE instead.");
 						Logger.LogError("{Result}", Result.ToString());
 						throw new CompilationResultException(CompilationResult.FailedDueToEngineChange);
 					}
