Only add source-map commands to Android debug config if bUseVFS is true for target

diff --git a/Engine/Source/Programs/UnrealBuildTool/Platform/Android/AndroidProjectGenerator.cs b/Engine/Source/Programs/UnrealBuildTool/Platform/Android/AndroidProjectGenerator.cs
--- a/Engine/Source/Programs/UnrealBuildTool/Platform/Android/AndroidProjectGenerator.cs
+++ b/Engine/Source/Programs/UnrealBuildTool/Platform/Android/AndroidProjectGenerator.cs
@@ -198,7 +198,7 @@ namespace UnrealBuildTool
 				StringBuilder Out = new StringBuilder();
 				Out.AppendLine("  <PropertyGroup " + InConditionString + ">");
 
-				Dictionary<DirectoryReference, DirectoryReference> SourceFileMap = ProjectFileGenerator.GenerateSourceMap(InTargetRules.ProjectFile);
+				Dictionary<DirectoryReference, DirectoryReference> SourceFileMap = ProjectFileGenerator.GenerateSourceMap(InTargetRules, InTargetRules.ProjectFile);
 				Out.Append("    <AndroidLldbStartupCommands>");
 				Out.Append($"command script import \"{Path.Combine(Unreal.EngineDirectory.FullName, "Extras", "LLDBDataFormatters", "UEDataFormatters_2ByteChars.py")}\";");
 				Out.AppendJoin(';', SourceFileMap.Select(from_to => $"setting append target.source-map {from_to.Key} {from_to.Value}"));
diff --git a/Engine/Source/Programs/UnrealBuildTool/ProjectFiles/ProjectFileGenerator.cs b/Engine/Source/Programs/UnrealBuildTool/ProjectFiles/ProjectFileGenerator.cs
--- a/Engine/Source/Programs/UnrealBuildTool/ProjectFiles/ProjectFileGenerator.cs
+++ b/Engine/Source/Programs/UnrealBuildTool/ProjectFiles/ProjectFileGenerator.cs
@@ -864,25 +864,29 @@ namespace UnrealBuildTool
 		/// <summary>
 		/// Generate a mapping for source file directories when using VFS for debugger configuration 
 		/// </summary>
+		/// <param name="InTargetRules">Target the debugging settings are being generated for</param>
 		/// <param name="UProjectFile">Path to uproject file for the target we're generating debugger config for</param>
 		/// <returns>A mapping where keys are VFS paths and values are disk paths</returns>
-		public static Dictionary<DirectoryReference, DirectoryReference> GenerateSourceMap(FileReference? UProjectFile)
+		public static Dictionary<DirectoryReference, DirectoryReference> GenerateSourceMap(TargetRules InTargetRules, FileReference? UProjectFile)
 		{
 			Dictionary<DirectoryReference, DirectoryReference> SourceFileMap = new();
 			if (Unreal.OriginalCompilationRootDirectory != Unreal.RootDirectory)
 			{
 				SourceFileMap.Add(Unreal.OriginalCompilationRootDirectory, Unreal.RootDirectory);
 			}
-			CppRootPaths RootPaths = new();
-			RootPaths[CppRootPathFolder.Root] = Unreal.RootDirectory;
-			if (UProjectFile != null)
+			if (InTargetRules.bUseVFS)
 			{
-				RootPaths[CppRootPathFolder.Project] = UProjectFile.Directory;
-			}
+				CppRootPaths RootPaths = new() { bUseVfs = InTargetRules.bUseVFS };
+				RootPaths[CppRootPathFolder.Root] = Unreal.RootDirectory;
+				if (UProjectFile != null)
+				{
+					RootPaths[CppRootPathFolder.Project] = UProjectFile.Directory;
+				}
 
-			foreach ((uint id, DirectoryReference vfs, DirectoryReference local) entry in RootPaths)
-			{
-				SourceFileMap.Add(entry.vfs, entry.local);
+				foreach ((uint id, DirectoryReference vfs, DirectoryReference local) entry in RootPaths)
+				{
+					SourceFileMap.Add(entry.vfs, entry.local);
+				}
 			}
 			return SourceFileMap;
 		}
diff --git a/Engine/Source/Programs/UnrealBuildTool/ProjectFiles/VisualStudioCode/VSCodeProjectFileGenerator.cs b/Engine/Source/Programs/UnrealBuildTool/ProjectFiles/VisualStudioCode/VSCodeProjectFileGenerator.cs
--- a/Engine/Source/Programs/UnrealBuildTool/ProjectFiles/VisualStudioCode/VSCodeProjectFileGenerator.cs
+++ b/Engine/Source/Programs/UnrealBuildTool/ProjectFiles/VisualStudioCode/VSCodeProjectFileGenerator.cs
@@ -439,8 +439,17 @@ namespace UnrealBuildTool
 			{
 				public string Name;
 				public TargetType Type;
+				// For native projects only
+				public TargetRules? TargetRules;
 				public List<BuildProduct> BuildProducts = new List<BuildProduct>();
 
+				public Target(Project InParentProject, TargetRules InTargetRules)
+				{
+					Name = InTargetRules.Name;
+					Type = InTargetRules.Type;
+					TargetRules = InTargetRules;
+					InParentProject.Targets.Add(this);
+				}
 				public Target(Project InParentProject, string InName, TargetType InType)
 				{
 					Name = InName;
@@ -496,7 +505,7 @@ namespace UnrealBuildTool
 						UnrealTargetConfiguration[] Configs = (UnrealTargetConfiguration[])Enum.GetValues(typeof(UnrealTargetConfiguration));
 						List<UnrealTargetPlatform> Platforms = new List<UnrealTargetPlatform>(Target.TargetRules!.GetSupportedPlatforms());
 
-						ProjectData.Target NewTarget = new ProjectData.Target(NewProject, Target.TargetRules.Name, Target.TargetRules.Type);
+						ProjectData.Target NewTarget = new ProjectData.Target(NewProject, Target.TargetRules);
 
 						if (HostPlatform != UnrealTargetPlatform.Win64)
 						{
@@ -1397,19 +1406,22 @@ namespace UnrealBuildTool
 						OutFile.AddField("type", "lldb");
 					}
 
-					Dictionary<DirectoryReference, DirectoryReference> SourceFileMap = GenerateSourceMap(BuildProduct.UProjectFile);
-					if (SourceFileMap.Any())
-					{
-						OutFile.BeginObject("sourceFileMap");
+					if (Target.TargetRules is not null)
+					{ 
+						Dictionary<DirectoryReference, DirectoryReference> SourceFileMap = GenerateSourceMap(Target.TargetRules, BuildProduct.UProjectFile);
+						if (SourceFileMap.Any())
 						{
-							foreach (KeyValuePair<DirectoryReference, DirectoryReference> entry in SourceFileMap)
+							OutFile.BeginObject("sourceFileMap");
 							{
-								OutFile.AddField(
-									MakeUnquotedPathString(entry.Key, EPathType.Absolute),
-									MakeUnquotedPathString(entry.Value, EPathType.Absolute));
+								foreach (KeyValuePair<DirectoryReference, DirectoryReference> entry in SourceFileMap)
+								{
+									OutFile.AddField(
+										MakeUnquotedPathString(entry.Key, EPathType.Absolute),
+										MakeUnquotedPathString(entry.Value, EPathType.Absolute));
+								}
 							}
+							OutFile.EndObject();
 						}
-						OutFile.EndObject();
 					}
 				}
 				OutFile.EndObject();
