Win64 arm64 and arm64ec, fix for Make Installed Build Win64
-Fix for "Asking CppCompileEnvironment for a single Architecture, but it has multiple Architectures (arm64, x64). This indicates logic needs to be updated".

#tests I was able to run the FinalizeInstalledBuild command locally and reproduce the issue, then with the fix the command succeded.

diff --git a/Engine/Source/Programs/UnrealBuildTool/Platform/Windows/UEBuildWindows.cs b/Engine/Source/Programs/UnrealBuildTool/Platform/Windows/UEBuildWindows.cs
--- a/Engine/Source/Programs/UnrealBuildTool/Platform/Windows/UEBuildWindows.cs
+++ b/Engine/Source/Programs/UnrealBuildTool/Platform/Windows/UEBuildWindows.cs
@@ -1169,7 +1169,8 @@ namespace UnrealBuildTool
 
 		public override bool RequiresArchitectureFilenames(UnrealArchitectures Architectures)
 		{
-			return Architectures.SingleArchitecture != UnrealArch.X64;
+			// for now x64 only uses architecture-less filenames.
+			return Architectures.bIsMultiArch || !Architectures.Contains(UnrealArch.X64);
 		}
 
 		public override UnrealArch GetHostArchitecture()
