Don't request WRITE_EXTERNAL_STORAGE permission if TargetSDK doesn't support non-scoped storage

#android

diff --git a/Engine/Source/Programs/UnrealBuildTool/Platform/Android/UEDeployAndroid.cs b/Engine/Source/Programs/UnrealBuildTool/Platform/Android/UEDeployAndroid.cs
--- a/Engine/Source/Programs/UnrealBuildTool/Platform/Android/UEDeployAndroid.cs
+++ b/Engine/Source/Programs/UnrealBuildTool/Platform/Android/UEDeployAndroid.cs
@@ -3172,7 +3172,7 @@ namespace UnrealBuildTool
 
 			// Figure out the required startup permissions if targetting devices supporting runtime permissions
 			string StartupPermissions = "";
-			if (TargetSDKVersion >= 23 && !bPackageForMetaQuest)
+			if (TargetSDKVersion >= 23 && TargetSDKVersion < 30 && !bPackageForMetaQuest)
 			{
 				if (Configuration != "Shipping" || !bUseExternalFilesDir)
 				{
@@ -3325,9 +3325,12 @@ namespace UnrealBuildTool
 				// Apps requesting these permissions are rejected by the Quest App Store
 				if (!bPackageForMetaQuest)
 				{
-					if (Configuration != "Shipping" || !bUseExternalFilesDir)
+					if (TargetSDKVersion >= 23 && TargetSDKVersion < 30)
 					{
-						Text.AppendLine("\t<uses-permission android:name=\"android.permission.WRITE_EXTERNAL_STORAGE\"/>");
+						if (Configuration != "Shipping" || !bUseExternalFilesDir)
+						{
+							Text.AppendLine("\t<uses-permission android:name=\"android.permission.WRITE_EXTERNAL_STORAGE\"/>");
+						}
 					}
 					Text.AppendLine("\t<uses-permission android:name=\"android.permission.WAKE_LOCK\"/>");
 					//Text.AppendLine("\t<uses-permission android:name=\"android.permission.READ_PHONE_STATE\"/>");
