Fix ASan lib/dll paths for Clang on Windows.

Use the Clang and not the MSVC version of the libraries.

diff --git a/Engine/Source/Programs/UnrealBuildTool/Platform/Windows/VCToolChain.cs b/Engine/Source/Programs/UnrealBuildTool/Platform/Windows/VCToolChain.cs
--- a/Engine/Source/Programs/UnrealBuildTool/Platform/Windows/VCToolChain.cs
+++ b/Engine/Source/Programs/UnrealBuildTool/Platform/Windows/VCToolChain.cs
@@ -1855,7 +1855,8 @@ namespace UnrealBuildTool
 					string ASanArchSuffix;
 					if (EnvVars.Architecture == UnrealArch.X64)
 					{
-						ASanRuntimeDir = DirectoryReference.Combine(EnvVars.ToolChainDir, "lib", "x64");
+						VersionNumber ClangVersion = Target.WindowsPlatform.Compiler == WindowsCompiler.Intel ? MicrosoftPlatformSDK.GetClangVersionForIntelCompiler(EnvVars.CompilerPath) : EnvVars.CompilerVersion;
+						ASanRuntimeDir = DirectoryReference.Combine(EnvVars.CompilerDir, "lib", "clang", $"{ClangVersion.Components[0]}", "lib", "windows");
 						ASanArchSuffix = "x86_64";
 					}
 					else
@@ -2630,7 +2631,15 @@ namespace UnrealBuildTool
 				string ASanArchSuffix;
 				if (EnvVars.Architecture == UnrealArch.X64)
 				{
-					ASanRuntimeDir = DirectoryReference.Combine(EnvVars.ToolChainDir, "bin", "Hostx64", "x64");
+					if (Target.WindowsPlatform.Compiler.IsClang())
+					{
+						VersionNumber ClangVersion = Target.WindowsPlatform.Compiler == WindowsCompiler.Intel ? MicrosoftPlatformSDK.GetClangVersionForIntelCompiler(EnvVars.CompilerPath) : EnvVars.CompilerVersion;
+						ASanRuntimeDir = DirectoryReference.Combine(EnvVars.CompilerDir, "lib", "clang", $"{ClangVersion.Components[0]}", "lib", "windows");
+					}
+					else
+					{
+						ASanRuntimeDir = DirectoryReference.Combine(EnvVars.ToolChainDir, "bin", "Hostx64", "x64");
+					}
 					ASanArchSuffix = "x86_64";
 				}
 				else
