[UBT]
* Fixed so .precompiled manifest location is not in non-architecture intermediate folder since arm64 and x64 needs to have their own precompiled manifest files

diff --git a/Engine/Source/Programs/UnrealBuildTool/Configuration/UEBuildModuleCPP.cs b/Engine/Source/Programs/UnrealBuildTool/Configuration/UEBuildModuleCPP.cs
--- a/Engine/Source/Programs/UnrealBuildTool/Configuration/UEBuildModuleCPP.cs
+++ b/Engine/Source/Programs/UnrealBuildTool/Configuration/UEBuildModuleCPP.cs
@@ -459,7 +459,7 @@ namespace UnrealBuildTool
 		/// <summary>
 		/// Path to the precompiled manifest location
 		/// </summary>
-		public virtual FileReference PrecompiledManifestLocation => FileReference.Combine(IntermediateDirectoryNoArch, String.Format("{0}.precompiled", Name));
+		public virtual FileReference PrecompiledManifestLocation => FileReference.Combine(IntermediateDirectory, String.Format("{0}.precompiled", Name));
 
 		/// <summary>
 		/// Sets up the environment for compiling any module that includes the public interface of this module.
