[IWYU]
* Switched to use win64 targets for iwyu instead of linux because it is easier to support (and more people build win targets than linux)
* Upgraded include-what-you-use.exe to clang 20 (also with tons of more fixes for various c++ features)
* Added some more mappings for unreal code
* Fixed so -mode=iwyu -updatedependents works properly with -pathtoupdate

diff --git a/Engine/Source/Programs/UnrealBuildTool/Configuration/UEBuildModule.cs b/Engine/Source/Programs/UnrealBuildTool/Configuration/UEBuildModule.cs
--- a/Engine/Source/Programs/UnrealBuildTool/Configuration/UEBuildModule.cs
+++ b/Engine/Source/Programs/UnrealBuildTool/Configuration/UEBuildModule.cs
@@ -173,7 +173,7 @@ namespace UnrealBuildTool
 		/// <summary>
 		/// Set of modules that have referenced this module
 		/// </summary>
-		protected HashSet<UEBuildModule>? ReferenceStackParentModules;
+		public HashSet<UEBuildModule>? ReferenceStackParentModules;
 
 		/// <summary>
 		/// Set of all allowed restricted folder references
diff --git a/Engine/Source/Programs/UnrealBuildTool/Modes/IWYUMode.cs b/Engine/Source/Programs/UnrealBuildTool/Modes/IWYUMode.cs
--- a/Engine/Source/Programs/UnrealBuildTool/Modes/IWYUMode.cs
+++ b/Engine/Source/Programs/UnrealBuildTool/Modes/IWYUMode.cs
@@ -4,7 +4,6 @@ using System;
 using System.Collections.Concurrent;
 using System.Collections.Generic;
 using System.Diagnostics;
-using System.Diagnostics.CodeAnalysis;
 using System.IO;
 using System.Linq;
 using System.Text;
@@ -13,6 +12,7 @@ using System.Threading;
 using System.Threading.Tasks;
 using EpicGames.Core;
 using Microsoft.Extensions.Logging;
+using Microsoft.Extensions.ObjectPool;
 using UnrealBuildBase;
 
 // IWYUMode is a mode that can be used to clean up includes in source code. It uses the clang based tool include-what-you-use (IWYU) to figure out what is needed in each .h/.cpp file
@@ -75,18 +75,18 @@ namespace UnrealBuildTool
 	}
 
 	/// <summary>
-	/// Comparer for include entries using printable to compare
+	/// Comparer for include entries using full path to compare
 	/// </summary>
-	class IWYUIncludeEntryPrintableComparer : IEqualityComparer<IWYUIncludeEntry>
+	class IWYUIncludeEntryFullPathComparer : IEqualityComparer<IWYUIncludeEntry>
 	{
 		public bool Equals(IWYUIncludeEntry? x, IWYUIncludeEntry? y)
 		{
-			return x!.Printable == y!.Printable;
+			return x!.Full == y!.Full;
 		}
 
-		public int GetHashCode([DisallowNull] IWYUIncludeEntry obj)
+		public int GetHashCode(IWYUIncludeEntry x)
 		{
-			throw new NotImplementedException();
+			return x.Full.GetHashCode();
 		}
 	}
 
@@ -317,6 +317,12 @@ namespace UnrealBuildTool
 		[CommandLine("-ReadToc")]
 		public bool bReadToc = false;
 
+		/// <summary>
+		/// Will print out info about the file in question. Transitive includes etc
+		/// </summary>
+		[CommandLine("-PrintInfo")]
+		public string PrintInfo = "";
+
 		private string? GetModuleToUpdateName(TargetDescriptor Descriptor)
 		{
 			if (!String.IsNullOrEmpty(ModuleToUpdate))
@@ -426,6 +432,7 @@ namespace UnrealBuildTool
 			Dictionary<string, UEBuildModule> PathToModule = new();
 			Dictionary<string, UEBuildModule> NameToModule = new();
 			{
+				using ITimelineEvent _ = Timeline.ScopeEvent("CreatingTargets");
 				// Create target to be able to traverse all existing modules
 				Logger.LogInformation($"Creating BuildTarget for {TargetName}...");
 				UEBuildTarget Target = UEBuildTarget.Create(TargetDescriptors[0], BuildConfiguration, Logger);
@@ -447,7 +454,7 @@ namespace UnrealBuildTool
 
 				if (PathToUpdateList.Length == 0 && UEModuleToUpdate == null)
 				{
-					Logger.LogError($"Need to provide -ModuleToUpdate <modulename> or -PathToUpdate <partofpath> to run IWYU.");
+					Logger.LogError($"Need to provide -ModuleToUpdate=<modulename> or -PathToUpdate=<partofpath> to run IWYU.");
 					return -1;
 				}
 
@@ -477,11 +484,19 @@ namespace UnrealBuildTool
 					}
 				}
 
+
+				Queue<UEBuildModule> toUpdateQueue = [];
+				HashSet<UEBuildModule> toUpdateHandled = [];
+				if (UEModuleToUpdate != null)
+				{
+					toUpdateQueue.Enqueue(UEModuleToUpdate);
+					toUpdateHandled.Add(UEModuleToUpdate);
+				}
+
+
 				foreach (UEBuildModule Module in NameToModule.Values)
 				{
-					bool ShouldUpdate = false;
 					string ModuleDir = Module.ModuleDirectory.FullName.Replace('\\', '/');
-
 					if (!PathToModule.TryAdd(ModuleDir, Module))
 					{
 						continue;
@@ -499,29 +514,6 @@ namespace UnrealBuildTool
 						continue;
 					}
 
-					if (UEModuleToUpdate != null)
-					{
-						bool DependsOnModule = Module == UEModuleToUpdate;
-						if (bUpdateDependents)
-						{
-							if (Module.PublicDependencyModules != null)
-							{
-								foreach (UEBuildModule Dependency in Module.PublicDependencyModules)
-								{
-									DependsOnModule = DependsOnModule || Dependency == UEModuleToUpdate;
-								}
-							}
-							if (Module.PrivateDependencyModules != null)
-							{
-								foreach (UEBuildModule Dependency in Module.PrivateDependencyModules!)
-								{
-									DependsOnModule = DependsOnModule || Dependency == UEModuleToUpdate;
-								}
-							}
-						}
-						ShouldUpdate = DependsOnModule;
-					}
-
 					if (PathToUpdateList.Length > 0)
 					{
 						bool Match = false;
@@ -532,46 +524,55 @@ namespace UnrealBuildTool
 
 						if (Match)
 						{
-							if (UEModuleToUpdate == null)
+							if (toUpdateHandled.Add(Module))
 							{
-								ShouldUpdate = true;
+								toUpdateQueue.Enqueue(Module);
 							}
 						}
-						else
-						{
-							ShouldUpdate = false;
-						}
-					}
-					else if (UEModuleToUpdate == null)
-					{
-						ShouldUpdate = true;
-					}
-
-					if (!ShouldUpdate)
-					{
-						continue;
 					}
+				}
 
-					if (Module.Rules.IWYUSupport == IWYUSupport.None)
+				void AddModuleDirs(UEBuildModule module)
+				{
+					if (module.Rules.IWYUSupport == IWYUSupport.None)
 					{
 						++ModulesSkippedCount;
-						continue;
+						return;
 					}
 
 					++ModulesToUpdateCount;
 
+					string ModuleDir = module.ModuleDirectory.FullName.Replace('\\', '/');
+
 					// When adding to ValidPaths, make sure the Module directory ends in a / so we match the exact folder later
 					ValidPaths.Add(AdjustModulePathForMatching(ModuleDir));
 
-					foreach (DirectoryReference AdditionalDir in Module.ModuleDirectories)
+					foreach (DirectoryReference AdditionalDir in module.ModuleDirectories)
 					{
-						if (AdditionalDir != Module.ModuleDirectory)
+						if (AdditionalDir != module.ModuleDirectory)
 						{
 							ValidPaths.Add(AdjustModulePathForMatching(AdditionalDir.FullName));
 						}
 					}
 				}
 
+				while (toUpdateQueue.Count != 0)
+				{
+					UEBuildModule module = toUpdateQueue.Dequeue();
+					AddModuleDirs(module);
+
+					if (bUpdateDependents)
+					{
+						foreach (UEBuildModule refModule in module.ReferenceStackParentModules!)
+						{
+							if (toUpdateHandled.Add(refModule))
+							{
+								toUpdateQueue.Enqueue(refModule);
+							}
+						}
+					}
+				}
+
 				if (ValidPaths.Count == 0 && PathToUpdateList.Length > 0)
 				{
 					foreach (string Path in PathToUpdateList)
@@ -587,8 +588,7 @@ namespace UnrealBuildTool
 			}
 
 			Dictionary<string, IWYUInfo> Infos = new();
-			HashSet<string> GeneratedHeaderInfos = new();
-			List<IWYUInfo> GeneratedCppInfos = new();
+			ConcurrentBag<IWYUInfo> GeneratedCppInfos = new();
 
 			using (ISourceFileWorkingSet WorkingSet = new EmptySourceFileWorkingSet())
 			{
@@ -665,10 +665,10 @@ namespace UnrealBuildTool
 					}
 				}
 
-				// Time to parse all the .iwyu files generated from iwyu.
-				// We Do this in parallel since it involves reading a ton of .iwyu files from disk.
+				// First parse all iwyu files and add main entry to lookup (that is the last entry in each file)
 				Logger.LogInformation($"Parsing {TocContent.Count} .iwyu files...");
-				Parallel.ForEach(TocContent, IWYUFilePath =>
+				using ITimelineEvent parsingIwyu = Timeline.ScopeEvent("ParsingIwyu");
+				Parallel2.ForEach(TocContent, (IWYUFilePath) =>
 				{
 					string? JsonContent = File.ReadAllText(IWYUFilePath);
 					if (JsonContent == null)
@@ -681,99 +681,148 @@ namespace UnrealBuildTool
 					try
 					{
 						IWYUFile? IWYUFile = JsonSerializer.Deserialize<IWYUFile>(JsonContent);
-						IWYUFile!.Name = IWYUFilePath;
-
-						// Traverse the cpp/inl/h file entries inside the .iwyu file
-						foreach (IWYUInfo Info in IWYUFile!.Files)
+						if (IWYUFile == null)
 						{
-							Info.Source = IWYUFile;
-							Info.IsCpp = Info.File.EndsWith(".cpp");
+							Logger.LogError($"Failed to deserialize file {IWYUFilePath}");
+							Interlocked.Exchange(ref ReadSuccess, 0);
+							return;
+						}
 
-							// We track .gen.cpp in a special list, they need special treatment later
-							if (Info.File.Contains(".gen.cpp", StringComparison.Ordinal))
-							{
-								lock (GeneratedCppInfos)
-								{
-									GeneratedCppInfos.Add(Info);
-								}
+						IWYUFile.Name = IWYUFilePath;
 
-								continue;
+						var mainInfo = IWYUFile.Files[IWYUFile.Files.Count - 1];
+						mainInfo.Source = IWYUFile;
+
+						lock (Infos)
+						{
+							if (!Infos.TryAdd(mainInfo.File, mainInfo))
+							{
+								Logger.LogError($"Found {mainInfo.File} multiple times");
+								Interlocked.Exchange(ref ReadSuccess, 0);
+								return;
 							}
+						}
+					}
+					catch (Exception e)
+					{
+						Logger.LogError($"Failed to parse json {IWYUFilePath}: {e.Message} - File will be deleted");
+						//File.Delete(IWYUFilePath);
+						Interlocked.Exchange(ref ReadSuccess, 0);
+						return;
+					}
+				});
+				parsingIwyu.Finish();
+
+				ConcurrentBag<(IWYUInfo mainInfo, List<IWYUInfo> generatedHeaders)> generatedHeadersToFixup = new();;
+
+				// Now traverse main infos and find additional infos and fix them up
+				Logger.LogInformation($"Adding additional iwyu infos...");
+				var mainInfos = Infos.Values.ToArray();
+				Parallel2.ForEach(mainInfos, (mainInfo) =>
+				{
+					IWYUFile IWYUFile = mainInfo.Source!;
 
-							// Ok, time to add file entry to lookup
-							lock (Infos)
+					IWYUInfo? headerAsMain = null;
+					IWYUInfo? headerInCpp = null;
+					List<IWYUInfo>? generatedHeaders = null;
+
+					// Go through all infos in iwyu
+					foreach (IWYUInfo info in IWYUFile.Files)
+					{
+						bool isCpp = info.File.EndsWith(".cpp", StringComparison.Ordinal);
+						info.IsCpp = isCpp;
+						info.Module = GetModule(info.File, PathToModule);
+
+						foreach (IWYUIncludeEntry entry in info.IncludesSeenInFile)
+						{
+							if (entry.Printable[0] == '<')
 							{
-								if (!Infos.TryAdd(Info.File, Info))
-								{
-									// This is a valid scenario when Foo.cpp also registers Foo.h... and then Foo.h registers itself.
-									IWYUInfo ExistingEntry = Infos[Info.File];
-									if (IWYUFile.Files.Count == 1)
-									{
-										if (ExistingEntry.Source!.Files.Count == 1)
-										{
-											Logger.LogError($"{Info.File} - built twice somehow?");
-											return;
-										}
-										else
-										{
-											Infos[Info.File] = Info;
-										}
-									}
-									else
-									{
-										//bool Equals = Info.Includes.SequenceEqual(ExistingEntry.Includes, new IWYUIncludeEntryPrintableComparer());
-										//if (!Equals)
-										//{
-										//	Logger.LogWarning($"{Info.File} - mismatch found in multiple .iwyu-files");
-										//}
-									}
-								}
+								entry.Printable = entry.Printable.Substring(1, entry.Printable.Length - 2);
+								entry.System = true;
 							}
+						}
+
+						if (info == mainInfo)
+						{
+							break;
+						}
 
-							// TODO: Fix bad formatting coming from iwyu
-							foreach (IWYUIncludeEntry Entry in Info.IncludesSeenInFile)
+						if (isCpp && info.File.Contains(".gen.cpp", StringComparison.Ordinal))
+						{
+							info.Source = IWYUFile;
+							GeneratedCppInfos.Add(info);
+							continue;
+						}
+
+						bool isHeader = !isCpp && info.File.EndsWith(".h");
+						bool isGeneratedHeader = isHeader && (info.File.EndsWith(".generated.h", StringComparison.Ordinal) || info.File.EndsWith(".gen.h", StringComparison.Ordinal));
+
+						lock (Infos)
+						{
+							if (!Infos.TryAdd(info.File, info))
 							{
-								if (Entry.Printable[0] == '<')
+								if (isHeader && !isGeneratedHeader)
 								{
-									Entry.Printable = Entry.Printable.Substring(1, Entry.Printable.Length - 2);
-									Entry.System = true;
+									headerInCpp = info;
+									headerAsMain = Infos[info.File];
+									continue;
 								}
+								// TODO: Lots of .inl files here
+								//Logger.LogError($"Tried to add file {info.File}");
+								//Interlocked.Exchange(ref ReadSuccess, 0);
 							}
+						}
 
-							Info.Module = GetModule(Info.File, PathToModule);
+						if (isGeneratedHeader)
+						{
+							if (generatedHeaders == null)
+								generatedHeaders = new();
+							generatedHeaders.Add(info);
+						}
+					}
 
-							// Some special logic for headers.
-							if (Info.File.EndsWith(".h"))
-							{
-								// We need to add entries for .generated.h. They are not in the iwyu files
-								foreach (IWYUIncludeEntry Include in Info.Includes)
-								{
-									if (Include.Full.Contains("/UHT/"))
-									{
-										lock (GeneratedHeaderInfos)
-										{
-											GeneratedHeaderInfos.Add(Include.Full);
-										}
+					// Fix up main file
 
-										break;
-									}
-									else if (Include.Full.Contains("/VNI/"))
-									{
-										lock (GeneratedHeaderInfos)
-										{
-											GeneratedHeaderInfos.Add(Include.Full);
-										}
-									}
+					if (!mainInfo.IsCpp & generatedHeaders != null)
+					{
+						generatedHeadersToFixup.Add((mainInfo!, generatedHeaders!));
+					}
+
+					// This is annoying. But for some reason list of files needed in header can be different
+					// based on if it builds standalone or with the cpp file.
+					// This should not happen if header is first in include list.
+					// What we need to do is to take all includes that are in
+					if (mainInfo.IsCpp && headerInCpp != null)
+					{
+						if (headerInCpp.Includes.Count != headerAsMain!.Includes.Count)
+						{
+							var exists = new HashSet<IWYUIncludeEntry>(mainInfo.Includes, new IWYUIncludeEntryFullPathComparer());
+							exists.UnionWith(headerAsMain!.Includes);
+							foreach (var i in headerInCpp.Includes)
+								if (!exists.Contains(i))
+								{
+									// TODO: Note that h and cpp might have different printable include paths since h and cpp are often not in same folder
+									mainInfo.Includes.Add(i);
 								}
-							}
 						}
+
 					}
-					catch (Exception e)
+				});
+
+				Parallel2.ForEach(generatedHeadersToFixup, (entry) =>
+				{
+					IWYUInfo mainInfo = entry.mainInfo;
+					List<IWYUInfo> generatedHeaders = entry.generatedHeaders;
+					var exists = new HashSet<IWYUIncludeEntry>(mainInfo.Includes, new IWYUIncludeEntryFullPathComparer());
+					foreach (var generatedHeader in generatedHeaders)
+						exists.Union(generatedHeader.IncludesSeenInFile);
+
+					foreach (var generatedHeader in generatedHeaders)
 					{
-						Logger.LogError($"Failed to parse json {IWYUFilePath}: {e.Message} - File will be deleted");
-						File.Delete(IWYUFilePath);
-						Interlocked.Exchange(ref ReadSuccess, 0);
-						return;
+						foreach (var i in generatedHeader.Includes)
+							if (!exists.Contains(i))
+								mainInfo.Includes.Add(i);
+						generatedHeader.Includes = generatedHeader.IncludesSeenInFile;
 					}
 				});
 
@@ -835,7 +884,7 @@ namespace UnrealBuildTool
 				{ "TCompiledToken", SpecialIncludes["Misc/ExpressionParserTypesFwd.h"] },
 				{ "TExpressionToken", SpecialIncludes["Misc/ExpressionParserTypesFwd.h"] },
 			};
-
+			/*
 			// Add all .generated.h files as entries in the lookup and explicitly add the includes they have which will never be removed
 			Logger.LogInformation($"Generating infos for .generated.h files...");
 			if (GeneratedHeaderInfos.Count > 0)
@@ -872,19 +921,20 @@ namespace UnrealBuildTool
 					Infos.Add(Gen, GenInfo);
 				}
 			}
-
+			*/
 			Logger.LogInformation($"Found {Infos.Count} IWYU entries...");
 
 			Logger.LogInformation($"Resolving Includes...");
+			using ITimelineEvent resolvingIncludes = Timeline.ScopeEvent("ResolvingIncludes");
 
 			LinuxPlatformSDK PlatformSDK = new LinuxPlatformSDK(Logger);
 			DirectoryReference? BaseLinuxPath = PlatformSDK.GetBaseLinuxPathForArchitecture(LinuxPlatform.DefaultHostArchitecture);
 			string SystemPath = BaseLinuxPath!.FullName.Replace('\\', '/');
 
-			HashSet<IWYUInfo> IncludedInsideDecl = new();
+			ConcurrentDictionary<IWYUInfo, IWYUInfo> IncludedInsideDecl = new();
 			Dictionary<string, IWYUInfo> SkippedHeaders = new();
 
-			Parallel.ForEach(Infos.Values, Info =>
+			Parallel2.ForEach(Infos.Values, Info =>
 			{
 				// If KeepAsIs we transfer all "seen includes" into the include list.
 				if (Info.Module != null && Info.Module.Rules.IWYUSupport == IWYUSupport.KeepAsIs)
@@ -933,11 +983,8 @@ namespace UnrealBuildTool
 
 					if (Entry.InsideDecl && Infos.TryGetValue(Entry.Full, out Entry.Resolved))
 					{
-						lock (IncludedInsideDecl)
-						{
-							IncludedInsideDecl.Add(Entry.Resolved);
-							Info.Includes.Add(Entry);
-						}
+						IncludedInsideDecl.TryAdd(Entry.Resolved, Info);
+						Info.Includes.Add(Entry);
 					}
 				}
 
@@ -987,35 +1034,42 @@ namespace UnrealBuildTool
 				}
 
 			});
+			resolvingIncludes.Finish();
 
+			using ITimelineEvent parsingSkipped = Timeline.ScopeEvent("ParsingSkippedHeaders");
 			int InfosCount = Infos.Count;
 			Logger.LogInformation($"Parsing included headers not supporting being compiled...");
 			ParseSkippedHeaders(SkippedHeaders, Infos, PathToModule, NameToModule);
 			Logger.LogInformation($"Added {Infos.Count - InfosCount} more read-only IWYU entries...");
 
-			foreach (IWYUInfo Info in IncludedInsideDecl)
+			foreach ((IWYUInfo Info, IWYUInfo From) in IncludedInsideDecl)
 			{
-				if (Info.IncludesSeenInFile.Count != 0 && !Info.File.EndsWith("ScriptSerialization.h")) // Remove include in ScriptSerialization.h and remove this check
+				if (Info.IncludesSeenInFile.Count != 0) // Remove include in ScriptSerialization.h and remove this check
 				{
-					Logger.LogWarning($"{Info.File} - Included inside declaration in other file but has includes itself.");
+					Logger.LogWarning($"{Info.File} - Included inside declaration in {From.File} but has includes itself.");
 				}
 				Info.Includes.Clear();
 				Info.ForwardDeclarations.Clear();
 			}
+			parsingSkipped.Finish();
 
 			Logger.LogInformation($"Generating transitive include lists and forward declaration lists...");
-			Parallel.ForEach(Infos.Values, Info =>
+			using ITimelineEvent generateTransitive = Timeline.ScopeEvent("GenerateTransitiveIncludeLists");
+			ObjectPool<StringBuilder> stacks = ObjectPool.Create<StringBuilder>();
+			Parallel2.ForEach(Infos.Values, Info =>
 			{
-				Stack<string> Stack = new();
+				StringBuilder stack = stacks.Get();
 				Info.TransitiveIncludes.EnsureCapacity(300);
 				Info.TransitiveForwardDeclarations.EnsureCapacity(300);
-				CalculateTransitive(Info, Info, Stack, Info.TransitiveIncludes, Info.TransitiveForwardDeclarations, false);
+				CalculateTransitive(Info, Info, stack, Info.TransitiveIncludes, Info.TransitiveForwardDeclarations, false);
+				stacks.Return(stack);
 			});
+			generateTransitive.Finish();
 
 			// If we have built .gen.cpp it means that it is not inlined in another cpp file
 			// And we need to promote the includes needed to the header since we can never modify the .gen.cpp file
 			Logger.LogInformation($"Transferring needed includes from .gen.cpp to owning file...");
-			Parallel.ForEach(GeneratedCppInfos, GeneratedCpp =>
+			Parallel2.ForEach(GeneratedCppInfos, GeneratedCpp =>
 			{
 				// First, check which files .gen.cpp will see
 				HashSet<string> SeenTransitiveIncludes = new();
@@ -1093,12 +1147,40 @@ namespace UnrealBuildTool
 				}
 			});
 
+			if (!String.IsNullOrEmpty(PrintInfo))
+			{
+				PrintInfo = PrintInfo.Replace('\\', '/');
+				bool found = false;
+				foreach (var info in Infos.Values)
+				{
+					if (info.File.Contains(PrintInfo, StringComparison.OrdinalIgnoreCase))
+					{
+						Logger.LogInformation($"--- {info.File} ---");
+						Logger.LogInformation($"  Transitive includes");
+						foreach (var i in info.TransitiveIncludes)
+						{
+							Logger.LogInformation($"    {i.Key}");
+						}
+						found = true;
+					}
+				}
+				if (found)
+				{
+					return 0;
+				}
+
+				Logger.LogError($"PrintInfo file not found ({PrintInfo})");
+				return 1;
+			}
+
 			if (bCompare)
 			{
+				using ITimelineEvent compareFiles = Timeline.ScopeEvent("CompareFiles");
 				return CompareFiles(Infos, ValidPaths, PathToModule, NameToModule, Logger);
 			}
 			else
 			{
+				using ITimelineEvent updateFiles = Timeline.ScopeEvent("UpdateFiles");
 				return UpdateFiles(Infos, ValidPaths, Logger);
 			}
 		}
@@ -1194,7 +1276,7 @@ namespace UnrealBuildTool
 			while (true)
 			{
 				Dictionary<string, IWYUInfo> NewSkippedHeaders = new();
-				Parallel.ForEach(SkippedHeaders.Values, Info =>
+				Parallel2.ForEach(SkippedHeaders.Values, Info =>
 				{
 					Info.Module = GetModule(Info.File, PathToModule);
 
@@ -1307,7 +1389,9 @@ namespace UnrealBuildTool
 
 		int UpdateFiles(Dictionary<string, IWYUInfo> Infos, HashSet<string> ValidPaths, ILogger Logger)
 		{
-			object? ShouldLog = bWrite ? null : new();
+			SortedDictionary<string, List<string>>? sortedLog = null;
+			if (!bWrite)
+				sortedLog = [];
 
 			List<ValueTuple<IWYUInfo, List<string>>> UpdatedFiles = new(); // <FilePath>, <NewLines>
 			HashSet<IWYUFile> SkippedFiles = new();
@@ -1319,7 +1403,7 @@ namespace UnrealBuildTool
 			uint FilesParseCount = 0;
 			int ProcessSuccess = 1;
 
-			Parallel.ForEach(Infos.Values, Info =>
+			Parallel2.ForEach(Infos.Values, Info =>
 			{
 				if (!IsValidForUpdate(Info, ValidPaths, true))
 				{
@@ -1380,7 +1464,7 @@ namespace UnrealBuildTool
 								string? TransitivePath;
 								if (Include2.Resolved.TransitiveIncludes!.TryGetValue(Key, out TransitivePath))
 								{
-									if (ShouldLog != null)
+									if (sortedLog != null)
 									{
 										TransitivelyIncluded.TryAdd(QuotedPrintable, String.Join(" -> ", Include2.Printable, TransitivePath));
 									}
@@ -1447,6 +1531,11 @@ namespace UnrealBuildTool
 				HashSet<string> SeenIncludes = new();
 				foreach (IWYUIncludeEntry SeenInclude in Info.IncludesSeenInFile)
 				{
+					if (SeenInclude.Printable.EndsWith(".ipp")) // Abuse this to make sure we never remove .ipp includes
+					{
+						continue;
+					}
+
 					if (SeenInclude.System)
 					{
 						SeenIncludes.Add($"<{SeenInclude.Printable}>");
@@ -1626,30 +1715,33 @@ namespace UnrealBuildTool
 					LinesToAdd.Add("#include " + IncludeToAdd);
 				}
 
-				if (ShouldLog != null)
+				if (sortedLog != null)
 				{
-					lock (ShouldLog)
+					List<string> LogLines = new();
 					{
-						System.Console.WriteLine(Info.File);
 						foreach (string I in LinesToAdd)
 						{
-							System.Console.WriteLine("  +" + I);
+							LogLines.Add("  +" + I);
 						}
 						foreach (KeyValuePair<string, string> Pair in LinesToRemove)
 						{
-							System.Console.Write("  -" + Pair.Key);
+							string Line = "  -" + Pair.Key;
 							string? Reason;
 							if (TransitivelyIncluded.TryGetValue(Pair.Value, out Reason))
 							{
-								System.Console.Write("  (Transitively included from " + Reason + ")");
+								Line += "  (Transitively included from " + Reason + ")";
 							}
-							System.Console.WriteLine();
+							LogLines.Add(Line);
 						}
 						foreach (string I in ForwardDeclarationsToAdd)
 						{
-							System.Console.WriteLine("  +" + I);
+							LogLines.Add("  +" + I);
 						}
-						System.Console.WriteLine();
+					}
+
+					lock (sortedLog)
+					{
+						sortedLog.Add(Info.File, LogLines);
 					}
 				}
 
@@ -1787,6 +1879,8 @@ namespace UnrealBuildTool
 
 					if (ForwardDeclarationsToAdd.Count > 0)
 					{
+						// TODO: Need to make sure we are not in #if/#endif
+
 						int InsertPos;
 						if (FirstForwardDeclareLineIndex == -1)
 						{
@@ -1890,8 +1984,15 @@ namespace UnrealBuildTool
 			}
 
 			// If we have been logging we can exit now since we don't want to write any files to disk
-			if (ShouldLog != null)
+			if (sortedLog != null)
 			{
+				foreach (var kv in sortedLog)
+				{
+					Console.WriteLine(kv.Key);
+					foreach (string line in kv.Value)
+						Console.WriteLine(line);
+					Console.WriteLine();
+				}
 				return 0;
 			}
 
@@ -2026,7 +2127,7 @@ namespace UnrealBuildTool
 		/// <summary>
 		/// Calculate all indirect transitive includes for a file. This list contains does not contain itself and will handle circular dependencies
 		/// </summary>
-		static void CalculateTransitive(IWYUInfo Root, IWYUInfo Info, Stack<string> Stack, Dictionary<string, string> TransitiveIncludes, Dictionary<string, string> TransitiveForwardDeclarations, bool UseSeenIncludes)
+		static void CalculateTransitive(IWYUInfo Root, IWYUInfo Info, StringBuilder Stack, Dictionary<string, string> TransitiveIncludes, Dictionary<string, string> TransitiveForwardDeclarations, bool UseSeenIncludes)
 		{
 			List<IWYUIncludeEntry> Includes = UseSeenIncludes ? Info.IncludesSeenInFile : Info.Includes;
 			foreach (IWYUIncludeEntry Include in Includes)
@@ -2038,14 +2139,18 @@ namespace UnrealBuildTool
 					continue;
 				}
 
-				Stack.Push(Include.Printable);
-				string TransitivePath = String.Join(" -> ", Stack.Reverse());
-				TransitiveIncludes.Add(Key, TransitivePath);
+				int length = Stack.Length;
+				if (length > 0)
+				{
+					Stack.Append(" <- ");
+				}
+				Stack.Append(Include.Printable);
+				TransitiveIncludes.Add(Key, Stack.ToString());
 				if (Include.Resolved != null)
 				{
 					CalculateTransitive(Root, Include.Resolved, Stack, TransitiveIncludes, TransitiveForwardDeclarations, UseSeenIncludes);
 				}
-				Stack.Pop();
+				Stack.Length = length;
 			}
 
 			foreach (IWYUForwardEntry ForwardDeclaration in Info.ForwardDeclarations)
@@ -2062,7 +2167,7 @@ namespace UnrealBuildTool
 			Dictionary<string, IWYUInfo> SkippedHeaders = new();
 
 			Logger.LogInformation($"Parsing seen headers not supporting being compiled...");
-			Parallel.ForEach(Infos.Values, Info =>
+			Parallel2.ForEach(Infos.Values, Info =>
 			{
 				foreach (IWYUIncludeEntry Include in Info.IncludesSeenInFile)
 				{
@@ -2093,7 +2198,7 @@ namespace UnrealBuildTool
 			List<ValueTuple<IWYUInfo, Dictionary<string, string>>> InfosToCompare = new();
 
 			Logger.LogInformation($"Generating transitive seen include lists...");
-			Parallel.ForEach(Infos.Values, Info =>
+			Parallel2.ForEach(Infos.Values, Info =>
 			{
 				if (!IsValidForUpdate(Info, ValidPaths, false) || Info.File.EndsWith(".h"))
 				{
@@ -2101,8 +2206,8 @@ namespace UnrealBuildTool
 				}
 
 				Dictionary<string, string> SeenTransitiveIncludes = new();
-				Stack<string> Stack = new();
-				CalculateTransitive(Info, Info, Stack, SeenTransitiveIncludes, new Dictionary<string, string>(), true);
+				StringBuilder stack = new();
+				CalculateTransitive(Info, Info, stack, SeenTransitiveIncludes, new Dictionary<string, string>(), true);
 
 				lock (InfosToCompare)
 				{
@@ -2129,14 +2234,14 @@ namespace UnrealBuildTool
 			Dictionary<string, long> FileToSize = new Dictionary<string, long>(AllSeenIncludes);
 
 			Logger.LogInformation($"Reading file sizes...");
-			Parallel.ForEach(AllSeenIncludes, Info =>
+			Parallel2.ForEach(AllSeenIncludes, Info =>
 			{
 				FileToSize[Info.Key] = (new FileInfo(Info.Key)).Length;
 			});
 
 			Logger.LogInformation($"Calculate total amount of bytes included per file...");
 			List<ValueTuple<long, IWYUInfo>> ByteSizeDiffList = new();
-			Parallel.ForEach(InfosToCompare, Kvp =>
+			Parallel2.ForEach(InfosToCompare, Kvp =>
 			{
 				long OptimizedSize = 0;
 				long SeenSize = 0;
diff --git a/Engine/Source/Programs/UnrealBuildTool/ToolChain/IWYUToolChain.cs b/Engine/Source/Programs/UnrealBuildTool/ToolChain/IWYUToolChain.cs
--- a/Engine/Source/Programs/UnrealBuildTool/ToolChain/IWYUToolChain.cs
+++ b/Engine/Source/Programs/UnrealBuildTool/ToolChain/IWYUToolChain.cs
@@ -125,6 +125,8 @@ namespace UnrealBuildTool
 			// Arguments.Add("-Xiwyu --no_check_matching_header");
 			Arguments.Add("-Xiwyu --cxx17ns");
 
+			CompileEnvironment.CppCompileWarnings.UnusedWarningLevel = WarningLevel.Off;
+
 			Arguments.Add("-Wno-delete-incomplete");
 			Arguments.Add("-Wno-unknown-warning-option");
 
@@ -237,14 +239,14 @@ namespace UnrealBuildTool
 		{
 			FileItem fileItem = base.GetCompileArguments_FileType(CompileEnvironment, SourceFile, OutputDir, Arguments, CompileAction, CompileResult);
 
-			Arguments.Add("-Xiwyu --write_json_path=\"" + fileItem.AbsolutePath.Replace('\\', '/') + "\" ");
+			Arguments.Add($"-Xiwyu --write_json_path=\"{fileItem.AbsolutePath.Replace('\\', '/')}\"");
 			if (SourceFile.HasExtension(".cpp"))
 			{
 				if (CompileEnvironment.FileInlineSourceMap.TryGetValue(SourceFile, out HashSet<FileItem>? InlinedFiles))
 				{
 					foreach (FileItem InlinedFile in InlinedFiles)
 					{
-						Arguments.Add($"-Xiwyu --check_also={InlinedFile.FullName.Replace('\\', '/')} ");
+						Arguments.Add($"-Xiwyu --check_also={InlinedFile.FullName.Replace('\\', '/')}");
 					}
 				}
 			}
@@ -252,6 +254,7 @@ namespace UnrealBuildTool
 			{
 				// We want to build header files first
 				CompileAction.bIsHighPriority = true;
+				Arguments.Add($"-Xiwyu --check_also=*/{SourceFile.Location.GetFileNameWithoutAnyExtensions()}.generated.h");
 			}
 			return fileItem;
 		}
