Resolve various issues with iOS remote build.
1 - Fix remote codesigning by ensuring that intemediate certificates (Specifically Apple WWDR G3) are installed to the remote Mac's temporary keychain that's setup for remote codesigning.
2 - Use legacy rsync 2.6.9 for Mac Remote builds. MacOS Sequoia transitioned from rsync to openrsync causing some incompatibilities of the --delete* arguments resulting in erroneous file deletions on the invoking Windows PC.

diff --git a/Engine/Build/Rsync/RsyncEngine.txt b/Engine/Build/Rsync/RsyncEngine.txt
--- a/Engine/Build/Rsync/RsyncEngine.txt
+++ b/Engine/Build/Rsync/RsyncEngine.txt
@@ -214,6 +214,7 @@ protect /Source/Programs/Shared/EpicGames.*/***
 + /Source/ThirdParty/**/IOS/**.xcframework/**
 + /Source/ThirdParty/**/TVOS/**.xcframework/**
 - /Source/ThirdParty/**
+- /Source/**/.vs/**
 
 # Source: Include everything else
 + /Source/**/
diff --git a/Engine/Source/Programs/UnrealBuildTool/Platform/IOS/IOSToolChain.cs b/Engine/Source/Programs/UnrealBuildTool/Platform/IOS/IOSToolChain.cs
--- a/Engine/Source/Programs/UnrealBuildTool/Platform/IOS/IOSToolChain.cs
+++ b/Engine/Source/Programs/UnrealBuildTool/Platform/IOS/IOSToolChain.cs
@@ -1448,7 +1448,7 @@ namespace UnrealBuildTool
 					// Copy the mobile provision into the system store
 					if (Target.ImportProvision != null)
 					{
-						Writer.WriteLine("cp -f {0} ~/Library/MobileDevice/Provisioning\\ Profiles/", Utils.EscapeShellArgument(Target.ImportProvision));
+						Writer.WriteLine("cp -f {0} {1}", Utils.EscapeShellArgument(Target.ImportProvision), Utils.EscapeShellArgument(AppleExports.GetProvisionDirectory().FullName));
 					}
 
 					// Get the signing certificate to use
diff --git a/Engine/Source/Programs/UnrealBuildTool/Platform/IOS/UEBuildIOS.cs b/Engine/Source/Programs/UnrealBuildTool/Platform/IOS/UEBuildIOS.cs
--- a/Engine/Source/Programs/UnrealBuildTool/Platform/IOS/UEBuildIOS.cs
+++ b/Engine/Source/Programs/UnrealBuildTool/Platform/IOS/UEBuildIOS.cs
@@ -472,17 +472,7 @@ namespace UnrealBuildTool
 
 			if (!String.IsNullOrEmpty(MobileProvision))
 			{
-				DirectoryReference MobileProvisionDir;
-				if (BuildHostPlatform.Current.Platform == UnrealTargetPlatform.Mac)
-				{
-					MobileProvisionDir = DirectoryReference.Combine(new DirectoryReference(Environment.GetEnvironmentVariable("HOME")!), "Library", "MobileDevice", "Provisioning Profiles");
-				}
-				else
-				{
-					MobileProvisionDir = DirectoryReference.Combine(DirectoryReference.GetSpecialFolder(Environment.SpecialFolder.LocalApplicationData)!, "Apple Computer", "MobileDevice", "Provisioning Profiles");
-				}
-
-				FileReference PossibleMobileProvisionFile = FileReference.Combine(MobileProvisionDir, MobileProvision);
+				FileReference PossibleMobileProvisionFile = FileReference.Combine(AppleExports.GetProvisionDirectory(), MobileProvision);
 				if (FileReference.Exists(PossibleMobileProvisionFile))
 				{
 					MobileProvisionFile = PossibleMobileProvisionFile;
diff --git a/Engine/Source/Programs/UnrealBuildTool/ToolChain/AppleToolChain.cs b/Engine/Source/Programs/UnrealBuildTool/ToolChain/AppleToolChain.cs
--- a/Engine/Source/Programs/UnrealBuildTool/ToolChain/AppleToolChain.cs
+++ b/Engine/Source/Programs/UnrealBuildTool/ToolChain/AppleToolChain.cs
@@ -1239,7 +1239,7 @@ namespace UnrealBuildTool
 				}
 
 				// copy the provision into standard location
-				Writer.WriteLine("cp -f {0} ~/Library/MobileDevice/Provisioning\\ Profiles/", Utils.EscapeShellArgument(Target.RemoteImportProvision));
+				Writer.WriteLine("cp -f {0} {1}", Utils.EscapeShellArgument(Target.RemoteImportProvision), Utils.EscapeShellArgument(AppleExports.GetProvisionDirectory().FullName));
 				MobileProvisionContents MobileProvision = MobileProvisionContents.Read(new FileReference(Target.RemoteImportProvision));
 				MobileProvisionUUID = MobileProvision.GetUniqueId();
 
@@ -1263,6 +1263,7 @@ namespace UnrealBuildTool
 				Writer.WriteLine("security list-keychains");
 				Writer.WriteLine("security set-keychain-settings -t 3600 -l  \"{0}\"", TempKeychain);
 				Writer.WriteLine("security -v unlock-keychain -p \"A\" \"{0}\"", TempKeychain);
+				Writer.WriteLine("security import {0}/Build/IOS/AppleWorldwideDeveloperRelationsCA.pem -k \"{1}\"", Utils.EscapeShellArgument(Unreal.EngineDirectory.FullName), TempKeychain);
 				Writer.WriteLine("security import {0} -P {1} -k \"{2}\" -T /usr/bin/codesign -T /usr/bin/security -t agg", Utils.EscapeShellArgument(Target.RemoteImportCertificate), Utils.EscapeShellArgument(Target.RemoteImportCertificatePassword!), TempKeychain);
 				Writer.WriteLine("security set-key-partition-list -S apple-tool:,apple:,codesign: -s -k \"A\" -D '{0}' -t private {1}", SigningCertificate, TempKeychain);
 			}
diff --git a/Engine/Source/Programs/UnrealBuildTool/ToolChain/RemoteMac.cs b/Engine/Source/Programs/UnrealBuildTool/ToolChain/RemoteMac.cs
--- a/Engine/Source/Programs/UnrealBuildTool/ToolChain/RemoteMac.cs
+++ b/Engine/Source/Programs/UnrealBuildTool/ToolChain/RemoteMac.cs
@@ -119,6 +119,20 @@ namespace UnrealBuildTool
 		/// </summary>
 		private List<string> CommonRsyncArguments;
 
+		private List<string> RsyncArgumentsForDownload(List<string> BaseArguments)
+		{
+			List<string> Arguments = new List<string>(BaseArguments);
+			Arguments.Add(String.Format("--rsync-path=\"{0}/rsync\"", RemoteBaseDir));
+			return Arguments;
+		}
+
+		private List<string> RsyncArgumentsForUpload(List<string> BaseArguments, string RemoteDirectory)
+		{
+			List<string> Arguments = new List<string>(BaseArguments);
+			Arguments.Add(String.Format("--rsync-path=\"mkdir -p {0} && {1}/rsync\"", RemoteDirectory, RemoteBaseDir));
+			return Arguments;
+		}
+
 		private string? IniBundleIdentifier = "";
 
 		/// <summary>
@@ -131,8 +145,9 @@ namespace UnrealBuildTool
 		/// <param name="bPrepareForSecondaryMac">Added arguments when preparing a secondary Mac for debug. False by default.</param>
 		public RemoteMac(FileReference? ProjectFile, ILogger Logger, string? RemoteIniPath = null, bool bIsPrimary = true, bool bPrepareForSecondaryMac = false)
 		{
-			RsyncExe = FileReference.Combine(Unreal.EngineDirectory, "Extras", "ThirdPartyNotUE", "cwrsync", "bin", "rsync.exe");
-			SshExe = FileReference.Combine(Unreal.EngineDirectory, "Extras", "ThirdPartyNotUE", "cwrsync", "bin", "ssh.exe");
+			DirectoryReference Cwrsync = DirectoryReference.Combine(Unreal.EngineDirectory, "Extras", "ThirdPartyNotUE", "cwrsync");
+			RsyncExe = FileReference.Combine(Cwrsync, "bin", "rsync.exe");
+			SshExe = FileReference.Combine(Cwrsync, "bin", "ssh.exe");
 			this.ProjectFile = ProjectFile;
 			if (ProjectFile != null)
 			{
@@ -187,12 +202,12 @@ namespace UnrealBuildTool
 			int PortIdx = ServerName.LastIndexOf(':');
 			if (PortIdx != -1)
 			{
-				string Port = ServerName.Substring(PortIdx + 1);
+				string Port = ServerName[(PortIdx + 1)..];
 				if (!Int32.TryParse(Port, out ServerPort))
 				{
 					throw new BuildException("Unable to parse port number from '{0}'", ServerName);
 				}
-				ServerName = ServerName.Substring(0, PortIdx);
+				ServerName = ServerName[..PortIdx];
 			}
 
 			// If a user name is not set, use the current user
@@ -315,6 +330,23 @@ namespace UnrealBuildTool
 					}
 				}
 			}
+
+			string RemoteFile = GetRemotePath("rsync");
+
+			List<string> Arguments = new List<string>(CommonRsyncArguments);
+			// use MacOS default system rsync for the upload of the 2.6.9 rsync binary
+			Arguments.Add(String.Format("--rsync-path=\"mkdir -p {0} && rsync\"", RemoteBaseDir));
+			Arguments.Add(String.Format("\"{0}\"", GetLocalCygwinPath(FileReference.Combine(Cwrsync, "mac", "rsync"))));
+			Arguments.Add(String.Format("\"{0}@{1}\":'{2}'", UserName, ServerName, RemoteFile));
+			Arguments.Add("-q");
+
+			int Result = Rsync(String.Join(" ", Arguments), Logger);
+			if (Result != 0)
+			{
+				throw new BuildException("Error while running Rsync (exit code {0})", Result);
+			}
+			// Remove quaratine extended attribute if remote machine set it.
+			Execute(RemoteBaseDir, "if xattr -p com.apple.quarantine rsync > /dev/null 2>&1; then xattr -d com.apple.quarantine rsync; fi", Logger);
 		}
 
 		/// <summary>
@@ -662,8 +694,8 @@ namespace UnrealBuildTool
 					continue;
 				}
 
-				string Key = LocalArgument.Substring(0, EqualsIdx);
-				string Value = LocalArgument.Substring(EqualsIdx + 1);
+				string Key = LocalArgument[..EqualsIdx];
+				string Value = LocalArgument[(EqualsIdx + 1)..];
 
 				if (Key.Equals("-Log", StringComparison.InvariantCultureIgnoreCase))
 				{
@@ -797,7 +829,7 @@ namespace UnrealBuildTool
 			{
 				if (RemotePath.StartsWith(Mapping.RemoteDirectory, StringComparison.InvariantCultureIgnoreCase) && RemotePath.Length > Mapping.RemoteDirectory.Length && RemotePath[Mapping.RemoteDirectory.Length] == '/')
 				{
-					return FileReference.Combine(Mapping.LocalDirectory, RemotePath.Substring(Mapping.RemoteDirectory.Length + 1));
+					return FileReference.Combine(Mapping.LocalDirectory, RemotePath[(Mapping.RemoteDirectory.Length + 1)..]);
 				}
 			}
 			throw new BuildException("Unable to map remote path '{0}' to local path", RemotePath);
@@ -834,7 +866,7 @@ namespace UnrealBuildTool
 			{
 				throw new BuildException("Invalid local path for converting to cygwin format ({0}).", InPath);
 			}
-			return String.Format("/cygdrive/{0}{1}", InPath.FullName.Substring(0, 1), InPath.FullName.Substring(2).Replace('\\', '/'));
+			return String.Format("/cygdrive/{0}{1}", InPath.FullName[..1], InPath.FullName[2..].Replace('\\', '/'));
 		}
 
 		/// <summary>
@@ -857,8 +889,7 @@ namespace UnrealBuildTool
 			string RemoteFile = GetRemotePath(LocalFile);
 			string RemoteDirectory = GetRemotePath(LocalFile.Directory);
 
-			List<string> Arguments = new List<string>(CommonRsyncArguments);
-			Arguments.Add(String.Format("--rsync-path=\"mkdir -p {0} && rsync\"", RemoteDirectory));
+			List<string> Arguments = RsyncArgumentsForUpload(CommonRsyncArguments, RemoteDirectory);
 			Arguments.Add(String.Format("\"{0}\"", GetLocalCygwinPath(LocalFile)));
 			Arguments.Add(String.Format("\"{0}@{1}\":'{2}'", UserName, ServerName, RemoteFile));
 			Arguments.Add("-q");
@@ -879,8 +910,7 @@ namespace UnrealBuildTool
 		/// <param name="Logger">Logger for output</param>
 		void UploadFiles(DirectoryReference LocalDirectory, string RemoteDirectory, FileReference LocalFileList, ILogger Logger)
 		{
-			List<string> Arguments = new List<string>(BasicRsyncArguments);
-			Arguments.Add(String.Format("--rsync-path=\"mkdir -p {0} && rsync\"", RemoteDirectory));
+			List<string> Arguments = RsyncArgumentsForUpload(BasicRsyncArguments, RemoteDirectory);
 			Arguments.Add(String.Format("--files-from=\"{0}\"", GetLocalCygwinPath(LocalFileList)));
 			Arguments.Add(String.Format("\"{0}/\"", GetLocalCygwinPath(LocalDirectory)));
 			Arguments.Add(String.Format("\"{0}@{1}\":'{2}/'", UserName, ServerName, RemoteDirectory));
@@ -902,8 +932,7 @@ namespace UnrealBuildTool
 		{
 			string RemoteDirectory = GetRemotePath(LocalDirectory);
 
-			List<string> Arguments = new List<string>(CommonRsyncArguments);
-			Arguments.Add(String.Format("--rsync-path=\"mkdir -p {0} && rsync\"", RemoteDirectory));
+			List<string> Arguments = RsyncArgumentsForUpload(CommonRsyncArguments, RemoteDirectory);
 			Arguments.Add(String.Format("\"{0}/\"", GetLocalCygwinPath(LocalDirectory)));
 			Arguments.Add(String.Format("\"{0}@{1}\":'{2}/'", UserName, ServerName, RemoteDirectory));
 			Arguments.Add("-q");
@@ -924,8 +953,7 @@ namespace UnrealBuildTool
 		/// <param name="Logger">Logger for output</param>
 		void UploadDirectory(DirectoryReference LocalDirectory, string RemoteDirectory, List<FileReference> FilterLocations, ILogger Logger)
 		{
-			List<string> Arguments = new List<string>(CommonRsyncArguments);
-			Arguments.Add(String.Format("--rsync-path=\"mkdir -p {0} && rsync\"", RemoteDirectory));
+			List<string> Arguments = RsyncArgumentsForUpload(CommonRsyncArguments, RemoteDirectory);
 			foreach (FileReference FilterLocation in FilterLocations)
 			{
 				Arguments.Add(String.Format("--filter=\"merge {0}\"", GetLocalCygwinPath(FilterLocation)));
@@ -1093,7 +1121,7 @@ namespace UnrealBuildTool
 				throw new BuildException("File for download '{0}' is not under any mapped directory.", LocalFile);
 			}
 
-			List<string> Arguments = new List<string>(CommonRsyncArguments);
+			List<string> Arguments = RsyncArgumentsForDownload(CommonRsyncArguments);
 			Arguments.Add(String.Format("\"{0}@{1}\":'{2}/{3}'", UserName, ServerName, Mapping.RemoteDirectory, LocalFile.MakeRelativeTo(Mapping.LocalDirectory).Replace('\\', '/')));
 			Arguments.Add(String.Format("\"{0}/\"", GetLocalCygwinPath(LocalFile.Directory)));
 			Arguments.Add("-q");
@@ -1134,7 +1162,7 @@ namespace UnrealBuildTool
 					DirectoryReference.CreateDirectory(DownloadListLocation.Directory);
 					FileReference.WriteAllLines(DownloadListLocation, FileGroups[Idx].Select(x => x.MakeRelativeTo(Mappings[Idx].LocalDirectory).Replace('\\', '/')));
 
-					List<string> Arguments = new List<string>(CommonRsyncArguments);
+					List<string> Arguments = RsyncArgumentsForDownload(CommonRsyncArguments);
 					Arguments.Add(String.Format("--files-from=\"{0}\"", GetLocalCygwinPath(DownloadListLocation)));
 					Arguments.Add(String.Format("\"{0}@{1}\":'{2}/'", UserName, ServerName, Mappings[Idx].RemoteDirectory));
 					Arguments.Add(String.Format("\"{0}/\"", GetLocalCygwinPath(Mappings[Idx].LocalDirectory)));
@@ -1171,7 +1199,7 @@ namespace UnrealBuildTool
 
 			string RemoteDirectory = GetRemotePath(LocalDirectory);
 
-			List<string> Arguments = new List<string>(CommonRsyncArguments);
+			List<string> Arguments = RsyncArgumentsForDownload(CommonRsyncArguments);
 			Arguments.Add(String.Format("\"{0}@{1}\":'{2}/'", UserName, ServerName, RemoteDirectory));
 			Arguments.Add(String.Format("\"{0}/\"", GetLocalCygwinPath(LocalDirectory)));
 
@@ -1363,7 +1391,7 @@ namespace UnrealBuildTool
 				{
 					if (Path.StartsWith(Mapping.RemoteDirectory))
 					{
-						Path = Mapping.LocalDirectory + Path.Substring(Mapping.RemoteDirectory.Length).Replace('/', '\\');
+						Path = Mapping.LocalDirectory + Path[Mapping.RemoteDirectory.Length..].Replace('/', '\\');
 						break;
 					}
 				}
