UnrealBuildTool: Adjust .precompiled location to architecture specific intermediate directory only for single architecture per target platforms

#rnx

diff --git a/Engine/Source/Programs/UnrealBuildTool/Configuration/UEBuildModuleCPP.cs b/Engine/Source/Programs/UnrealBuildTool/Configuration/UEBuildModuleCPP.cs
--- a/Engine/Source/Programs/UnrealBuildTool/Configuration/UEBuildModuleCPP.cs
+++ b/Engine/Source/Programs/UnrealBuildTool/Configuration/UEBuildModuleCPP.cs
@@ -459,7 +459,9 @@ namespace UnrealBuildTool
 		/// <summary>
 		/// Path to the precompiled manifest location
 		/// </summary>
-		public virtual FileReference PrecompiledManifestLocation => FileReference.Combine(IntermediateDirectory, String.Format("{0}.precompiled", Name));
+		public virtual FileReference PrecompiledManifestLocation => UnrealArchitectureConfig.ForPlatform(Rules.Target.Platform).Mode <= UnrealArchitectureMode.OneTargetPerArchitecture
+			? FileReference.Combine(IntermediateDirectory, $"{Name}.precompiled")
+			: FileReference.Combine(IntermediateDirectoryNoArch, $"{Name}.precompiled");
 
 		/// <summary>
 		/// Sets up the environment for compiling any module that includes the public interface of this module.
