UnrealBuildTool: Rewrite HotReloadState serializer to use BinaryArchive

#rnx

diff --git a/Engine/Source/Programs/UnrealBuildTool/System/HotReload.cs b/Engine/Source/Programs/UnrealBuildTool/System/HotReload.cs
--- a/Engine/Source/Programs/UnrealBuildTool/System/HotReload.cs
+++ b/Engine/Source/Programs/UnrealBuildTool/System/HotReload.cs
@@ -31,9 +31,13 @@ namespace UnrealBuildTool
 	/// <summary>
 	/// Stores the current hot reload state, tracking temporary files created by previous invocations.
 	/// </summary>
-	[Serializable]
 	class HotReloadState
 	{
+		/// <summary>
+		/// The version number to write
+		/// </summary>
+		public const int CurrentVersion = 1;
+
 		/// <summary>
 		/// Suffix to use for the next hot reload invocation
 		/// </summary>
@@ -49,6 +53,28 @@ namespace UnrealBuildTool
 		/// </summary>
 		public HashSet<FileReference> TemporaryFiles = new HashSet<FileReference>();
 
+		/// <summary>
+		/// Public constructor
+		/// </summary>
+		public HotReloadState()
+		{
+		}
+
+		/// <summary>
+		/// Private constructor
+		/// </summary>
+		/// <param name="Reader"></param>
+		HotReloadState(BinaryArchiveReader Reader)
+		{
+			int ArchiveVersion = Reader.ReadInt();
+			if (CurrentVersion == ArchiveVersion)
+			{
+				NextSuffix = Reader.ReadInt();
+				OriginalFileToHotReloadFile = Reader.ReadDictionary(Reader.ReadFileReference, Reader.ReadFileReference) ?? [];
+				TemporaryFiles = Reader.ReadHashSet(Reader.ReadFileReference) ?? [];
+			}
+		}
+
 		/// <summary>
 		/// Adds all the actions into the hot reload state, so we can restore the action graph on next iteration
 		/// </summary>
@@ -123,7 +149,7 @@ namespace UnrealBuildTool
 		public static FileReference GetLocation(FileReference? ProjectFile, string TargetName, UnrealTargetPlatform Platform, UnrealTargetConfiguration Configuration, UnrealArchitectures Architectures)
 		{
 			DirectoryReference BaseDir = DirectoryReference.FromFile(ProjectFile) ?? Unreal.EngineDirectory;
-			return FileReference.Combine(BaseDir, UEBuildTarget.GetPlatformIntermediateFolder(Platform, Architectures, false), TargetName, Configuration.ToString(), "HotReload.state");
+			return FileReference.Combine(BaseDir, UEBuildTarget.GetPlatformIntermediateFolder(Platform, Architectures, false), TargetName, Configuration.ToString(), "HotReloadState.bin");
 		}
 
 		/// <summary>
@@ -133,7 +159,8 @@ namespace UnrealBuildTool
 		/// <returns>New hot reload state instance</returns>
 		public static HotReloadState Load(FileReference Location)
 		{
-			return BinaryFormatterUtils.Load<HotReloadState>(Location);
+			using BinaryArchiveReader Reader = new(Location);
+			return new(Reader);
 		}
 
 		/// <summary>
@@ -143,7 +170,11 @@ namespace UnrealBuildTool
 		public void Save(FileReference Location)
 		{
 			DirectoryReference.CreateDirectory(Location.Directory);
-			BinaryFormatterUtils.Save(Location, this);
+			using BinaryArchiveWriter Writer = new(Location);
+			Writer.WriteInt(CurrentVersion);
+			Writer.WriteInt(NextSuffix);
+			Writer.WriteDictionary(OriginalFileToHotReloadFile, Writer.WriteFileReference, Writer.WriteFileReference);
+			Writer.WriteHashSet(TemporaryFiles, Writer.WriteFileReference);
 		}
 	}
 
