Fixed crash in UBT when reporting Engine Files error

diff --git a/Engine/Source/Programs/UnrealBuildTool/Modes/BuildMode.cs b/Engine/Source/Programs/UnrealBuildTool/Modes/BuildMode.cs
--- a/Engine/Source/Programs/UnrealBuildTool/Modes/BuildMode.cs
+++ b/Engine/Source/Programs/UnrealBuildTool/Modes/BuildMode.cs
@@ -730,7 +730,7 @@ namespace UnrealBuildTool
 				// Make sure we're not modifying any existing engine files while allowing new engine files to be created
 				if (Options.HasFlag(BuildOptions.NoEngineChanges))
 				{
-					IEnumerable<FileItem> EngineChanges = MergedActionsToExecute.SelectMany(x => x.ProducedItems).Where(x => x.Location.IsUnderDirectory(Unreal.EngineDirectory) && x.Exists).Distinct().Order();
+					IEnumerable<FileItem> EngineChanges = MergedActionsToExecute.SelectMany(x => x.ProducedItems).Where(x => x.Location.IsUnderDirectory(Unreal.EngineDirectory) && x.Exists).Distinct().OrderBy(x => x.AbsolutePath);
 					if (EngineChanges.Any())
 					{
 						StringBuilder Result = new("Building would modify the following existing engine files:");
