[UBT]
* Added more timer scopes

diff --git a/Engine/Source/Programs/UnrealBuildTool/ProjectFiles/VisualStudio/VCProjectFileGenerator.cs b/Engine/Source/Programs/UnrealBuildTool/ProjectFiles/VisualStudio/VCProjectFileGenerator.cs
--- a/Engine/Source/Programs/UnrealBuildTool/ProjectFiles/VisualStudio/VCProjectFileGenerator.cs
+++ b/Engine/Source/Programs/UnrealBuildTool/ProjectFiles/VisualStudio/VCProjectFileGenerator.cs
@@ -676,6 +676,8 @@ namespace UnrealBuildTool
 
 		protected override bool WritePrimaryProjectFile(ProjectFile? UBTProject, PlatformProjectGeneratorCollection PlatformProjectGenerators, ILogger Logger)
 		{
+			using ITimelineEvent _ = Timeline.ScopeEvent("Writing solution file", Logger);
+
 			bool bSuccess = true;
 
 			string SolutionFileName = PrimaryProjectName + ".sln";
diff --git a/Engine/Source/Programs/UnrealBuildTool/System/RulesCompiler.cs b/Engine/Source/Programs/UnrealBuildTool/System/RulesCompiler.cs
--- a/Engine/Source/Programs/UnrealBuildTool/System/RulesCompiler.cs
+++ b/Engine/Source/Programs/UnrealBuildTool/System/RulesCompiler.cs
@@ -136,6 +136,8 @@ namespace UnrealBuildTool
 		/// <returns>New rules assembly</returns>
 		private static RulesAssembly CreateEngineRulesAssemblyInternal(List<DirectoryReference> RootDirectories, string AssemblyPrefix, IReadOnlyList<PluginInfo> Plugins, bool bReadOnly, bool bSkipCompile, bool bForceCompile, RulesAssembly? Parent, ILogger Logger)
 		{
+			using ITimelineEvent _ = Timeline.ScopeEvent("EngineRulesAssembly", Logger);
+
 			// Scope hierarchy
 			RulesScope Scope = new RulesScope("Engine", null);
 			RulesScope PluginsScope = new RulesScope("Engine Plugins", Scope);
