[UBT]
* Optimized ModuleRules.IsValidForTarget by using cache for type attribute queries
* Fixed so SourceFileMetadataCache is using Lazy in concurrentdictyionary value to prevent multiple creations

diff --git a/Engine/Source/Programs/UnrealBuildTool/Configuration/ModuleRules.cs b/Engine/Source/Programs/UnrealBuildTool/Configuration/ModuleRules.cs
--- a/Engine/Source/Programs/UnrealBuildTool/Configuration/ModuleRules.cs
+++ b/Engine/Source/Programs/UnrealBuildTool/Configuration/ModuleRules.cs
@@ -1,6 +1,7 @@
 // Copyright Epic Games, Inc. All Rights Reserved.
 
 using System;
+using System.Collections.Concurrent;
 using System.Collections.Generic;
 using System.Diagnostics.CodeAnalysis;
 using System.IO;
@@ -1060,7 +1061,7 @@ namespace UnrealBuildTool
 		/// <summary>
 		/// Source file types that will be placed into unity files when enabled.
 		/// </summary>
-		public UnityFileType UnityFileTypes { get; set;  } = UnityFileType.Cpp;
+		public UnityFileType UnityFileTypes { get; set; } = UnityFileType.Cpp;
 
 		/// <summary>
 		/// The number of source files in this module before unity build will be activated for that module.  If set to
@@ -1152,7 +1153,7 @@ namespace UnrealBuildTool
 		/// that uses transactionally-disabled code in places that are transactionally-enabled.
 		/// </summary>
 		public bool bInjectAutoRTFMAttributeDisables { get; set; }
-		
+
 		/// <summary>
 		/// If true, autogenerated functions of USTRUCTS are always exported
 		/// </summary>
@@ -1167,12 +1168,12 @@ namespace UnrealBuildTool
 		/// If true, this module can contain UE types (i.e. UCLASS) in namespaces
 		/// </summary>
 		public bool bAllowUETypesInNamespaces = false;
-  
+
 		/// <summary>
 		/// If true, minimizes the amount of header files included in .generated.h files, which may require more headers from CoreUObject to be explicitly included.
 		/// </summary>
 		public bool bMinimizeGeneratedIncludes = false;
-		
+
 		/// <summary>
 		/// List of modules names (no path needed) with header files that our module's public headers needs access to, but we don't need to "import" or link against.
 		/// </summary>
@@ -1694,15 +1695,15 @@ namespace UnrealBuildTool
 		/// <returns>True if the module is valid, false otherwise</returns>
 		internal static bool IsValidForTarget(Type moduleType, ReadOnlyTargetRules targetRules, [NotNullWhen(false)] out string? invalidReason)
 		{
-			IEnumerable<TargetType> supportedTargetTypes = moduleType.GetCustomAttributes<SupportedTargetTypesAttribute>().SelectMany(x => x.TargetTypes).Distinct();
-			if (supportedTargetTypes.Any() && !supportedTargetTypes.Contains(targetRules.Type))
+			HashSet<TargetType>? targetTypes = GetSupportedTargetTypesByAttribute(moduleType);
+			if (targetTypes != null && !targetTypes.Contains(targetRules.Type))
 			{
 				invalidReason = $"TargetType '{targetRules.Type}'";
 				return false;
 			}
 
-			IEnumerable<UnrealTargetConfiguration> supportedConfigurations = moduleType.GetCustomAttributes<SupportedConfigurationsAttribute>().SelectMany(x => x.Configurations).Distinct();
-			if (supportedConfigurations.Any() && !supportedConfigurations.Contains(targetRules.Configuration))
+			UnrealTargetConfiguration[]? supportedConfigurations = TargetRules.GetSupportedConfigurationFromAttribute(moduleType);
+			if (supportedConfigurations != null && !supportedConfigurations.Contains(targetRules.Configuration))
 			{
 				invalidReason = $"Configuration '{targetRules.Configuration}'";
 				return false;
@@ -1710,11 +1711,17 @@ namespace UnrealBuildTool
 
 			// Skip platform extension modules. We only care about the base modules, not the platform overrides.
 			// The platform overrides get applied at a later stage when we actually come to build the module.
-			if (!UEBuildPlatform.GetPlatformFolderNames().Any(name => moduleType.Name.EndsWith("_" + name, StringComparison.OrdinalIgnoreCase)))
+			bool isPlatformExtension = false;
+			int lastUnderscore = moduleType.Name.LastIndexOf('_');
+			if (lastUnderscore != -1)
 			{
-				IEnumerable<SupportedPlatformsAttribute> platformAttributes = moduleType.GetCustomAttributes<SupportedPlatformsAttribute>();
-				IEnumerable<UnrealTargetPlatform> supportedPlatforms = platformAttributes.SelectMany(x => x.Platforms).Distinct();
-				if (platformAttributes.Any() && !supportedPlatforms.Contains(targetRules.Platform))
+				string potentialPlatform = moduleType.Name.Substring(lastUnderscore);
+				isPlatformExtension = UEBuildPlatform.GetPlatformFolderNames().Any(s => s.Equals(potentialPlatform, StringComparison.OrdinalIgnoreCase));
+			}
+			if (!isPlatformExtension)
+			{
+				HashSet<UnrealTargetPlatform>? supportedPlatforms = GetSupportedPlatformsByAttribute(moduleType);
+				if (supportedPlatforms != null && !supportedPlatforms.Contains(targetRules.Platform))
 				{
 					invalidReason = $"Platform '{targetRules.Platform}'";
 					return false;
@@ -2017,5 +2024,73 @@ namespace UnrealBuildTool
 			}
 			definitions.Add(newModule);
 		}
+
+		/// <summary>
+		/// Cached function for getting supported target types from type attributes
+		/// </summary>
+		/// <param name="t"></param>
+		/// <returns></returns>
+		internal static HashSet<TargetType>? GetSupportedTargetTypesByAttribute(Type t) => s_supportedTargetTypesAttribute.GetOrAdd(t, static type => new Lazy<HashSet<TargetType>?>(() =>
+		{
+			object[] attributes = type.GetCustomAttributes(typeof(SupportedTargetTypesAttribute), false);
+			if (attributes.Length == 0)
+			{
+				if (type.BaseType != null)
+				{
+					return GetSupportedTargetTypesByAttribute(type.BaseType);
+				}
+				else
+				{
+					return null;
+				}
+			}
+			HashSet<TargetType> targetTypes = new(attributes.OfType<SupportedTargetTypesAttribute>().SelectMany(x => x.TargetTypes));
+			if (type.BaseType != null)
+			{
+				HashSet<TargetType>? parentTargetTypes = GetSupportedTargetTypesByAttribute(type.BaseType);
+				if (parentTargetTypes != null)
+				{
+					targetTypes.UnionWith(parentTargetTypes);
+				}
+			}
+
+			return targetTypes;
+
+		})).Value;
+		static readonly ConcurrentDictionary<Type, Lazy<HashSet<TargetType>?>> s_supportedTargetTypesAttribute = new();
+
+		/// <summary>
+		/// Cached function for getting supported platforms from type attributes
+		/// </summary>
+		/// <param name="t"></param>
+		/// <returns></returns>
+		internal static HashSet<UnrealTargetPlatform>? GetSupportedPlatformsByAttribute(Type t) => s_supportedPlatformsAttribute.GetOrAdd(t, static type => new Lazy<HashSet<UnrealTargetPlatform>?>(() =>
+		{
+			object[] attributes = type.GetCustomAttributes(typeof(SupportedPlatformsAttribute), false);
+			if (attributes.Length == 0)
+			{
+				if (type.BaseType != null)
+				{
+					return GetSupportedPlatformsByAttribute(type.BaseType);
+				}
+				else
+				{
+					return null;
+				}
+			}
+			HashSet<UnrealTargetPlatform> platforms = new(attributes.OfType<SupportedPlatformsAttribute>().SelectMany(x => x.Platforms));
+			if (type.BaseType != null)
+			{
+				HashSet<UnrealTargetPlatform>? parentPlatforms = GetSupportedPlatformsByAttribute(type.BaseType);
+				if (parentPlatforms != null)
+				{
+					platforms.UnionWith(parentPlatforms);
+				}
+			}
+
+			return platforms;
+
+		})).Value;
+		static readonly ConcurrentDictionary<Type, Lazy<HashSet<UnrealTargetPlatform>?>> s_supportedPlatformsAttribute = new();
 	}
-}
+}
\ No newline at end of file
diff --git a/Engine/Source/Programs/UnrealBuildTool/Configuration/TargetRules.cs b/Engine/Source/Programs/UnrealBuildTool/Configuration/TargetRules.cs
--- a/Engine/Source/Programs/UnrealBuildTool/Configuration/TargetRules.cs
+++ b/Engine/Source/Programs/UnrealBuildTool/Configuration/TargetRules.cs
@@ -3598,7 +3598,7 @@ namespace UnrealBuildTool
 		/// </summary>
 		/// <param name="t"></param>
 		/// <returns></returns>
-		static UnrealTargetConfiguration[]? GetSupportedConfigurationFromAttribute(Type t) => s_supportedConfigurations.GetOrAdd(t, static type => new Lazy<UnrealTargetConfiguration[]?>(() =>
+		internal static UnrealTargetConfiguration[]? GetSupportedConfigurationFromAttribute(Type t) => s_supportedConfigurations.GetOrAdd(t, static type => new Lazy<UnrealTargetConfiguration[]?>(() =>
 		{
 			object[] attributes = type.GetCustomAttributes(typeof(SupportedConfigurationsAttribute), false);
 			if (attributes.Length > 0)
@@ -3618,7 +3618,7 @@ namespace UnrealBuildTool
 		/// </summary>
 		/// <param name="t"></param>
 		/// <returns></returns>
-		static List<PropertyInfo> GetSubObjectProperties(Type t) => s_subPropsCache.GetOrAdd(t, static type => new Lazy<List<PropertyInfo>>(() =>
+		internal static List<PropertyInfo> GetSubObjectProperties(Type t) => s_subPropsCache.GetOrAdd(t, static type => new Lazy<List<PropertyInfo>>(() =>
 		{
 			var list = new List<PropertyInfo>();
 
diff --git a/Engine/Source/Programs/UnrealBuildTool/System/SourceFileMetadataCache.cs b/Engine/Source/Programs/UnrealBuildTool/System/SourceFileMetadataCache.cs
--- a/Engine/Source/Programs/UnrealBuildTool/System/SourceFileMetadataCache.cs
+++ b/Engine/Source/Programs/UnrealBuildTool/System/SourceFileMetadataCache.cs
@@ -154,7 +154,7 @@ namespace UnrealBuildTool
 		/// <summary>
 		/// Static cache of all constructed dependency caches
 		/// </summary>
-		static ConcurrentDictionary<FileReference, SourceFileMetadataCache> Caches = new ConcurrentDictionary<FileReference, SourceFileMetadataCache>();
+		static ConcurrentDictionary<FileReference, Lazy<SourceFileMetadataCache>> Caches = new();
 
 		/// <summary>
 		/// Constructs a dependency cache. This method is private; call CppDependencyCache.Create() to create a cache hierarchy for a given project.
@@ -559,11 +559,11 @@ namespace UnrealBuildTool
 		/// <returns>Reference to a dependency cache with the given settings</returns>
 		static SourceFileMetadataCache FindOrAddCache(FileReference Location, DirectoryReference BaseDirectory, SourceFileMetadataCache? Parent, ILogger Logger)
 		{
-			SourceFileMetadataCache Cache = Caches.GetOrAdd(Location, _ =>
+			SourceFileMetadataCache Cache = Caches.GetOrAdd(Location, new Lazy<SourceFileMetadataCache>(() =>
 			{
 				return new SourceFileMetadataCache(Location, BaseDirectory, Parent, Logger);
 				;
-			});
+			})).Value;
 
 			Debug.Assert(Cache.BaseDirectory == BaseDirectory);
 			Debug.Assert(Cache.Parent == Parent);
@@ -576,8 +576,9 @@ namespace UnrealBuildTool
 		/// </summary>
 		public static void SaveAll()
 		{
-			Parallel.ForEach(Caches.Values, Cache =>
+			Parallel.ForEach(Caches.Values, CacheLazy =>
 			{
+				SourceFileMetadataCache Cache = CacheLazy.Value;
 				if (Cache.bModified)
 				{
 					Cache.Write();
