[UBT]
* Project generation optimizations. Do not read ProjectGuid from vcxproj files that are generated by UBT (they are based on path so are fully deterministic anyway)

diff --git a/Engine/Source/Programs/UnrealBuildTool/ProjectFiles/VisualStudio/VCProject.cs b/Engine/Source/Programs/UnrealBuildTool/ProjectFiles/VisualStudio/VCProject.cs
--- a/Engine/Source/Programs/UnrealBuildTool/ProjectFiles/VisualStudio/VCProject.cs
+++ b/Engine/Source/Programs/UnrealBuildTool/ProjectFiles/VisualStudio/VCProject.cs
@@ -100,7 +100,8 @@ namespace UnrealBuildTool
 		/// </summary>
 		/// <param name="InitFilePath">The path to the project file on disk</param>
 		/// <param name="BaseDir">The base directory for files within this project</param>
-		public MSBuildProjectFile(FileReference InitFilePath, DirectoryReference BaseDir)
+		/// <param name="TryReadProjectGuid">Try read the project guid from already existsing file</param>
+		public MSBuildProjectFile(FileReference InitFilePath, DirectoryReference BaseDir, bool TryReadProjectGuid)
 			: base(InitFilePath, BaseDir)
 		{
 			// Each project gets its own GUID.  This is stored in the project file and referenced in the solution file.
@@ -110,7 +111,7 @@ namespace UnrealBuildTool
 				Guid projectGuid = Guid.Empty;
 				// First, check to see if we have an existing file on disk.  If we do, then we'll try to preserve the
 				// GUID by loading it from the existing file.
-				if (FileReference.Exists(ProjectFilePath))
+				if (TryReadProjectGuid && FileReference.Exists(ProjectFilePath))
 				{
 					try
 					{
@@ -316,8 +317,9 @@ namespace UnrealBuildTool
 		/// <param name="bMakeProjectPerTarget">Whether to add roll the target type into the config (ie "Development Editor")</param>
 		/// <param name="BuildToolOverride">Optional arguments to pass to UBT when building</param>
 		/// <param name="Settings">Other settings</param>
-		public VCProjectFile(FileReference FilePath, DirectoryReference BaseDir, VCProjectFileFormat ProjectFileFormat, bool bUsePrecompiled, bool bMakeProjectPerTarget, string? BuildToolOverride, VCProjectFileSettings Settings)
-			: base(FilePath, BaseDir)
+		/// <param name="TryReadProjectGuid">Try read the project guid from already existsing file</param>
+		public VCProjectFile(FileReference FilePath, DirectoryReference BaseDir, VCProjectFileFormat ProjectFileFormat, bool bUsePrecompiled, bool bMakeProjectPerTarget, string? BuildToolOverride, VCProjectFileSettings Settings, bool TryReadProjectGuid)
+			: base(FilePath, BaseDir, TryReadProjectGuid)
 		{
 			this.ProjectFileFormat = ProjectFileFormat;
 			this.bUsePrecompiled = bUsePrecompiled;
@@ -906,7 +908,7 @@ namespace UnrealBuildTool
 			{
 				if (!IgnorePaths.Contains(IncludePath) && !IncludePathIsFilteredOut(IncludePath))
 				{
-					Builder.Append(NormalizeProjectPath(IncludePath.FullName));
+					Builder.Append(NormalizeProjectPath(IncludePath));
 					Builder.Append(';');
 				}
 			}
@@ -1346,7 +1348,7 @@ namespace UnrealBuildTool
 			string CommonForcedIncludes = String.Empty;
 			string CommonAdditionalOptions = String.Empty;
 			{
-				if (DirectoryToForceIncludePaths.Any())
+				if (!DirectoryToForceIncludePaths.IsEmpty)
 				{
 					string ForceIncludePathToCheck = DirectoryToForceIncludePaths.Values.First();
 					if (DirectoryToForceIncludePaths.Values.All(x => x == ForceIncludePathToCheck))
@@ -1356,7 +1358,7 @@ namespace UnrealBuildTool
 					}
 				}
 
-				if (DirectoryToPchFile.Any())
+				if (!DirectoryToPchFile.IsEmpty)
 				{
 					string? PchFileToCheck = DirectoryToPchFile.Values.FirstOrDefault();
 					if (DirectoryToPchFile.Values.All(x => x == PchFileToCheck))
@@ -1461,7 +1463,7 @@ namespace UnrealBuildTool
 				{
 					StringBuilder CommonProjectFileContent = new StringBuilder();
 
-					Action<IDictionary<DirectoryReference, string>, string, string, string> AddProperties = (IDictionary<DirectoryReference, string> DirectoryToStringDict, string CommonPropertyPrefix, string PropertyNamePrefix, string PropertyValuePrefix) =>
+					void AddProperties(IDictionary<DirectoryReference, string> DirectoryToStringDict, string CommonPropertyPrefix, string PropertyNamePrefix, string PropertyValuePrefix)
 					{
 						Dictionary<string, string> UpdatedValues = new();
 						List<KeyValuePair<DirectoryReference, string>> KVPList = DirectoryToStringDict.ToList();
@@ -1644,7 +1646,7 @@ namespace UnrealBuildTool
 								ClCompileInfo.AppendLine($"      <ForcedIncludeFiles>$({DirectoryToForceIncludePathValue})</ForcedIncludeFiles>");
 							}
 
-							if (DirectoryToPchFile.Any())
+							if (!DirectoryToPchFile.IsEmpty)
 							{
 								string? PchHeaderFile = DirectoryToPchFile[Directory];
 								if (PchHeaderFile != null && ProjectFileFormat >= VCProjectFileFormat.VisualStudio2022)
@@ -2498,7 +2500,7 @@ namespace UnrealBuildTool
 					} 
 					else 
 					{
-						VCProjectFileContent.AppendLine("    <NMakeOutput>{0}</NMakeOutput>", NormalizeProjectPath(NMakePath.FullName));
+						VCProjectFileContent.AppendLine("    <NMakeOutput>{0}</NMakeOutput>", NormalizeProjectPath(NMakePath));
 					}
 					if (ProjectFileGenerator.bVisualStudioLinux)
 					{
@@ -2595,7 +2597,7 @@ namespace UnrealBuildTool
 		/// <param name="Logger">Logger for output</param>
 		/// <param name="BaseDir">The base directory for files within this project - if not specified, InitFilePath.Directory will be used</param>
 		public VCSharpProjectFile(FileReference InitFilePath, ILogger Logger, DirectoryReference? BaseDir = null)
-			: base(InitFilePath, BaseDir ?? InitFilePath.Directory)
+			: base(InitFilePath, BaseDir ?? InitFilePath.Directory, true)
 		{
 			try
 			{
diff --git a/Engine/Source/Programs/UnrealBuildTool/ProjectFiles/VisualStudio/VCProjectFileGenerator.cs b/Engine/Source/Programs/UnrealBuildTool/ProjectFiles/VisualStudio/VCProjectFileGenerator.cs
--- a/Engine/Source/Programs/UnrealBuildTool/ProjectFiles/VisualStudio/VCProjectFileGenerator.cs
+++ b/Engine/Source/Programs/UnrealBuildTool/ProjectFiles/VisualStudio/VCProjectFileGenerator.cs
@@ -218,7 +218,7 @@ namespace UnrealBuildTool
 		/// <returns>The newly allocated project file object</returns>
 		protected override ProjectFile AllocateProjectFile(FileReference InitFilePath, DirectoryReference BaseDir)
 		{
-			return new VCProjectFile(InitFilePath, BaseDir, Settings.ProjectFileFormat, bUsePrecompiled, bMakeProjectPerTarget, BuildToolOverride, Settings);
+			return new VCProjectFile(InitFilePath, BaseDir, Settings.ProjectFileFormat, bUsePrecompiled, bMakeProjectPerTarget, BuildToolOverride, Settings, false);
 		}
 
 		/// "4.0", "12.0", or "14.0", etc...
diff --git a/Engine/Source/Programs/UnrealBuildTool/ProjectFiles/VisualStudioWorkspace/VSWorkspaceProjectFileGenerator.cs b/Engine/Source/Programs/UnrealBuildTool/ProjectFiles/VisualStudioWorkspace/VSWorkspaceProjectFileGenerator.cs
--- a/Engine/Source/Programs/UnrealBuildTool/ProjectFiles/VisualStudioWorkspace/VSWorkspaceProjectFileGenerator.cs
+++ b/Engine/Source/Programs/UnrealBuildTool/ProjectFiles/VisualStudioWorkspace/VSWorkspaceProjectFileGenerator.cs
@@ -526,7 +526,7 @@ namespace UnrealBuildTool
 				var Settings = new VCProjectFileSettings();
 				Settings.MaxSharedIncludePaths = 0;
 
-				var VCProject = new VCProjectFile(Output, RootPath, Format, false, true, null, Settings);
+				var VCProject = new VCProjectFile(Output, RootPath, Format, false, true, null, Settings, true);
 				VCProject.ProjectTargets.Add(Target);
 				if (VCProject.WriteProjectFile(FilteredPlatforms, FilteredConfigurations, PlatformProjectGenerators, Logger))
 				{
