[EpicGames.Core/UBT]
* Added support in LegacyEventLogger to log in a separate thread and have minimal contention on calling log. Disabled by default
* Enabled threaded logging in UBT
* Removed threaded logger in UBAExecutor since it is now always threaded

diff --git a/Engine/Source/Programs/Shared/EpicGames.Core/LegacyEventLogger.cs b/Engine/Source/Programs/Shared/EpicGames.Core/LegacyEventLogger.cs
--- a/Engine/Source/Programs/Shared/EpicGames.Core/LegacyEventLogger.cs
+++ b/Engine/Source/Programs/Shared/EpicGames.Core/LegacyEventLogger.cs
@@ -1,18 +1,21 @@
 // Copyright Epic Games, Inc. All Rights Reserved.
 
 using System;
+using System.Threading;
+using System.Threading.Channels;
 using Microsoft.Extensions.Logging;
 
 namespace EpicGames.Core
 {
 	/// <summary>
 	/// Wrapper around a custom logger interface which flushes the event parser when switching between legacy
-	/// and native structured logging
+	/// and native structured logging. Also support threading the log messages to minimize performance overhead
 	/// </summary>
 	sealed class LegacyEventLogger : ILogger, IDisposable
 	{
-		private ILogger _inner;
+		private volatile ILogger _inner;
 		private readonly LogEventParser _parser;
+		private ThreadingState? _state;
 
 		public LogEventParser Parser => _parser;
 
@@ -24,9 +27,15 @@ namespace EpicGames.Core
 
 		public void Dispose()
 		{
+			// Drain and stop the worker if enabled, then dispose parser.
+			DisableThreaded();
 			_parser.Dispose();
 		}
 
+		/// <summary>
+		/// Set inner logger to transfer log calls to
+		/// </summary>
+		/// <param name="inner"></param>
 		public void SetInnerLogger(ILogger inner)
 		{
 			lock (_parser)
@@ -37,17 +46,117 @@ namespace EpicGames.Core
 			}
 		}
 
+		/// <summary>
+		/// Enable background logging (idempotent).
+		/// </summary>
+		public void EnableThreaded()
+		{
+			ThreadingState created = new();
+			Interlocked.CompareExchange(ref _state, created, null);
+		}
+
+		/// <summary>
+		/// Disable background logging and drain queued items (idempotent).
+		/// </summary>
+		public void DisableThreaded()
+		{
+			ThreadingState? state = Interlocked.Exchange(ref _state, null);
+			state?.Stop(); // completes channel, drains, joins worker thread
+		}
+
+		// ---- ILogger ----
 		public IDisposable? BeginScope<TState>(TState state) where TState : notnull => _inner.BeginScope(state);
 
 		public bool IsEnabled(LogLevel logLevel) => _inner.IsEnabled(logLevel);
 
 		public void Log<TState>(LogLevel logLevel, EventId eventId, TState state, Exception? exception, Func<TState, Exception?, string> formatter)
 		{
-			lock (_parser)
+			void DoLog()
 			{
-				_parser.Flush();
+				var logger = _inner;
+				lock (_parser)
+				{
+					_parser.Flush();
+				}
+				logger.Log(logLevel, eventId, state, exception, formatter);
+			}
+
+			var s = Volatile.Read(ref _state);
+			if (s == null)
+			{
+				DoLog();
+				return;
+			}
+
+			if (logLevel < LogLevel.Error)
+			{
+				// Non-blocking enqueue; unbounded channel TryWrite is always true.
+				s.Writer.TryWrite(DoLog);
+				return;
+			}
+
+			// Error+ logs: preserve synchronous semantics.
+			using ManualResetEventSlim done = new(false, spinCount: 1);
+			s.Writer.TryWrite(() =>
+			{
+				try { DoLog(); }
+				finally { done.Set(); }
+			});
+			done.Wait();
+		}
+
+		private sealed class ThreadingState
+		{
+			public readonly ChannelWriter<Action> Writer;
+			private readonly ChannelReader<Action> _reader;
+			private readonly Thread _thread;
+
+			public ThreadingState()
+			{
+				// Single reader = true => fewer fences and better throughput.
+				Channel<Action> ch = Channel.CreateUnbounded<Action>(new UnboundedChannelOptions
+				{
+					SingleReader = true,
+					SingleWriter = false,
+					AllowSynchronousContinuations = false
+				});
+				Writer = ch.Writer;
+				_reader = ch.Reader;
+
+				_thread = new Thread(static o =>
+				{
+					ThreadingState self = (ThreadingState)o!;
+					ChannelReader<Action> reader = self._reader;
+
+					while (reader.WaitToReadAsync().AsTask().GetAwaiter().GetResult())
+					{
+						while (reader.TryRead(out Action? work))
+						{
+							try
+							{
+								work();
+							}
+							catch
+							{
+								// Swallow exceptions to keep the logging loop healthy.
+								// Consider publishing to a fallback sink if needed.
+							}
+						}
+					}
+				})
+				{
+					IsBackground = true,
+					Name = "LegacyEventLogger",
+				};
+				_thread.Start(this);
+			}
+
+			public void Stop()
+			{
+				// Complete the writer so the reader drains and then exits cleanly.
+				Writer.TryComplete();
+				_thread.Join();
 			}
-			_inner.Log(logLevel, eventId, state, exception, formatter);
 		}
 	}
 }
\ No newline at end of file
diff --git a/Engine/Source/Programs/Shared/EpicGames.Core/Log.cs b/Engine/Source/Programs/Shared/EpicGames.Core/Log.cs
--- a/Engine/Source/Programs/Shared/EpicGames.Core/Log.cs
+++ b/Engine/Source/Programs/Shared/EpicGames.Core/Log.cs
@@ -205,6 +205,20 @@ namespace EpicGames.Core
 			await s_defaultLogger.FlushAsync();
 		}
 
+		public static void EnableThreaded()
+		{
+			s_legacyLogger.EnableThreaded();
+		}
+
+		/// <summary>
+		/// Finish async and fallback to sync logging
+		/// </summary>
+		/// <returns></returns>
+		public static void DisableThreaded()
+		{
+			s_legacyLogger.DisableThreaded();
+		}
+
 		/// <summary>
 		/// Backup an existing log file if it already exists at the outputpath
 		/// </summary>
diff --git a/Engine/Source/Programs/Shared/EpicGames.UBA/ThreadedLogger.cs b/Engine/Source/Programs/Shared/EpicGames.UBA/ThreadedLogger.cs
deleted file mode 100644
diff --git a/Engine/Source/Programs/UnrealBuildTool/Executors/UnrealBuildAccelerator/UBAExecutor.cs b/Engine/Source/Programs/UnrealBuildTool/Executors/UnrealBuildAccelerator/UBAExecutor.cs
--- a/Engine/Source/Programs/UnrealBuildTool/Executors/UnrealBuildAccelerator/UBAExecutor.cs
+++ b/Engine/Source/Programs/UnrealBuildTool/Executors/UnrealBuildAccelerator/UBAExecutor.cs
@@ -1,11 +1,9 @@
 // Copyright Epic Games, Inc. All Rights Reserved.
 
 using EpicGames.Core;
-using EpicGames.Horde.Projects;
 using EpicGames.UBA;
 using Microsoft.Extensions.Logging;
 using System;
-using System.Collections.Concurrent;
 using System.Collections.Generic;
 using System.IO;
 using System.Linq;
@@ -16,6 +14,7 @@ using System.Threading;
 using System.Threading.Tasks;
 using UnrealBuildBase;
 using UnrealBuildTool.Artifacts;
+using ILogger = Microsoft.Extensions.Logging.ILogger;
 
 namespace UnrealBuildTool
 {
@@ -32,7 +31,6 @@ namespace UnrealBuildTool
 		readonly List<IUBAAgentCoordinator> _agentCoordinators = new();
 		DirectoryReference? _rootDirRef;
 		bool _bIsCancelled;
-		readonly ThreadedLogger _threadedLogger;
 
 		DateTime _ubaStartTimeUtc = DateTime.UtcNow;
 		TimeSpan _ubaDurationWaitingForRemote = TimeSpan.Zero;
@@ -58,7 +56,6 @@ namespace UnrealBuildTool
 				_cacheClients.Clear();
 				_session?.Dispose();
 				_session = null;
-				_threadedLogger.Dispose();
 				_ubaLogger?.Dispose();
 				_ubaLogger = null;
 			}
@@ -134,7 +131,6 @@ namespace UnrealBuildTool
 			}
 
 			_agentCoordinators.AddRange(UBAAgentCoordinatorHorde.Init(logger, targetDescriptors, UBAConfig, ref _remoteConnectionMode));
-			_threadedLogger = new ThreadedLogger(logger);
 		}
 
 		private void LoadCacheConfigs(IEnumerable<TargetDescriptor> targetDescriptors, Microsoft.Extensions.Logging.ILogger logger)
@@ -179,17 +175,17 @@ namespace UnrealBuildTool
 			}
 		}
 
-		private void PrintConfiguration()
+		private void PrintConfiguration(ILogger logger)
 		{
-			_threadedLogger.LogInformation("  Storage capacity {StoreCapacityGb}Gb", UBAConfig.StoreCapacityGb);
+			logger.LogInformation("  Storage capacity {StoreCapacityGb}Gb", UBAConfig.StoreCapacityGb);
 			if (UBAConfig.bStoreObjFilesCompressed)
 			{
-				_threadedLogger.LogInformation("  Object Compression Allowed");
+				logger.LogInformation("  Object Compression Allowed");
 			}
-			_threadedLogger.LogDebug($"  RootDir {_rootDirRef}");
+			logger.LogDebug($"  RootDir {_rootDirRef}");
 		}
 
-		private async Task WriteActionOutputFileAsync(IEnumerable<LinkedAction> inputActions)
+		private async Task WriteActionOutputFileAsync(ILogger logger, IEnumerable<LinkedAction> inputActions)
 		{
 			if (String.IsNullOrEmpty(UBAConfig.ActionsOutputFile))
 			{
@@ -198,7 +194,7 @@ namespace UnrealBuildTool
 
 			if (!UBAConfig.ActionsOutputFile.EndsWith(".yaml", StringComparison.OrdinalIgnoreCase))
 			{
-				_threadedLogger.LogError("UBA actions output file needs to have extension .yaml for UbaCli to understand it");
+				logger.LogError("UBA actions output file needs to have extension .yaml for UbaCli to understand it");
 			}
 			using System.IO.StreamWriter streamWriter = new(UBAConfig.ActionsOutputFile);
 			using System.CodeDom.Compiler.IndentedTextWriter writer = new(streamWriter, "  ");
@@ -229,7 +225,7 @@ namespace UnrealBuildTool
 				{
 					await writer.WriteLineAsync("remote: false");
 				}
-				uint bucket = GetCacheBucket(action);
+				uint bucket = GetCacheBucket(logger, action);
 				if (bucket != 0)
 				{
 					await writer.WriteLineAsync($"cache: {bucket}");
@@ -296,9 +292,7 @@ namespace UnrealBuildTool
 
 			_cacheableActionsLeft = inputActions.Count(x => x.ArtifactMode.HasFlag(ArtifactMode.Enabled));
 
-			await WriteActionOutputFileAsync(inputActions);
-
-			logger = _threadedLogger;
+			await WriteActionOutputFileAsync(logger, inputActions);
 
 			if (!String.IsNullOrEmpty(UBAConfig.RootDir))
 			{
@@ -339,7 +333,7 @@ namespace UnrealBuildTool
 				}
 			}
 
-			PrintConfiguration();
+			PrintConfiguration(logger);
 
 			DirectoryReference.CreateDirectory(_rootDirRef);
 
@@ -457,7 +451,7 @@ namespace UnrealBuildTool
 								session.DisableRemoteExecution();
 							}
 
-							bool success = ExecuteActionsInternal(inputActions, Server, ubaStorage, session, _cacheClients.Values);
+							bool success = ExecuteActionsInternal(logger, inputActions, Server, ubaStorage, session, _cacheClients.Values);
 
 							if (!UBAConfig.bDisableRemote)
 							{
@@ -494,7 +488,6 @@ namespace UnrealBuildTool
 				{
 					await coordinator.CloseAsync();
 				}
-				await _threadedLogger.FinishAsync();
 				_session = null;
 			}
 		}
@@ -608,11 +601,11 @@ namespace UnrealBuildTool
 		/// Executes the provided actions
 		/// </summary>
 		/// <returns>True if all the tasks successfully executed, or false if any of them failed.</returns>
-		bool ExecuteActionsInternal(IEnumerable<LinkedAction> inputActions, IServer server, IStorageServer storage, ISessionServer session, IEnumerable<ICacheClient> cache)
+		bool ExecuteActionsInternal(ILogger logger, IEnumerable<LinkedAction> inputActions, IServer server, IStorageServer storage, ISessionServer session, IEnumerable<ICacheClient> cache)
 		{
 			_ubaStartTimeUtc = DateTime.UtcNow;
 
-			ActionLogger actionLogger = new(inputActions.Count(), "Compiling C++ source code...", x => WriteToolOutput(x), () => FlushToolOutput(), _threadedLogger)
+			ActionLogger actionLogger = new(inputActions.Count(), "Compiling C++ source code...", x => WriteToolOutput(x), () => FlushToolOutput(), logger)
 			{
 				ShowCompilationTimes = bShowCompilationTimes,
 				ShowCPUUtilization = bShowCPUUtilization,
@@ -694,10 +687,10 @@ namespace UnrealBuildTool
 
 							var rerun = (string error) =>
 							{
-								_threadedLogger.LogInformation("{Description} {StatusDescription} [RemoteExecutor: {ExecutingHost}]: Exited with error code {ExitCode} ({Error}). This action will retry locally", action.CommandDescription, action.StatusDescription, executingHost, exitCode, error);
+								logger.LogInformation("{Description} {StatusDescription} [RemoteExecutor: {ExecutingHost}]: Exited with error code {ExitCode} ({Error}). This action will retry locally", action.CommandDescription, action.StatusDescription, executingHost, exitCode, error);
 								foreach (string line in logLines)
 								{
-									_threadedLogger.LogInformation(line);
+									logger.LogInformation(line);
 								}
 								++localRetryActions;
 								return ProcessFinishedResponse.RerunLocal;
@@ -705,7 +698,7 @@ namespace UnrealBuildTool
 
 							if (exitCode != 0 && !logLines.Any())
 							{
-								_threadedLogger.LogInformation("{Description} {StatusDescription} [RemoteExecutor: {ExecutingHost}]: Exited with error code {ExitCode} with no output. This action will retry locally", action.CommandDescription, action.StatusDescription, executingHost, exitCode);
+								logger.LogInformation("{Description} {StatusDescription} [RemoteExecutor: {ExecutingHost}]: Exited with error code {ExitCode} with no output. This action will retry locally", action.CommandDescription, action.StatusDescription, executingHost, exitCode);
 								return ProcessFinishedResponse.RerunLocal;
 							}
 							else if ((uint)exitCode == 0xC0000005)
@@ -743,13 +736,13 @@ namespace UnrealBuildTool
 							{
 								if (UBAConfig.bStoreObjFilesCompressed)
 								{
-									_threadedLogger.LogInformation("{Description} {StatusDescription}: Exited with error code {ExitCode}. This action will retry", action.CommandDescription, action.StatusDescription, exitCode);
+									logger.LogInformation("{Description} {StatusDescription}: Exited with error code {ExitCode}. This action will retry", action.CommandDescription, action.StatusDescription, exitCode);
 									++forcedRetryActions;
 									return ProcessFinishedResponse.RerunLocal;
 								}
 								else
 								{
-									_threadedLogger.LogInformation("{Description} {StatusDescription}: Exited with error code {ExitCode}. This action will retry without UBA", action.CommandDescription, action.StatusDescription, exitCode);
+									logger.LogInformation("{Description} {StatusDescription}: Exited with error code {ExitCode}. This action will retry without UBA", action.CommandDescription, action.StatusDescription, exitCode);
 									++forcedRetryActions;
 									return ProcessFinishedResponse.RerunNative;
 								}
@@ -802,7 +795,7 @@ namespace UnrealBuildTool
 					{
 						if (executionType != ProcessExecutionType.Cache)
 						{
-							WriteToCache(action, exitCode, info.ProcessHandle);
+							WriteToCache(logger, action, exitCode, info.ProcessHandle);
 						}
 						++succeededActions;
 					}
@@ -857,7 +850,7 @@ namespace UnrealBuildTool
 
 				var (knownInputs, knownInputsCount) = BuildKnownInputs(action);
 
-				uint bucket = GetCacheBucket(action);
+				uint bucket = GetCacheBucket(logger, action);
 				uint memoryGroup = GetMemoryGroup(action);
 				ulong predictedMemoryUsage = 0; // uba would do a better job scheduling if we knew approx memory usage
 				scheduler.EnqueueProcess(startInfo, action.Weight, enableDetour, canRunRemotely, prereqActionsSortIndex, knownInputs, knownInputsCount, bucket, memoryGroup, predictedMemoryUsage, action);
@@ -871,7 +864,7 @@ namespace UnrealBuildTool
 				scheduler.SetAllowDisableRemoteExecution();
 			}
 
-			_threadedLogger.LogDebug("All processes queued");
+			logger.LogDebug("All processes queued");
 
 			try
 			{
@@ -1021,7 +1014,7 @@ namespace UnrealBuildTool
 
 		HashSet<uint> _usedBuckets = new();
 
-		uint GetCacheBucket(LinkedAction action)
+		uint GetCacheBucket(ILogger logger, LinkedAction action)
 		{
 			if (!action.ArtifactMode.HasFlag(ArtifactMode.Enabled))
 			{
@@ -1068,7 +1061,7 @@ namespace UnrealBuildTool
 			{
 				if (_usedBuckets.Add(bucket))
 				{
-					_threadedLogger.LogDebug($"Cache bucket {bucket} used. Action: {action.CommandDescription} {action.StatusDescription} ({String.Join(" + ", action.GroupNames)})");
+					logger.LogDebug($"Cache bucket {bucket} used. Action: {action.CommandDescription} {action.StatusDescription} ({String.Join(" + ", action.GroupNames)})");
 				}
 			}
 
@@ -1138,7 +1131,7 @@ namespace UnrealBuildTool
 			public DepsData? Data { get; init; }
 		}
 
-		bool WriteToCache(LinkedAction action, int exitCode, nint processHandle)
+		bool WriteToCache(ILogger logger, LinkedAction action, int exitCode, nint processHandle)
 		{
 			if (exitCode != 0 || !action.ArtifactMode.HasFlag(ArtifactMode.Enabled) || !_cacheClients.Any(x => x.Key.CanWrite))
 			{
@@ -1186,7 +1179,7 @@ namespace UnrealBuildTool
 					}
 					catch (Exception)
 					{
-						_threadedLogger.LogInformation("{Description} {StatusDescription}: Failed to read dependency file {DependencyFile} for WriteCache", action.CommandDescription, action.StatusDescription, action.DependencyListFile);
+						logger.LogInformation("{Description} {StatusDescription}: Failed to read dependency file {DependencyFile} for WriteCache", action.CommandDescription, action.StatusDescription, action.DependencyListFile);
 						return false;
 					}
 				}
@@ -1203,7 +1196,7 @@ namespace UnrealBuildTool
 				}
 			}
 
-			uint bucket = GetCacheBucket(action);
+			uint bucket = GetCacheBucket(logger, action);
 
 			bool success = true;
 			foreach (KeyValuePair<UnrealBuildAcceleratorCacheConfig, ICacheClient> item in _cacheClients.Where(x => x.Key.CanWrite))
diff --git a/Engine/Source/Programs/UnrealBuildTool/UnrealBuildTool.cs b/Engine/Source/Programs/UnrealBuildTool/UnrealBuildTool.cs
--- a/Engine/Source/Programs/UnrealBuildTool/UnrealBuildTool.cs
+++ b/Engine/Source/Programs/UnrealBuildTool/UnrealBuildTool.cs
@@ -453,6 +453,9 @@ namespace UnrealBuildTool
 				StartupTraceListener StartupTrace = new StartupTraceListener();
 				Log.AddTraceListener(StartupTrace);
 
+				// Enable threading on logger since it can stall quite a lot when logging happens from multiple threads
+				Log.EnableThreaded();
+
 				if (Options.TraceWrites != null)
 				{
 					Logger.LogInformation("All attempts to write to \"{TraceWrites}\" via WriteFileIfChanged() will be logged", Options.TraceWrites);
@@ -545,8 +548,11 @@ namespace UnrealBuildTool
 					SessionIdentifier = Options.TelemetrySession;
 				}
 
-				TelemetryService.Get().AddTelemetryConfigProviders(Options.TelemetryProviders.Concat(Options.CmdTelemetryProviders));
-				TelemetryService.Get().AddEndpointsFromConfig(Logger);
+				using (Timeline.ScopeEvent("AddTelemetry"))
+				{
+					TelemetryService.Get().AddTelemetryConfigProviders(Options.TelemetryProviders.Concat(Options.CmdTelemetryProviders));
+					TelemetryService.Get().AddEndpointsFromConfig(Logger);
+				}
 
 				// Create a UbtRun file
 				try
@@ -712,6 +718,8 @@ namespace UnrealBuildTool
 				}
 				Logger.Log(ExecutionLogLevel, "Total execution time: {Time:0.00} seconds", Timeline.Elapsed.TotalSeconds);
 
+				Log.DisableThreaded();
+
 				// Make sure we flush the logs however we exit
 				Trace.Close();
 
