[UBT]
* Added timing scopes around vni and uht

diff --git a/Engine/Source/Programs/UnrealBuildTool/System/ExternalExecution.cs b/Engine/Source/Programs/UnrealBuildTool/System/ExternalExecution.cs
--- a/Engine/Source/Programs/UnrealBuildTool/System/ExternalExecution.cs
+++ b/Engine/Source/Programs/UnrealBuildTool/System/ExternalExecution.cs
@@ -1200,6 +1200,8 @@ namespace UnrealBuildTool
 			}
 			using (ProgressWriter Progress = new("Generating code...", false, Logger))
 			{
+				//using var uhtInit = Timeline.ScopeEvent("UHTInit");
+
 				string RootLocalPath = Unreal.RootDirectory.FullName;
 
 				// Get the path to the assemblies we care about
@@ -1248,12 +1250,17 @@ namespace UnrealBuildTool
 					bUHTNeedsToRun = true;
 				}
 
+
 				if (bUHTNeedsToRun)
 				{
 					Progress.Write(1, 3);
 
-					WriteUHTManifest(Makefile, TargetName, ModuleInfoFileName, ExternalDependenciesFile);
+					//uhtInit.Dispose();
 
+					{
+						using var _ = Timeline.ScopeEvent("UHTManifest");
+						WriteUHTManifest(Makefile, TargetName, ModuleInfoFileName, ExternalDependenciesFile);
+					}
 					string ActualTargetName = String.IsNullOrEmpty(TargetName) ? "UE5" : TargetName;
 					Logger.LogInformation("Parsing headers for {ActualTargetName}", ActualTargetName);
 
@@ -1298,12 +1305,17 @@ namespace UnrealBuildTool
 
 					Logger.LogInformation("  Running Internal UnrealHeaderTool {CmdLine}", CmdLine);
 
+					using var uhtExecute = Timeline.ScopeEvent("UHT");
+
 					// Run UHT
 					IScope Timer = GlobalTracer.Instance.BuildSpan("Executing UnrealHeaderTool").StartActive();
 					UnrealHeaderToolMode UHTTool = new();
+
 					CompilationResult UHTResult = (CompilationResult)await UHTTool.ExecuteAsync(Arguments, Logger);
 					Timer.Span.Finish();
 
+					uhtExecute.Dispose();
+
 					if (UHTResult != CompilationResult.Succeeded)
 					{
 						// On Linux and Mac, the shell will return 128+signal number exit codes if UHT gets a signal (e.g. crashes or is interrupted)
@@ -1345,6 +1357,7 @@ namespace UnrealBuildTool
 
 				using (GlobalTracer.Instance.BuildSpan("ExternalExecution.UpdateTimestamps()").StartActive())
 				{
+					using var uhtExecute = Timeline.ScopeEvent("UHTTimestamps");
 					UpdateTimestamps(Makefile.UObjectModules);
 				}
 
