[UBT]
* Fixed so UnrealBuildTool processes spawned by UnrealBuildTool processes are not creating trace

diff --git a/Engine/Source/Programs/Shared/EpicGames.UBA/ITrace.cs b/Engine/Source/Programs/Shared/EpicGames.UBA/ITrace.cs
--- a/Engine/Source/Programs/Shared/EpicGames.UBA/ITrace.cs
+++ b/Engine/Source/Programs/Shared/EpicGames.UBA/ITrace.cs
@@ -49,7 +49,14 @@ namespace EpicGames.UBA
 		/// <param name="channel">Set which channel this trace should be visible to</param>
 		public static ITrace Create(string? name = null, string? traceFile = null, bool makeGlobal = false, string? channel = "Default")
 		{
-			return new TraceImpl(name, traceFile, makeGlobal, channel);
+			if (name != null)
+			{
+				return new TraceImpl(name, traceFile, makeGlobal, channel);
+			}
+			else
+			{
+				return new TraceNull();
+			}
 		}
 	}
 
diff --git a/Engine/Source/Programs/Shared/EpicGames.UBA/Impl/TraceImpl.cs b/Engine/Source/Programs/Shared/EpicGames.UBA/Impl/TraceImpl.cs
--- a/Engine/Source/Programs/Shared/EpicGames.UBA/Impl/TraceImpl.cs
+++ b/Engine/Source/Programs/Shared/EpicGames.UBA/Impl/TraceImpl.cs
@@ -100,4 +100,14 @@ namespace EpicGames.UBA
 
 		public void UpdateStatus(uint statusRow, uint statusColumn, string statusText, LogEntryType statusType, string? statusLink) => Trace_UpdateStatus(_handle, statusRow, statusColumn, statusText, (byte)statusType, statusLink);
 	}
+
+	internal class TraceNull : ITrace
+	{
+		public nint GetHandle() => 0;
+		public void Dispose() { }
+		public uint TaskBegin(string description, string details) => 0;
+		public void TaskHint(uint taskId, string hint) { }
+		public void TaskEnd(uint taskId) { }
+		public void UpdateStatus(uint statusRow, uint statusColumn, string statusText, LogEntryType statusType, string? statusLink = null) { }
+	}
 }
\ No newline at end of file
diff --git a/Engine/Source/Programs/UnrealBuildTool/UnrealBuildTool.cs b/Engine/Source/Programs/UnrealBuildTool/UnrealBuildTool.cs
--- a/Engine/Source/Programs/UnrealBuildTool/UnrealBuildTool.cs
+++ b/Engine/Source/Programs/UnrealBuildTool/UnrealBuildTool.cs
@@ -413,9 +413,12 @@ namespace UnrealBuildTool
 
 			// If command line contains -Session it is a recursive action and we don't want to report this trace to the default trace channel
 			string? traceChannel = "Default";
+			string? traceName = "UnrealBuildTool";
 			if (ArgumentsArray.FirstOrDefault(s => s.Contains("-Session=")) != null)
+			{
 				traceChannel = null;
-
+				traceName = null;
+			}
 			// Setup trace file and backup old one
 			FileReference ubaTraceFile;
 			if (Unreal.IsBuildMachine())
@@ -434,7 +437,7 @@ namespace UnrealBuildTool
 			Log.BackupLogFile(ubaTraceFile);
 
 			// Attach timeline to uba trace
-			using EpicGames.UBA.ITrace trace = EpicGames.UBA.ITrace.Create("UnrealBuildTool", ubaTraceFile.FullName, true, traceChannel);
+			using EpicGames.UBA.ITrace trace = EpicGames.UBA.ITrace.Create(traceName, ubaTraceFile.FullName, true, traceChannel);
 			using TimelineEventScope _ = new((name) => trace.TaskBegin(name, ""), (id) => trace.TaskEnd(id));
 
 			bool bShowTimeline = false;
