[UBT]
* Added timeline scopes for preparing action graph

diff --git a/Engine/Source/Programs/UnrealBuildTool/System/ActionGraph.cs b/Engine/Source/Programs/UnrealBuildTool/System/ActionGraph.cs
--- a/Engine/Source/Programs/UnrealBuildTool/System/ActionGraph.cs
+++ b/Engine/Source/Programs/UnrealBuildTool/System/ActionGraph.cs
@@ -46,6 +46,8 @@ namespace UnrealBuildTool
 		/// <param name="Sort">Optional sorting of actions</param>
 		public static void Link(List<LinkedAction> Actions, ILogger Logger, bool Sort = true)
 		{
+			using ITimelineEvent _ = Timeline.ScopeEvent($"LinkActions");
+
 			// Build a map from item to its producing action
 			Dictionary<FileItem, LinkedAction> ItemToProducingAction = new Dictionary<FileItem, LinkedAction>();
 			foreach (LinkedAction Action in Actions)
@@ -542,6 +544,8 @@ namespace UnrealBuildTool
 		/// </summary>
 		static void DetectActionGraphCycles(List<LinkedAction> Actions, Dictionary<FileItem, LinkedAction> ItemToProducingAction, ILogger Logger)
 		{
+			using ITimelineEvent _ = Timeline.ScopeEvent($"DetectActionGraphCycles");
+
 			// Starting with actions that only depend on non-produced items, iteratively expand a set of actions that are only dependent on
 			// non-cyclical actions.
 			Dictionary<LinkedAction, bool> ActionIsNonCyclical = new Dictionary<LinkedAction, bool>();
@@ -673,6 +677,8 @@ namespace UnrealBuildTool
 		/// <returns>Set of prerequisite actions</returns>
 		public static List<LinkedAction> GatherPrerequisiteActions(List<LinkedAction> Actions, HashSet<FileItem> OutputItems)
 		{
+			using ITimelineEvent _ = Timeline.ScopeEvent($"GatherPrerequisiteActions");
+
 			HashSet<LinkedAction> PrerequisiteActions = new HashSet<LinkedAction>();
 			foreach (LinkedAction Action in Actions)
 			{
@@ -938,6 +944,8 @@ namespace UnrealBuildTool
 			ConcurrentDictionary<LinkedAction, bool> OutdatedActions, CppDependencyCache CppDependencies,
 			bool bIgnoreOutdatedImportLibraries, ILogger Logger)
 		{
+			using ITimelineEvent __ = Timeline.ScopeEvent($"GatherOutdatedActions");
+
 			using (GlobalTracer.Instance.BuildSpan("Prefetching include dependencies").StartActive())
 			{
 				List<FileItem> Dependencies = new List<FileItem>();
@@ -949,7 +957,7 @@ namespace UnrealBuildTool
 					}
 				}
 
-				Parallel.ForEach(Dependencies, File => { CppDependencies.TryGetDependencies(File, Logger, out _); });
+				Parallel2.ForEach(Dependencies, File => { CppDependencies.TryGetDependencies(File, Logger, out _); });
 			}
 
 			using (GlobalTracer.Instance.BuildSpan("Cache individual outdated actions").StartActive())
@@ -975,6 +983,8 @@ namespace UnrealBuildTool
 		/// <param name="Logger">Logger for output</param>
 		public static void DeleteOutdatedProducedItems(List<LinkedAction> OutdatedActions, ILogger Logger)
 		{
+			using ITimelineEvent __ = Timeline.ScopeEvent($"DeleteOutdatedProducedItems");
+
 			foreach (LinkedAction OutdatedAction in OutdatedActions)
 			{
 				foreach (FileItem DeleteItem in OutdatedAction.DeleteItems)
