[UBT]
* Added timing scope around creating UHT manifest

diff --git a/Engine/Source/Programs/UnrealBuildTool/System/ExternalExecution.cs b/Engine/Source/Programs/UnrealBuildTool/System/ExternalExecution.cs
--- a/Engine/Source/Programs/UnrealBuildTool/System/ExternalExecution.cs
+++ b/Engine/Source/Programs/UnrealBuildTool/System/ExternalExecution.cs
@@ -1200,7 +1200,7 @@ namespace UnrealBuildTool
 			}
 			using (ProgressWriter Progress = new("Generating code...", false, Logger))
 			{
-				//using var uhtInit = Timeline.ScopeEvent("UHTInit");
+				using var uhtManifestScope = Timeline.ScopeEvent("UHTManifest");
 
 				string RootLocalPath = Unreal.RootDirectory.FullName;
 
@@ -1255,12 +1255,7 @@ namespace UnrealBuildTool
 				{
 					Progress.Write(1, 3);
 
-					//uhtInit.Dispose();
-
-					{
-						using var _ = Timeline.ScopeEvent("UHTManifest");
-						WriteUHTManifest(Makefile, TargetName, ModuleInfoFileName, ExternalDependenciesFile);
-					}
+					WriteUHTManifest(Makefile, TargetName, ModuleInfoFileName, ExternalDependenciesFile);
 					string ActualTargetName = String.IsNullOrEmpty(TargetName) ? "UE5" : TargetName;
 					Logger.LogInformation("Parsing headers for {ActualTargetName}", ActualTargetName);
 
@@ -1303,6 +1298,8 @@ namespace UnrealBuildTool
 						CmdLine.AppendArgument(Arg);
 					}
 
+					uhtManifestScope.Finish();
+
 					Logger.LogInformation("  Running Internal UnrealHeaderTool {CmdLine}", CmdLine);
 
 					using var uhtExecute = Timeline.ScopeEvent("UHT");
