[UBT]
* Renamed ExternalExecution.cs to UHTExecution.cs  to make it easier to find what is executing UHT

diff --git a/Engine/Source/Programs/UnrealBuildTool/System/ExternalExecution.cs b/Engine/Source/Programs/UnrealBuildTool/System/ExternalExecution.cs
deleted file mode 100644
diff --git a/Engine/Source/Programs/UnrealBuildTool/System/UHTExecution.cs b/Engine/Source/Programs/UnrealBuildTool/System/UHTExecution.cs
new file mode 100644
--- /dev/null
+++ b/Engine/Source/Programs/UnrealBuildTool/System/UHTExecution.cs
@@ -0,0 +1,1377 @@
+// Copyright Epic Games, Inc. All Rights Reserved.
+
+using System;
+using System.Collections;
+using System.Collections.Generic;
+using System.Data;
+using System.Diagnostics;
+using System.Diagnostics.CodeAnalysis;
+using System.IO;
+using System.Linq;
+using System.Reflection;
+using System.Text;
+using System.Text.Json;
+using System.Threading.Tasks;
+using EpicGames.Core;
+using EpicGames.UHT.Utils;
+using Microsoft.Extensions.Logging;
+using OpenTracing;
+using OpenTracing.Util;
+using UnrealBuildBase;
+using UnrealBuildTool.Modes;
+
+namespace UnrealBuildTool
+{
+	static class UHTModuleTypeExtensions
+	{
+		public static UHTModuleType? GetEngineModuleType(this ModuleHostType ModuleType)
+		{
+			return ModuleType switch
+			{
+				ModuleHostType.Program => UHTModuleType.Program,
+
+				ModuleHostType.Runtime or 
+				ModuleHostType.RuntimeNoCommandlet or 
+				ModuleHostType.RuntimeAndProgram or 
+				ModuleHostType.CookedOnly or 
+				ModuleHostType.ServerOnly or 
+				ModuleHostType.ClientOnly or 
+				ModuleHostType.ClientOnlyNoCommandlet => UHTModuleType.EngineRuntime,
+
+				ModuleHostType.Developer or 
+				ModuleHostType.DeveloperTool => UHTModuleType.EngineDeveloper,
+
+				ModuleHostType.Editor or
+				ModuleHostType.EditorNoCommandlet or 
+				ModuleHostType.EditorAndProgram => UHTModuleType.EngineEditor,
+
+				ModuleHostType.UncookedOnly => UHTModuleType.EngineUncooked,
+
+				_ => null,
+			};
+		}
+
+		public static UHTModuleType? GetGameModuleType(this ModuleHostType ModuleType)
+		{
+			return ModuleType switch
+			{
+				ModuleHostType.Program => UHTModuleType.Program,
+
+				ModuleHostType.Runtime or 
+				ModuleHostType.RuntimeNoCommandlet or 
+				ModuleHostType.RuntimeAndProgram or 
+				ModuleHostType.CookedOnly or 
+				ModuleHostType.ServerOnly or 
+				ModuleHostType.ClientOnly or 
+				ModuleHostType.ClientOnlyNoCommandlet => UHTModuleType.GameRuntime,
+
+				ModuleHostType.Developer or 
+				ModuleHostType.DeveloperTool => UHTModuleType.GameDeveloper,
+
+				ModuleHostType.Editor or 
+				ModuleHostType.EditorNoCommandlet or 
+				ModuleHostType.EditorAndProgram => UHTModuleType.GameEditor,
+
+				ModuleHostType.UncookedOnly => UHTModuleType.GameUncooked,
+
+				_ => null,
+			};
+		}
+	}
+
+	/// <summary>
+	/// Information about a module that needs to be passed to UnrealHeaderTool for code generation
+	/// </summary>
+	class UHTModuleInfo
+	{
+		/// <summary>
+		/// Module name
+		/// </summary>
+		public required string ModuleName;
+
+		/// <summary>
+		/// Path to the module rules file
+		/// </summary>
+		public required FileReference ModuleRulesFile;
+
+		/// <summary>
+		/// Paths to all potential module source directories (with platform extension directories added in)
+		/// </summary>
+		public required DirectoryReference[] ModuleDirectories;
+
+		/// <summary>
+		/// The include search path for generated headers to include other headers
+		/// </summary>
+		public required DirectoryReference[] ModuleIncludePaths;
+
+		/// <summary>
+		/// Module type
+		/// </summary>
+		public required string ModuleType;
+
+		/// <summary>
+		/// Overridden package module type to add more flags
+		/// </summary>
+		public required string OverrideModuleType;
+
+		/// <summary>
+		/// Public UObject headers found in the Classes directory (legacy)
+		/// </summary>
+		public List<FileItem> PublicUObjectClassesHeaders = [];
+
+		/// <summary>
+		/// Public headers with UObjects
+		/// </summary>
+		public List<FileItem> PublicUObjectHeaders = [];
+
+		/// <summary>
+		/// Internal headers with UObjects
+		/// </summary>
+		public List<FileItem> InternalUObjectHeaders = [];
+
+		/// <summary>
+		/// Private headers with UObjects
+		/// </summary>
+		public List<FileItem> PrivateUObjectHeaders = [];
+
+		/// <summary>
+		/// Directory containing generated code
+		/// </summary>
+		public required DirectoryItem GeneratedCodeDirectory;
+
+		/// <summary>
+		/// Collection of the module's public defines
+		/// </summary>
+		public List<string> PublicDefines = [];
+
+		/// <summary>
+		/// Path and base filename, without extension, of the .gen files
+		/// </summary>
+		public string? GeneratedCPPFilenameBase;
+
+		/// <summary>
+		/// Version of code generated by UHT
+		/// </summary>
+		public required EGeneratedCodeVersion GeneratedCodeVersion;
+
+		/// <summary>
+		/// Whether this module is read-only
+		/// </summary>
+		public required bool bIsReadOnly;
+
+		/// <summary>
+		/// Path of the verse generated types 
+		/// </summary>
+		public required string VersePath;
+
+		/// <summary>
+		/// Scope of the verse definitions
+		/// </summary>
+		public required VerseScope VerseScope;
+
+		/// <summary>
+		/// Mount point for UE definitions
+		/// </summary>
+		public required string VerseMountPoint;
+
+		/// <summary>
+		/// Verse module name
+		/// </summary>
+		public required string VersePackageName;
+
+		/// <summary>
+		/// Relative location of the verse code
+		/// </summary>
+		public required string VerseDirectoryPath;
+
+		/// <summary>
+		/// Collection of verse packages this package depends on
+		/// </summary>
+		public HashSet<string> VersePublicDependencies = [];
+
+		/// <summary>
+		/// Collection of verse packages this package depends on
+		/// </summary>
+		public HashSet<string> VersePrivateDependencies = [];
+
+		/// <summary>
+		/// If true, verse scripts were found
+		/// </summary>
+		public required bool HasVerse;
+		
+		/// <summary>
+		/// If true, autogenerated functions of USTRUCTS are always exported
+		/// </summary>
+		public required bool AlwaysExportStructs;
+
+		/// <summary>
+		/// If true, autogenerated functions of UENUMS are always exported
+		/// </summary>
+		public required bool AlwaysExportEnums;
+
+		/// <summary>
+		/// If true, UE types (i.e. UCLASS) will be allowed in namespaces
+		/// </summary>
+		public required bool AllowUETypesInNamespaces;
+
+		/// <summary>
+		/// If true, minimizes the amount of header files included in .generated.h files, which may require more headers from CoreUObject to be explicitly included.
+		/// </summary>
+		public required bool MinimizeGeneratedIncludes;
+
+		public UHTModuleInfo()
+		{
+		}
+
+		[SetsRequiredMembers]
+		public UHTModuleInfo(BinaryArchiveReader Reader)
+		{
+			ModuleName = Reader.ReadString()!;
+			ModuleRulesFile = Reader.ReadFileReference();
+			ModuleDirectories = Reader.ReadArray<DirectoryReference>(Reader.ReadDirectoryReferenceNotNull)!;
+			ModuleIncludePaths = Reader.ReadArray<DirectoryReference>(Reader.ReadDirectoryReferenceNotNull)!;
+			ModuleType = Reader.ReadString()!;
+			OverrideModuleType = Reader.ReadString()!;
+			PublicUObjectClassesHeaders = Reader.ReadList(() => Reader.ReadFileItem())!;
+			PublicUObjectHeaders = Reader.ReadList(() => Reader.ReadFileItem())!;
+			InternalUObjectHeaders = Reader.ReadList(() => Reader.ReadFileItem())!;
+			PrivateUObjectHeaders = Reader.ReadList(() => Reader.ReadFileItem())!;
+			GeneratedCPPFilenameBase = Reader.ReadString();
+			GeneratedCodeDirectory = Reader.ReadDirectoryItem()!;
+			GeneratedCodeVersion = (EGeneratedCodeVersion)Reader.ReadInt();
+			bIsReadOnly = Reader.ReadBool();
+			PublicDefines = Reader.ReadList(() => Reader.ReadString())!;
+			VersePath = Reader.ReadString()!;
+			VerseScope = (VerseScope)Reader.ReadInt();
+			HasVerse = Reader.ReadBool();
+			VerseMountPoint = Reader.ReadString()!;
+			VersePackageName = Reader.ReadString()!;
+			VerseDirectoryPath = Reader.ReadString()!;
+			VersePublicDependencies = Reader.ReadHashSet(() => Reader.ReadString())!;
+			VersePrivateDependencies = Reader.ReadHashSet(() => Reader.ReadString())!;
+			AlwaysExportStructs = Reader.ReadBool();
+			AlwaysExportEnums = Reader.ReadBool();
+			AllowUETypesInNamespaces = Reader.ReadBool();
+			MinimizeGeneratedIncludes = Reader.ReadBool();
+		}
+
+		public void Write(BinaryArchiveWriter Writer)
+		{
+			Writer.WriteString(ModuleName);
+			Writer.WriteFileReference(ModuleRulesFile);
+			Writer.WriteArray<DirectoryReference>(ModuleDirectories, Writer.WriteDirectoryReference);
+			Writer.WriteArray<DirectoryReference>(ModuleIncludePaths, Writer.WriteDirectoryReference);
+			Writer.WriteString(ModuleType);
+			Writer.WriteString(OverrideModuleType);
+			Writer.WriteList(PublicUObjectClassesHeaders, Item => Writer.WriteFileItem(Item));
+			Writer.WriteList(PublicUObjectHeaders, Item => Writer.WriteFileItem(Item));
+			Writer.WriteList(InternalUObjectHeaders, Item => Writer.WriteFileItem(Item));
+			Writer.WriteList(PrivateUObjectHeaders, Item => Writer.WriteFileItem(Item));
+			Writer.WriteString(GeneratedCPPFilenameBase);
+			Writer.WriteDirectoryItem(GeneratedCodeDirectory);
+			Writer.WriteInt((int)GeneratedCodeVersion);
+			Writer.WriteBool(bIsReadOnly);
+			Writer.WriteList(PublicDefines, Item => Writer.WriteString(Item));
+			Writer.WriteString(VersePath);
+			Writer.WriteInt((int)VerseScope);
+			Writer.WriteBool(HasVerse);
+			Writer.WriteString(VerseMountPoint);
+			Writer.WriteString(VersePackageName);
+			Writer.WriteString(VerseDirectoryPath);
+			Writer.WriteHashSet(VersePublicDependencies, Item => Writer.WriteString(Item));
+			Writer.WriteHashSet(VersePrivateDependencies, Item => Writer.WriteString(Item));
+			Writer.WriteBool(AlwaysExportStructs);
+			Writer.WriteBool(AlwaysExportEnums);
+			Writer.WriteBool(AllowUETypesInNamespaces);
+			Writer.WriteBool(MinimizeGeneratedIncludes);
+		}
+
+		public override string ToString()
+		{
+			return ModuleName;
+		}
+	}
+
+	class UHTModuleHeaderInfo
+	{
+		public DirectoryItem SourceFolder;
+		public List<FileItem> HeaderFiles;
+		public bool bUsePrecompiled;
+
+		public UHTModuleHeaderInfo(DirectoryItem SourceFolder, List<FileItem> HeaderFiles, bool bUsePrecompiled)
+		{
+			this.SourceFolder = SourceFolder;
+			this.HeaderFiles = HeaderFiles;
+			this.bUsePrecompiled = bUsePrecompiled;
+		}
+
+		public UHTModuleHeaderInfo(BinaryArchiveReader Reader)
+		{
+			SourceFolder = Reader.ReadDirectoryItem()!;
+			HeaderFiles = Reader.ReadList(() => Reader.ReadFileItem())!;
+			bUsePrecompiled = Reader.ReadBool();
+		}
+
+		public void Write(BinaryArchiveWriter Writer)
+		{
+			Writer.WriteDirectoryItem(SourceFolder);
+			Writer.WriteList(HeaderFiles, Item => Writer.WriteFileItem(Item));
+			Writer.WriteBool(bUsePrecompiled);
+		}
+	}
+
+	/// <summary>
+	/// This handles all running of the UnrealHeaderTool
+	/// </summary>
+	class ExternalExecution
+	{
+		static UHTModuleType GetEngineModuleTypeFromDescriptor(ModuleDescriptor Module)
+		{
+			return Module.Type.GetEngineModuleType() ?? throw new BuildException("Unhandled engine module type {0} for {1}", Module.Type.ToString(), Module.Name);
+		}
+
+		static UHTModuleType GetGameModuleTypeFromDescriptor(ModuleDescriptor Module)
+		{
+			return Module.Type.GetGameModuleType() ?? throw new BuildException("Unhandled game module type {0} for {1}", Module.Type.ToString(), Module.Name);
+		}
+
+		/// <summary>
+		/// Return the mount point for the given module
+		/// </summary>
+		/// <param name="Module"></param>
+		/// <returns></returns>
+		/// <exception cref="BuildException"></exception>
+		public static string GetVerseMountPointForModule(UEBuildModule Module)
+		{
+			if (Module.Rules.Plugin != null)
+			{
+				if (Module.bHasVerse && !Module.Rules.Plugin.Descriptor.bCanContainVerse)
+				{
+					throw new BuildException("Module '{0}' has an associated Verse directory but its containing plugin '{1}' does not specify \"CanContainVerse\": true.", Module.Name, Module.Rules.Plugin.Name);
+				}
+
+				return Module.Rules.Plugin.Name;
+			}
+			// We currently allow Verse only in plugins and programs
+			if (Module.bHasVerse && Module.Rules.Context.DefaultUHTModuleType.GetValueOrDefault(UHTModuleType.Program) != UHTModuleType.Program)
+			{
+				throw new BuildException("Module '{0}' has an associated Verse directory but it is neither in a plugin nor program.", Module.Name);
+			}
+
+			return Module.Rules.Context.Scope.Name == "Project" ? "Game" : "Engine";
+		}
+
+		/// <summary>
+		/// Gets the module type for a given rules object
+		/// </summary>
+		/// <param name="RulesObject">The rules object</param>
+		/// <param name="ProjectDescriptor">Descriptor for the project being built</param>
+		/// <returns>The module type</returns>
+		static UHTModuleType GetModuleType(ModuleRules RulesObject, ProjectDescriptor? ProjectDescriptor)
+		{
+			ModuleRulesContext Context = RulesObject.Context;
+			if (Context.bClassifyAsGameModuleForUHT)
+			{
+				if (RulesObject.Type == ModuleRules.ModuleType.External)
+				{
+					return UHTModuleType.GameThirdParty;
+				}
+				if (Context.DefaultUHTModuleType.HasValue)
+				{
+					return Context.DefaultUHTModuleType.Value;
+				}
+				if (RulesObject.Plugin != null)
+				{
+					ModuleDescriptor? Module = RulesObject.Plugin.Descriptor.Modules?.FirstOrDefault(x => x.Name == RulesObject.Name);
+					if (Module != null)
+					{
+						return GetGameModuleTypeFromDescriptor(Module);
+					}
+				}
+				if (ProjectDescriptor != null && ProjectDescriptor.Modules != null)
+				{
+					ModuleDescriptor? Module = ProjectDescriptor.Modules.FirstOrDefault(x => x.Name == RulesObject.Name);
+					if (Module != null)
+					{
+						return Module.Type.GetGameModuleType() ?? UHTModuleType.GameRuntime;
+					}
+				}
+				return UHTModuleType.GameRuntime;
+			}
+			else
+			{
+				if (RulesObject.Type == ModuleRules.ModuleType.External)
+				{
+					return UHTModuleType.EngineThirdParty;
+				}
+				if (Context.DefaultUHTModuleType.HasValue)
+				{
+					return Context.DefaultUHTModuleType.Value;
+				}
+				if (RulesObject.Plugin != null)
+				{
+					string PluginName = !RulesObject.IsTestModule ? RulesObject.Name : TargetDescriptor.GetTestedName(RulesObject.Name);
+					ModuleDescriptor? Module = RulesObject.Plugin.Descriptor.Modules?.FirstOrDefault(x => x.Name == PluginName);
+					if (Module != null)
+					{
+						return GetEngineModuleTypeFromDescriptor(Module);
+					}
+				}
+				throw new BuildException("Unable to determine UHT module type for {0}", RulesObject.File);
+			}
+		}
+
+		/// <summary>
+		/// Find all the headers under the given base directory, excluding any other platform folders.
+		/// </summary>
+		/// <param name="BaseDirectory">Base directory to search</param>
+		/// <param name="ExcludeFolders">Array of folders to exclude</param>
+		/// <param name="Headers">Receives the list of headers that was found</param>
+		static void FindHeaders(DirectoryItem BaseDirectory, IReadOnlySet<string> ExcludeFolders, List<FileItem> Headers)
+		{
+			if (BaseDirectory.TryGetFile(".ubtignore", out FileItem? OutIgnoreFile))
+			{
+				return;
+			}
+
+			// Check for all the headers in this folder
+			Headers.AddRange(BaseDirectory.EnumerateFiles().Where((fi) => fi.HasExtension(".h")));
+
+			foreach (DirectoryItem SubDirectory in BaseDirectory.EnumerateDirectories())
+			{
+				if (!ExcludeFolders.Contains(SubDirectory.Name))
+				{
+					FindHeaders(SubDirectory, ExcludeFolders, Headers);
+				}
+			}
+		}
+
+		public static void SetupUObjectModules(
+			IEnumerable<UEBuildModuleCPP> ModulesToGenerateHeadersFor,
+			UnrealTargetPlatform Platform, 
+			ProjectDescriptor? ProjectDescriptor, 
+			List<UHTModuleInfo> UObjectModules,
+			List<UHTModuleHeaderInfo> UObjectModuleHeaders,
+			DirectoryReference ExecutableDirectory,
+			EGeneratedCodeVersion GeneratedCodeVersion, 
+			SourceFileMetadataCache MetadataCache)
+		{
+			// Find the type of each module
+			Dictionary<UEBuildModuleCPP, UHTModuleType> ModuleToType = [];
+			foreach (UEBuildModuleCPP Module in ModulesToGenerateHeadersFor)
+			{
+				ModuleToType[Module] = GetModuleType(Module.Rules, ProjectDescriptor);
+			}
+
+			// Sort modules by type, then by dependency
+			List<UEBuildModuleCPP> ModulesSortedByType = [.. ModulesToGenerateHeadersFor.OrderBy(c => ModuleToType[c])];
+			ModulesSortedByType = [.. UEBuildModule.StableTopologicalSort([.. ModulesSortedByType.Cast<UEBuildModule>()]).Cast<UEBuildModuleCPP>()];
+
+			// Create a module to index lookup
+			Dictionary<UEBuildModule, int> ModuleToIndex = [];
+			for (int Idx = 0; Idx < ModulesSortedByType.Count; Idx++)
+			{
+				ModuleToIndex.Add(ModulesSortedByType[Idx], Idx);
+			}
+
+			// We neeed the base dir to make a relative path for the Verse source location
+			DirectoryReference BaseDir = ExecutableDirectory;
+			string? RelativeBaseDir = UEBuildTarget.GetRelativeBaseDir(ExecutableDirectory, Platform);
+			if (RelativeBaseDir != null)
+			{
+				// Copy the logic from UEBuildTarget.SetupGlobalEnvironment ~ln 3645 that defines UE_RELATIVE_BASE_DIR to make the BaseDir relative to the Engine directory
+				// This ensures that the source file gathering can always use FPlatformProcess:BaseDir reliably since it has UE_RELATIVE_BASE_DIR baked in
+				BaseDir = DirectoryReference.Combine(ExecutableDirectory, RelativeBaseDir);
+			}
+
+			// Create the info for each module in parallel
+			UHTModuleInfo[] ModuleInfoArray = new UHTModuleInfo[ModulesSortedByType.Count];
+			using (ThreadPoolWorkQueue Queue = new())
+			{
+				IReadOnlySet<string> ExcludedFolders = UEBuildPlatform.GetBuildPlatform(Platform).GetExcludedFolderNames();
+				for (int Idx = 0; Idx < ModulesSortedByType.Count; Idx++)
+				{
+					UEBuildModuleCPP Module = ModulesSortedByType[Idx];
+
+					DirectoryItem GeneratedCodeDirectory = DirectoryItem.GetItemByDirectoryReference(Module.GeneratedCodeDirectoryUHT!);
+
+					DirectoryReference[] ModuleIncludePaths = [.. Module.PublicIncludePaths.Union(Module.InternalIncludePaths).Union(Module.PrivateIncludePaths)];
+
+					string VerseMountPoint = string.Empty;
+					string VersePackageName = string.Empty;
+					string VerseDirectoryPath = string.Empty;
+					if (Module.bHasVerse)
+					{
+						DirectoryReference SourceDirectory = DirectoryReference.Combine(Module.Rules.Directory, "Verse");
+						VerseMountPoint = GetVerseMountPointForModule(Module);
+						VersePackageName = $"{VerseMountPoint}/{Module.Name}";
+						VerseDirectoryPath = SourceDirectory.MakeRelativeTo(BaseDir).Replace("\\", "/", StringComparison.Ordinal);
+					}
+
+					UHTModuleInfo Info = new()
+					{
+						ModuleName = Module.Name,
+						ModuleRulesFile = Module.RulesFile,
+						ModuleDirectories = Module.ModuleDirectories,
+						ModuleIncludePaths = ModuleIncludePaths,
+						ModuleType = ModuleToType[Module].ToString(),
+						GeneratedCodeDirectory = GeneratedCodeDirectory,
+						GeneratedCodeVersion = GeneratedCodeVersion,
+						bIsReadOnly = Module.Rules.bUsePrecompiled,
+						OverrideModuleType = Module.Rules.OverridePackageType.ToString(),
+						VersePath = Module.Rules.VersePath ?? "",
+						VerseScope = Module.Rules.VerseScope,
+						HasVerse = Module.bHasVerse,
+						VerseMountPoint = VerseMountPoint,
+						VersePackageName = VersePackageName,
+						VerseDirectoryPath = VerseDirectoryPath,
+						AlwaysExportStructs = Module.Rules.bAlwaysExportStructs,
+						AlwaysExportEnums = Module.Rules.bAlwaysExportEnums,
+						AllowUETypesInNamespaces = Module.Rules.bAllowUETypesInNamespaces,
+						MinimizeGeneratedIncludes = Module.Rules.bMinimizeGeneratedIncludes,
+					};
+					ModuleInfoArray[Idx] = Info;
+
+					Queue.Enqueue(() => SetupUObjectModule(Info, ExcludedFolders, MetadataCache, Queue));
+				}
+			}
+
+			// Gather all the verse dependencies
+			BitArray ModuleHasBeenGathered = new(ModulesSortedByType.Count, false);
+			for (int Idx = 0; Idx < ModulesSortedByType.Count; Idx++)
+			{
+				GatherVerseDependencies(Idx, ModulesSortedByType, ModuleInfoArray, ModuleToIndex, ModuleHasBeenGathered);
+			}
+
+			// Filter out all the modules with reflection data
+			for (int Idx = 0; Idx < ModulesSortedByType.Count; Idx++)
+			{
+				UEBuildModuleCPP Module = ModulesSortedByType[Idx];
+				UHTModuleInfo Info = ModuleInfoArray[Idx];
+				Info.PublicDefines.AddRange(Module.PublicDefinitions);
+
+				if (Info.PublicUObjectClassesHeaders.Count > 0 || Info.PrivateUObjectHeaders.Count > 0 || Info.PublicUObjectHeaders.Count > 0 || Info.InternalUObjectHeaders.Count > 0)
+				{
+					Module.bHasUObjects = true;
+
+					// If we've got this far and there are no source files then it's likely we're installed and ignoring
+					// engine files, so we don't need a .gen.cpp either
+					DirectoryReference GeneratedCodeDirectoryUHT = Module.GeneratedCodeDirectoryUHT!;
+					Info.GeneratedCPPFilenameBase = Path.Combine(GeneratedCodeDirectoryUHT.FullName, Info.ModuleName) + ".gen";
+					if (!Module.Rules.bUsePrecompiled)
+					{
+						Module.GeneratedCppDirectories ??= [];
+						Module.GeneratedCppDirectories.Add(GeneratedCodeDirectoryUHT.FullName);
+					}
+
+					UObjectModules.Add(Info);
+
+					DirectoryItem ModuleDirectoryItem = DirectoryItem.GetItemByDirectoryReference(Module.ModuleDirectory);
+
+					List<FileItem> ReflectedHeaderFiles =
+					[
+						.. Info.PublicUObjectClassesHeaders,
+						.. Info.PublicUObjectHeaders,
+						.. Info.InternalUObjectHeaders,
+						.. Info.PrivateUObjectHeaders,
+					];
+					UObjectModuleHeaders.Add(new UHTModuleHeaderInfo(ModuleDirectoryItem, ReflectedHeaderFiles, Module.Rules.bUsePrecompiled));
+				}
+				else
+				{
+					// Remove any stale generated code directory
+					if (Module.GeneratedCodeDirectoryUHT != null && !Module.Rules.bUsePrecompiled)
+					{
+						if (DirectoryReference.Exists(Module.GeneratedCodeDirectoryUHT))
+						{
+							Directory.Delete(Module.GeneratedCodeDirectoryUHT.FullName, true);
+							// Also delete parent directory if now empty
+							if (!Directory.EnumerateFileSystemEntries(Module.GeneratedCodeDirectory!.FullName).Any())
+							{
+								Directory.Delete(Module.GeneratedCodeDirectory!.FullName, true);
+							}
+						}
+					}
+				}
+			}
+
+			// Set Module.bHasUObjects for any IncludePathModules not already processed.
+			// This is necessary to keep include paths consistent between targets built with -AllModules and without
+			// Note that PublicIncludePathModules can be recursive so we need to traverse the entire chain
+			List<UEBuildModuleCPP> IncludePathModules = [];
+			CollectModulesOnlyIncluded(ModulesToGenerateHeadersFor, IncludePathModules);
+			if (IncludePathModules.Count > 0)
+			{
+				Dictionary<UEBuildModuleCPP, UHTModuleInfo> IncludePathInfo = [];
+
+				using ThreadPoolWorkQueue Queue = new();
+				IReadOnlySet<string> ExcludedFolders = UEBuildPlatform.GetBuildPlatform(Platform).GetExcludedFolderNames();
+				foreach (UEBuildModuleCPP Module in IncludePathModules)
+				{
+					Queue.Enqueue(() =>
+					{
+						foreach (DirectoryItem ModuleDirectoryItem in Module.ModuleDirectories.Select(x => DirectoryItem.GetItemByDirectoryReference(x)))
+						{
+							List<FileItem> HeaderFiles = [];
+							FindHeaders(ModuleDirectoryItem, ExcludedFolders, HeaderFiles);
+
+							if (HeaderFiles.Any(x => MetadataCache.ContainsReflectionMarkup(x)))
+							{
+								Module.bHasUObjects = true;
+								break;
+							}
+						}
+					});
+				}
+			}
+		}
+
+		static UHTModuleInfo GatherVerseDependencies(int Idx, List<UEBuildModuleCPP> ModulesSortedByType, UHTModuleInfo[] ModuleInfoArray, Dictionary<UEBuildModule, int> ModuleToIndex, BitArray ModuleHasBeenGathered)
+		{
+			UHTModuleInfo ModuleInfo = ModuleInfoArray[Idx];
+			if (ModuleHasBeenGathered[Idx])
+			{
+				return ModuleInfo;
+			}
+			ModuleHasBeenGathered[Idx] = true;
+
+			UEBuildModuleCPP BuildModule = ModulesSortedByType[Idx];
+			if (BuildModule.bHasVerse)
+			{
+				GatherVerseDependencies(ModulesSortedByType, ModuleInfoArray, ModuleToIndex, ModuleHasBeenGathered, BuildModule.PublicDependencyModules, ModuleInfo.VersePublicDependencies);
+				GatherVerseDependencies(ModulesSortedByType, ModuleInfoArray, ModuleToIndex, ModuleHasBeenGathered, BuildModule.PrivateDependencyModules, ModuleInfo.VersePrivateDependencies);
+			}
+			return ModuleInfo;
+		}
+
+		static void GatherVerseDependencies(List<UEBuildModuleCPP> ModulesSortedByType, UHTModuleInfo[] ModuleInfoArray, Dictionary<UEBuildModule, int> ModuleToIndex, BitArray ModuleHasBeenGathered, 
+			List<UEBuildModule>? InDependencies, HashSet<string> OutDependencies)
+		{
+			if (InDependencies != null)
+			{
+				foreach (UEBuildModule Dependency in InDependencies)
+				{
+					if (ModuleToIndex.TryGetValue(Dependency, out int GatherIdx))
+					{
+						UHTModuleInfo GatherModuleInfo = GatherVerseDependencies(GatherIdx, ModulesSortedByType, ModuleInfoArray, ModuleToIndex, ModuleHasBeenGathered);
+						if (GatherModuleInfo.HasVerse)
+						{
+							OutDependencies.Add(GatherModuleInfo.VersePackageName);
+						}
+						OutDependencies.UnionWith(GatherModuleInfo.VersePublicDependencies); // This is correct, private inherits public
+					}
+				}
+			}
+		}
+
+		static void CollectModulesOnlyIncluded(HashSet<UEBuildModuleCPP> HandledModules, List<UEBuildModule> IncludedModules, List<UEBuildModuleCPP> OutList)
+		{
+			foreach (UEBuildModule pi in IncludedModules)
+			{
+				if (pi is UEBuildModuleCPP IncludedModule)
+				{
+					if (IncludedModule.bHasUObjects)
+					{
+						continue;
+					}
+
+					if (!HandledModules.Add(IncludedModule))
+					{
+						continue;
+					}
+
+					OutList.Add(IncludedModule);
+
+					if (IncludedModule.PublicIncludePathModules != null)
+					{
+						CollectModulesOnlyIncluded(HandledModules, IncludedModule.PublicIncludePathModules, OutList);
+					}
+				}
+			}
+		}
+
+		static void CollectModulesOnlyIncluded(IEnumerable<UEBuildModuleCPP> ModulesToGenerateHeadersFor, List<UEBuildModuleCPP> OutList)
+		{
+			HashSet<UEBuildModuleCPP> HandledModules = [];
+			foreach (UEBuildModuleCPP Module in ModulesToGenerateHeadersFor)
+			{
+				HandledModules.Add(Module);
+			}
+
+			foreach (UEBuildModuleCPP Module in ModulesToGenerateHeadersFor)
+			{
+				if (Module.PublicIncludePathModules != null)
+				{
+					CollectModulesOnlyIncluded(HandledModules, Module.PublicIncludePathModules, OutList);
+				}
+				if (Module.PrivateIncludePathModules != null)
+				{
+					CollectModulesOnlyIncluded(HandledModules, Module.PrivateIncludePathModules, OutList);
+				}
+			}
+		}
+
+		static void SetupUObjectModule(UHTModuleInfo ModuleInfo, IReadOnlySet<string> ExcludedFolders, SourceFileMetadataCache MetadataCache, ThreadPoolWorkQueue Queue)
+		{
+			foreach (DirectoryReference ModuleDirectory in ModuleInfo.ModuleDirectories)
+			{
+				DirectoryItem ModuleDirectoryItem = DirectoryItem.GetItemByDirectoryReference(ModuleDirectory);
+
+				List<FileItem> HeaderFiles = [];
+				FindHeaders(ModuleDirectoryItem, ExcludedFolders, HeaderFiles);
+
+				foreach (FileItem HeaderFile in HeaderFiles)
+				{
+					Queue.Enqueue(() => SetupUObjectModuleHeader(ModuleInfo, HeaderFile, MetadataCache));
+				}
+			}
+		}
+
+		static void SetupUObjectModuleHeader(UHTModuleInfo ModuleInfo, FileItem HeaderFile, SourceFileMetadataCache MetadataCache)
+		{
+			// Check to see if we know anything about this file.  If we have up-to-date cached information about whether it has
+			// UObjects or not, we can skip doing a test here.
+			if (MetadataCache.ContainsReflectionMarkup(HeaderFile))
+			{
+				lock (ModuleInfo)
+				{
+					bool bFoundHeaderLocation = false;
+					foreach (DirectoryReference ModuleDirectory in ModuleInfo.ModuleDirectories)
+					{
+						if (HeaderFile.Location.IsUnderDirectory(DirectoryReference.Combine(ModuleDirectory, "Classes")))
+						{
+							ModuleInfo.PublicUObjectClassesHeaders.Add(HeaderFile);
+							bFoundHeaderLocation = true;
+						}
+						else if (HeaderFile.Location.IsUnderDirectory(DirectoryReference.Combine(ModuleDirectory, "Public")))
+						{
+							ModuleInfo.PublicUObjectHeaders.Add(HeaderFile);
+							bFoundHeaderLocation = true;
+						}
+						else if (HeaderFile.Location.IsUnderDirectory(DirectoryReference.Combine(ModuleDirectory, "Internal")))
+						{
+							ModuleInfo.InternalUObjectHeaders.Add(HeaderFile);
+							bFoundHeaderLocation = true;
+						}
+					}
+					if (!bFoundHeaderLocation)
+					{
+						ModuleInfo.PrivateUObjectHeaders.Add(HeaderFile);
+					}
+				}
+			}
+		}
+
+		/// <summary>
+		/// Gets the timestamp of CoreUObject.gen.cpp file.
+		/// </summary>
+		/// <returns>Last write time of CoreUObject.gen.cpp or DateTime.MaxValue if it doesn't exist.</returns>
+		private static DateTime GetCoreGeneratedTimestampUtc(string ModuleName, string ModuleGeneratedCodeDirectory)
+		{
+			// In Installed Builds, we don't check the timestamps on engine headers.  Default to a very old date.
+			if (Unreal.IsEngineInstalled())
+			{
+				return DateTime.MinValue;
+			}
+
+			// Otherwise look for CoreUObject.init.gen.cpp
+			FileInfo CoreGeneratedFileInfo = new(Path.Combine(ModuleGeneratedCodeDirectory, ModuleName + ".init.gen.cpp"));
+			if (CoreGeneratedFileInfo.Exists)
+			{
+				return CoreGeneratedFileInfo.LastWriteTimeUtc;
+			}
+
+			// Doesn't exist, so use a 'newer that everything' date to force rebuild headers.
+			return DateTime.MaxValue;
+		}
+
+		/// <summary>
+		/// Gets the timestamp of the Version.h header file. Throws an exception if not found.
+		/// </summary>
+		/// <returns>Last write time of Version.h</returns>
+		private static DateTime GetEngineVersionHeaderTimestampUtc()
+		{
+			// In Installed Builds, we don't check the timestamps on engine headers.  Default to a very old date.
+			if (Unreal.IsEngineInstalled())
+			{
+				return DateTime.MinValue;
+			}
+
+			string EngineVersionHeaderPath = Path.Combine(Unreal.EngineSourceDirectory.FullName, "Runtime", "Launch", "Resources", "Version.h");
+			FileInfo EngineVersionHeaderFileInfo = new(EngineVersionHeaderPath);
+			if (EngineVersionHeaderFileInfo.Exists)
+			{
+				return EngineVersionHeaderFileInfo.LastWriteTimeUtc;
+			}
+
+			// The Version.h header doesn't exist somehow (maybe it was moved).
+			throw new BuildException("Could not find engine version header, expected: {0}", EngineVersionHeaderPath);
+		}
+
+		/// <summary>
+		/// Checks the class header files and determines if generated UObject code files are out of date in comparison.
+		/// </summary>
+		/// <param name="UObjectModules">Modules that we generate headers for</param>
+		/// <param name="HeaderToolTimestampUtc">Timestamp for UHT</param>
+		/// <param name="Logger">Logger for output</param>
+		/// <returns>True if the code files are out of date</returns>
+		private static bool AreGeneratedCodeFilesOutOfDate(List<UHTModuleInfo> UObjectModules, DateTime HeaderToolTimestampUtc, ILogger Logger)
+		{
+			// Get CoreUObject.init.gen.cpp timestamp.  If the source files are older than the CoreUObject generated code, we'll
+			// need to regenerate code for the module
+			DateTime? CoreGeneratedTimestampUtc = null;
+			{
+				// Find the CoreUObject module
+				foreach (UHTModuleInfo Module in UObjectModules)
+				{
+					if (Module.ModuleName.Equals("CoreUObject", StringComparison.InvariantCultureIgnoreCase))
+					{
+						CoreGeneratedTimestampUtc = GetCoreGeneratedTimestampUtc(Module.ModuleName, Path.GetDirectoryName(Module.GeneratedCPPFilenameBase)!);
+						break;
+					}
+				}
+				if (CoreGeneratedTimestampUtc == null)
+				{
+					throw new BuildException("Could not find CoreUObject in list of all UObjectModules");
+				}
+			}
+
+			// Also get the Version.h timestamp.  If it is newer than the last UHT run, it could affect preprocessor conditionals
+			// that check for the engine version
+			DateTime EngineVersionHeaderTimestampUtc = GetEngineVersionHeaderTimestampUtc();
+
+			foreach (UHTModuleInfo Module in UObjectModules)
+			{
+				// If we're using a precompiled engine, skip checking timestamps for modules that are under the engine directory
+				if (Module.bIsReadOnly)
+				{
+					continue;
+				}
+
+				// Make sure we have an existing folder for generated code.  If not, then we definitely need to generate code!
+				string GeneratedCodeDirectory = Path.GetDirectoryName(Module.GeneratedCPPFilenameBase)!;
+				FileSystemInfo TestDirectory = (FileSystemInfo)new DirectoryInfo(GeneratedCodeDirectory);
+				if (!TestDirectory.Exists)
+				{
+					// Generated code directory is missing entirely!
+					Logger.LogDebug("UnrealHeaderTool needs to run because no generated code directory was found for module {ModuleName}", Module.ModuleName);
+					return true;
+				}
+
+				// Grab our special "Timestamp" file that we saved after the last set of headers were generated.  This file
+				// contains the list of source files which contained UObjects, so that we can compare to see if any
+				// UObject source files were deleted (or no longer contain UObjects), which means we need to run UHT even
+				// if no other source files were outdated
+				string TimestampFile = Path.Combine(GeneratedCodeDirectory, @"Timestamp");
+				FileSystemInfo SavedTimestampFileInfo = (FileSystemInfo)new FileInfo(TimestampFile);
+				if (!SavedTimestampFileInfo.Exists)
+				{
+					// Timestamp file was missing (possibly deleted/cleaned), so headers are out of date
+					Logger.LogDebug("UnrealHeaderTool needs to run because UHT Timestamp file did not exist for module {ModuleName}", Module.ModuleName);
+					return true;
+				}
+
+				// Make sure the last UHT run completed after UnrealHeaderTool.exe was compiled last, and after the CoreUObject headers were touched last.
+				DateTime SavedTimestampUtc = SavedTimestampFileInfo.LastWriteTimeUtc;
+				if (HeaderToolTimestampUtc > SavedTimestampUtc)
+				{
+					// Generated code is older than UnrealHeaderTool.exe.  Out of date!
+					Logger.LogDebug("UnrealHeaderTool needs to run because UnrealHeaderTool timestamp ({Time}) is later than timestamp for module {ModuleName} ({ModuleTime})", HeaderToolTimestampUtc.ToLocalTime(), Module.ModuleName, SavedTimestampUtc.ToLocalTime());
+					return true;
+				}
+				if (CoreGeneratedTimestampUtc > SavedTimestampUtc)
+				{
+					// Generated code is older than CoreUObject headers.  Out of date!
+					Logger.LogDebug("UnrealHeaderTool needs to run because CoreUObject timestamp ({Time}) is newer than timestamp for module {ModuleName} ({ModuleTime})", CoreGeneratedTimestampUtc.Value.ToLocalTime(), Module.ModuleName, SavedTimestampUtc.ToLocalTime());
+					return true;
+				}
+				if (EngineVersionHeaderTimestampUtc > SavedTimestampUtc)
+				{
+					// Generated code is older than engine version header.  Out of date!
+					Logger.LogDebug("UnrealHeaderTool needs to run because engine version header timestamp ({Time}) is newer than timestamp for module {ModuleName} ({ModuleTime})", EngineVersionHeaderTimestampUtc.ToLocalTime(), Module.ModuleName, SavedTimestampUtc.ToLocalTime());
+					return true;
+				}
+
+				// Has the .build.cs file changed since we last generated headers successfully?
+				FileInfo ModuleRulesFile = new(Module.ModuleRulesFile.FullName);
+				if (!ModuleRulesFile.Exists || ModuleRulesFile.LastWriteTimeUtc > SavedTimestampUtc)
+				{
+					Logger.LogDebug("UnrealHeaderTool needs to run because SavedTimestamp is older than the rules file ({ModuleModuleRulesFile}) for module {ModuleModuleName}", Module.ModuleRulesFile, Module.ModuleName);
+					return true;
+				}
+
+				// Iterate over our UObjects headers and figure out if any of them have changed
+				List<FileItem> AllUObjectHeaders =
+				[
+					.. Module.PublicUObjectClassesHeaders,
+					.. Module.PublicUObjectHeaders,
+					.. Module.InternalUObjectHeaders,
+					.. Module.PrivateUObjectHeaders,
+				];
+
+				// Load up the old timestamp file and check to see if anything has changed
+				{
+					string[] UObjectFilesFromPreviousRun = File.ReadAllLines(TimestampFile);
+					if (AllUObjectHeaders.Count != UObjectFilesFromPreviousRun.Length)
+					{
+						Logger.LogDebug("UnrealHeaderTool needs to run because there are a different number of UObject source files in module {ModuleModuleName}", Module.ModuleName);
+						return true;
+					}
+
+					// Iterate over our UObjects headers and figure out if any of them have changed
+					HashSet<string> ObjectHeadersSet = new(AllUObjectHeaders.Select(x => x.AbsolutePath), FileReference.Comparer);
+					foreach (string FileName in UObjectFilesFromPreviousRun)
+					{
+						if (!ObjectHeadersSet.Contains(FileName))
+						{
+							Logger.LogDebug("UnrealHeaderTool needs to run because the set of UObject source files in module {ModuleModuleName} has changed ({FileName})", Module.ModuleName, FileName);
+							return true;
+						}
+					}
+				}
+
+				foreach (FileItem HeaderFile in AllUObjectHeaders)
+				{
+					DateTime HeaderFileTimestampUtc = HeaderFile.LastWriteTimeUtc;
+
+					// Has the source header changed since we last generated headers successfully?
+					if (HeaderFileTimestampUtc > SavedTimestampUtc)
+					{
+						Logger.LogDebug("UnrealHeaderTool needs to run because SavedTimestamp is older than HeaderFileTimestamp ({HeaderFileAbsolutePath}) for module {ModuleModuleName}", HeaderFile.AbsolutePath, Module.ModuleName);
+						return true;
+					}
+				}
+			}
+
+			return false;
+		}
+
+		/// <summary>
+		/// Determines if any external dependencies for generated code is out of date
+		/// </summary>
+		/// <param name="ExternalDependenciesFile">Path to the external dependencies file</param>
+		/// <returns>True if any external dependencies are out of date</returns>
+		private static bool AreExternalDependenciesOutOfDate(FileReference ExternalDependenciesFile)
+		{
+			if (!FileReference.Exists(ExternalDependenciesFile))
+			{
+				return true;
+			}
+
+			DateTime LastWriteTime = File.GetLastWriteTimeUtc(ExternalDependenciesFile.FullName);
+
+			string[] Lines = File.ReadAllLines(ExternalDependenciesFile.FullName);
+			foreach (string Line in Lines)
+			{
+				string ExternalDependencyFile = Line.Trim();
+				if (ExternalDependencyFile.Length > 0)
+				{
+					if (!File.Exists(ExternalDependencyFile) || File.GetLastWriteTimeUtc(ExternalDependencyFile) > LastWriteTime)
+					{
+						return true;
+					}
+				}
+			}
+
+			return false;
+		}
+
+		/// <summary>
+		/// Updates the recorded timestamps of all the passed in UObject modules
+		/// </summary>
+		private static void UpdateTimestamps(List<UHTModuleInfo> UObjectModules)
+		{
+			Parallel.ForEach(UObjectModules, async Module =>
+			{
+				if (!Module.bIsReadOnly)
+				{
+					string GeneratedCodeDirectory = Path.GetDirectoryName(Module.GeneratedCPPFilenameBase)!;
+					DirectoryInfo GeneratedCodeDirectoryInfo = new(GeneratedCodeDirectory);
+
+					try
+					{
+						if (GeneratedCodeDirectoryInfo.Exists)
+						{
+							FileReference TimestampFile = FileReference.Combine(new DirectoryReference(GeneratedCodeDirectoryInfo.FullName), "Timestamp");
+
+							// Save all of the UObject files to a timestamp file.  We'll load these on the next run to see if any new
+							// files with UObject classes were deleted, so that we'll know to run UHT even if the timestamps of all
+							// of the other source files were unchanged
+							{
+								List<string> AllUObjectFiles =
+								[
+									.. Module.PublicUObjectClassesHeaders.ConvertAll(Item => Item.AbsolutePath),
+									.. Module.PublicUObjectHeaders.ConvertAll(Item => Item.AbsolutePath),
+									.. Module.InternalUObjectHeaders.ConvertAll(Item => Item.AbsolutePath),
+									.. Module.PrivateUObjectHeaders.ConvertAll(Item => Item.AbsolutePath),
+								];
+								await FileReference.WriteAllLinesAsync(TimestampFile, AllUObjectFiles);
+							}
+
+							// Because new .cpp and .h files may have been generated by UHT, invalidate the DirectoryLookupCache
+							DirectoryLookupCache.InvalidateCachedDirectory(new DirectoryReference(GeneratedCodeDirectoryInfo.FullName));
+						}
+					}
+					catch (Exception Exception)
+					{
+						throw new BuildException(Exception, "Couldn't write Timestamp file: " + Exception.Message);
+					}
+				}
+			});
+		}
+
+		/// <summary>
+		/// Run an external native executable (and capture the output), given the executable path and the commandline.
+		/// </summary>
+		public static int RunExternalNativeExecutable(FileReference ExePath, string Commandline, ILogger Logger)
+		{
+			Logger.LogDebug("RunExternalExecutable {ExePathFullName} {Commandline}", ExePath.FullName, Commandline);
+
+			using LogEventParser Parser = new(Logger);
+			Parser.AddMatchersFromAssembly(Assembly.GetExecutingAssembly());
+
+			using Process GameProcess = new();
+			GameProcess.StartInfo.FileName = ExePath.FullName;
+			GameProcess.StartInfo.Arguments = Commandline;
+			GameProcess.StartInfo.UseShellExecute = false;
+			GameProcess.StartInfo.RedirectStandardOutput = true;
+			GameProcess.OutputDataReceived += (s, e) => PrintProcessOutputAsync(e, Parser);
+			GameProcess.Start();
+			GameProcess.BeginOutputReadLine();
+			GameProcess.WaitForExit();
+
+			return GameProcess.ExitCode;
+		}
+
+		/// <summary>
+		/// Simple function to pipe output asynchronously
+		/// </summary>
+		private static void PrintProcessOutputAsync(DataReceivedEventArgs Event, LogEventParser Parser)
+		{
+			// DataReceivedEventHandler is fired with a null string when the output stream is closed.  We don't want to
+			// print anything for that event.
+			if (!String.IsNullOrEmpty(Event.Data))
+			{
+				Parser.WriteLine(Event.Data);
+			}
+		}
+
+		/// <summary>
+		/// Builds and runs the header tool and touches the header directories.
+		/// Performs any early outs if headers need no changes, given the UObject modules, tool path, game name, and configuration
+		/// </summary>
+		public static async Task ExecuteHeaderToolIfNecessaryAsync(BuildConfiguration BuildConfiguration, FileReference? ProjectFile, TargetMakefile Makefile, string TargetName, ILogger Logger)
+		{
+
+			// No need to run UHT on itself
+			if (TargetName.Equals("UnrealHeaderTool", StringComparison.InvariantCultureIgnoreCase))
+			{
+				Logger.LogInformation("DEPRECATED: C++ UHT is being built by request.  C++ UHT has been deprecated and will be removed in 5.2");
+				return;
+			}
+
+			await ExecuteHeaderToolIfNecessaryInternalAsync(BuildConfiguration, ProjectFile, Makefile, TargetName, Logger);
+		}
+
+		/// <summary>
+		/// Return the file reference for the UHT manifest file
+		/// </summary>
+		/// <param name="Makefile">Input makefile</param>
+		/// <param name="TargetName">Name of the target</param>
+		/// <returns>Manifest file name</returns>
+		public static FileReference GetUHTModuleInfoFileName(TargetMakefile Makefile, string TargetName)
+		{
+			return FileReference.Combine(Makefile.ProjectIntermediateDirectoryNoArch, TargetName + ".uhtmanifest");
+		}
+
+		/// <summary>
+		/// Return the file reference for the UHT deps file
+		/// </summary>
+		/// <param name="ModuleInfoFileName">Manifest info file name</param>
+		/// <returns>UHT dependency file name</returns>
+		public static FileReference GetUHTDepsFileName(FileReference ModuleInfoFileName)
+		{
+			return ModuleInfoFileName.ChangeExtension(".deps");
+		}
+
+		/// <summary>
+		/// Convert the makefile to a UHTManifest object
+		/// </summary>
+		/// <param name="Makefile">Input makefile</param>
+		/// <param name="TargetName">Name of the target</param>
+		/// <param name="DepsFileName">Name of the dependencies file.</param>
+		/// <returns>Output UHT manifest</returns>
+		public static UHTManifest CreateUHTManifest(TargetMakefile Makefile, string TargetName, FileReference DepsFileName)
+		{
+			List<UHTManifest.Module> Modules = [];
+			foreach (UHTModuleInfo UObjectModule in Makefile.UObjectModules)
+			{
+				List<string> VerseDependencies = [.. UObjectModule.VersePublicDependencies.Union(UObjectModule.VersePrivateDependencies)];
+				VerseDependencies.Sort(StringComparer.Ordinal);
+
+				Modules.Add(
+					new UHTManifest.Module
+					{
+						Name = UObjectModule.ModuleName,
+						ModuleType = (UHTModuleType)Enum.Parse(typeof(UHTModuleType), UObjectModule.ModuleType),
+						OverrideModuleType = (EPackageOverrideType)Enum.Parse(typeof(EPackageOverrideType), UObjectModule.OverrideModuleType),
+						BaseDirectory = UObjectModule.ModuleDirectories[0].FullName,
+						IncludePaths = [.. UObjectModule.ModuleIncludePaths.Select(IncludePath => IncludePath.FullName)],
+						OutputDirectory = Path.GetDirectoryName(UObjectModule.GeneratedCPPFilenameBase)!,
+						ClassesHeaders = [.. UObjectModule.PublicUObjectClassesHeaders.Select((Header) => Header.AbsolutePath)],
+						PublicHeaders = [.. UObjectModule.PublicUObjectHeaders.Select((Header) => Header.AbsolutePath)],
+						InternalHeaders = [.. UObjectModule.InternalUObjectHeaders.Select((Header) => Header.AbsolutePath)],
+						PrivateHeaders = [.. UObjectModule.PrivateUObjectHeaders.Select((Header) => Header.AbsolutePath)],
+						PublicDefines = UObjectModule.PublicDefines,
+						GeneratedCPPFilenameBase = UObjectModule.GeneratedCPPFilenameBase,
+						SaveExportedHeaders = !UObjectModule.bIsReadOnly,
+						GeneratedCodeVersion = UObjectModule.GeneratedCodeVersion,
+						VersePath = UObjectModule.VersePath,
+						VerseScope = (UHTVerseScope)Enum.Parse(typeof(UHTVerseScope), UObjectModule.VerseScope.ToString()),
+						HasVerse = UObjectModule.HasVerse,
+						VerseMountPoint = UObjectModule.VerseMountPoint,
+						VersePackageName = UObjectModule.VersePackageName,
+						VerseDirectoryPath = UObjectModule.VerseDirectoryPath,
+						VerseDependencies = VerseDependencies,
+						AlwaysExportStructs = UObjectModule.AlwaysExportStructs,
+						AlwaysExportEnums = UObjectModule.AlwaysExportEnums,
+						AllowUETypesInNamespaces = UObjectModule.AllowUETypesInNamespaces,
+						MinimizeGeneratedIncludes = UObjectModule.MinimizeGeneratedIncludes,
+					});
+			}
+
+			List<string> UhtPlugins = [];
+			if (Makefile.EnabledUhtPlugins != null)
+			{
+				foreach (FileReference UhtPlugin in Makefile.EnabledUhtPlugins)
+				{
+					UhtPlugins.Add(UhtPlugin.FullName);
+				}
+			}
+
+			UHTManifest Manifest = new()
+			{
+				TargetName = TargetName,
+				IsGameTarget = Makefile.TargetType != TargetType.Program,
+				RootLocalPath = Unreal.RootDirectory.FullName,
+				ExternalDependenciesFile = DepsFileName.FullName,
+				Modules = Modules,
+				UhtPlugins = UhtPlugins,
+			};
+
+			return Manifest;
+		}
+
+		/// <summary>
+		/// Write the manifest to disk
+		/// </summary>
+		/// <param name="Makefile">Input makefile</param>
+		/// <param name="TargetName">Name of the target</param>
+		/// <param name="ModuleInfoFileName">Destination file name.  If not supplied, it will be generated.</param>
+		/// <param name="DepsFileName">Name of the dependencies file.  If not supplied, it will be generated.</param>
+		public static void WriteUHTManifest(TargetMakefile Makefile, string TargetName, FileReference ModuleInfoFileName, FileReference DepsFileName)
+		{
+			// @todo ubtmake: Optimization: Ideally we could avoid having to generate this data in the case where UHT doesn't even need to run!  Can't we use the existing copy?  (see below use of Manifest)
+			UHTManifest Manifest = CreateUHTManifest(Makefile, TargetName, DepsFileName);
+			Directory.CreateDirectory(ModuleInfoFileName.Directory.FullName);
+			System.IO.File.WriteAllText(ModuleInfoFileName.FullName, JsonSerializer.Serialize(Manifest, new JsonSerializerOptions { WriteIndented = true }));
+		}
+
+		/// <summary>
+		/// Return the latest of two date/times
+		/// </summary>
+		/// <param name="Lhs">First date/time to compare</param>
+		/// <param name="Rhs">Second date/time to compare</param>
+		/// <returns>Latest of the two dates</returns>
+		private static DateTime LatestDateTime(DateTime Lhs, DateTime Rhs)
+		{
+			return Lhs > Rhs ? Lhs : Rhs;
+		}
+
+		/// <summary>
+		/// Builds and runs the header tool and touches the header directories.
+		/// Performs any early outs if headers need no changes, given the UObject modules, tool path, game name, and configuration
+		/// </summary>
+		private static async Task ExecuteHeaderToolIfNecessaryInternalAsync(BuildConfiguration BuildConfiguration, FileReference? ProjectFile, TargetMakefile Makefile, string TargetName, ILogger Logger)
+		{
+			if (ProgressWriter.bWriteMarkup)
+			{
+				Logger.LogInformation("@progress push 5%");
+			}
+			using (ProgressWriter Progress = new("Generating code...", false, Logger))
+			{
+				using var uhtManifestScope = Timeline.ScopeEvent("UHTManifest");
+
+				string RootLocalPath = Unreal.RootDirectory.FullName;
+
+				// Get the path to the assemblies we care about
+				string UhtAssemblyPath = typeof(UhtSession).Assembly.Location;
+
+				// Get a composite date/time for critical assemblies being used
+				DateTime CompositeTimestamp = LatestDateTime(new FileInfo(Assembly.GetExecutingAssembly().Location).LastWriteTimeUtc, new FileInfo(UhtAssemblyPath).LastWriteTimeUtc);
+				if (Makefile.EnabledUhtPlugins != null)
+				{
+					foreach (FileReference Plugin in Makefile.EnabledUhtPlugins)
+					{
+						CompositeTimestamp = LatestDateTime(CompositeTimestamp, new FileInfo(Plugin.FullName).LastWriteTimeUtc);
+					}
+				}
+
+				// ensure the headers are up to date
+				bool bUHTNeedsToRun = false;
+				if (BuildConfiguration.bForceHeaderGeneration)
+				{
+					bUHTNeedsToRun = true;
+				}
+				else if (AreGeneratedCodeFilesOutOfDate(Makefile.UObjectModules, CompositeTimestamp, Logger))
+				{
+					bUHTNeedsToRun = true;
+				}
+
+				// Check we're not using a different version of UHT
+				FileReference ModuleInfoFileName = GetUHTModuleInfoFileName(Makefile, TargetName);
+				FileReference ToolInfoFile = ModuleInfoFileName.ChangeExtension(".uhtpath");
+				if (!bUHTNeedsToRun)
+				{
+					if (!FileReference.Exists(ToolInfoFile))
+					{
+						bUHTNeedsToRun = true;
+					}
+					else if (FileReference.ReadAllText(ToolInfoFile) != UhtAssemblyPath)
+					{
+						bUHTNeedsToRun = true;
+					}
+				}
+
+				// Get the file containing dependencies for the generated code
+				FileReference ExternalDependenciesFile = GetUHTDepsFileName(ModuleInfoFileName);
+				if (AreExternalDependenciesOutOfDate(ExternalDependenciesFile))
+				{
+					bUHTNeedsToRun = true;
+				}
+
+
+				if (bUHTNeedsToRun)
+				{
+					Progress.Write(1, 3);
+
+					WriteUHTManifest(Makefile, TargetName, ModuleInfoFileName, ExternalDependenciesFile);
+					string ActualTargetName = String.IsNullOrEmpty(TargetName) ? "UE5" : TargetName;
+					Logger.LogInformation("Parsing headers for {ActualTargetName}", ActualTargetName);
+
+					// Generate the command line
+					List<string> CmdArgs = [
+						(ProjectFile != null) ? ProjectFile.FullName : TargetName,
+						ModuleInfoFileName.FullName,
+						"-WarningsAsErrors",
+					];
+
+					if (Unreal.IsEngineInstalled())
+					{
+						CmdArgs.Add("-installed");
+					}
+
+					if (BuildConfiguration.bFailIfGeneratedCodeChanges)
+					{
+						CmdArgs.Add("-FailIfGeneratedCodeChanges");
+					}
+
+					if (Makefile.UHTAdditionalArguments != null)
+					{
+						foreach (string Arg in Makefile.UHTAdditionalArguments)
+						{
+							if (Arg[0] == '"' && Arg[^1] == '"')
+							{
+								CmdArgs.Add(Arg[1..^1]);
+							}
+							else
+							{
+								CmdArgs.Add(Arg);
+							}
+						}
+					}
+					CommandLineArguments Arguments = new([.. CmdArgs]);
+
+					StringBuilder CmdLine = new();
+					foreach (string Arg in CmdArgs)
+					{
+						CmdLine.AppendArgument(Arg);
+					}
+
+					uhtManifestScope.Finish();
+
+					Logger.LogInformation("  Running Internal UnrealHeaderTool {CmdLine}", CmdLine);
+
+					using var uhtExecute = Timeline.ScopeEvent("UHT");
+
+					// Run UHT
+					IScope Timer = GlobalTracer.Instance.BuildSpan("Executing UnrealHeaderTool").StartActive();
+					UnrealHeaderToolMode UHTTool = new();
+
+					CompilationResult UHTResult = (CompilationResult)await UHTTool.ExecuteAsync(Arguments, Logger);
+					Timer.Span.Finish();
+
+					uhtExecute.Dispose();
+
+					if (UHTResult != CompilationResult.Succeeded)
+					{
+						// On Linux and Mac, the shell will return 128+signal number exit codes if UHT gets a signal (e.g. crashes or is interrupted)
+						if ((BuildHostPlatform.Current.Platform == UnrealTargetPlatform.Linux ||
+							BuildHostPlatform.Current.Platform == UnrealTargetPlatform.Mac) &&
+							(int)(UHTResult) >= 128
+							)
+						{
+							// SIGINT is 2, so 128 + SIGINT is 130
+							UHTResult = ((int)(UHTResult) == 130) ? CompilationResult.Canceled : CompilationResult.CrashOrAssert;
+						}
+
+						if (BuildHostPlatform.Current.Platform == UnrealTargetPlatform.Win64 &&
+							(int)(UHTResult) < 0)
+						{
+							Logger.LogError("UnrealHeaderTool failed with exit code 0x{Result:X} - check that Unreal Engine prerequisites are installed.", (int)UHTResult);
+						}
+
+						throw new CompilationResultException(UHTResult);
+					}
+
+					// Update the tool info file
+					DirectoryReference.CreateDirectory(ToolInfoFile.Directory);
+					FileReference.WriteAllText(ToolInfoFile, UhtAssemblyPath);
+
+					// Now that UHT has successfully finished generating code, we need to update all cached FileItems in case their last write time has changed.
+					// Otherwise UBT might not detect changes UHT made.
+					using (GlobalTracer.Instance.BuildSpan("ExternalExecution.ResetCachedHeaderInfo()").StartActive())
+					{
+						ResetCachedHeaderInfo(Makefile.UObjectModules);
+					}
+				}
+				else
+				{
+					Logger.LogDebug("Generated code is up to date.");
+				}
+
+				Progress.Write(2, 3);
+
+				using (GlobalTracer.Instance.BuildSpan("ExternalExecution.UpdateTimestamps()").StartActive())
+				{
+					using var uhtExecute = Timeline.ScopeEvent("UHTTimestamps");
+					UpdateTimestamps(Makefile.UObjectModules);
+				}
+
+				Progress.Write(3, 3);
+			}
+			if (ProgressWriter.bWriteMarkup)
+			{
+				Logger.LogInformation("@progress pop");
+			}
+		}
+
+		static void ResetCachedHeaderInfo(List<UHTModuleInfo> UObjectModules)
+		{
+			foreach (UHTModuleInfo ModuleInfo in UObjectModules)
+			{
+				ModuleInfo.GeneratedCodeDirectory.ResetCachedInfo();
+			}
+		}
+	}
+}
