[UBT]
* Moved intermediate files used for uht out of buildconfigurations (debug,development, etc). Since Inc-folder is outside configuration it does not make sense to have the intermediate files inside configurations.

diff --git a/Engine/Source/Programs/UnrealBuildTool/System/UHTExecution.cs b/Engine/Source/Programs/UnrealBuildTool/System/UHTExecution.cs
--- a/Engine/Source/Programs/UnrealBuildTool/System/UHTExecution.cs
+++ b/Engine/Source/Programs/UnrealBuildTool/System/UHTExecution.cs
@@ -1081,7 +1081,7 @@ namespace UnrealBuildTool
 		/// <returns>Manifest file name</returns>
 		public static FileReference GetUHTModuleInfoFileName(TargetMakefile Makefile, string TargetName)
 		{
-			return FileReference.Combine(Makefile.ProjectIntermediateDirectoryNoArch, TargetName + ".uhtmanifest");
+			return FileReference.Combine(Makefile.ProjectIntermediateDirectoryNoArch.ParentDirectory!, TargetName + ".uhtmanifest");
 		}
 
 		/// <summary>
@@ -1347,6 +1347,8 @@ namespace UnrealBuildTool
 				}
 				else
 				{
+					uhtManifestScope.Finish();
+
 					Logger.LogDebug("Generated code is up to date.");
 				}
 
