[UnrealBuildTool]
* Renamed ExternalExecution to UHTExecution because that is what it is

diff --git a/Engine/Source/Programs/UnrealBuildTool/Configuration/UEBuildTarget.cs b/Engine/Source/Programs/UnrealBuildTool/Configuration/UEBuildTarget.cs
--- a/Engine/Source/Programs/UnrealBuildTool/Configuration/UEBuildTarget.cs
+++ b/Engine/Source/Programs/UnrealBuildTool/Configuration/UEBuildTarget.cs
@@ -2591,9 +2591,9 @@ namespace UnrealBuildTool
 			HashSet<UEBuildModuleCPP> ModulesToGenerateHeadersFor = GatherDependencyModules(OriginalBinaries.ToList());
 
 			// Prepare cached data for UHT header generation
-			using (GlobalTracer.Instance.BuildSpan("ExternalExecution.SetupUObjectModules()").StartActive())
+			using (GlobalTracer.Instance.BuildSpan("UHTExecution.SetupUObjectModules()").StartActive())
 			{
-				ExternalExecution.SetupUObjectModules(
+				UHTExecution.SetupUObjectModules(
 					ModulesToGenerateHeadersFor, 
 					Rules.Platform,
 					ProjectDescriptor, 
@@ -2631,7 +2631,7 @@ namespace UnrealBuildTool
 			// files that are injected as top level prerequisites.  If UHT only emitted included header files, we wouldn't need to run it during the Gather phase at all.
 			if (Makefile.UObjectModules.Count > 0)
 			{
-				await ExternalExecution.ExecuteHeaderToolIfNecessaryAsync(BuildConfiguration, ProjectFile, Makefile, TargetName, Logger);
+				await UHTExecution.ExecuteHeaderToolIfNecessaryAsync(BuildConfiguration, ProjectFile, Makefile, TargetName, Logger);
 			}
 			
 			GatherExtraGeneratedCPPFileTypes(Makefile, TargetName, ref GlobalCompileEnvironment);
@@ -6456,7 +6456,7 @@ namespace UnrealBuildTool
 		private static void GatherExtraGeneratedCPPFileTypes(TargetMakefile Makefile, string TargetName, ref CppCompileEnvironment InCompileEnvironment)
 		{
 			UhtStdFileManager fileManager = new UhtStdFileManager();
-			FileReference moduleInfoFileName = ExternalExecution.GetUHTModuleInfoFileName(Makefile, TargetName);
+			FileReference moduleInfoFileName = UHTExecution.GetUHTModuleInfoFileName(Makefile, TargetName);
 			if (fileManager.ReadSource(moduleInfoFileName.FullName, out UhtSourceFragment fragment))
 			{
 				UHTManifest? manifest = System.Text.Json.JsonSerializer.Deserialize<UHTManifest>(fragment.Data.ToString());
diff --git a/Engine/Source/Programs/UnrealBuildTool/Modes/BuildMode.cs b/Engine/Source/Programs/UnrealBuildTool/Modes/BuildMode.cs
--- a/Engine/Source/Programs/UnrealBuildTool/Modes/BuildMode.cs
+++ b/Engine/Source/Programs/UnrealBuildTool/Modes/BuildMode.cs
@@ -1196,7 +1196,7 @@ namespace UnrealBuildTool
 				// If the target needs UHT to be run, we'll go ahead and do that now
 				if (Makefile.UObjectModules.Count > 0)
 				{
-					await ExternalExecution.ExecuteHeaderToolIfNecessaryAsync(BuildConfiguration, TargetDescriptor.ProjectFile, Makefile, TargetDescriptor.Name, Logger);
+					await UHTExecution.ExecuteHeaderToolIfNecessaryAsync(BuildConfiguration, TargetDescriptor.ProjectFile, Makefile, TargetDescriptor.Name, Logger);
 				}
 
 #if __VPROJECT_AVAILABLE__
diff --git a/Engine/Source/Programs/UnrealBuildTool/Modes/UnrealHeaderToolMode.cs b/Engine/Source/Programs/UnrealBuildTool/Modes/UnrealHeaderToolMode.cs
--- a/Engine/Source/Programs/UnrealBuildTool/Modes/UnrealHeaderToolMode.cs
+++ b/Engine/Source/Programs/UnrealBuildTool/Modes/UnrealHeaderToolMode.cs
@@ -1227,10 +1227,10 @@ namespace UnrealBuildTool.Modes
 					// Create the makefile
 					TargetMakefile Makefile = await Target.BuildAsync(BuildConfiguration, WorkingSet, TargetDesc, Logger, true);
 
-					FileReference ModuleInfoFileName = ExternalExecution.GetUHTModuleInfoFileName(Makefile, Target.TargetName);
-					FileReference DepsFileName = ExternalExecution.GetUHTDepsFileName(ModuleInfoFileName);
+					FileReference ModuleInfoFileName = UHTExecution.GetUHTModuleInfoFileName(Makefile, Target.TargetName);
+					FileReference DepsFileName = UHTExecution.GetUHTDepsFileName(ModuleInfoFileName);
 					ManifestPath = ModuleInfoFileName.FullName;
-					ExternalExecution.WriteUHTManifest(Makefile, Target.TargetName, ModuleInfoFileName, DepsFileName);
+					UHTExecution.WriteUHTManifest(Makefile, Target.TargetName, ModuleInfoFileName, DepsFileName);
 
 					if (Target.ProjectFile != null)
 					{
diff --git a/Engine/Source/Programs/UnrealBuildTool/System/UHTExecution.cs b/Engine/Source/Programs/UnrealBuildTool/System/UHTExecution.cs
--- a/Engine/Source/Programs/UnrealBuildTool/System/UHTExecution.cs
+++ b/Engine/Source/Programs/UnrealBuildTool/System/UHTExecution.cs
@@ -323,7 +323,7 @@ namespace UnrealBuildTool
 	/// <summary>
 	/// This handles all running of the UnrealHeaderTool
 	/// </summary>
-	class ExternalExecution
+	class UHTExecution
 	{
 		static UHTModuleType GetEngineModuleTypeFromDescriptor(ModuleDescriptor Module)
 		{
@@ -1020,42 +1020,6 @@ namespace UnrealBuildTool
 			});
 		}
 
-		/// <summary>
-		/// Run an external native executable (and capture the output), given the executable path and the commandline.
-		/// </summary>
-		public static int RunExternalNativeExecutable(FileReference ExePath, string Commandline, ILogger Logger)
-		{
-			Logger.LogDebug("RunExternalExecutable {ExePathFullName} {Commandline}", ExePath.FullName, Commandline);
-
-			using LogEventParser Parser = new(Logger);
-			Parser.AddMatchersFromAssembly(Assembly.GetExecutingAssembly());
-
-			using Process GameProcess = new();
-			GameProcess.StartInfo.FileName = ExePath.FullName;
-			GameProcess.StartInfo.Arguments = Commandline;
-			GameProcess.StartInfo.UseShellExecute = false;
-			GameProcess.StartInfo.RedirectStandardOutput = true;
-			GameProcess.OutputDataReceived += (s, e) => PrintProcessOutputAsync(e, Parser);
-			GameProcess.Start();
-			GameProcess.BeginOutputReadLine();
-			GameProcess.WaitForExit();
-
-			return GameProcess.ExitCode;
-		}
-
-		/// <summary>
-		/// Simple function to pipe output asynchronously
-		/// </summary>
-		private static void PrintProcessOutputAsync(DataReceivedEventArgs Event, LogEventParser Parser)
-		{
-			// DataReceivedEventHandler is fired with a null string when the output stream is closed.  We don't want to
-			// print anything for that event.
-			if (!String.IsNullOrEmpty(Event.Data))
-			{
-				Parser.WriteLine(Event.Data);
-			}
-		}
-
 		/// <summary>
 		/// Builds and runs the header tool and touches the header directories.
 		/// Performs any early outs if headers need no changes, given the UObject modules, tool path, game name, and configuration
@@ -1340,7 +1304,7 @@ namespace UnrealBuildTool
 
 					// Now that UHT has successfully finished generating code, we need to update all cached FileItems in case their last write time has changed.
 					// Otherwise UBT might not detect changes UHT made.
-					using (GlobalTracer.Instance.BuildSpan("ExternalExecution.ResetCachedHeaderInfo()").StartActive())
+					using (GlobalTracer.Instance.BuildSpan("UHTExecution.ResetCachedHeaderInfo()").StartActive())
 					{
 						ResetCachedHeaderInfo(Makefile.UObjectModules);
 					}
@@ -1354,7 +1318,7 @@ namespace UnrealBuildTool
 
 				Progress.Write(2, 3);
 
-				using (GlobalTracer.Instance.BuildSpan("ExternalExecution.UpdateTimestamps()").StartActive())
+				using (GlobalTracer.Instance.BuildSpan("UHTExecution.UpdateTimestamps()").StartActive())
 				{
 					using var uhtExecute = Timeline.ScopeEvent("UHTTimestamps");
 					UpdateTimestamps(Makefile.UObjectModules);
