[UBT]
* Removed code that was deprecated in 5.2

diff --git a/Engine/Source/Programs/UnrealBuildTool/System/UHTExecution.cs b/Engine/Source/Programs/UnrealBuildTool/System/UHTExecution.cs
--- a/Engine/Source/Programs/UnrealBuildTool/System/UHTExecution.cs
+++ b/Engine/Source/Programs/UnrealBuildTool/System/UHTExecution.cs
@@ -1020,23 +1020,6 @@ namespace UnrealBuildTool
 			});
 		}
 
-		/// <summary>
-		/// Builds and runs the header tool and touches the header directories.
-		/// Performs any early outs if headers need no changes, given the UObject modules, tool path, game name, and configuration
-		/// </summary>
-		public static async Task ExecuteHeaderToolIfNecessaryAsync(BuildConfiguration BuildConfiguration, FileReference? ProjectFile, TargetMakefile Makefile, string TargetName, ILogger Logger)
-		{
-
-			// No need to run UHT on itself
-			if (TargetName.Equals("UnrealHeaderTool", StringComparison.InvariantCultureIgnoreCase))
-			{
-				Logger.LogInformation("DEPRECATED: C++ UHT is being built by request.  C++ UHT has been deprecated and will be removed in 5.2");
-				return;
-			}
-
-			await ExecuteHeaderToolIfNecessaryInternalAsync(BuildConfiguration, ProjectFile, Makefile, TargetName, Logger);
-		}
-
 		/// <summary>
 		/// Return the file reference for the UHT manifest file
 		/// </summary>
@@ -1156,7 +1139,7 @@ namespace UnrealBuildTool
 		/// Builds and runs the header tool and touches the header directories.
 		/// Performs any early outs if headers need no changes, given the UObject modules, tool path, game name, and configuration
 		/// </summary>
-		private static async Task ExecuteHeaderToolIfNecessaryInternalAsync(BuildConfiguration BuildConfiguration, FileReference? ProjectFile, TargetMakefile Makefile, string TargetName, ILogger Logger)
+		public static async Task ExecuteHeaderToolIfNecessaryAsync(BuildConfiguration BuildConfiguration, FileReference? ProjectFile, TargetMakefile Makefile, string TargetName, ILogger Logger)
 		{
 			if (ProgressWriter.bWriteMarkup)
 			{
