[UBT]
* Fixed bug where multiple spawned procesess could backup the trace file even though they didn't create a trace file themselves

diff --git a/Engine/Source/Programs/UnrealBuildTool/UnrealBuildTool.cs b/Engine/Source/Programs/UnrealBuildTool/UnrealBuildTool.cs
--- a/Engine/Source/Programs/UnrealBuildTool/UnrealBuildTool.cs
+++ b/Engine/Source/Programs/UnrealBuildTool/UnrealBuildTool.cs
@@ -412,32 +412,35 @@ namespace UnrealBuildTool
             }
 
 			// If command line contains -Session it is a recursive action and we don't want to report this trace to the default trace channel
-			string? traceChannel = "Default";
-			string? traceName = "UnrealBuildTool";
-			if (ArgumentsArray.FirstOrDefault(s => s.Contains("-Session=")) != null)
+			string? traceChannel = null;
+			string? traceName =null;
+			string? traceFile = null;
+			if (ArgumentsArray.FirstOrDefault(s => s.Contains("-Session=")) == null)
 			{
-				traceChannel = null;
-				traceName = null;
-			}
-			// Setup trace file and backup old one
-			FileReference ubaTraceFile;
-			if (Unreal.IsBuildMachine())
-			{
-				ubaTraceFile = FileReference.Combine(Unreal.EngineProgramSavedDirectory, "AutomationTool", "Saved", "Logs", "Trace.uba");
-			}
-			else if (Log.OutputFile != null)
-			{
-				ubaTraceFile = Log.OutputFile.ChangeExtension(".uba");
-			}
-			else
-			{
-				ubaTraceFile = FileReference.Combine(Unreal.EngineProgramSavedDirectory, "UnrealBuildTool", "Trace.uba");
+				FileReference traceFileRef;
+				// Setup trace file and backup old one
+				if (Unreal.IsBuildMachine())
+				{
+					traceFileRef = FileReference.Combine(Unreal.EngineProgramSavedDirectory, "AutomationTool", "Saved", "Logs", "Trace.uba");
+				}
+				else if (Log.OutputFile != null)
+				{
+					traceFileRef = Log.OutputFile.ChangeExtension(".uba");
+				}
+				else
+				{
+					traceFileRef = FileReference.Combine(Unreal.EngineProgramSavedDirectory, "UnrealBuildTool", "Trace.uba");
+				}
+				DirectoryReference.CreateDirectory(traceFileRef.Directory);
+				Log.BackupLogFile(traceFileRef);
+
+				traceChannel = "Default";
+				traceName = "UnrealBuildTool";
+				traceFile = traceFileRef.FullName;
 			}
-			DirectoryReference.CreateDirectory(ubaTraceFile.Directory);
-			Log.BackupLogFile(ubaTraceFile);
 
 			// Attach timeline to uba trace
-			using EpicGames.UBA.ITrace trace = EpicGames.UBA.ITrace.Create(traceName, ubaTraceFile.FullName, true, traceChannel);
+			using EpicGames.UBA.ITrace trace = EpicGames.UBA.ITrace.Create(traceName, traceFile, true, traceChannel);
 			using TimelineEventScope _ = new((name) => trace.TaskBegin(name, ""), (id) => trace.TaskEnd(id));
 
 			bool bShowTimeline = false;
