[UBT]
* Changed CppCompileWarnings to fetch custom attributes once instead of over and over again. Saves a few seconds when generating makefile

diff --git a/Engine/Source/Programs/UnrealBuildTool/Configuration/CppCompileWarnings.cs b/Engine/Source/Programs/UnrealBuildTool/Configuration/CppCompileWarnings.cs
--- a/Engine/Source/Programs/UnrealBuildTool/Configuration/CppCompileWarnings.cs
+++ b/Engine/Source/Programs/UnrealBuildTool/Configuration/CppCompileWarnings.cs
@@ -1486,8 +1486,8 @@ namespace UnrealBuildTool
 
 		#endregion -- Public API --
 
-		private static Lazy<List<MemberInfo>> s_membersWithApplyWarningsAttribute = new(() => GetMembersWithAttributes(typeof(CppCompileWarnings), typeof(ApplyWarningsAttribute), typeof(WarningLevel)).ToList());
-		private static Lazy<List<MemberInfo>> s_membersWithWarningLevelDefaultAttribute = new(() => GetMembersWithAttributes(typeof(CppCompileWarnings), typeof(WarningLevelDefaultAttribute), typeof(WarningLevel)).ToList());
+		private static readonly Lazy<List<(MemberInfo, List<ApplyWarningsAttribute>)>> s_membersWithApplyWarningsAttribute = new(() => GetMembersWithAttributes<ApplyWarningsAttribute>(typeof(CppCompileWarnings), typeof(WarningLevel)));
+		private static readonly Lazy<List<(MemberInfo, List<WarningLevelDefaultAttribute>)>> s_membersWithWarningLevelDefaultAttribute = new(() => GetMembersWithAttributes<WarningLevelDefaultAttribute>(typeof(CppCompileWarnings), typeof(WarningLevel)).ToList());
 		private static ConcurrentDictionary<MemberInfo, Lazy<Dictionary<Type, List<WarningLevelDefaultAttribute>>>>[] s_groupedAttributes = { [], [], [] };
 
 		/// <summary>
@@ -1512,14 +1512,14 @@ namespace UnrealBuildTool
 			List<string> warningArguments = [];
 			List<ValueTuple<ApplyWarningsAttribute, WarningLevel>> warningAttributesBuffer = [];
 
-			foreach (MemberInfo memberInfo in s_membersWithApplyWarningsAttribute.Value)
+			foreach (var (mi, miAttr) in s_membersWithApplyWarningsAttribute.Value)
 			{
 				ApplyWarningsAttribute? applyWarningsAttribute = null;
-				IEnumerable<ApplyWarningsAttribute?> applyWarningsAttributes = memberInfo.GetCustomAttributes<ApplyWarningsAttribute>().Where(x => x.CanApplyToContext(toolChainContext));
+				IEnumerable<ApplyWarningsAttribute?> applyWarningsAttributes = miAttr.Where(x => x.CanApplyToContext(toolChainContext));
 
-				if (applyWarningsAttributes.Count() > 1)
+				if (applyWarningsAttributes.Skip(1).Any())
 				{
-					_logger?.LogWarning("Ambiguous ApplyWarningsAttributes on the property ({PropertyName}), for the provided tool-chain ({ToolChainName}). Skipping this property.", memberInfo.Name, toolChainContext._toolChainType.Name);
+					_logger?.LogWarning("Ambiguous ApplyWarningsAttributes on the property ({PropertyName}), for the provided tool-chain ({ToolChainName}). Skipping this property.", mi.Name, toolChainContext._toolChainType.Name);
 					continue;
 				}
 
@@ -1527,7 +1527,7 @@ namespace UnrealBuildTool
 
 				if (applyWarningsAttribute != null)
 				{
-					WarningLevel instanceWarningLevel = GetCurrentWarningLevel(memberInfo, this);
+					WarningLevel instanceWarningLevel = GetCurrentWarningLevel(mi, this);
 					warningAttributesBuffer.Add(new ValueTuple<ApplyWarningsAttribute, WarningLevel>(applyWarningsAttribute, instanceWarningLevel));
 				}
 			}
@@ -1556,10 +1556,10 @@ namespace UnrealBuildTool
 		{
 			cppCompileWarnings._logger?.LogTrace("Applying default C++ Compiler warnings for the Target({TargetName}) context.", cppCompileWarnings.BuildContextProvider?._targetRulesPrivate?.Name);
 
-			foreach (MemberInfo mi in s_membersWithWarningLevelDefaultAttribute.Value)
+			foreach (var (mi, miAttr) in s_membersWithWarningLevelDefaultAttribute.Value)
 			{
 				WarningLevel defaultWarningLevel = WarningLevel.Default;
-				defaultWarningLevel = CoalesceDefaultWarningFromMemberInfo(mi, InitializationContext.Target, defaultWarningLevel, cppCompileWarnings.BuildContextProvider, cppCompileWarnings._logger);
+				defaultWarningLevel = CoalesceDefaultWarningFromMemberInfo(mi, miAttr, InitializationContext.Target, defaultWarningLevel, cppCompileWarnings.BuildContextProvider, cppCompileWarnings._logger);
 
 				// Obtain instance value, and see if we need to preserve the TargetRules' constructor setting, or apply the default.
 				WarningLevel currentWarningLevel = GetCurrentWarningLevel(mi, cppCompileWarnings);
@@ -1570,12 +1570,12 @@ namespace UnrealBuildTool
 			}
 		}
 
-		internal static WarningLevel CoalesceDefaultWarningFromMemberInfo(MemberInfo mi, InitializationContext defaultWarningContext, WarningLevel defaultWarningLevel, BuildSystemContext? buildSystemContext, ILogger? logger)
+		internal static WarningLevel CoalesceDefaultWarningFromMemberInfo(MemberInfo mi, List<WarningLevelDefaultAttribute> miAttr, InitializationContext defaultWarningContext, WarningLevel defaultWarningLevel, BuildSystemContext? buildSystemContext, ILogger? logger)
 		{
 			// This is really slow so needs to be cached.
 			Dictionary<Type, List<WarningLevelDefaultAttribute>> groupedAttributes = s_groupedAttributes[(int)defaultWarningContext].GetOrAdd(mi, new Lazy<Dictionary<Type, List<WarningLevelDefaultAttribute>>>(() =>
 			{
-				IList<WarningLevelDefaultAttribute> attributes = [.. mi.GetCustomAttributes<WarningLevelDefaultAttribute>().Where(prop => prop.Context == InitializationContext.Any || prop.Context == defaultWarningContext)];
+				IList<WarningLevelDefaultAttribute> attributes = [.. miAttr.Where(prop => prop.Context == InitializationContext.Any || prop.Context == defaultWarningContext)];
 				Type fallbackType = typeof(BasicWarningLevelDefaultAttribute);
 				return attributes
 					.GroupBy(attr => attr.GetType())
@@ -1649,10 +1649,10 @@ namespace UnrealBuildTool
 
 			BuildSettingsVersion buildSettingsVersion = BuildContextProvider != null ? BuildContextProvider._buildContext.GetBuildSettings() : BuildSettingsVersion.V1;
 
-			foreach (MemberInfo mi in s_membersWithWarningLevelDefaultAttribute.Value)
+			foreach (var (mi, miAttr) in s_membersWithWarningLevelDefaultAttribute.Value)
 			{
 				WarningLevel defaultWarningLevel = WarningLevel.Default;
-				defaultWarningLevel = CoalesceDefaultWarningFromMemberInfo(mi, InitializationContext.Constructor, defaultWarningLevel, BuildContextProvider, _logger);
+				defaultWarningLevel = CoalesceDefaultWarningFromMemberInfo(mi, miAttr, InitializationContext.Constructor, defaultWarningLevel, BuildContextProvider, _logger);
 
 				SetWarningLevel(mi, this, defaultWarningLevel);
 			}
@@ -1660,17 +1660,34 @@ namespace UnrealBuildTool
 
 		#region -- Reflection helpers --
 
-		private static IEnumerable<MemberInfo> GetMembersWithAttributes(Type type, Type attributeType, Type targetType)
+		private static List<(MemberInfo, List<TAttr>)> GetMembersWithAttributes<TAttr>(Type type, Type targetType) where TAttr : Attribute
 		{
-			IEnumerable<MemberInfo> properties = type.GetProperties(BindingFlags.Instance | BindingFlags.Public | BindingFlags.NonPublic)
-				.Where(prop => Attribute.IsDefined(prop, attributeType) && prop.PropertyType == targetType)
-				.Select(prop => (MemberInfo)prop);
+			List<(MemberInfo, List<TAttr>)> list = [];
 
-			IEnumerable<MemberInfo> fields = type.GetFields(BindingFlags.Instance | BindingFlags.Public | BindingFlags.NonPublic)
-				.Where(field => Attribute.IsDefined(field, attributeType) && field.FieldType == targetType)
-				.Select(field => (MemberInfo)field);
+			void TryAddMemberInfo(MemberInfo mi, Type memberType)
+			{
+				if (memberType == targetType)
+				{
+					IEnumerable<TAttr> attrs = mi.GetCustomAttributes<TAttr>();
+					if (attrs.Any())
+					{
+						list.Add((mi, attrs.ToList()));
+					}
+				}
+			}
+
+			BindingFlags flags = BindingFlags.Instance | BindingFlags.Public | BindingFlags.NonPublic;
+			foreach (PropertyInfo prop in type.GetProperties(flags))
+			{
+				TryAddMemberInfo(prop, prop.PropertyType);
+			}
+
+			foreach (FieldInfo field in type.GetFields(flags))
+			{
+				TryAddMemberInfo(field, field.FieldType);
+			}
 
-			return properties.Concat(fields);
+			return list;
 		}
 
 		private static void SetWarningLevel(MemberInfo member, object instance, WarningLevel value)
