[UBT]
* Re-implemented Stable topological sort to be single threaded and quite much faster (in my tests main thread accumulated time went from 6 seconds to 1 second)

diff --git a/Engine/Source/Programs/UnrealBuildTool/Configuration/UEBuildModule.cs b/Engine/Source/Programs/UnrealBuildTool/Configuration/UEBuildModule.cs
--- a/Engine/Source/Programs/UnrealBuildTool/Configuration/UEBuildModule.cs
+++ b/Engine/Source/Programs/UnrealBuildTool/Configuration/UEBuildModule.cs
@@ -1595,82 +1595,108 @@ namespace UnrealBuildTool
 		/// <param name="NodeList">The list of nodes to sort.</param>
 		public static List<UEBuildModule> StableTopologicalSort(List<UEBuildModule> NodeList)
 		{
-			int NodeCount = NodeList.Count;
+			int n = NodeList.Count;
+			if (n <= 1)
+			{
+				// Nothing to do
+				return [.. NodeList];
+			}
 
-			// For each Node in NodeList, populated with the full circular dependency list from
-			// Node.GetAllDependencyModules()
-			List<Task<HashSet<UEBuildModule>>> NodeDependencies = new List<Task<HashSet<UEBuildModule>>>(NodeCount);
-			// Used to populate an element of NodeDependencies
-			HashSet<UEBuildModule> FetchDependencies(int NodeIndex)
+			// Map modules to their index for O(1) membership/ordering checks
+			Dictionary<UEBuildModule, int> indexOf = new (n);
+			for (int i = 0; i != n; i++)
 			{
-				HashSet<UEBuildModule> Dependencies = new HashSet<UEBuildModule>();
-				NodeList[NodeIndex].GetAllDependencyModules(new List<UEBuildModule>(), Dependencies, true, true, false);
-				return Dependencies;
+				indexOf[NodeList[i]] = i;
 			}
 
-			// For each Node in NodeList, populated with the nodes with a lower index in NodeList that Node depends on
-			List<Task<HashSet<UEBuildModule>>> PrecedingDependents = new List<Task<HashSet<UEBuildModule>>>(NodeCount);
-			HashSet<UEBuildModule> ComputePrecedingDependents(int NodeIndex)
+			// Build graph: edge (dep -> node) for every in-list dependency
+			// indegree[v] = number of in-list prerequisites of node v
+			List<int>[] adj = new List<int>[n];
+			for (int i = 0; i < n; i++)
+			{
+				adj[i] = [];
+			}
+
+			int[] indegree = new int[n];
+
+			// Compute each module's (transitive) dependencies once
+			// Only count dependencies that are present in NodeList
+			HashSet<UEBuildModule> depsBuffer = []; // reused per-iteration
+			List<UEBuildModule> scratch = [];       // reused per-iteration (API needs a List)
+
+			for (int v = 0; v != n; v++)
 			{
-				HashSet<UEBuildModule> Results = new HashSet<UEBuildModule>();
+				depsBuffer.Clear();
+				scratch.Clear();
 
-				UEBuildModule Node = NodeList[NodeIndex];
-				HashSet<UEBuildModule> Dependencies = NodeDependencies[NodeIndex].Result;
+				// Note: assumes GetAllDependencyModules fills 'dependencies' with transitive closure
+				NodeList[v].GetAllDependencyModules(scratch, depsBuffer, true, true, false);
 
-				for (int I = 0; I < NodeIndex; ++I)
+				foreach (UEBuildModule dep in depsBuffer)
 				{
-					if (NodeDependencies[I].Result.Contains(Node) && !Dependencies.Contains(NodeList[I]))
+					if (dep == NodeList[v])
 					{
-						Results.Add(NodeList[I]);
+						continue; // ignore self-dependency if any
 					}
-				}
 
-				return Results;
+					if (indexOf.TryGetValue(dep, out int u))
+					{
+						adj[u].Add(v);   // dep -> node
+						indegree[v]++;   // node has one more prerequisite within this set
+					}
+				}
 			}
 
-			for (int I = 0; I < NodeCount; ++I)
+			// Initialize queue with indegree-0 nodes in original order for stability
+			var queue = new Queue<int>();
+			for (int i = 0; i != n; i++)
 			{
-				int LocalI = I;
-				NodeDependencies.Add(Task.Run(() => FetchDependencies(LocalI)));
-				PrecedingDependents.Add(Task.Run(() => ComputePrecedingDependents(LocalI)));
+				if (indegree[i] == 0)
+				{
+					queue.Enqueue(i);
+				}
 			}
 
-			List<UEBuildModule> Out = new List<UEBuildModule>(NodeCount);
-			// Write the ordered output			
-			for (int Index1 = 0; Index1 != NodeCount; ++Index1)
+			List<int> order = new(n);
+			while (queue.Count > 0)
 			{
-				UEBuildModule Node1 = NodeList[Index1];
-				HashSet<UEBuildModule> NodesThatDependOnNode1 = PrecedingDependents[Index1].Result;
-				Out.Add(Node1);
+				int u = queue.Dequeue();
+				order.Add(u);
 
-				if (NodesThatDependOnNode1.Count == 0)
+				var list = adj[u];
+				for (int k = 0; k < list.Count; k++)
 				{
-					continue;
+					int v = list[k];
+					if (--indegree[v] == 0)
+					{
+						// Enqueue in discovery order; since we iterate adj lists
+						// of nodes emitted in input order, this preserves stability.
+						queue.Enqueue(v);
+					}
 				}
+			}
 
-				for (int Index2 = 0; Index2 != Index1; ++Index2)
+			// If there is a cycle, append remaining nodes in their original order.
+			// This keeps circularly-dependent nodes in their same relative order.
+			if (order.Count < n)
+			{
+				for (int i = 0; i < n; i++)
 				{
-					UEBuildModule Node2 = Out[Index2];
-
-					if (NodesThatDependOnNode1.Contains(Node2))
+					if (indegree[i] > 0)
 					{
-						// Rotate element at Index1 into position at Index2
-						for (int Index3 = Index1; Index3 != Index2;)
-						{
-							--Index3;
-							Out[Index3 + 1] = Out[Index3];
-						}
-
-						Out[Index2] = Node1;
-
-						// Break out of this loop, because this iteration must have covered all existing cases
-						// involving the node formerly at position Index1
-						break;
+						order.Add(i);
 					}
 				}
 			}
 
-			return Out;
+			// Materialize result
+			List<UEBuildModule> result = new(n);
+			for (int i = 0; i < order.Count; i++)
+			{
+				result.Add(NodeList[order[i]]);
+			}
+
+			return result;
 		}
 	};
 }
