[UBT]
* Optimized UEBuildModule.GetAllDependencyModules. Did not need to create a list just to traverse modules

diff --git a/Engine/Source/Programs/UnrealBuildTool/Configuration/UEBuildModule.cs b/Engine/Source/Programs/UnrealBuildTool/Configuration/UEBuildModule.cs
--- a/Engine/Source/Programs/UnrealBuildTool/Configuration/UEBuildModule.cs
+++ b/Engine/Source/Programs/UnrealBuildTool/Configuration/UEBuildModule.cs
@@ -1067,31 +1067,33 @@ namespace UnrealBuildTool
 		/// <param name="bOnlyDirectDependencies">True to return only this module's direct dependencies</param>
 		public void GetAllDependencyModules(List<UEBuildModule> ReferencedModules, HashSet<UEBuildModule> IgnoreReferencedModules, bool bIncludeDynamicallyLoaded, bool bForceCircular, bool bOnlyDirectDependencies)
 		{
-			List<UEBuildModule> AllDependencyModules = new List<UEBuildModule>(PrivateDependencyModules!.Count + PublicDependencyModules!.Count + (bIncludeDynamicallyLoaded ? DynamicallyLoadedModules!.Count : 0));
-			AllDependencyModules.AddRange(PrivateDependencyModules!);
-			AllDependencyModules.AddRange(PublicDependencyModules!);
-			if (bIncludeDynamicallyLoaded)
-			{
-				AllDependencyModules.AddRange(DynamicallyLoadedModules!);
-			}
-
-			foreach (UEBuildModule DependencyModule in AllDependencyModules)
+			void AddReferencedModules(List<UEBuildModule> Modules)
 			{
-				// Don't follow circular back-references!
-				if (bForceCircular || !HasCircularDependencyOn(DependencyModule))
+				foreach (UEBuildModule DependencyModule in Modules)
 				{
-					if (IgnoreReferencedModules.Add(DependencyModule))
+					// Don't follow circular back-references!
+					if (bForceCircular || !HasCircularDependencyOn(DependencyModule))
 					{
-						if (!bOnlyDirectDependencies)
+						if (IgnoreReferencedModules.Add(DependencyModule))
 						{
-							// Recurse into dependent modules first
-							DependencyModule.GetAllDependencyModules(ReferencedModules, IgnoreReferencedModules, bIncludeDynamicallyLoaded, bForceCircular, bOnlyDirectDependencies);
-						}
+							if (!bOnlyDirectDependencies)
+							{
+								// Recurse into dependent modules first
+								DependencyModule.GetAllDependencyModules(ReferencedModules, IgnoreReferencedModules, bIncludeDynamicallyLoaded, bForceCircular, bOnlyDirectDependencies);
+							}
 
-						ReferencedModules.Add(DependencyModule);
+							ReferencedModules.Add(DependencyModule);
+						}
 					}
 				}
 			}
+
+			AddReferencedModules(PrivateDependencyModules!);
+			AddReferencedModules(PublicDependencyModules!);
+			if (bIncludeDynamicallyLoaded)
+			{
+				AddReferencedModules(DynamicallyLoadedModules!);
+			}
 		}
 
 		/// <summary>
