[UBT]
* Added "Cleaning" timescope in CleanMode

diff --git a/Engine/Source/Programs/UnrealBuildTool/Modes/CleanMode.cs b/Engine/Source/Programs/UnrealBuildTool/Modes/CleanMode.cs
--- a/Engine/Source/Programs/UnrealBuildTool/Modes/CleanMode.cs
+++ b/Engine/Source/Programs/UnrealBuildTool/Modes/CleanMode.cs
@@ -60,6 +60,8 @@ namespace UnrealBuildTool
 			using ScopedTimer CleanTimer = new ScopedTimer("CleanMode.Clean()", Logger);
 			using IScope Scope = GlobalTracer.Instance.BuildSpan("CleanMode.Clean()").StartActive();
 
+			using var _ = Timeline.ScopeEvent("Cleaning");
+
 			if (TargetDescriptors.Count == 0)
 			{
 				throw new BuildException("No targets specified to clean");
