[UBT]
* Moved UHT/VNI to run in parallel with setting up targets. We do this by fake-populate directory cache with the files we know UHT/VNI will create. We then reset the directory cache before reaching the code that checks if things are outdated.
* Added DirectoryItem.PrepopulateCache that can be used to populate a directory cache without checking if files are there. This way we can populate directories while actually writing the files
* Changed FileMetadataPrefetch to use ParallelQueue and added tasks around each call to make it possible to wait for prefetching
* Changed so UHT always output .gen.cpp files even if empty  (like a #if 0 around the entire file). There are only 2-3 of these so does not affect perf
*

diff --git a/Engine/Source/Programs/Shared/EpicGames.Build/System/DirectoryItem.cs b/Engine/Source/Programs/Shared/EpicGames.Build/System/DirectoryItem.cs
--- a/Engine/Source/Programs/Shared/EpicGames.Build/System/DirectoryItem.cs
+++ b/Engine/Source/Programs/Shared/EpicGames.Build/System/DirectoryItem.cs
@@ -118,6 +118,23 @@ namespace UnrealBuildBase
 			}
 		}
 
+		/// <summary>
+		/// This function can be used to force the cache to be populated with files that not yet exists.
+		/// Useful when wanting to run things in parallel and one thread write files while the other use the list of files in the directory
+		/// </summary>
+		/// <param name="files"></param>
+		public void PrepopulateCache(List<string> files)
+		{
+			Dictionary<string, FileItem> cacheFiles = new(files.Count);
+			foreach (string file in files)
+			{
+				FileReference fileRef = FileReference.Combine(Location, file);
+				cacheFiles.Add(file, new FileItem(fileRef, file, new() { Attributes = FileAttributes.Normal }, this));
+			}
+
+			_cache = new(([], cacheFiles));
+		}
+
 		/// <summary>
 		/// Resets the cached info, if the DirectoryInfo is not found don't create a new entry
 		/// </summary>
diff --git a/Engine/Source/Programs/Shared/EpicGames.UHT/Exporters/CodeGen/UhtCodeGenerator.cs b/Engine/Source/Programs/Shared/EpicGames.UHT/Exporters/CodeGen/UhtCodeGenerator.cs
--- a/Engine/Source/Programs/Shared/EpicGames.UHT/Exporters/CodeGen/UhtCodeGenerator.cs
+++ b/Engine/Source/Programs/Shared/EpicGames.UHT/Exporters/CodeGen/UhtCodeGenerator.cs
@@ -116,6 +116,15 @@ namespace EpicGames.UHT.Exporters.CodeGen
 							new UhtHeaderCodeGeneratorCppFile(this, headerFile).Generate(factory);
 						});
 				}
+				else
+				{
+					// Still needs to write cpp file since UBT expects it to exist (there are very few of these)
+					HeaderInfos[headerFile.HeaderFileTypeIndex].Task = Factory.CreateTask(prereqs,
+						(IUhtExportFactory factory) =>
+						{
+							new UhtHeaderCodeGeneratorCppFile(this, headerFile).Generate(factory);
+						});
+				}
 			}
 
 			// Generate the files for the modules
diff --git a/Engine/Source/Programs/UnrealBuildTool/Configuration/UEBuildModuleCPP.cs b/Engine/Source/Programs/UnrealBuildTool/Configuration/UEBuildModuleCPP.cs
--- a/Engine/Source/Programs/UnrealBuildTool/Configuration/UEBuildModuleCPP.cs
+++ b/Engine/Source/Programs/UnrealBuildTool/Configuration/UEBuildModuleCPP.cs
@@ -659,11 +659,6 @@ namespace UnrealBuildTool
 				Dictionary<string, FileItem> GeneratedFiles = new Dictionary<string, FileItem>();
 				foreach (string GeneratedDir in GeneratedCppDirectories)
 				{
-					if (!Directory.Exists(GeneratedDir))
-					{
-						continue;
-					}
-
 					string Prefix = Path.GetFileName(GeneratedDir) + '/'; // "UHT/" or "VNI/"
 
 					DirectoryItem DirItem = DirectoryItem.GetItemByPath(GeneratedDir);
diff --git a/Engine/Source/Programs/UnrealBuildTool/Configuration/UEBuildTarget.cs b/Engine/Source/Programs/UnrealBuildTool/Configuration/UEBuildTarget.cs
--- a/Engine/Source/Programs/UnrealBuildTool/Configuration/UEBuildTarget.cs
+++ b/Engine/Source/Programs/UnrealBuildTool/Configuration/UEBuildTarget.cs
@@ -2511,6 +2511,8 @@ namespace UnrealBuildTool
 		/// </summary>
 		public async Task<TargetMakefile> BuildAsync(BuildConfiguration BuildConfiguration, ISourceFileWorkingSet WorkingSet, TargetDescriptor TargetDescriptor, ILogger Logger, bool bInitOnly = false)
 		{
+			await Task.CompletedTask;
+
 			CppConfiguration CppConfiguration = GetCppConfiguration(Configuration);
 
 			SourceFileMetadataCache MetadataCache = SourceFileMetadataCache.CreateHierarchy(ProjectFile, Logger);
@@ -2619,32 +2621,45 @@ namespace UnrealBuildTool
 			{
 				return Makefile;
 			}
-			
-#if __VPROJECT_AVAILABLE__
-			// Copy Verse BPVM usage flag
-			Makefile.bUseVerseBPVM = Rules.bUseVerseBPVM;
 
-			Task VNITask = VNIExecution.RunVNIAsync(ModulesToGenerateHeadersFor, Makefile, TargetDescriptor, Logger);
-#endif
+			// We inject all files that we know UHT/VNI will create into the directory cache.. this way we can run UHT/VNI without waiting
+			// .. and when we are about to test if actions are out-of-date we just reset directory cache so times are re-fetched
+			UHTExecution.PrepopulateDirectoryCache(Makefile.UObjectModules);
 
 			// NOTE: Even in Gather mode, we need to run UHT to make sure the files exist for the static action graph to be setup correctly.  This is because UHT generates .cpp
 			// files that are injected as top level prerequisites.  If UHT only emitted included header files, we wouldn't need to run it during the Gather phase at all.
+			Task? uhtTask = null;
 			if (Makefile.UObjectModules.Count > 0)
 			{
-				await UHTExecution.ExecuteHeaderToolIfNecessaryAsync(BuildConfiguration, ProjectFile, Makefile, TargetName, Logger);
+				uhtTask = Task.Run(() => UHTExecution.ExecuteHeaderToolIfNecessaryAsync(BuildConfiguration, ProjectFile, Makefile, TargetName, Logger, false));
 			}
-			
-			GatherExtraGeneratedCPPFileTypes(Makefile, TargetName, ref GlobalCompileEnvironment);
 
-			// Prefetch directory items for UHT folders since they are going to be used later
-			foreach (UEBuildModuleCPP Module in ModulesToGenerateHeadersFor)
+#if __VPROJECT_AVAILABLE__
+			// Prepare cached data for VNI header generation
+			using (new ScopedTimer("UHTExecution.SetupVNIModules()", Logger, LogLevel.Debug, false))
+			{
+				VNIExecution.SetupVNIModules(ModulesToGenerateHeadersFor, out Makefile.VNIModules);
+			}
+			VNIExecution.PrepopulateDirectoryCache(Makefile);
+
+			// Copy Verse BPVM usage flag
+			Makefile.bUseVerseBPVM = Rules.bUseVerseBPVM;
+			Task? vniTask = null;
+			if (Makefile.VNIModules.Count > 0)
 			{
-				if (Module.GeneratedCodeDirectoryUHT != null)
+				vniTask = Task.Run(() =>
 				{
-					FileMetadataPrefetch.QueueDirectoryTree(Module.GeneratedCodeDirectoryUHT);
-				}
+					using (new ScopedTimer("ExternalExecution.ExecuteVNITool()", Logger, LogLevel.Debug, false))
+					{
+						VNIExecution.ExecuteVNITool(Makefile, TargetDescriptor, Logger, false);
+					}
+				});
 			}
 
+#endif
+
+			GatherExtraGeneratedCPPFileTypes(Makefile, TargetName, ref GlobalCompileEnvironment);
+
 			if (Rules.bUseSharedPCHs)
 			{
 				// Find all the shared PCHs.
@@ -2654,11 +2669,6 @@ namespace UnrealBuildTool
 				CreateSharedPCHInstances(Rules, TargetToolChain, OriginalBinaries, GlobalCompileEnvironment, MakefileBuilder, Logger);
 			}
 
-			// Can probably be moved further down?
-#if __VPROJECT_AVAILABLE__
-			await VNITask;
-#endif
-
 			foreach (UEBuildModuleCPP Module in Modules.Values.OfType<UEBuildModuleCPP>())
 			{
 				foreach ((string Subdirectory, Action<ILogger, DirectoryReference> Func) in Module.Rules.GenerateHeaderFuncs)
@@ -3053,6 +3063,21 @@ namespace UnrealBuildTool
 			// Wait for makefile tasks to finish (tasks writing files to disk etc)
 			MakefileBuilder.WaitOnWriteTasks();
 
+#if __VPROJECT_AVAILABLE__
+			if (vniTask != null)
+			{
+				VNIExecution.WaitForExecution(vniTask, Makefile);
+			}
+#endif
+
+			if (uhtTask != null)
+			{
+				UHTExecution.WaitForExecution(uhtTask, Makefile);
+
+				// Prefetch directory items for UHT folders since they are going to be used later for check out-of-date actions
+				Task task = FileMetadataPrefetch.QueueDirectoryTrees(ModulesToGenerateHeadersFor.Select(m => m.GeneratedCodeDirectoryUHT));
+			}
+
 			// Clean any stale modules which exist in multiple output directories. This can lead to the wrong DLL being loaded on Windows.
 			CleanStaleModules(Logger);
 
diff --git a/Engine/Source/Programs/UnrealBuildTool/Modes/BuildMode.cs b/Engine/Source/Programs/UnrealBuildTool/Modes/BuildMode.cs
--- a/Engine/Source/Programs/UnrealBuildTool/Modes/BuildMode.cs
+++ b/Engine/Source/Programs/UnrealBuildTool/Modes/BuildMode.cs
@@ -224,7 +224,7 @@ namespace UnrealBuildTool
 						if (TargetDesc.ProjectFile != null)
 						{
 							DirectoryReference ProjectDirectory = TargetDesc.ProjectFile.Directory;
-							FileMetadataPrefetch.QueueProjectDirectory(ProjectDirectory);
+							Task task = FileMetadataPrefetch.QueueProjectDirectory(ProjectDirectory);
 							ProjectDirs.Add(ProjectDirectory);
 						}
 
diff --git a/Engine/Source/Programs/UnrealBuildTool/System/FileMetadataPrefetch.cs b/Engine/Source/Programs/UnrealBuildTool/System/FileMetadataPrefetch.cs
--- a/Engine/Source/Programs/UnrealBuildTool/System/FileMetadataPrefetch.cs
+++ b/Engine/Source/Programs/UnrealBuildTool/System/FileMetadataPrefetch.cs
@@ -1,8 +1,11 @@
 // Copyright Epic Games, Inc. All Rights Reserved.
 
+using System;
+using System.Collections.Concurrent;
 using System.Collections.Generic;
 using System.Linq;
 using System.Threading;
+using System.Threading.Tasks;
 using EpicGames.Core;
 using UnrealBuildBase;
 
@@ -13,11 +16,6 @@ namespace UnrealBuildTool
 	/// </summary>
 	static class FileMetadataPrefetch
 	{
-		/// <summary>
-		/// Queue for tasks added to the thread pool
-		/// </summary>
-		static ThreadPoolWorkQueue Queue = new ThreadPoolWorkQueue();
-
 		/// <summary>
 		/// Used to cancel any queued tasks
 		/// </summary>
@@ -31,58 +29,60 @@ namespace UnrealBuildTool
 		/// <summary>
 		/// Set of all the directory trees that have been queued up, to save adding any more than once.
 		/// </summary>
-		static HashSet<DirectoryReference> QueuedDirectories = new HashSet<DirectoryReference>();
+		static ConcurrentDictionary<DirectoryReference, Lazy<Task>> QueuedDirectories = new();
 
 		/// <summary>
 		/// Enqueue the engine directory for prefetching
 		/// </summary>
-		public static void QueueEngineDirectory()
+		public static Task QueueEngineDirectory()
 		{
-			lock (QueuedDirectories)
+			return QueuedDirectories.GetOrAdd(Unreal.EngineDirectory, new Lazy<Task>(() =>
 			{
-				if (QueuedDirectories.Add(Unreal.EngineDirectory))
+				return Task.Run(() =>
 				{
-					Enqueue(() => ScanEngineDirectory());
-				}
-			}
+					//using var _ = Timeline.ScopeEvent("PrefetchEngineDir");
+					ParallelQueue queue = new();
+					queue.Enqueue(() => ScanEngineDirectory(queue));
+					queue.Drain();
+				});
+			})).Value;
 		}
 
 		/// <summary>
 		/// Enqueue a project directory for prefetching
 		/// </summary>
 		/// <param name="ProjectDirectory">The project directory to prefetch</param>
-		public static void QueueProjectDirectory(DirectoryReference ProjectDirectory)
+		public static Task QueueProjectDirectory(DirectoryReference ProjectDirectory)
 		{
-			lock (QueuedDirectories)
+			return QueuedDirectories.GetOrAdd(ProjectDirectory, new Lazy<Task>(() =>
 			{
-				if (QueuedDirectories.Add(ProjectDirectory))
+				return Task.Run(() =>
 				{
-					Enqueue(() => ScanProjectDirectory(DirectoryItem.GetItemByDirectoryReference(ProjectDirectory)));
-				}
-			}
+					//using var _ = Timeline.ScopeEvent("PrefetchProjectDir");
+					ParallelQueue queue = new();
+					queue.Enqueue(() => ScanProjectDirectory(queue, DirectoryItem.GetItemByDirectoryReference(ProjectDirectory)));
+					queue.Drain();
+				});
+			})).Value;
 		}
 
 		/// <summary>
 		/// Enqueue a directory tree for prefetching
 		/// </summary>
-		/// <param name="Directory">Directory to start searching from</param>
-		public static void QueueDirectoryTree(DirectoryReference Directory)
+		/// <param name="Directories">Directory to start searching from</param>
+		public static Task QueueDirectoryTrees(IEnumerable<DirectoryReference?> Directories)
 		{
-			lock (QueuedDirectories)
+			return Task.Run(() =>
 			{
-				if (QueuedDirectories.Add(Directory))
+				//using var _ = Timeline.ScopeEvent("PrefetchDir");
+				ParallelQueue queue = new();
+				foreach (DirectoryReference? dir in Directories)
 				{
-					Enqueue(() => ScanDirectoryTree(DirectoryItem.GetItemByDirectoryReference(Directory)));
+					if (dir != null)
+						queue.Enqueue(() => ScanDirectoryTree(queue, DirectoryItem.GetItemByDirectoryReference(dir)));
 				}
-			}
-		}
-
-		/// <summary>
-		/// Wait for the prefetcher to complete all reqeusted tasks
-		/// </summary>
-		public static void Wait()
-		{
-			Queue.Wait();
+				queue.Drain();
+			});
 		}
 
 		/// <summary>
@@ -90,29 +90,17 @@ namespace UnrealBuildTool
 		/// </summary>
 		public static void Stop()
 		{
-			CancelSource.Cancel();
-			Queue.Wait();
-		}
-
-		/// <summary>
-		/// Enqueue a task which checks for the cancellation token first
-		/// </summary>
-		/// <param name="Action">Action to enqueue</param>
-		static void Enqueue(System.Action Action)
-		{
-			Queue.Enqueue(() =>
+			foreach (var  directory in QueuedDirectories.Values)
 			{
-				if (!CancelToken.IsCancellationRequested)
-				{
-					Action();
-				}
-			});
+				directory.Value.Wait();
+			}
+			QueuedDirectories.Clear();
 		}
 
 		/// <summary>
 		/// Scans the engine directory, adding tasks for subdirectories
 		/// </summary>
-		static void ScanEngineDirectory()
+		static void ScanEngineDirectory(ParallelQueue queue)
 		{
 			foreach (DirectoryReference ExtensionDir in Unreal.GetExtensionDirs(Unreal.EngineDirectory))
 			{
@@ -120,27 +108,28 @@ namespace UnrealBuildTool
 				BaseDirectory.CacheDirectories();
 
 				DirectoryItem BasePluginsDirectory = DirectoryItem.Combine(BaseDirectory, "Plugins");
-				Enqueue(() => ScanPluginFolder(BasePluginsDirectory));
+				queue.Enqueue(() => ScanPluginFolder(queue, BasePluginsDirectory));
 
 				DirectoryItem BaseSourceDirectory = DirectoryItem.Combine(BaseDirectory, "Source");
 				BaseSourceDirectory.CacheDirectories();
 
 				DirectoryItem BaseSourceRuntimeDirectory = DirectoryItem.Combine(BaseSourceDirectory, "Runtime");
-				Enqueue(() => ScanDirectoryTree(BaseSourceRuntimeDirectory));
+				queue.Enqueue(() => ScanDirectoryTree(queue, BaseSourceRuntimeDirectory));
 
 				DirectoryItem BaseSourceDeveloperDirectory = DirectoryItem.Combine(BaseSourceDirectory, "Developer");
-				Enqueue(() => ScanDirectoryTree(BaseSourceDeveloperDirectory));
+				queue.Enqueue(() => ScanDirectoryTree(queue, BaseSourceDeveloperDirectory));
 
 				DirectoryItem BaseSourceEditorDirectory = DirectoryItem.Combine(BaseSourceDirectory, "Editor");
-				Enqueue(() => ScanDirectoryTree(BaseSourceEditorDirectory));
+				queue.Enqueue(() => ScanDirectoryTree(queue, BaseSourceEditorDirectory));
 			}
 		}
 
 		/// <summary>
 		/// Scans a project directory, adding tasks for subdirectories
 		/// </summary>
+		/// <param name="queue">The project directory to search</param>
 		/// <param name="ProjectDirectory">The project directory to search</param>
-		static void ScanProjectDirectory(DirectoryItem ProjectDirectory)
+		static void ScanProjectDirectory(ParallelQueue queue, DirectoryItem ProjectDirectory)
 		{
 			foreach (DirectoryReference ExtensionDir in Unreal.GetExtensionDirs(ProjectDirectory.Location))
 			{
@@ -148,18 +137,19 @@ namespace UnrealBuildTool
 				BaseDirectory.CacheDirectories();
 
 				DirectoryItem BasePluginsDirectory = DirectoryItem.Combine(BaseDirectory, "Plugins");
-				Enqueue(() => ScanPluginFolder(BasePluginsDirectory));
+				queue.Enqueue(() => ScanPluginFolder(queue, BasePluginsDirectory));
 
 				DirectoryItem BaseSourceDirectory = DirectoryItem.Combine(BaseDirectory, "Source");
-				Enqueue(() => ScanDirectoryTree(BaseSourceDirectory));
+				queue.Enqueue(() => ScanDirectoryTree(queue, BaseSourceDirectory));
 			}
 		}
 
 		/// <summary>
 		/// Scans a plugin parent directory, adding tasks for subdirectories
 		/// </summary>
+		/// <param name="queue">Root of the directory tree</param>
 		/// <param name="Directory">The directory which may contain plugin directories</param>
-		static void ScanPluginFolder(DirectoryItem Directory)
+		static void ScanPluginFolder(ParallelQueue queue, DirectoryItem Directory)
 		{
 			if (CancelToken.IsCancellationRequested || Directory.TryGetFile(".ubtignore", out FileItem? _))
 			{
@@ -170,11 +160,11 @@ namespace UnrealBuildTool
 			{
 				if (SubDirectory.EnumerateFiles().Any((fi) => fi.HasExtension(".uplugin")))
 				{
-					Enqueue(() => ScanDirectoryTree(DirectoryItem.Combine(SubDirectory, "Source")));
+					queue.Enqueue(() => ScanDirectoryTree(queue, DirectoryItem.Combine(SubDirectory, "Source")));
 				}
 				else if (!SubDirectory.TryGetFile(".ubtignore", out FileItem? OutFile))
 				{
-					Enqueue(() => ScanPluginFolder(SubDirectory));
+					queue.Enqueue(() => ScanPluginFolder(queue, SubDirectory));
 				}
 			}
 		}
@@ -182,8 +172,9 @@ namespace UnrealBuildTool
 		/// <summary>
 		/// Scans an arbitrary directory tree
 		/// </summary>
+		/// <param name="queue">Root of the directory tree</param>
 		/// <param name="Directory">Root of the directory tree</param>
-		static void ScanDirectoryTree(DirectoryItem Directory)
+		static void ScanDirectoryTree(ParallelQueue queue, DirectoryItem Directory)
 		{
 			if (CancelToken.IsCancellationRequested || Directory.TryGetFile(".ubtignore", out FileItem? _))
 			{
@@ -192,7 +183,7 @@ namespace UnrealBuildTool
 
 			foreach (DirectoryItem SubDirectory in Directory.EnumerateDirectories())
 			{
-				Enqueue(() => ScanDirectoryTree(SubDirectory));
+				queue.Enqueue(() => ScanDirectoryTree(queue, SubDirectory));
 			}
 		}
 	}
diff --git a/Engine/Source/Programs/UnrealBuildTool/System/UHTExecution.cs b/Engine/Source/Programs/UnrealBuildTool/System/UHTExecution.cs
--- a/Engine/Source/Programs/UnrealBuildTool/System/UHTExecution.cs
+++ b/Engine/Source/Programs/UnrealBuildTool/System/UHTExecution.cs
@@ -979,7 +979,7 @@ namespace UnrealBuildTool
 		/// <summary>
 		/// Updates the recorded timestamps of all the passed in UObject modules
 		/// </summary>
-		private static void UpdateTimestamps(List<UHTModuleInfo> UObjectModules)
+		private static void UpdateTimestamps(List<UHTModuleInfo> UObjectModules, bool resetDirectoryCaches)
 		{
 			Parallel.ForEach(UObjectModules, async Module =>
 			{
@@ -1008,8 +1008,11 @@ namespace UnrealBuildTool
 								await FileReference.WriteAllLinesAsync(TimestampFile, AllUObjectFiles);
 							}
 
-							// Because new .cpp and .h files may have been generated by UHT, invalidate the DirectoryLookupCache
-							DirectoryLookupCache.InvalidateCachedDirectory(new DirectoryReference(GeneratedCodeDirectoryInfo.FullName));
+							if (resetDirectoryCaches)
+							{
+								// Because new .cpp and .h files may have been generated by UHT, invalidate the DirectoryLookupCache
+								DirectoryLookupCache.InvalidateCachedDirectory(new DirectoryReference(GeneratedCodeDirectoryInfo.FullName));
+							}
 						}
 					}
 					catch (Exception Exception)
@@ -1135,11 +1138,35 @@ namespace UnrealBuildTool
 			return Lhs > Rhs ? Lhs : Rhs;
 		}
 
+		/// <summary>
+		/// Will prepopulate DirectoryInfo cache with files that will be written.
+		/// This can be combined with ExecuteHeaderToolIfNecessaryAsync(resetDirectoryCaches = false) and WaitForExecution
+		/// </summary>
+		/// <param name="moduleInfos"></param>
+		public static void PrepopulateDirectoryCache(List<UHTModuleInfo> moduleInfos)
+		{
+			foreach (UHTModuleInfo mi in moduleInfos)
+			{
+				List<string> outputFiles = [];
+				outputFiles.Add(mi.ModuleName + ".init.gen.cpp");
+				outputFiles.Add(mi.ModuleName + "Classes.cpp");
+
+				foreach (FileItem? f in mi.PublicUObjectClassesHeaders.Concat(mi.PublicUObjectHeaders).Concat(mi.InternalUObjectHeaders).Concat(mi.PrivateUObjectHeaders))
+				{
+					string fileName = f.Location.GetFileNameWithoutExtension();
+					outputFiles.Add(fileName + ".generated.h");
+					outputFiles.Add(fileName + ".gen.cpp");
+				}
+
+				mi.GeneratedCodeDirectory.PrepopulateCache(outputFiles);
+			}
+		}
+
 		/// <summary>
 		/// Builds and runs the header tool and touches the header directories.
 		/// Performs any early outs if headers need no changes, given the UObject modules, tool path, game name, and configuration
 		/// </summary>
-		public static async Task ExecuteHeaderToolIfNecessaryAsync(BuildConfiguration BuildConfiguration, FileReference? ProjectFile, TargetMakefile Makefile, string TargetName, ILogger Logger)
+		public static async Task ExecuteHeaderToolIfNecessaryAsync(BuildConfiguration BuildConfiguration, FileReference? ProjectFile, TargetMakefile Makefile, string TargetName, ILogger Logger, bool resetDirectoryCaches = true)
 		{
 			if (ProgressWriter.bWriteMarkup)
 			{
@@ -1285,11 +1312,14 @@ namespace UnrealBuildTool
 					DirectoryReference.CreateDirectory(ToolInfoFile.Directory);
 					FileReference.WriteAllText(ToolInfoFile, UhtAssemblyPath);
 
-					// Now that UHT has successfully finished generating code, we need to update all cached FileItems in case their last write time has changed.
-					// Otherwise UBT might not detect changes UHT made.
-					using (GlobalTracer.Instance.BuildSpan("UHTExecution.ResetCachedHeaderInfo()").StartActive())
+					if (resetDirectoryCaches)
 					{
-						ResetCachedHeaderInfo(Makefile.UObjectModules);
+						// Now that UHT has successfully finished generating code, we need to update all cached FileItems in case their last write time has changed.
+						// Otherwise UBT might not detect changes UHT made.
+						using (GlobalTracer.Instance.BuildSpan("UHTExecution.ResetCachedHeaderInfo()").StartActive())
+						{
+							ResetCachedHeaderInfo(Makefile.UObjectModules);
+						}
 					}
 				}
 				else
@@ -1304,7 +1334,7 @@ namespace UnrealBuildTool
 				using (GlobalTracer.Instance.BuildSpan("UHTExecution.UpdateTimestamps()").StartActive())
 				{
 					using var uhtExecute = Timeline.ScopeEvent("UHTTimestamps");
-					UpdateTimestamps(Makefile.UObjectModules);
+					UpdateTimestamps(Makefile.UObjectModules, resetDirectoryCaches);
 				}
 
 				Progress.Write(3, 3);
@@ -1315,6 +1345,18 @@ namespace UnrealBuildTool
 			}
 		}
 
+		/// <summary>
+		/// Will wait for execution task to finish and the reset cached directories
+		/// </summary>
+		/// <param name="executionTask"></param>
+		/// <param name="Makefile"></param>
+		public static void WaitForExecution(Task executionTask, TargetMakefile Makefile)
+		{
+			using var uhtManifestScope = Timeline.ScopeEvent("UHTWait");
+			executionTask.Wait();
+			ResetCachedHeaderInfo(Makefile.UObjectModules);
+		}
+
 		static void ResetCachedHeaderInfo(List<UHTModuleInfo> UObjectModules)
 		{
 			foreach (UHTModuleInfo ModuleInfo in UObjectModules)
diff --git a/Engine/Source/Programs/UnrealBuildTool/System/Unity.cs b/Engine/Source/Programs/UnrealBuildTool/System/Unity.cs
--- a/Engine/Source/Programs/UnrealBuildTool/System/Unity.cs
+++ b/Engine/Source/Programs/UnrealBuildTool/System/Unity.cs
@@ -371,7 +371,11 @@ namespace UnrealBuildTool
 			foreach (FileItem CPPFile in CPPFiles)
 			{
 				bool bHeaderInWorkingSet = false;
-				if (CompileEnvironment.MetadataCache.GetFirstInclude(CPPFile) is string FirstInclude &&
+				if (CPPFile.Name.EndsWith(".gen.cpp"))
+				{
+					// Ignore these files, does not need to be in adapt 
+				}
+				else if (CompileEnvironment.MetadataCache.GetFirstInclude(CPPFile) is string FirstInclude &&
 					NameToHeaderFileInWorkingSet.TryGetValue(Path.GetFileName(FirstInclude), out FileItem? HeaderFile))
 				{
 					bHeaderInWorkingSet = true;
