[UBT]
* SourceFileMetadataCache - Changed lambda to local function

diff --git a/Engine/Source/Programs/UnrealBuildTool/System/SourceFileMetadataCache.cs b/Engine/Source/Programs/UnrealBuildTool/System/SourceFileMetadataCache.cs
--- a/Engine/Source/Programs/UnrealBuildTool/System/SourceFileMetadataCache.cs
+++ b/Engine/Source/Programs/UnrealBuildTool/System/SourceFileMetadataCache.cs
@@ -192,7 +192,7 @@ namespace UnrealBuildTool
 			}
 			else
 			{
-				Func<FileItem, SourceFileInfo> UpdateSourceFileInfo = (FileItem SourceFile) =>
+				SourceFileInfo UpdateSourceFileInfo(FileItem SourceFile)
 				{
 					SourceFileInfo SourceFileInfo = new SourceFileInfo();
 					string FileText = FileReference.ReadAllText(SourceFile.Location);
