[UBT]
* Fixed so ActionGraph.DeleteOutdatedProducedItems is running in parallel

diff --git a/Engine/Source/Programs/UnrealBuildTool/System/ActionGraph.cs b/Engine/Source/Programs/UnrealBuildTool/System/ActionGraph.cs
--- a/Engine/Source/Programs/UnrealBuildTool/System/ActionGraph.cs
+++ b/Engine/Source/Programs/UnrealBuildTool/System/ActionGraph.cs
@@ -962,7 +962,7 @@ namespace UnrealBuildTool
 
 			using (GlobalTracer.Instance.BuildSpan("Cache individual outdated actions").StartActive())
 			{
-				Parallel.ForEach(Actions,
+				Parallel2.ForEach(Actions,
 					Action => IsIndividualActionOutdated(Action, OutdatedActions,
 						ActionHistory, CppDependencies, bIgnoreOutdatedImportLibraries, Logger));
 			}
@@ -985,7 +985,7 @@ namespace UnrealBuildTool
 		{
 			using ITimelineEvent __ = Timeline.ScopeEvent($"DeleteOutdatedProducedItems");
 
-			foreach (LinkedAction OutdatedAction in OutdatedActions)
+			Parallel2.ForEach(OutdatedActions, OutdatedAction =>
 			{
 				foreach (FileItem DeleteItem in OutdatedAction.DeleteItems)
 				{
@@ -995,7 +995,7 @@ namespace UnrealBuildTool
 						DeleteItem.Delete(Logger);
 					}
 				}
-			}
+			});
 		}
 
 		/// <summary>
