[UBT]
* Added retry on writing .uhtmanifest.. Can't figure out what is holding it and also what is releasing it to make it succeed when trying again. Will revisit this but this seems to solve the farm issue

diff --git a/Engine/Source/Programs/UnrealBuildTool/System/UHTExecution.cs b/Engine/Source/Programs/UnrealBuildTool/System/UHTExecution.cs
--- a/Engine/Source/Programs/UnrealBuildTool/System/UHTExecution.cs
+++ b/Engine/Source/Programs/UnrealBuildTool/System/UHTExecution.cs
@@ -1124,7 +1124,24 @@ namespace UnrealBuildTool
 			// @todo ubtmake: Optimization: Ideally we could avoid having to generate this data in the case where UHT doesn't even need to run!  Can't we use the existing copy?  (see below use of Manifest)
 			UHTManifest Manifest = CreateUHTManifest(Makefile, TargetName, DepsFileName);
 			Directory.CreateDirectory(ModuleInfoFileName.Directory.FullName);
-			System.IO.File.WriteAllText(ModuleInfoFileName.FullName, JsonSerializer.Serialize(Manifest, new JsonSerializerOptions { WriteIndented = true }));
+
+			int retry = 0;
+			while (true)
+			{
+				try
+				{
+					System.IO.File.WriteAllText(ModuleInfoFileName.FullName, JsonSerializer.Serialize(Manifest, new JsonSerializerOptions { WriteIndented = true }));
+					break;
+				}
+				catch (IOException) // I have no idea what could cause this. Something is holding on to the file but I can't see what. and it goes away when waiting a little bit
+				{
+					if (retry++ > 5)
+					{
+						throw;
+					}
+					System.Threading.Thread.Sleep(500);
+				}
+			}
 		}
 
 		/// <summary>
