[UBT]
* Skip .generated.h files in FindIncludedHeaders in UEBuildModuleCpp. This is used for shared pch and we know that the only way .generated.h is included is through the representing header which is in the same module

diff --git a/Engine/Source/Programs/UnrealBuildTool/Configuration/UEBuildModuleCPP.cs b/Engine/Source/Programs/UnrealBuildTool/Configuration/UEBuildModuleCPP.cs
--- a/Engine/Source/Programs/UnrealBuildTool/Configuration/UEBuildModuleCPP.cs
+++ b/Engine/Source/Programs/UnrealBuildTool/Configuration/UEBuildModuleCPP.cs
@@ -1290,6 +1290,12 @@ namespace UnrealBuildTool
 		{
 			foreach (string HeaderInclude in CompileEnvironment.MetadataCache.GetHeaderIncludes(FileToSearch))
 			{
+				// Skip .generated.h. They should always be in the same module as the header including them.. so no need to track them
+				if (HeaderInclude.Contains(".generated.h"))
+				{
+					continue;
+				}
+
 				string TranformedHeaderInclude = HeaderInclude;
 				if (TranformedHeaderInclude.Contains("COMPILED_PLATFORM_HEADER("))
 				{
