[UBA]
* Fixed so ctrl-c is just setting variable that process shouldCancel instead of actually doing work. This will likely fix the crash that can happen when pressing ctrl-c

diff --git a/Engine/Source/Programs/UnrealBuildTool/Executors/UnrealBuildAccelerator/UBAExecutor.cs b/Engine/Source/Programs/UnrealBuildTool/Executors/UnrealBuildAccelerator/UBAExecutor.cs
--- a/Engine/Source/Programs/UnrealBuildTool/Executors/UnrealBuildAccelerator/UBAExecutor.cs
+++ b/Engine/Source/Programs/UnrealBuildTool/Executors/UnrealBuildAccelerator/UBAExecutor.cs
@@ -833,6 +833,7 @@ namespace UnrealBuildTool
 				Console.CancelKeyPress -= cancelHandler;
 				cancelHandler = null;
 				e.Cancel = true;
+				shouldCancel = true;
 				cancelBuild();
 			};
 			Console.CancelKeyPress += cancelHandler;
