[UBT]
* Added so DependenciesForStableTopologicalSort is calculated in parallel for all modules and then cached.

diff --git a/Engine/Source/Programs/UnrealBuildTool/Configuration/UEBuildBinary.cs b/Engine/Source/Programs/UnrealBuildTool/Configuration/UEBuildBinary.cs
--- a/Engine/Source/Programs/UnrealBuildTool/Configuration/UEBuildBinary.cs
+++ b/Engine/Source/Programs/UnrealBuildTool/Configuration/UEBuildBinary.cs
@@ -414,6 +414,15 @@ namespace UnrealBuildTool
 		/// <returns>List of all referenced modules</returns>
 		public List<UEBuildModule> GetAllDependencyModules(bool bIncludeDynamicallyLoaded, bool bForceCircular)
 		{
+			if (Modules.Count == 1 && bIncludeDynamicallyLoaded && bForceCircular)
+			{
+				UEBuildModule First = Modules[0];
+				if (First.DependenciesForStableTopologicalSort != null)
+				{
+					return First.DependenciesForStableTopologicalSort.ToList();
+				}
+			}
+
 			List<UEBuildModule> ReferencedModules = new List<UEBuildModule>();
 			HashSet<UEBuildModule> IgnoreReferencedModules = new HashSet<UEBuildModule>();
 
diff --git a/Engine/Source/Programs/UnrealBuildTool/Configuration/UEBuildModule.cs b/Engine/Source/Programs/UnrealBuildTool/Configuration/UEBuildModule.cs
--- a/Engine/Source/Programs/UnrealBuildTool/Configuration/UEBuildModule.cs
+++ b/Engine/Source/Programs/UnrealBuildTool/Configuration/UEBuildModule.cs
@@ -215,6 +215,9 @@ namespace UnrealBuildTool
 		/// </summary>
 		public HashSet<FileItem> PublicAutoRTFMExternalMappingFiles = new();
 
+
+		public HashSet<UEBuildModule>? DependenciesForStableTopologicalSort { get; private set; }
+
 		/// <summary>
 		/// Constructor
 		/// </summary>
@@ -906,11 +909,7 @@ namespace UnrealBuildTool
 				if (!bIsBuildingAStaticLibrary && bIsModuleBinaryAStaticLibrary)
 				{
 					// Gather all dependencies and recursively call SetupPublicLinkEnvironmnet
-					List<UEBuildModule> AllDependencyModules = new List<UEBuildModule>();
-					AllDependencyModules.AddRange(PrivateDependencyModules!);
-					AllDependencyModules.AddRange(PublicDependencyModules!);
-
-					foreach (UEBuildModule DependencyModule in AllDependencyModules)
+					foreach (UEBuildModule DependencyModule in PrivateDependencyModules!.Concat(PublicDependencyModules!))
 					{
 						bool bIsExternalModule = (DependencyModule as UEBuildModuleExternal != null);
 						bool bIsInStaticLibrary = (DependencyModule.Binary != null && DependencyModule.Binary.Type == UEBuildBinaryType.StaticLibrary);
@@ -934,13 +933,16 @@ namespace UnrealBuildTool
 				DelayLoadDLLs.AddRange(PublicDelayLoadDLLs);
 
 				// merge in to the outgoing dictionary
-				foreach (KeyValuePair<string, List<UnrealArch>> Pair in DependenciesToSkipPerArchitecture)
+				lock (DependenciesToSkipPerArchitecture)
 				{
-					if (!DependenciesToSkip.ContainsKey(Pair.Key))
+					foreach (KeyValuePair<string, List<UnrealArch>> Pair in DependenciesToSkipPerArchitecture)
 					{
-						DependenciesToSkip[Pair.Key] = new HashSet<UnrealArch>();
+						if (!DependenciesToSkip.ContainsKey(Pair.Key))
+						{
+							DependenciesToSkip[Pair.Key] = new HashSet<UnrealArch>();
+						}
+						DependenciesToSkip[Pair.Key] = DependenciesToSkip[Pair.Key].Union(Pair.Value).ToHashSet();
 					}
-					DependenciesToSkip[Pair.Key] = DependenciesToSkip[Pair.Key].Union(Pair.Value).ToHashSet();
 				}
 			}
 		}
@@ -965,11 +967,7 @@ namespace UnrealBuildTool
 				LinkEnvironment.AdditionalFrameworks, LinkEnvironment.AdditionalBundleResources, LinkEnvironment.DelayLoadDLLs, BinaryDependencies, LinkEnvironment.DependenciesToSkipPerArchitecture, VisitedModules, ExeDir, LinkEnvironment.RootPaths, Logger);
 
 			// Also allow the module's public and private dependencies to modify the link environment.
-			List<UEBuildModule> AllDependencyModules = new List<UEBuildModule>();
-			AllDependencyModules.AddRange(PrivateDependencyModules!);
-			AllDependencyModules.AddRange(PublicDependencyModules!);
-
-			foreach (UEBuildModule DependencyModule in AllDependencyModules)
+			foreach (UEBuildModule DependencyModule in PrivateDependencyModules!.Concat(PublicDependencyModules!))
 			{
 				DependencyModule.SetupPublicLinkEnvironment(SourceBinary, LinkEnvironment.Libraries, LinkEnvironment.SystemLibraryPaths, LinkEnvironment.SystemLibraries, LinkEnvironment.RuntimeLibraryPaths, LinkEnvironment.Frameworks, LinkEnvironment.WeakFrameworks,
 					LinkEnvironment.AdditionalFrameworks, LinkEnvironment.AdditionalBundleResources, LinkEnvironment.DelayLoadDLLs, BinaryDependencies, LinkEnvironment.DependenciesToSkipPerArchitecture, VisitedModules, ExeDir, LinkEnvironment.RootPaths, Logger);
@@ -1067,7 +1065,10 @@ namespace UnrealBuildTool
 		/// <param name="bOnlyDirectDependencies">True to return only this module's direct dependencies</param>
 		public void GetAllDependencyModules(List<UEBuildModule> ReferencedModules, HashSet<UEBuildModule> IgnoreReferencedModules, bool bIncludeDynamicallyLoaded, bool bForceCircular, bool bOnlyDirectDependencies)
 		{
-			void AddReferencedModules(List<UEBuildModule> Modules)
+			var stack = new Stack<UEBuildModule>();
+			stack.Push(this);
+
+			void QueueReferencedModules(List<UEBuildModule> Modules)
 			{
 				foreach (UEBuildModule DependencyModule in Modules)
 				{
@@ -1078,8 +1079,7 @@ namespace UnrealBuildTool
 						{
 							if (!bOnlyDirectDependencies)
 							{
-								// Recurse into dependent modules first
-								DependencyModule.GetAllDependencyModules(ReferencedModules, IgnoreReferencedModules, bIncludeDynamicallyLoaded, bForceCircular, bOnlyDirectDependencies);
+								stack.Push(DependencyModule);
 							}
 
 							ReferencedModules.Add(DependencyModule);
@@ -1088,11 +1088,16 @@ namespace UnrealBuildTool
 				}
 			}
 
-			AddReferencedModules(PrivateDependencyModules!);
-			AddReferencedModules(PublicDependencyModules!);
-			if (bIncludeDynamicallyLoaded)
+			while (stack.Count != 0)
 			{
-				AddReferencedModules(DynamicallyLoadedModules!);
+				UEBuildModule module = stack.Pop();
+
+				QueueReferencedModules(module.PrivateDependencyModules!);
+				QueueReferencedModules(module.PublicDependencyModules!);
+				if (bIncludeDynamicallyLoaded)
+				{
+					QueueReferencedModules(module.DynamicallyLoadedModules!);
+				}
 			}
 		}
 
@@ -1590,6 +1595,19 @@ namespace UnrealBuildTool
 			Writer.WriteArrayEnd();
 		}
 
+		public static void CacheDependenciesForStableTopologicalSort(IEnumerable<UEBuildModule> Modules)
+		{
+			Parallel2.ForEach(Modules, module =>
+			{
+				if (module.DependenciesForStableTopologicalSort == null && module.PrivateDependencyModules != null)
+				{
+					HashSet<UEBuildModule> Deps = [];
+					module.GetAllDependencyModules([], Deps, true, true, false);
+					module.DependenciesForStableTopologicalSort = Deps;
+				}
+			});
+		}
+
 		/// <summary>
 		/// Returns a copy of Nodes sorted by dependency.  Independent or circularly-dependent nodes should
 		/// remain in their same relative order within the original Nodes sequence.
@@ -1597,82 +1615,156 @@ namespace UnrealBuildTool
 		/// <param name="NodeList">The list of nodes to sort.</param>
 		public static List<UEBuildModule> StableTopologicalSort(List<UEBuildModule> NodeList)
 		{
-			int NodeCount = NodeList.Count;
-
-			// For each Node in NodeList, populated with the full circular dependency list from
-			// Node.GetAllDependencyModules()
-			List<Task<HashSet<UEBuildModule>>> NodeDependencies = new List<Task<HashSet<UEBuildModule>>>(NodeCount);
-			// Used to populate an element of NodeDependencies
-			HashSet<UEBuildModule> FetchDependencies(int NodeIndex)
+			int nodeCount = NodeList.Count;
+			if (nodeCount <= 1)
 			{
-				HashSet<UEBuildModule> Dependencies = new HashSet<UEBuildModule>();
-				NodeList[NodeIndex].GetAllDependencyModules(new List<UEBuildModule>(), Dependencies, true, true, false);
-				return Dependencies;
+				return [.. NodeList];
 			}
 
-			// For each Node in NodeList, populated with the nodes with a lower index in NodeList that Node depends on
-			List<Task<HashSet<UEBuildModule>>> PrecedingDependents = new List<Task<HashSet<UEBuildModule>>>(NodeCount);
-			HashSet<UEBuildModule> ComputePrecedingDependents(int NodeIndex)
+			// Use default equality (same as the old code)
+			Dictionary<UEBuildModule, int> indexOf = new (nodeCount);
+			for (int i = 0; i < nodeCount; i++)
 			{
-				HashSet<UEBuildModule> Results = new HashSet<UEBuildModule>();
+				indexOf[NodeList[i]] = i;
+			}
 
-				UEBuildModule Node = NodeList[NodeIndex];
-				HashSet<UEBuildModule> Dependencies = NodeDependencies[NodeIndex].Result;
+			int words = (nodeCount + 63) >> 6;
+			long bytesNeeded = (long)nodeCount * words * sizeof(ulong);
+			var depBits = new ulong[(long)nodeCount * words];
 
-				for (int I = 0; I < NodeIndex; ++I)
+			for (int i = 0; i < nodeCount; i++)
+			{
+				int rowBase = i * words;
+				foreach (var d in NodeList[i].DependenciesForStableTopologicalSort!)
 				{
-					if (NodeDependencies[I].Result.Contains(Node) && !Dependencies.Contains(NodeList[I]))
+					if (indexOf.TryGetValue(d, out int j) && j != i)
 					{
-						Results.Add(NodeList[I]);
+						depBits[rowBase + (j >> 6)] |= 1UL << (j & 63);
 					}
 				}
-
-				return Results;
 			}
 
-			for (int I = 0; I < NodeCount; ++I)
+			List<int> outIdx = new (nodeCount);
+			for (int i = 0; i < nodeCount; i++)
 			{
-				int LocalI = I;
-				NodeDependencies.Add(Task.Run(() => FetchDependencies(LocalI)));
-				PrecedingDependents.Add(Task.Run(() => ComputePrecedingDependents(LocalI)));
+				int insertAt = outIdx.Count;
+
+				// Precompute these once per i
+				int baseI = i * words;
+				int wI = i >> 6;
+				ulong maskI = 1UL << (i & 63);
+
+				for (int pos = 0; pos < outIdx.Count; pos++)
+				{
+					int j = outIdx[pos];
+					int baseJ = j * words;
+
+					// earlierDependsOnNode? (does j depend on i)
+					if ((depBits[baseJ + wI] & maskI) != 0)
+					{
+						// Only if needed, check mutual dep (does i depend on j)
+						int wJ = j >> 6;
+						ulong maskJ = 1UL << (j & 63);
+
+						if ((depBits[baseI + wJ] & maskJ) == 0)
+						{
+							insertAt = pos;
+							break;
+						}
+					}
+				}
+				outIdx.Insert(insertAt, i);
 			}
 
-			List<UEBuildModule> Out = new List<UEBuildModule>(NodeCount);
-			// Write the ordered output			
-			for (int Index1 = 0; Index1 != NodeCount; ++Index1)
+			List<UEBuildModule> ordered = new (nodeCount);
+			foreach (int idx in outIdx)
+				ordered.Add(NodeList[idx]);
+
+			#if false
+			// THIS IS THE OLD CODE
 			{
-				UEBuildModule Node1 = NodeList[Index1];
-				HashSet<UEBuildModule> NodesThatDependOnNode1 = PrecedingDependents[Index1].Result;
-				Out.Add(Node1);
+				int NodeCount = NodeList.Count;
 
-				if (NodesThatDependOnNode1.Count == 0)
+				// For each Node in NodeList, populated with the full circular dependency list from
+				// Node.GetAllDependencyModules()
+				List<Task<HashSet<UEBuildModule>>> NodeDependencies = new List<Task<HashSet<UEBuildModule>>>(NodeCount);
+				// Used to populate an element of NodeDependencies
+				HashSet<UEBuildModule> FetchDependencies(int NodeIndex)
 				{
-					continue;
+					HashSet<UEBuildModule> Dependencies = new HashSet<UEBuildModule>();
+					NodeList[NodeIndex].GetAllDependencyModules(new List<UEBuildModule>(), Dependencies, true, true, false);
+					return Dependencies;
 				}
 
-				for (int Index2 = 0; Index2 != Index1; ++Index2)
+				// For each Node in NodeList, populated with the nodes with a lower index in NodeList that Node depends on
+				List<Task<HashSet<UEBuildModule>>> PrecedingDependents = new List<Task<HashSet<UEBuildModule>>>(NodeCount);
+				HashSet<UEBuildModule> ComputePrecedingDependents(int NodeIndex)
 				{
-					UEBuildModule Node2 = Out[Index2];
+					HashSet<UEBuildModule> Results = new HashSet<UEBuildModule>();
 
-					if (NodesThatDependOnNode1.Contains(Node2))
+					UEBuildModule Node = NodeList[NodeIndex];
+					HashSet<UEBuildModule> Dependencies = NodeDependencies[NodeIndex].Result;
+
+					for (int I = 0; I < NodeIndex; ++I)
 					{
-						// Rotate element at Index1 into position at Index2
-						for (int Index3 = Index1; Index3 != Index2;)
+						if (NodeDependencies[I].Result.Contains(Node) && !Dependencies.Contains(NodeList[I]))
 						{
-							--Index3;
-							Out[Index3 + 1] = Out[Index3];
+							Results.Add(NodeList[I]);
 						}
+					}
+
+					return Results;
+				}
 
-						Out[Index2] = Node1;
+				for (int I = 0; I < NodeCount; ++I)
+				{
+					int LocalI = I;
+					NodeDependencies.Add(Task.Run(() => FetchDependencies(LocalI)));
+					PrecedingDependents.Add(Task.Run(() => ComputePrecedingDependents(LocalI)));
+				}
 
-						// Break out of this loop, because this iteration must have covered all existing cases
-						// involving the node formerly at position Index1
-						break;
+				List<UEBuildModule> Out = new List<UEBuildModule>(NodeCount);
+				// Write the ordered output			
+				for (int Index1 = 0; Index1 != NodeCount; ++Index1)
+				{
+					UEBuildModule Node1 = NodeList[Index1];
+					HashSet<UEBuildModule> NodesThatDependOnNode1 = PrecedingDependents[Index1].Result;
+					Out.Add(Node1);
+
+					if (NodesThatDependOnNode1.Count == 0)
+					{
+						continue;
 					}
+
+					for (int Index2 = 0; Index2 != Index1; ++Index2)
+					{
+						UEBuildModule Node2 = Out[Index2];
+
+						if (NodesThatDependOnNode1.Contains(Node2))
+						{
+							// Rotate element at Index1 into position at Index2
+							for (int Index3 = Index1; Index3 != Index2;)
+							{
+								--Index3;
+								Out[Index3 + 1] = Out[Index3];
+							}
+
+							Out[Index2] = Node1;
+
+							// Break out of this loop, because this iteration must have covered all existing cases
+							// involving the node formerly at position Index1
+							break;
+						}
+					}
+				}
+				if (!ordered.SequenceEqual(Out))
+				{
+					throw new Exception("Sorting diff!!!");
 				}
 			}
+			#endif
 
-			return Out;
+			return ordered;
 		}
 	};
 }
diff --git a/Engine/Source/Programs/UnrealBuildTool/Configuration/UEBuildTarget.cs b/Engine/Source/Programs/UnrealBuildTool/Configuration/UEBuildTarget.cs
--- a/Engine/Source/Programs/UnrealBuildTool/Configuration/UEBuildTarget.cs
+++ b/Engine/Source/Programs/UnrealBuildTool/Configuration/UEBuildTarget.cs
@@ -2592,6 +2592,9 @@ namespace UnrealBuildTool
 			// Gather modules we might want to generate headers for
 			HashSet<UEBuildModuleCPP> ModulesToGenerateHeadersFor = GatherDependencyModules(OriginalBinaries.ToList());
 
+			// Cache module dependencies for StableTopologicalSort
+			UEBuildModule.CacheDependenciesForStableTopologicalSort(ModulesToGenerateHeadersFor);
+
 			// Prepare cached data for UHT header generation
 			using (GlobalTracer.Instance.BuildSpan("UHTExecution.SetupUObjectModules()").StartActive())
 			{
@@ -2701,6 +2704,9 @@ namespace UnrealBuildTool
 				Makefile.OutputItems.AddRange(Items);
 			}
 
+			// Cache module dependencies for StableTopologicalSort
+			UEBuildModule.CacheDependenciesForStableTopologicalSort(Modules.Values);
+
 			// Build the target's binaries.
 			DirectoryReference ExeDir = GetExecutableDir();
 			using (GlobalTracer.Instance.BuildSpan("UEBuildBinary.Build()").StartActive())
@@ -4061,6 +4067,9 @@ namespace UnrealBuildTool
 				Binary.CreateAllDependentModules((Name, RefChain) => FindOrCreateModuleByName(Name, RefChain, Logger), Logger);
 			}
 
+			// Cache dependencies since we can do this in parallel for all modules
+			UEBuildModule.CacheDependenciesForStableTopologicalSort(Binaries.SelectMany(B => B.Modules));
+
 			// Bind every referenced C++ module to a binary
 			for (int Idx = 0; Idx < Binaries.Count; Idx++)
 			{
