Add support for Additional wav, aiff, caff resources to allow users to deploy custom notification sounds for iOS applications.

diff --git a/Engine/Source/Programs/UnrealBuildTool/ProjectFiles/Xcode/XcodeProject.File.cs b/Engine/Source/Programs/UnrealBuildTool/ProjectFiles/Xcode/XcodeProject.File.cs
--- a/Engine/Source/Programs/UnrealBuildTool/ProjectFiles/Xcode/XcodeProject.File.cs
+++ b/Engine/Source/Programs/UnrealBuildTool/ProjectFiles/Xcode/XcodeProject.File.cs
@@ -128,6 +128,9 @@ namespace UnrealBuildTool.XcodeProjectXcconfig
 			".xcprivacy",
 			".plist",
 			".bin",
+			".wav",
+			".aiff",
+			".caff",
 			// no extensions will be folders, like "cookeddata" that we want to copy as a resource
 			"",
 		};
