[UBT]
* Added AppDomain extension class that has the method EnumerateSubclasses. method is searching in all assemblies that depend on the assembly storing the base class to see if there are any classes inhering from it

diff --git a/Engine/Source/Programs/UnrealBuildTool/System/AppDomainExtensions.cs b/Engine/Source/Programs/UnrealBuildTool/System/AppDomainExtensions.cs
new file mode 100644
--- /dev/null
+++ b/Engine/Source/Programs/UnrealBuildTool/System/AppDomainExtensions.cs
@@ -0,0 +1,74 @@
+// Copyright Epic Games, Inc. All Rights Reserved.
+
+using System;
+using System.Collections.Generic;
+using System.Reflection;
+
+namespace UnrealBuildBase
+{
+
+	/// <summary>
+	/// Extensions to AppDomain
+	/// </summary>
+	public static class AppDomainExtensions
+	{
+		/// <summary>
+		/// Using reflection to find all subclasses for a class.
+		/// Will check which assembly that class is in, then traverse that assmbly plus all assemblies depending on that assembly
+		/// to find the base classes. This is much more efficient than parsing all loaded assemblies to find subclasses
+		/// </summary>
+		/// <param name="domain">AppDomain to use</param>
+		/// <param name="baseClass">Base class to use</param>
+		public static IEnumerable<Type> EnumerateSubclasses(this AppDomain domain, Type baseClass)
+		{
+			var baseAsm = baseClass.Assembly;
+			var baseAsmName = baseAsm.GetName();
+
+			// Only search in assemblies that are either baseClass's assembly or assemblies depending on it
+			foreach (var asm in domain.GetAssemblies())
+			{
+				if (asm.IsDynamic)
+				{
+					continue;
+				}
+
+				if (!ReferenceEquals(asm, baseAsm))
+				{
+					try
+					{
+						AssemblyName[] refs = asm.GetReferencedAssemblies();
+
+						bool referencesBase = false;
+						for (int i = 0; i < refs.Length; i++)
+						{
+							if (AssemblyName.ReferenceMatchesDefinition(refs[i], baseAsmName))
+							{
+								referencesBase = true;
+								break;
+							}
+						}
+						if (!referencesBase)
+						{
+							continue;
+						}
+					}
+					catch
+					{
+						continue;
+					}
+				}
+
+				Type[] types = asm.SafeGetLoadedTypes();
+
+				for (int i = 0; i < types.Length; i++)
+				{
+					var t = types[i];
+					if (t != null && t.IsClass && !t.IsAbstract && t.IsSubclassOf(baseClass))
+					{
+						yield return t;
+					}
+				}
+			}
+		}
+	}
+}
\ No newline at end of file
