[UBT]
* Removed UHT feature UhtExporterOptions.CompileOutput. This is not supported anymore since we need to know all UHT outputs ahead of time (before UHT actually runs). Needs to be implemented in a different way

diff --git a/Engine/Source/Programs/Shared/EpicGames.UHT/Tables/UhtExporterTable.cs b/Engine/Source/Programs/Shared/EpicGames.UHT/Tables/UhtExporterTable.cs
--- a/Engine/Source/Programs/Shared/EpicGames.UHT/Tables/UhtExporterTable.cs
+++ b/Engine/Source/Programs/Shared/EpicGames.UHT/Tables/UhtExporterTable.cs
@@ -31,13 +31,11 @@ namespace EpicGames.UHT.Tables
 		/// The exporter should be run by default
 		/// </summary>
 		Default = 1 << 0,
-		
+
 		/// <summary>
-		/// Generated CPP files that match any filter in the CppFilters property of this UhtExporter will be included in the compilation process.
-		/// For this option, all custom generated files should be placed into the same folder as the default unreal generated files.
-		/// Make sure the CppFilters property is set up properly and does not conflict with other exporters (including the default).
-		/// Header files should be included as normal in code.
+		/// DEPRECATED. This does not work
 		/// </summary>
+		[Obsolete("This is not supported anymore. Contact epic if you were using this and really needs it", error: true)]
 		CompileOutput = 1 << 1,
 	}
 
diff --git a/Engine/Source/Programs/UnrealBuildTool/Configuration/UEBuildTarget.cs b/Engine/Source/Programs/UnrealBuildTool/Configuration/UEBuildTarget.cs
--- a/Engine/Source/Programs/UnrealBuildTool/Configuration/UEBuildTarget.cs
+++ b/Engine/Source/Programs/UnrealBuildTool/Configuration/UEBuildTarget.cs
@@ -2661,8 +2661,6 @@ namespace UnrealBuildTool
 
 #endif
 
-			GatherExtraGeneratedCPPFileTypes(Makefile, TargetName, ref GlobalCompileEnvironment);
-
 			if (Rules.bUseSharedPCHs)
 			{
 				// Find all the shared PCHs.
@@ -6486,29 +6484,5 @@ namespace UnrealBuildTool
 			}
 			return NewPathList;
 		}
-		
-		private static void GatherExtraGeneratedCPPFileTypes(TargetMakefile Makefile, string TargetName, ref CppCompileEnvironment InCompileEnvironment)
-		{
-			UhtStdFileManager fileManager = new UhtStdFileManager();
-			FileReference moduleInfoFileName = UHTExecution.GetUHTModuleInfoFileName(Makefile, TargetName);
-			if (fileManager.ReadSource(moduleInfoFileName.FullName, out UhtSourceFragment fragment))
-			{
-				UHTManifest? manifest = System.Text.Json.JsonSerializer.Deserialize<UHTManifest>(fragment.Data.ToString());
-
-				if (manifest != null && manifest.UhtPlugins.Count > 0)
-				{
-					UhtTables tables = new();
-					tables.AddPlugins(manifest.UhtPlugins);
-
-					foreach (UhtExporter exporter in tables.ExporterTable)
-					{
-						if (exporter.Options.HasAnyFlags(UhtExporterOptions.CompileOutput))
-						{
-							InCompileEnvironment.ExtraGeneratedCPPFileTypes.AddRange(exporter.CppFilters);
-						}
-					}
-				}
-			}
-		}
 	}
 }
