UBT: Allow passing multiple modules to compile as a list using the default "+" unreal separator so that "-Module=A -Module=B -Module=C" can be condensed to "-Module=A+B+C".

#rnx

diff --git a/Engine/Source/Programs/UnrealBuildTool/Configuration/TargetDescriptor.cs b/Engine/Source/Programs/UnrealBuildTool/Configuration/TargetDescriptor.cs
--- a/Engine/Source/Programs/UnrealBuildTool/Configuration/TargetDescriptor.cs
+++ b/Engine/Source/Programs/UnrealBuildTool/Configuration/TargetDescriptor.cs
@@ -56,7 +56,7 @@ namespace UnrealBuildTool
 		/// <summary>
 		/// Set of module names to compile.
 		/// </summary>
-		[CommandLine("-Module=")]
+		[CommandLine("-Module=", ListSeparator = '+')]
 		public HashSet<string> OnlyModuleNames = new HashSet<string>(StringComparer.OrdinalIgnoreCase);
 
 		/// <summary>
