[UBT]
* Optimized GenerateWarningCommandLineArgs by not using Linq

diff --git a/Engine/Source/Programs/UnrealBuildTool/Configuration/CppCompileWarnings.cs b/Engine/Source/Programs/UnrealBuildTool/Configuration/CppCompileWarnings.cs
--- a/Engine/Source/Programs/UnrealBuildTool/Configuration/CppCompileWarnings.cs
+++ b/Engine/Source/Programs/UnrealBuildTool/Configuration/CppCompileWarnings.cs
@@ -1515,15 +1515,29 @@ namespace UnrealBuildTool
 			foreach (var (mi, miAttr) in s_membersWithApplyWarningsAttribute.Value)
 			{
 				ApplyWarningsAttribute? applyWarningsAttribute = null;
-				IEnumerable<ApplyWarningsAttribute?> applyWarningsAttributes = miAttr.Where(x => x.CanApplyToContext(toolChainContext));
 
-				if (applyWarningsAttributes.Skip(1).Any())
+				bool isFirst = true;
+
+				// We only want to iterate this enumerator once since it does quite a lot of work. Also, this is called a lot so avoid Linq
+				foreach (ApplyWarningsAttribute? attr in miAttr)
 				{
-					_logger?.LogWarning("Ambiguous ApplyWarningsAttributes on the property ({PropertyName}), for the provided tool-chain ({ToolChainName}). Skipping this property.", mi.Name, toolChainContext._toolChainType.Name);
-					continue;
-				}
+					if (!attr.CanApplyToContext(toolChainContext))
+					{
+						continue;
+					}
 
-				applyWarningsAttribute = applyWarningsAttributes.FirstOrDefault();
+					if (isFirst)
+					{
+						applyWarningsAttribute = attr;
+						isFirst = false;
+					}
+					else
+					{
+						applyWarningsAttribute = null;
+						_logger?.LogWarning("Ambiguous ApplyWarningsAttributes on the property ({PropertyName}), for the provided tool-chain ({ToolChainName}). Skipping this property.", mi.Name, toolChainContext._toolChainType.Name);
+						break;
+					}
+				}
 
 				if (applyWarningsAttribute != null)
 				{
