[UBT]
* Changed ActionGraph test to use FirstOrDefault instead of Contains. Saves half a second

diff --git a/Engine/Source/Programs/UnrealBuildTool/System/ActionGraph.cs b/Engine/Source/Programs/UnrealBuildTool/System/ActionGraph.cs
--- a/Engine/Source/Programs/UnrealBuildTool/System/ActionGraph.cs
+++ b/Engine/Source/Programs/UnrealBuildTool/System/ActionGraph.cs
@@ -916,7 +916,7 @@ namespace UnrealBuildTool
 			{
 				foreach (LinkedAction PrerequisiteAction in RootAction.PrerequisiteActions)
 				{
-					if (PrerequisiteAction.bProducesImportLibrary && PrerequisiteAction.ProducedItems.Contains(PrerequisiteItem))
+					if (PrerequisiteAction.bProducesImportLibrary && PrerequisiteAction.ProducedItems.FirstOrDefault(PrerequisiteItem) != null)
 					{
 						return true;
 					}
