[UBT]
* Made ISPCToolChain GetISPCHostCompilerVersion threadsafe

diff --git a/Engine/Source/Programs/UnrealBuildTool/ToolChain/ISPCToolChain.cs b/Engine/Source/Programs/UnrealBuildTool/ToolChain/ISPCToolChain.cs
--- a/Engine/Source/Programs/UnrealBuildTool/ToolChain/ISPCToolChain.cs
+++ b/Engine/Source/Programs/UnrealBuildTool/ToolChain/ISPCToolChain.cs
@@ -4,6 +4,7 @@ using EpicGames.Core;
 using Microsoft.CodeAnalysis;
 using Microsoft.Extensions.Logging;
 using System;
+using System.Collections.Concurrent;
 using System.Collections.Generic;
 using System.IO;
 using System.Linq;
@@ -207,7 +208,7 @@ namespace UnrealBuildTool
 			return null;
 		}
 
-		static readonly Dictionary<string, string> s_ISPCCompilerVersions = [];
+		static readonly ConcurrentDictionary<string, Lazy<string>> s_ISPCCompilerVersions = [];
 
 		/// <summary>
 		/// Returns the version of the ISPC compiler for the specified platform. If GetISPCHostCompilerPath() doesn't return a valid path
@@ -218,20 +219,18 @@ namespace UnrealBuildTool
 		public virtual string GetISPCHostCompilerVersion(UnrealTargetPlatform platform)
 		{
 			string compilerPath = GetISPCHostCompilerPath(platform);
-			if (!s_ISPCCompilerVersions.ContainsKey(compilerPath))
+			return s_ISPCCompilerVersions.GetOrAdd(compilerPath, new Lazy<string>(() =>
 			{
 				if (File.Exists(compilerPath))
 				{
-					s_ISPCCompilerVersions[compilerPath] = RunToolAndCaptureOutput(new FileReference(compilerPath), "--version", "(.*)")!;
+					return RunToolAndCaptureOutput(new FileReference(compilerPath), "--version", "(.*)")!;
 				}
 				else
 				{
 					Logger.LogWarning("No ISPC compiler at {CompilerPath}", compilerPath);
-					s_ISPCCompilerVersions[compilerPath] = "-1";
+					return "-1";
 				}
-			}
-
-			return s_ISPCCompilerVersions[compilerPath];
+			})).Value;
 		}
 
 		/// <summary>
