[UBT]
* Changed UHT to use ParallelQueue and Parallel2.ForEach since they don't join started tasks and thus don't stall when there is lots of task pressure

diff --git a/Engine/Source/Programs/Shared/EpicGames.UHT/Utils/UhtSession.cs b/Engine/Source/Programs/Shared/EpicGames.UHT/Utils/UhtSession.cs
--- a/Engine/Source/Programs/Shared/EpicGames.UHT/Utils/UhtSession.cs
+++ b/Engine/Source/Programs/Shared/EpicGames.UHT/Utils/UhtSession.cs
@@ -17,6 +17,7 @@ using EpicGames.UHT.Tables;
 using EpicGames.UHT.Tokenizer;
 using EpicGames.UHT.Types;
 using Microsoft.Extensions.Logging;
+using UnrealBuildBase;
 
 namespace EpicGames.UHT.Utils
 {
@@ -515,7 +516,7 @@ namespace EpicGames.UHT.Utils
 				// Perform the renames
 				if (Session.GoWide)
 				{
-					Parallel.ForEach(saves, (UhtExportFactory.Output output) =>
+					Parallel2.ForEach(saves, (UhtExportFactory.Output output) =>
 					{
 						RenameSource(output);
 					});
@@ -535,7 +536,7 @@ namespace EpicGames.UHT.Utils
 			{
 				if (Session.GoWide)
 				{
-					Parallel.ForEach(outputsByDirectory, (KeyValuePair<string, HashSet<string>> kvp) =>
+					Parallel2.ForEach(outputsByDirectory, (KeyValuePair<string, HashSet<string>> kvp) =>
 					{
 						CullOutputDirectory(kvp.Key, kvp.Value);
 					});
@@ -1323,7 +1324,7 @@ namespace EpicGames.UHT.Utils
 		{
 			if (GoWide && allowGoWide)
 			{
-				Parallel.ForEach(_headerFiles, headerFile =>
+				Parallel2.ForEach(_headerFiles, headerFile =>
 				{
 					TryHeader(action, headerFile);
 				});
diff --git a/Engine/Source/Programs/UnrealBuildTool/System/UHTExecution.cs b/Engine/Source/Programs/UnrealBuildTool/System/UHTExecution.cs
--- a/Engine/Source/Programs/UnrealBuildTool/System/UHTExecution.cs
+++ b/Engine/Source/Programs/UnrealBuildTool/System/UHTExecution.cs
@@ -486,8 +486,9 @@ namespace UnrealBuildTool
 
 			// Create the info for each module in parallel
 			UHTModuleInfo[] ModuleInfoArray = new UHTModuleInfo[ModulesSortedByType.Count];
-			using (ThreadPoolWorkQueue Queue = new())
+
 			{
+				ParallelQueue Queue = new();
 				IReadOnlySet<string> ExcludedFolders = UEBuildPlatform.GetBuildPlatform(Platform).GetExcludedFolderNames();
 				for (int Idx = 0; Idx < ModulesSortedByType.Count; Idx++)
 				{
@@ -534,6 +535,7 @@ namespace UnrealBuildTool
 
 					Queue.Enqueue(() => SetupUObjectModule(Info, ExcludedFolders, MetadataCache, Queue));
 				}
+				Queue.Drain();
 			}
 
 			// Gather all the verse dependencies
@@ -711,7 +713,7 @@ namespace UnrealBuildTool
 			}
 		}
 
-		static void SetupUObjectModule(UHTModuleInfo ModuleInfo, IReadOnlySet<string> ExcludedFolders, SourceFileMetadataCache MetadataCache, ThreadPoolWorkQueue Queue)
+		static void SetupUObjectModule(UHTModuleInfo ModuleInfo, IReadOnlySet<string> ExcludedFolders, SourceFileMetadataCache MetadataCache, ParallelQueue Queue)
 		{
 			foreach (DirectoryReference ModuleDirectory in ModuleInfo.ModuleDirectories)
 			{
