[UBT]
* Changed UEBuildBinary.DependentLinkLibraries to be a Lazy in order to prevent same code from running lots of times

diff --git a/Engine/Source/Programs/UnrealBuildTool/Configuration/UEBuildBinary.cs b/Engine/Source/Programs/UnrealBuildTool/Configuration/UEBuildBinary.cs
--- a/Engine/Source/Programs/UnrealBuildTool/Configuration/UEBuildBinary.cs
+++ b/Engine/Source/Programs/UnrealBuildTool/Configuration/UEBuildBinary.cs
@@ -121,7 +121,7 @@ namespace UnrealBuildTool
 		/// <summary>
 		/// Cached list of dependent link libraries.
 		/// </summary>
-		private List<FileReference>? DependentLinkLibraries;
+		private Lazy<List<FileReference>> DependentLinkLibraries;
 
 		/// <summary>
 		/// Expanded list of runtime dependencies generated by calling PrepareRuntimeDependencies.
@@ -173,6 +173,26 @@ namespace UnrealBuildTool
 			PrimaryModule.bCreatePerModuleFile = true;
 
 			Modules.Add(PrimaryModule);
+
+			// Cache the list of libraries in the dependent link environment between calls. We typically run this code path many times for each module.
+			DependentLinkLibraries = new(() =>
+			{
+				List<FileReference> libs = new ();
+				foreach (FileReference OutputFilePath in OutputFilePaths)
+				{
+					FileReference LibraryFileName;
+					if (Type == UEBuildBinaryType.StaticLibrary || !PrimaryModule.Rules.Target.Platform.IsInGroup(UnrealPlatformGroup.Microsoft))
+					{
+						LibraryFileName = OutputFilePath;
+					}
+					else
+					{
+						LibraryFileName = FileReference.Combine(IntermediateDirectory, OutputFilePath.GetFileNameWithoutExtension() + ".lib");
+					}
+					libs.Add(LibraryFileName);
+				}
+				return libs;
+			});
 		}
 
 		/// <summary>
@@ -376,25 +396,7 @@ namespace UnrealBuildTool
 		/// <param name="DependentLinkEnvironment">The link environment of the dependency</param>
 		public void SetupDependentLinkEnvironment(LinkEnvironment DependentLinkEnvironment)
 		{
-			// Cache the list of libraries in the dependent link environment between calls. We typically run this code path many times for each module.
-			if (DependentLinkLibraries == null)
-			{
-				DependentLinkLibraries = new List<FileReference>();
-				foreach (FileReference OutputFilePath in OutputFilePaths)
-				{
-					FileReference LibraryFileName;
-					if (Type == UEBuildBinaryType.StaticLibrary || !DependentLinkEnvironment.Platform.IsInGroup(UnrealPlatformGroup.Microsoft))
-					{
-						LibraryFileName = OutputFilePath;
-					}
-					else
-					{
-						LibraryFileName = FileReference.Combine(IntermediateDirectory, OutputFilePath.GetFileNameWithoutExtension() + ".lib");
-					}
-					DependentLinkLibraries.Add(LibraryFileName);
-				}
-			}
-			DependentLinkEnvironment.Libraries.AddRange(DependentLinkLibraries);
+			DependentLinkEnvironment.Libraries.AddRange(DependentLinkLibraries.Value);
 		}
 
 		/// <summary>
