[UBT]
* Create UEBuildBinary class in parallel
* Cached dependencies for StableTopologicalSort
* Optimized code figuring out which includes that are used by shared pch. Added use of bloom filter to minimize I/O
* Changed so makefile save is running in parallel in task that can overlap entire build
* Changed so CppDependencies load happens in parallel in task that can run in parallel until it is being used
* Changed so SourceFileMetadataCache is loaded in parallel with other things
* Changed so metadata action creation is running in parallel with other work
* Changed so CleanStaleModules is traversing files to delete in parallel

diff --git a/Engine/Source/Programs/UnrealBuildTool/Configuration/UEBuildBinary.cs b/Engine/Source/Programs/UnrealBuildTool/Configuration/UEBuildBinary.cs
--- a/Engine/Source/Programs/UnrealBuildTool/Configuration/UEBuildBinary.cs
+++ b/Engine/Source/Programs/UnrealBuildTool/Configuration/UEBuildBinary.cs
@@ -414,14 +414,14 @@ namespace UnrealBuildTool
 		/// <param name="bIncludeDynamicallyLoaded">True if dynamically loaded modules (and all of their dependent modules) should be included.</param>
 		/// <param name="bForceCircular">True if circular dependencies should be process</param>
 		/// <returns>List of all referenced modules</returns>
-		public List<UEBuildModule> GetAllDependencyModules(bool bIncludeDynamicallyLoaded, bool bForceCircular)
+		public IEnumerable<UEBuildModule> GetAllDependencyModules(bool bIncludeDynamicallyLoaded, bool bForceCircular)
 		{
 			if (Modules.Count == 1 && bIncludeDynamicallyLoaded && bForceCircular)
 			{
 				UEBuildModule First = Modules[0];
 				if (First.DependenciesForStableTopologicalSort != null)
 				{
-					return First.DependenciesForStableTopologicalSort.ToList();
+					return First.DependenciesForStableTopologicalSort;
 				}
 			}
 
diff --git a/Engine/Source/Programs/UnrealBuildTool/Configuration/UEBuildModule.cs b/Engine/Source/Programs/UnrealBuildTool/Configuration/UEBuildModule.cs
--- a/Engine/Source/Programs/UnrealBuildTool/Configuration/UEBuildModule.cs
+++ b/Engine/Source/Programs/UnrealBuildTool/Configuration/UEBuildModule.cs
@@ -1105,12 +1105,11 @@ namespace UnrealBuildTool
 		/// Gathers all the module dependencies a PCH would have
 		/// </summary>
 		/// <param name="bIncludePrivateModules">Whether to include private modules.</param>
-		/// <param name="bForceCircular">True if circular dependencies should be processed</param>
-		public HashSet<UEBuildModule> GetAllDependencyModulesForPCH(bool bIncludePrivateModules, bool bForceCircular)
+		public HashSet<UEBuildModule> GetAllDependencyModulesForPCH(bool bIncludePrivateModules)
 		{
 			HashSet<UEBuildModule> ReferencedModules = new HashSet<UEBuildModule>();
 			HashSet<UEBuildModule> IgnoreReferencedModules = new HashSet<UEBuildModule> { this };
-			InternalGetAllDependencyModulesForPCH(ReferencedModules, IgnoreReferencedModules, bIncludePrivateModules, bForceCircular);
+			InternalGetAllDependencyModulesForPCH(ReferencedModules, IgnoreReferencedModules, bIncludePrivateModules);
 			return ReferencedModules;
 		}
 
@@ -1120,8 +1119,7 @@ namespace UnrealBuildTool
 		/// <param name="ReferencedModules">Hash of all referenced modules with their addition index.</param>
 		/// <param name="IgnoreReferencedModules">Hashset used to ignore modules which are already added to the list</param>
 		/// <param name="bIncludePrivateModules">Whether to include private modules.</param>
-		/// <param name="bForceCircular">True if circular dependencies should be processed</param>
-		private void InternalGetAllDependencyModulesForPCH(HashSet<UEBuildModule> ReferencedModules, HashSet<UEBuildModule> IgnoreReferencedModules, bool bIncludePrivateModules, bool bForceCircular)
+		private void InternalGetAllDependencyModulesForPCH(HashSet<UEBuildModule> ReferencedModules, HashSet<UEBuildModule> IgnoreReferencedModules, bool bIncludePrivateModules)
 		{
 			List<UEBuildModule> AllDependencyModules = new List<UEBuildModule>(
 				((bIncludePrivateModules && PrivateDependencyModules != null) ? PrivateDependencyModules.Count : 0) +
@@ -1144,12 +1142,12 @@ namespace UnrealBuildTool
 			foreach (UEBuildModule DependencyModule in AllDependencyModules.Distinct())
 			{
 				// Don't follow circular back-references!
-				if (bForceCircular || !HasCircularDependencyOn(DependencyModule))
+				if (!HasCircularDependencyOn(DependencyModule))
 				{
 					if (IgnoreReferencedModules.Add(DependencyModule))
 					{
 						// Recurse into dependent modules first
-						DependencyModule.InternalGetAllDependencyModulesForPCH(ReferencedModules, IgnoreReferencedModules, false, bForceCircular);
+						DependencyModule.InternalGetAllDependencyModulesForPCH(ReferencedModules, IgnoreReferencedModules, false);
 
 						ReferencedModules.Add(DependencyModule);
 					}
diff --git a/Engine/Source/Programs/UnrealBuildTool/Configuration/UEBuildModuleCPP.cs b/Engine/Source/Programs/UnrealBuildTool/Configuration/UEBuildModuleCPP.cs
--- a/Engine/Source/Programs/UnrealBuildTool/Configuration/UEBuildModuleCPP.cs
+++ b/Engine/Source/Programs/UnrealBuildTool/Configuration/UEBuildModuleCPP.cs
@@ -1258,144 +1258,139 @@ namespace UnrealBuildTool
 			return HeaderFileItems;
 		}
 
-		// Cache of files found while searching the includes
-		static ConcurrentDictionary<DirectoryReference, HashSet<FileReference>?> KnownIncludeFilesDict = new();
-		static HashSet<FileReference>? GetIncludeFiles(DirectoryReference Directory)
+		/// <summary>
+		/// Gets the list of module dependencies that this module declares and then tries to optimize the list to what is actually used by the header.
+		/// </summary>
+		/// <param name="CompileEnvironment">Compile environment for this PCH</param>
+		/// <param name="PCHHeaderFile">PCH header file</param>
+		/// <param name="Logger">Logger for output</param>
+		/// <returns>Returns an optimized list of modules this module's shared PCH uses</returns>
+		private HashSet<UEBuildModule> GetSharedPCHModuleDependencies(CppCompileEnvironment CompileEnvironment, FileItem PCHHeaderFile, ILogger Logger)
 		{
-			if (Directory == null)
+			Dictionary<UEBuildModule, bool> ModuleToIncludePathsOnlyFlag = new Dictionary<UEBuildModule, bool>();
+			FindModulesInPublicCompileEnvironment(ModuleToIncludePathsOnlyFlag);
+
+			HashSet<UEBuildModule> AllModuleDeps = ModuleToIncludePathsOnlyFlag.Keys.ToHashSet();
+			HashSet<UEBuildModule> ModuleDeps = new HashSet<UEBuildModule>(AllModuleDeps) { this };
+
+			HashSet<string> VisitedIncludes = new();
+			HashSet<FileItem> FoundIncludeFileItems = new();
+			HashSet<string> IncludesNotFound = new();
+
+
+			// Go through user directories and create bloom filter for files and dirs so we can early out quickly
+			List<(DirectoryItem, BloomFilter, BloomFilter)> UserIncludePaths = new(CompileEnvironment.UserIncludePaths.Count);
+			foreach (DirectoryReference UserInclude in CompileEnvironment.UserIncludePaths)
 			{
-				return null;
+				DirectoryItem dirItem = DirectoryItem.GetItemByDirectoryReference(UserInclude);
+				BloomFilter dirFilter = new();
+				BloomFilter fileFilter = new();
+				foreach (DirectoryItem dir in dirItem.EnumerateDirectories())
+					dirFilter.Add(dir.Name);
+				foreach (FileItem file in dirItem.EnumerateFiles())
+					fileFilter.Add(file.Name);
+				UserIncludePaths.Add((dirItem, dirFilter, fileFilter));
 			}
-			return KnownIncludeFilesDict.GetOrAdd(Directory, (_) =>
+
+			Stack<(FileItem, List<string>)> stack = new ();
+
+			stack.Push((PCHHeaderFile, CompileEnvironment.MetadataCache.GetHeaderIncludes(PCHHeaderFile)));
+
+			while (stack.Count != 0)
 			{
-				if (DirectoryLookupCache.DirectoryExists(Directory))
+				(FileItem FileToSearch, List<string> HeaderIncludes) = stack.Pop();
+				foreach (string HeaderInclude in HeaderIncludes)
 				{
-					try
+					// Skip .generated.h. They should always be in the same module as the header including them.. so no need to track them
+					if (HeaderInclude.Contains(".generated.h"))
 					{
-						return DirectoryLookupCache.EnumerateFiles(Directory).ToHashSet();
+						continue;
 					}
-					catch (Exception)
+
+					string TransformedHeaderInclude = HeaderInclude;
+					if (TransformedHeaderInclude.Contains("COMPILED_PLATFORM_HEADER("))
 					{
-						return null;
+						string PlatformName = UEBuildPlatform.GetBuildPlatform(CompileEnvironment.Platform).GetPlatformName();
+						TransformedHeaderInclude = TransformedHeaderInclude.Replace("COMPILED_PLATFORM_HEADER(", PlatformName + "/" + PlatformName).TrimEnd(')');
 					}
-				}
-				else
-				{
-					return null;
-				}
-			});
-		}
 
-		private void FindIncludedHeaders(HashSet<string> VisitedIncludes, HashSet<FileItem> FoundIncludeFileItems, HashSet<string> IncludesNotFound, CppCompileEnvironment CompileEnvironment, FileItem FileToSearch, ILogger Logger)
-		{
-			foreach (string HeaderInclude in CompileEnvironment.MetadataCache.GetHeaderIncludes(FileToSearch))
-			{
-				// Skip .generated.h. They should always be in the same module as the header including them.. so no need to track them
-				if (HeaderInclude.Contains(".generated.h"))
-				{
-					continue;
-				}
-
-				string TranformedHeaderInclude = HeaderInclude;
-				if (TranformedHeaderInclude.Contains("COMPILED_PLATFORM_HEADER("))
-				{
-					string PlatformName = UEBuildPlatform.GetBuildPlatform(CompileEnvironment.Platform).GetPlatformName();
-					TranformedHeaderInclude = TranformedHeaderInclude.Replace("COMPILED_PLATFORM_HEADER(", PlatformName + "/" + PlatformName).TrimEnd(')');
-				}
-
-				if (VisitedIncludes.Add(TranformedHeaderInclude))
-				{
-					Func<DirectoryReference, FileItem?> SearchForFileItem = (DirectoryReference dir) =>
+					if (VisitedIncludes.Add(TransformedHeaderInclude))
 					{
-						FileReference FileRef = FileReference.Combine(dir, TranformedHeaderInclude);
-						HashSet<FileReference>? Files = GetIncludeFiles(FileRef.Directory);
-						if (Files != null && Files.Contains(FileRef))
+						FileItem? SearchForFileItem(DirectoryItem dir)
 						{
-							return FileItem.GetItemByFileReference(FileRef);
+							FileReference FileRef = FileReference.Combine(dir.Location, TransformedHeaderInclude);
+							DirectoryItem dirItem = DirectoryItem.GetItemByDirectoryReference(FileRef.Directory);
+							dirItem.TryGetFile(FileRef.GetFileName(), out FileItem? fileItem);
+							return fileItem;
 						}
-						return null;
-					};
+						;
 
-					// check to see if the file is a relative path first
-					FileItem? IncludeFileItem = SearchForFileItem(FileToSearch.Directory.Location);
+						// check to see if the file is a relative path first
+						FileItem? IncludeFileItem = SearchForFileItem(FileToSearch.Directory);
 
-					// search through the include paths if the file isn't relative
-					if (IncludeFileItem == null)
-					{
-						foreach (DirectoryReference IncludePath in CompileEnvironment.UserIncludePaths)
+						// search through the include paths if the file isn't relative
+						if (IncludeFileItem == null)
 						{
-							IncludeFileItem = SearchForFileItem(IncludePath);
-							if (IncludeFileItem != null)
+							foreach ((DirectoryItem IncludePath, BloomFilter dirFilter, BloomFilter fileFilter) in UserIncludePaths)
 							{
-								break;
+								ReadOnlySpan<char> RelativePath = TransformedHeaderInclude.AsSpan();
+								int firstSlash = RelativePath.IndexOf('/');
+								if (firstSlash != -1)
+								{
+									if (!dirFilter.MightContain(RelativePath.Slice(0, firstSlash)))
+									{
+										continue;
+									}
+								}
+								else
+								{
+									if (!fileFilter.MightContain(RelativePath))
+									{
+										continue;
+									}
+								}
+
+								IncludeFileItem = SearchForFileItem(IncludePath);
+								if (IncludeFileItem != null)
+								{
+									break;
+								}
 							}
 						}
-					}
 
-					if (IncludeFileItem != null)
-					{
-						FoundIncludeFileItems.Add(IncludeFileItem);
-						// Logger.LogDebug("{0} SharedPCH - '{1}' included '{2}'.", Name, FileToSearch.Location, IncludeFileItem.Location);
-
-						FindIncludedHeaders(VisitedIncludes, FoundIncludeFileItems, IncludesNotFound, CompileEnvironment, IncludeFileItem, Logger);
-					}
-					else
-					{
-						if (!TranformedHeaderInclude.Contains('.'))
+						if (IncludeFileItem != null)
 						{
-#if DEBUG
-							Logger.LogDebug("{0} SharedPCH - Skipping '{1}' found in '{2}' because it doesn't appear to be a module header.", Name, TranformedHeaderInclude, FileToSearch.Location);
-#endif
+							FoundIncludeFileItems.Add(IncludeFileItem);
+							// Logger.LogDebug("{0} SharedPCH - '{1}' included '{2}'.", Name, FileToSearch.Location, IncludeFileItem.Location);
+							stack.Push((IncludeFileItem, CompileEnvironment.MetadataCache.GetHeaderIncludes(IncludeFileItem)));
 						}
-						else if (TranformedHeaderInclude.EndsWith(".generated.h"))
+						else
 						{
+							if (!TransformedHeaderInclude.Contains('.'))
+							{
 #if DEBUG
-							Logger.LogDebug("{0} SharedPCH - Skipping '{1}' found in '{2}' because it appears to be a generated UHT header.", Name, TranformedHeaderInclude, FileToSearch.Location);
+							Logger.LogDebug("{0} SharedPCH - Skipping '{1}' found in '{2}' because it doesn't appear to be a module header.", Name, TransformedHeaderInclude, FileToSearch.Location);
 #endif
-						}
-						else
-						{
+							}
+							else if (TransformedHeaderInclude.EndsWith(".generated.h"))
+							{
+#if DEBUG
+							Logger.LogDebug("{0} SharedPCH - Skipping '{1}' found in '{2}' because it appears to be a generated UHT header.", Name, TransformedHeaderInclude, FileToSearch.Location);
+#endif
+							}
+							else
+							{
 #if DEBUG
-							Logger.LogDebug("{0} SharedPCH - Could not find include directory for '{1}' found in '{2}'.", Name, TranformedHeaderInclude, FileToSearch.Location);
+							Logger.LogDebug("{0} SharedPCH - Could not find include directory for '{1}' found in '{2}'.", Name, TransformedHeaderInclude, FileToSearch.Location);
 #endif
-							IncludesNotFound.Add(TranformedHeaderInclude);
+								IncludesNotFound.Add(TransformedHeaderInclude);
+							}
 						}
 					}
 				}
 			}
-		}
-
-		/// <summary>
-		/// Gets the list of module dependencies that this module declares and then tries to optimize the list to what is actually used by the header.
-		/// </summary>
-		/// <param name="CompileEnvironment">Compile environment for this PCH</param>
-		/// <param name="PCHHeaderFile">PCH header file</param>
-		/// <param name="Logger">Logger for output</param>
-		/// <returns>Returns an optimized list of modules this module's shared PCH uses</returns>
-		private HashSet<UEBuildModule> GetSharedPCHModuleDependencies(CppCompileEnvironment CompileEnvironment, FileItem PCHHeaderFile, ILogger Logger)
-		{
-			Dictionary<UEBuildModule, bool> ModuleToIncludePathsOnlyFlag = new Dictionary<UEBuildModule, bool>();
-			FindModulesInPublicCompileEnvironment(ModuleToIncludePathsOnlyFlag);
-
-			HashSet<UEBuildModule> AllModuleDeps = ModuleToIncludePathsOnlyFlag.Keys.ToHashSet();
-			HashSet<UEBuildModule> ModuleDeps = new HashSet<UEBuildModule>(AllModuleDeps) { this };
 
-			HashSet<string> VisitedIncludes = new();
-			HashSet<FileItem> FoundIncludeFileItems = new();
-			HashSet<string> IncludesNotFound = new();
 
-			// pre-populate all the known file locations for this compile environment
-			foreach (DirectoryReference dir in CompileEnvironment.UserIncludePaths)
-			{
-				GetIncludeFiles(dir);
-			}
-
-			FindIncludedHeaders(VisitedIncludes, FoundIncludeFileItems, IncludesNotFound, CompileEnvironment, PCHHeaderFile, Logger);
-			if (IncludesNotFound.Any())
-			{
-				// We are assuming that the includes not found aren't important. If this is found to be not true then we should return all the module dependencies
-				//return AllModuleDeps;
-			}
 
 			bool FoundAllModules = true;
 			HashSet<UEBuildModule> OptModules = new();
@@ -1536,116 +1531,119 @@ namespace UnrealBuildTool
 		/// <returns>Instance of a PCH</returns>
 		public PrecompiledHeaderInstance FindOrCreateSharedPCH(UEToolChain? ToolChain, PrecompiledHeaderTemplate Template, CppCompileEnvironment ModuleCompileEnvironment, IActionGraphBuilder Graph)
 		{
-			PrecompiledHeaderInstance? Instance = Template.Instances.Find(x => IsCompatibleForSharedPCH(x.CompileEnvironment, ModuleCompileEnvironment));
-			if (Instance == null)
+			lock (Template.Instances)
 			{
-				List<string> Definitions = Template.BaseCompileEnvironment.Definitions;
-
-				// Modify definitions if we need to create a new shared pch for the include order
-				if (ModuleCompileEnvironment.IncludeOrderVersion != Template.BaseCompileEnvironment.IncludeOrderVersion)
+				PrecompiledHeaderInstance? Instance = Template.Instances.Find(x => IsCompatibleForSharedPCH(x.CompileEnvironment, ModuleCompileEnvironment));
+				if (Instance == null)
 				{
-					Definitions = new List<string>(Definitions);
-					foreach (string OldDefine in EngineIncludeOrderHelper.GetDeprecationDefines(Template.BaseCompileEnvironment.IncludeOrderVersion))
+					List<string> Definitions = Template.BaseCompileEnvironment.Definitions;
+
+					// Modify definitions if we need to create a new shared pch for the include order
+					if (ModuleCompileEnvironment.IncludeOrderVersion != Template.BaseCompileEnvironment.IncludeOrderVersion)
 					{
-						Definitions.Remove(OldDefine);
+						Definitions = new List<string>(Definitions);
+						foreach (string OldDefine in EngineIncludeOrderHelper.GetDeprecationDefines(Template.BaseCompileEnvironment.IncludeOrderVersion))
+						{
+							Definitions.Remove(OldDefine);
+						}
+						Definitions.AddRange(EngineIncludeOrderHelper.GetDeprecationDefines(ModuleCompileEnvironment.IncludeOrderVersion));
 					}
-					Definitions.AddRange(EngineIncludeOrderHelper.GetDeprecationDefines(ModuleCompileEnvironment.IncludeOrderVersion));
-				}
 
-				// Modify definitions if we need to create a new shared pch for project modules
-				if (ModuleCompileEnvironment.bTreatAsEngineModule != Template.BaseCompileEnvironment.bTreatAsEngineModule)
-				{
-					Definitions = new List<string>(Definitions);
-					Definitions.RemoveAll(x => x.Contains("UE_IS_ENGINE_MODULE", StringComparison.Ordinal));
-					Definitions.RemoveAll(x => x.Contains("UE_DEPRECATED_FORGAME", StringComparison.Ordinal));
-					Definitions.RemoveAll(x => x.Contains("UE_DEPRECATED_FORENGINE", StringComparison.Ordinal));
-
-					Definitions.Add($"UE_IS_ENGINE_MODULE={(Rules.bTreatAsEngineModule || Rules.Target.bWarnAboutMonolithicHeadersIncluded ? "1" : "0")}");
-					Definitions.Add($"UE_DEPRECATED_FORGAME={(Rules.bTreatAsEngineModule ? "PREPROCESSOR_NOTHING_FUNCTION" : "UE_DEPRECATED")}");
-					Definitions.Add($"UE_DEPRECATED_FORENGINE={(Rules.bTreatAsEngineModule || Rules.Target.bDisableEngineDeprecations ? "PREPROCESSOR_NOTHING_FUNCTION" : "UE_DEPRECATED")}");
-				}
+					// Modify definitions if we need to create a new shared pch for project modules
+					if (ModuleCompileEnvironment.bTreatAsEngineModule != Template.BaseCompileEnvironment.bTreatAsEngineModule)
+					{
+						Definitions = new List<string>(Definitions);
+						Definitions.RemoveAll(x => x.Contains("UE_IS_ENGINE_MODULE", StringComparison.Ordinal));
+						Definitions.RemoveAll(x => x.Contains("UE_DEPRECATED_FORGAME", StringComparison.Ordinal));
+						Definitions.RemoveAll(x => x.Contains("UE_DEPRECATED_FORENGINE", StringComparison.Ordinal));
+
+						Definitions.Add($"UE_IS_ENGINE_MODULE={(Rules.bTreatAsEngineModule || Rules.Target.bWarnAboutMonolithicHeadersIncluded ? "1" : "0")}");
+						Definitions.Add($"UE_DEPRECATED_FORGAME={(Rules.bTreatAsEngineModule ? "PREPROCESSOR_NOTHING_FUNCTION" : "UE_DEPRECATED")}");
+						Definitions.Add($"UE_DEPRECATED_FORENGINE={(Rules.bTreatAsEngineModule || Rules.Target.bDisableEngineDeprecations ? "PREPROCESSOR_NOTHING_FUNCTION" : "UE_DEPRECATED")}");
+					}
 
-				// Modify definitions if we need to create a new shared pch for validating format strings
-				if (ModuleCompileEnvironment.bValidateFormatStrings != Template.BaseCompileEnvironment.bValidateFormatStrings)
-				{
-					Definitions = new List<string>(Definitions);
-					Definitions.RemoveAll(x => x.Contains("UE_VALIDATE_FORMAT_STRINGS", StringComparison.Ordinal));
-					Definitions.Add($"UE_VALIDATE_FORMAT_STRINGS={(ModuleCompileEnvironment.bValidateFormatStrings ? "1" : "0")}");
-				}
+					// Modify definitions if we need to create a new shared pch for validating format strings
+					if (ModuleCompileEnvironment.bValidateFormatStrings != Template.BaseCompileEnvironment.bValidateFormatStrings)
+					{
+						Definitions = new List<string>(Definitions);
+						Definitions.RemoveAll(x => x.Contains("UE_VALIDATE_FORMAT_STRINGS", StringComparison.Ordinal));
+						Definitions.Add($"UE_VALIDATE_FORMAT_STRINGS={(ModuleCompileEnvironment.bValidateFormatStrings ? "1" : "0")}");
+					}
 
-				// Modify definitions if we need to create a new shared pch for validating internal api usage
-				if (ModuleCompileEnvironment.bValidateInternalApi != Template.BaseCompileEnvironment.bValidateInternalApi)
-				{
-					Definitions = new List<string>(Definitions);
-					Definitions.RemoveAll(x => x.Contains("UE_VALIDATE_INTERNAL_API", StringComparison.Ordinal));
-					Definitions.Add($"UE_VALIDATE_INTERNAL_API={(ModuleCompileEnvironment.bValidateInternalApi ? "1" : "0")}");
-				}
+					// Modify definitions if we need to create a new shared pch for validating internal api usage
+					if (ModuleCompileEnvironment.bValidateInternalApi != Template.BaseCompileEnvironment.bValidateInternalApi)
+					{
+						Definitions = new List<string>(Definitions);
+						Definitions.RemoveAll(x => x.Contains("UE_VALIDATE_INTERNAL_API", StringComparison.Ordinal));
+						Definitions.Add($"UE_VALIDATE_INTERNAL_API={(ModuleCompileEnvironment.bValidateInternalApi ? "1" : "0")}");
+					}
 
-				// Modify definitions if we need to create a new shared pch for validating experimental api usage
-				if (ModuleCompileEnvironment.bValidateExperimentalApi != Template.BaseCompileEnvironment.bValidateExperimentalApi)
-				{
-					Definitions = new List<string>(Definitions);
-					Definitions.RemoveAll(x => x.Contains("UE_VALIDATE_EXPERIMENTAL_API", StringComparison.Ordinal));
-					Definitions.Add($"UE_VALIDATE_EXPERIMENTAL_API={(ModuleCompileEnvironment.bValidateExperimentalApi ? "1" : "0")}");
-				}
+					// Modify definitions if we need to create a new shared pch for validating experimental api usage
+					if (ModuleCompileEnvironment.bValidateExperimentalApi != Template.BaseCompileEnvironment.bValidateExperimentalApi)
+					{
+						Definitions = new List<string>(Definitions);
+						Definitions.RemoveAll(x => x.Contains("UE_VALIDATE_EXPERIMENTAL_API", StringComparison.Ordinal));
+						Definitions.Add($"UE_VALIDATE_EXPERIMENTAL_API={(ModuleCompileEnvironment.bValidateExperimentalApi ? "1" : "0")}");
+					}
 
-				// Create a suffix to distinguish this shared PCH variant from any others.
-				string Variant = GetSuffixForSharedPCH(ModuleCompileEnvironment, Template.BaseCompileEnvironment);
+					// Create a suffix to distinguish this shared PCH variant from any others.
+					string Variant = GetSuffixForSharedPCH(ModuleCompileEnvironment, Template.BaseCompileEnvironment);
 
-				FileReference SharedDefinitionsLocation = FileReference.Combine(Template.OutputDir, $"SharedDefinitions.{Template.Module.Name}{Variant}.h");
-				List<string> NewDefinitions = new();
-				StringBuilder Writer = new StringBuilder();
-				WriteDefinitions($"Shared Definitions for {Template.Module.Name}{Variant}", [], Definitions, [], Writer);
-				FileItem SharedDefinitionsFileItem = Graph.CreateIntermediateTextFile(SharedDefinitionsLocation, Writer.ToString(), AllowAsync: false);
+					FileReference SharedDefinitionsLocation = FileReference.Combine(Template.OutputDir, $"SharedDefinitions.{Template.Module.Name}{Variant}.h");
+					List<string> NewDefinitions = new();
+					StringBuilder Writer = new StringBuilder();
+					WriteDefinitions($"Shared Definitions for {Template.Module.Name}{Variant}", [], Definitions, [], Writer);
+					FileItem SharedDefinitionsFileItem = Graph.CreateIntermediateTextFile(SharedDefinitionsLocation, Writer.ToString(), AllowAsync: false);
 
-				// Create the wrapper file, which sets all the definitions needed to compile it
-				FileReference WrapperLocation = FileReference.Combine(Template.OutputDir, $"SharedPCH.{Template.Module.Name}{Variant}.h");
-				FileItem WrapperFile = CreatePCHWrapperFile(WrapperLocation, SharedDefinitionsFileItem, Template.HeaderFile, Graph);
+					// Create the wrapper file, which sets all the definitions needed to compile it
+					FileReference WrapperLocation = FileReference.Combine(Template.OutputDir, $"SharedPCH.{Template.Module.Name}{Variant}.h");
+					FileItem WrapperFile = CreatePCHWrapperFile(WrapperLocation, SharedDefinitionsFileItem, Template.HeaderFile, Graph);
 
-				// Create the compile environment for this PCH
-				CppCompileEnvironment CompileEnvironment = new CppCompileEnvironment(Template.BaseCompileEnvironment);
-				CompileEnvironment.Definitions.Clear();
-				CompileEnvironment.PrecompiledHeaderAction = PrecompiledHeaderAction.Create;
-				CompileEnvironment.PrecompiledHeaderIncludeFilename = WrapperFile.Location;
-				CopySettingsForSharedPCH(ModuleCompileEnvironment, CompileEnvironment);
+					// Create the compile environment for this PCH
+					CppCompileEnvironment CompileEnvironment = new CppCompileEnvironment(Template.BaseCompileEnvironment);
+					CompileEnvironment.Definitions.Clear();
+					CompileEnvironment.PrecompiledHeaderAction = PrecompiledHeaderAction.Create;
+					CompileEnvironment.PrecompiledHeaderIncludeFilename = WrapperFile.Location;
+					CopySettingsForSharedPCH(ModuleCompileEnvironment, CompileEnvironment);
 
-				// Setup PCH chaining
-				PrecompiledHeaderInstance? ParentPCHInstance = null;
-				if (Rules.Target.bChainPCHs)
-				{
-					// Find all the dependencies of this module
-					HashSet<UEBuildModule> ReferencedModules = new HashSet<UEBuildModule>();
-					Template.Module.GetAllDependencyModules(new List<UEBuildModule>(), ReferencedModules, bIncludeDynamicallyLoaded: false, bForceCircular: false, bOnlyDirectDependencies: false);
+					// Setup PCH chaining
+					PrecompiledHeaderInstance? ParentPCHInstance = null;
+					if (Rules.Target.bChainPCHs)
+					{
+						// Find all the dependencies of this module
+						HashSet<UEBuildModule> ReferencedModules = new HashSet<UEBuildModule>();
+						Template.Module.GetAllDependencyModules(new List<UEBuildModule>(), ReferencedModules, bIncludeDynamicallyLoaded: false, bForceCircular: false, bOnlyDirectDependencies: false);
+
+						PrecompiledHeaderTemplate? ParentTemplate = CompileEnvironment.SharedPCHs
+							.Skip(CompileEnvironment.SharedPCHs.IndexOf(Template) + 1)
+							.FirstOrDefault(x => ReferencedModules.Contains(x.Module) && x.IsValidFor(CompileEnvironment));
+						if (ParentTemplate != null)
+						{
+							ParentPCHInstance = FindOrCreateSharedPCH(ToolChain, ParentTemplate, ModuleCompileEnvironment, Graph);
+							CompileEnvironment.ParentPCHInstance = ParentPCHInstance;
+						}
+					}
 
-					PrecompiledHeaderTemplate? ParentTemplate = CompileEnvironment.SharedPCHs
-						.Skip(CompileEnvironment.SharedPCHs.IndexOf(Template) + 1)
-						.FirstOrDefault(x => ReferencedModules.Contains(x.Module) && x.IsValidFor(CompileEnvironment));
-					if (ParentTemplate != null) 
+					// Create the PCH
+					CPPOutput Output;
+					if (ToolChain == null)
 					{
-						ParentPCHInstance = FindOrCreateSharedPCH(ToolChain, ParentTemplate, ModuleCompileEnvironment, Graph);
-						CompileEnvironment.ParentPCHInstance = ParentPCHInstance;
+						Output = new CPPOutput();
 					}
+					else
+					{
+						Output = ToolChain.CompileAllCPPFiles(CompileEnvironment, new List<FileItem>() { WrapperFile }, Template.OutputDir, "Shared", Graph);
+					}
+					Instance = new PrecompiledHeaderInstance(WrapperFile, SharedDefinitionsFileItem, CompileEnvironment, Output, GetImmutableDefinitions(Template.BaseCompileEnvironment.Definitions))
+					{
+						ParentPCHInstance = ParentPCHInstance
+					};
+					Template.Instances.Add(Instance);
 				}
 
-				// Create the PCH
-				CPPOutput Output;
-				if (ToolChain == null)
-				{
-					Output = new CPPOutput();
-				}
-				else
-				{
-					Output = ToolChain.CompileAllCPPFiles(CompileEnvironment, new List<FileItem>() { WrapperFile }, Template.OutputDir, "Shared", Graph);
-				}
-				Instance = new PrecompiledHeaderInstance(WrapperFile, SharedDefinitionsFileItem, CompileEnvironment, Output, GetImmutableDefinitions(Template.BaseCompileEnvironment.Definitions))
-				{
-					ParentPCHInstance = ParentPCHInstance
-				};
-				Template.Instances.Add(Instance);
+				Instance.Modules.Add(this);
+				return Instance;
 			}
-
-			Instance.Modules.Add(this);
-			return Instance;
 		}
 
 		/// <summary>
@@ -2059,7 +2057,7 @@ namespace UnrealBuildTool
 				if (CompileEnvironment.bHasPrecompiledHeader == false && CompileEnvironment.SharedPCHs.Count > 0 && !CompileEnvironment.bIsBuildingLibrary && Rules.PCHUsage != ModuleRules.PCHUsageMode.NoSharedPCHs)
 				{
 					// Find the first shared PCH module we can use that doesn't reference this module
-					HashSet<UEBuildModule> AllDependencies = GetAllDependencyModulesForPCH(true, false);
+					HashSet<UEBuildModule> AllDependencies = GetAllDependencyModulesForPCH(true);
 					PrecompiledHeaderTemplate? Template = CompileEnvironment.SharedPCHs.FirstOrDefault(x => !x.ModuleDependencies.Contains(this) && AllDependencies.Contains(x.Module) && x.IsValidFor(CompileEnvironment));
 
 					if (Template != null)
@@ -2245,7 +2243,8 @@ namespace UnrealBuildTool
 					}
 
 					// Find the directly included files for each source file, and make sure it includes the matching header if possible
-					InvalidIncludeDirectives = new();
+
+					List<InvalidIncludeDirective>? list = null;
 					if (Rules != null && Rules.IWYUSupport != IWYUSupport.None && Target.bEnforceIWYU)
 					{
 						foreach (FileItem CppFile in CppFiles)
@@ -2260,12 +2259,19 @@ namespace UnrealBuildTool
 									FileReference? HeaderFile;
 									if (NameToHeaderFile.TryGetValue(ExpectedName, out HeaderFile))
 									{
-										InvalidIncludeDirectives.Add(new InvalidIncludeDirective(CppFile.Location, HeaderFile));
+										if (list == null)
+											list = new();
+										list.Add(new InvalidIncludeDirective(CppFile.Location, HeaderFile));
 									}
 								}
 							}
 						}
 					}
+
+					if (list != null)
+					{
+						InvalidIncludeDirectives = list;
+					}
 				}
 			}
 		}
diff --git a/Engine/Source/Programs/UnrealBuildTool/Configuration/UEBuildTarget.cs b/Engine/Source/Programs/UnrealBuildTool/Configuration/UEBuildTarget.cs
--- a/Engine/Source/Programs/UnrealBuildTool/Configuration/UEBuildTarget.cs
+++ b/Engine/Source/Programs/UnrealBuildTool/Configuration/UEBuildTarget.cs
@@ -1264,6 +1264,8 @@ namespace UnrealBuildTool
 				throw new BuildException("Platform {0} is not a valid platform to build. Check that the SDK is installed properly and that you have the necessary platform support files (DataDrivenPlatformInfo.ini, SDK.json, etc).", Descriptor.Platform);
 			}
 
+			Task<SourceFileMetadataCache> metaDataTask = Task.Run(() => SourceFileMetadataCache.CreateHierarchy(Descriptor.ProjectFile, Logger));
+
 			RulesAssembly RulesAssembly;
 			using (GlobalTracer.Instance.BuildSpan("RulesCompiler.CreateTargetRulesAssembly()").StartActive())
 			{
@@ -1273,6 +1275,7 @@ namespace UnrealBuildTool
 			TargetRules RulesObject;
 			using (GlobalTracer.Instance.BuildSpan("RulesAssembly.CreateTargetRules()").StartActive())
 			{
+				using var _ = Timeline.ScopeEvent("CreateTarget");
 				RulesObject = RulesAssembly.CreateTargetRules(Descriptor.Name, Descriptor.Platform, Descriptor.Configuration, Descriptor.Architectures, Descriptor.ProjectFile, Descriptor.AdditionalArguments, Logger, Descriptor.IsTestsTarget, IntermediateEnvironment: IntermediateEnvironment);
 			}
 			if ((ProjectFileGenerator.bGenerateProjectFiles == false) && !RulesObject.GetSupportedPlatforms().Contains(Descriptor.Platform))
@@ -1394,7 +1397,7 @@ namespace UnrealBuildTool
 			UEBuildTarget Target;
 			using (GlobalTracer.Instance.BuildSpan("UEBuildTarget constructor").StartActive())
 			{
-				Target = new UEBuildTarget(Descriptor, new ReadOnlyTargetRules(RulesObject), RulesAssembly, Logger);
+				Target = new UEBuildTarget(Descriptor, new ReadOnlyTargetRules(RulesObject), RulesAssembly, metaDataTask.Result, Logger);
 			}
 			using (GlobalTracer.Instance.BuildSpan("UEBuildTarget.PreBuildSetup()").StartActive())
 			{
@@ -1743,7 +1746,9 @@ namespace UnrealBuildTool
 		/// <summary>
 		/// Used to keep track of all modules by name.
 		/// </summary>
-		private Dictionary<string, UEBuildModule> Modules = new Dictionary<string, UEBuildModule>(StringComparer.InvariantCultureIgnoreCase);
+		private Dictionary<string, UEBuildModule> Modules = new (StringComparer.InvariantCultureIgnoreCase);
+
+		private Dictionary<string, UEBuildModule> UsedShortNames = new(StringComparer.InvariantCultureIgnoreCase);
 
 		/// <summary>
 		/// Filename for the receipt for this target.
@@ -1788,10 +1793,11 @@ namespace UnrealBuildTool
 		/// <param name="InDescriptor">Target descriptor</param>
 		/// <param name="InRules">The target rules, as created by RulesCompiler.</param>
 		/// <param name="InRulesAssembly">The chain of rules assemblies that this target was created with</param>
+		/// <param name="InMetadataCache"></param>
 		/// <param name="Logger">Logger for output</param>
-		private UEBuildTarget(TargetDescriptor InDescriptor, ReadOnlyTargetRules InRules, RulesAssembly InRulesAssembly, ILogger Logger)
+		private UEBuildTarget(TargetDescriptor InDescriptor, ReadOnlyTargetRules InRules, RulesAssembly InRulesAssembly, SourceFileMetadataCache InMetadataCache, ILogger Logger)
 		{
-			MetadataCache = SourceFileMetadataCache.CreateHierarchy(InDescriptor.ProjectFile, Logger);
+			MetadataCache = InMetadataCache;
 			ProjectFile = InDescriptor.ProjectFile;
 			TargetName = InRules.Name;
 			AppName = InRules.Name;
@@ -2471,9 +2477,9 @@ namespace UnrealBuildTool
 		static HashSet<UEBuildModuleCPP> GatherDependencyModules(List<UEBuildBinary> Binaries)
 		{
 			ConcurrentBag<UEBuildModuleCPP> Bag = new();
-			Parallel.ForEach(Binaries, Binary =>
+			Parallel2.ForEach(Binaries, Binary =>
 			{
-				List<UEBuildModule> DependencyModules = Binary.GetAllDependencyModules(bIncludeDynamicallyLoaded: false, bForceCircular: false);
+				IEnumerable<UEBuildModule> DependencyModules = Binary.GetAllDependencyModules(bIncludeDynamicallyLoaded: false, bForceCircular: false);
 				foreach (UEBuildModuleCPP Module in DependencyModules.OfType<UEBuildModuleCPP>())
 				{
 					if (Module.Binary != null)
@@ -2586,7 +2592,7 @@ namespace UnrealBuildTool
 				UEBuildBinary ExecutableBinary = Binaries[0];
 
 				// Add all the modules that the executable depends on. Plugins will be already included in this list.
-				List<UEBuildModule> AllReferencedModules = ExecutableBinary.GetAllDependencyModules(bIncludeDynamicallyLoaded: true, bForceCircular: true);
+				IEnumerable<UEBuildModule> AllReferencedModules = ExecutableBinary.GetAllDependencyModules(bIncludeDynamicallyLoaded: true, bForceCircular: true);
 				foreach (UEBuildModule CurModule in AllReferencedModules)
 				{
 					if (CurModule.Binary == null || CurModule.Binary == ExecutableBinary || CurModule.Binary.Type == UEBuildBinaryType.StaticLibrary)
@@ -2608,6 +2614,8 @@ namespace UnrealBuildTool
 			// Prepare cached data for UHT header generation
 			using (GlobalTracer.Instance.BuildSpan("UHTExecution.SetupUObjectModules()").StartActive())
 			{
+				using var _ = Timeline.ScopeEvent("SetupUObjectModules");
+
 				UHTExecution.SetupUObjectModules(
 					ModulesToGenerateHeadersFor, 
 					Rules.Platform,
@@ -2727,12 +2735,15 @@ namespace UnrealBuildTool
 			DirectoryReference ExeDir = GetExecutableDir();
 			using (GlobalTracer.Instance.BuildSpan("UEBuildBinary.Build()").StartActive())
 			{
-				foreach (UEBuildBinary Binary in Binaries)
+				Parallel2.ForEach(Binaries, Binary =>
 				{
 					Binary.PrepareRuntimeDependencies(ExeDir);
 					List<FileItem> BinaryOutputItems = Binary.Build(Rules, TargetToolChain, GlobalCompileEnvironment, GlobalLinkEnvironment, WorkingSet, ExeDir, MakefileBuilder, Logger);
-					Makefile.OutputItems.AddRange(BinaryOutputItems);
-				}
+					lock (Makefile)
+					{
+						Makefile.OutputItems.AddRange(BinaryOutputItems);
+					}
+				});
 			}
 
 			if (Rules.bStripExports)
@@ -2856,6 +2867,8 @@ namespace UnrealBuildTool
 				}
 			}
 
+			List<Task> writeMetadataTasks = [];
+
 			// Finalize and generate metadata for this target
 			if (!Rules.bDisableLinking)
 			{
@@ -2904,7 +2917,14 @@ namespace UnrealBuildTool
 				{
 					WriteMetadataTargetInfo TargetInfo = new WriteMetadataTargetInfo(ProjectFile, null, null, ReceiptFileName, Receipt, FileNameToModuleManifest, FileNameToLoadOrderManifest);
 					FileReference TargetInfoFile = FileReference.Combine(ProjectIntermediateDirectory, "TargetMetadata.dat");
-					CreateWriteMetadataAction(MakefileBuilder, ReceiptFileName.GetFileName(), TargetInfoFile, TargetInfo, Makefile.OutputItems);
+
+					// Make a copy of this since OutputItems might be added byh post build actions
+					List<FileItem> TargetPrereqItems = new (Makefile.OutputItems);
+
+					writeMetadataTasks.Add(Task.Run(() =>
+					{
+						CreateWriteMetadataAction(MakefileBuilder, ReceiptFileName.GetFileName(), TargetInfoFile, TargetInfo, TargetPrereqItems);
+					}));
 				}
 				else
 				{
@@ -2938,14 +2958,19 @@ namespace UnrealBuildTool
 						}
 					}
 
-					WriteMetadataTargetInfo EngineInfo = new WriteMetadataTargetInfo(null, VersionFile, Receipt.Version, null, null, EngineFileToManifest, null);
-					FileReference EngineInfoFile = FileReference.Combine(ProjectIntermediateDirectory, "EngineMetadata.dat");
-					string EngineInfoName = TargetName.Equals(VersionFile.GetFileNameWithoutExtension()) ? VersionFile.GetFileName() : $"{VersionFile.GetFileName()} ({TargetName})";
-					CreateWriteMetadataAction(MakefileBuilder, EngineInfoName, EngineInfoFile, EngineInfo, EnginePrereqItems);
-
-					WriteMetadataTargetInfo TargetInfo = new WriteMetadataTargetInfo(ProjectFile, VersionFile, null, ReceiptFileName, Receipt, TargetFileToManifest, FileNameToLoadOrderManifest);
-					FileReference TargetInfoFile = FileReference.Combine(ProjectIntermediateDirectory, "TargetMetadata.dat");
-					CreateWriteMetadataAction(MakefileBuilder, ReceiptFileName.GetFileName(), TargetInfoFile, TargetInfo, TargetPrereqItems);
+					writeMetadataTasks.Add(Task.Run(() =>
+					{
+						WriteMetadataTargetInfo EngineInfo = new WriteMetadataTargetInfo(null, VersionFile, Receipt.Version, null, null, EngineFileToManifest, null);
+						FileReference EngineInfoFile = FileReference.Combine(ProjectIntermediateDirectory, "EngineMetadata.dat");
+						string EngineInfoName = TargetName.Equals(VersionFile.GetFileNameWithoutExtension()) ? VersionFile.GetFileName() : $"{VersionFile.GetFileName()} ({TargetName})";
+						CreateWriteMetadataAction(MakefileBuilder, EngineInfoName, EngineInfoFile, EngineInfo, EnginePrereqItems);
+					}));
+					writeMetadataTasks.Add(Task.Run(() =>
+					{
+						WriteMetadataTargetInfo TargetInfo = new WriteMetadataTargetInfo(ProjectFile, VersionFile, null, ReceiptFileName, Receipt, TargetFileToManifest, FileNameToLoadOrderManifest);
+						FileReference TargetInfoFile = FileReference.Combine(ProjectIntermediateDirectory, "TargetMetadata.dat");
+						CreateWriteMetadataAction(MakefileBuilder, ReceiptFileName.GetFileName(), TargetInfoFile, TargetInfo, TargetPrereqItems);
+					}));
 				}
 
 				// Create actions to run the post build steps
@@ -3082,9 +3107,14 @@ namespace UnrealBuildTool
 				WritePublicHeader(Binary, Header, GlobalCompileEnvironment, Logger);
 			}
 
+			Task.WaitAll(writeMetadataTasks.ToArray());
+
 			// Wait for makefile tasks to finish (tasks writing files to disk etc)
 			MakefileBuilder.WaitOnWriteTasks();
 
+			// Clean any stale modules which exist in multiple output directories. This can lead to the wrong DLL being loaded on Windows.
+			CleanStaleModules(Logger);
+
 #if __VPROJECT_AVAILABLE__
 			if (vniTask != null)
 			{
@@ -3100,9 +3130,6 @@ namespace UnrealBuildTool
 				Task task = FileMetadataPrefetch.QueueDirectoryTrees(ModulesToGenerateHeadersFor.Select(m => m.GeneratedCodeDirectoryUHT));
 			}
 
-			// Clean any stale modules which exist in multiple output directories. This can lead to the wrong DLL being loaded on Windows.
-			CleanStaleModules(Logger);
-
 			return Makefile;
 		}
 
@@ -3353,7 +3380,7 @@ namespace UnrealBuildTool
 			}
 
 			// Search all the output directories for files with a name matching one of our output files
-			foreach (DirectoryReference OutputDirectory in OutputFiles.Select(x => x.Directory).Distinct())
+			Parallel2.ForEach(OutputFiles.Select(x => x.Directory).Distinct(), OutputDirectory =>
 			{
 				if (DirectoryReference.Exists(OutputDirectory))
 				{
@@ -3374,7 +3401,7 @@ namespace UnrealBuildTool
 						}
 					}
 				}
-			}
+			});
 		}
 
 		/// <summary>
@@ -3511,7 +3538,7 @@ namespace UnrealBuildTool
 				bool bLicenseViolation = false;
 				foreach (UEBuildBinary Binary in Binaries)
 				{
-					List<UEBuildModule> AllDependencies = Binary.GetAllDependencyModules(true, false);
+					IEnumerable<UEBuildModule> AllDependencies = Binary.GetAllDependencyModules(true, false);
 					IEnumerable<UEBuildModule> NonRedistModules = AllDependencies.Where((DependencyModule) =>
 							!IsRedistributable(DependencyModule) && DependencyModule.Name != AppName
 						);
@@ -4051,6 +4078,8 @@ namespace UnrealBuildTool
 		/// </summary>
 		public void PreBuildSetup(ILogger Logger)
 		{
+			using var _ = Timeline.ScopeEvent("PreBuildSetup");
+
 			// Describe what's being built.
 			Logger.LogDebug("Building {AppName} - {TargetName} - {Platform} - {Configuration}", AppName, TargetName, Platform, Configuration);
 
@@ -4084,20 +4113,27 @@ namespace UnrealBuildTool
 			}
 
 			// Cache dependencies since we can do this in parallel for all modules
-			UEBuildModule.CacheDependenciesForStableTopologicalSort(Binaries.SelectMany(B => B.Modules));
+			if (ShouldCompileMonolithic())
+			{
+				UEBuildModule.CacheDependenciesForStableTopologicalSort(Binaries.SelectMany(b => b.Modules.Take(1)));
+			}
+			else
+			{
+				UEBuildModule.CacheDependenciesForStableTopologicalSort(Modules.Values.OfType<UEBuildModuleCPP>());
+			}
 
 			// Bind every referenced C++ module to a binary
 			for (int Idx = 0; Idx < Binaries.Count; Idx++)
-			{
-				List<UEBuildModule> DependencyModules = Binaries[Idx].GetAllDependencyModules(true, true);
-				foreach (UEBuildModuleCPP DependencyModule in DependencyModules.OfType<UEBuildModuleCPP>())
 				{
-					if (DependencyModule.Binary == null)
+					IEnumerable<UEBuildModule> DependencyModules = Binaries[Idx].GetAllDependencyModules(true, true);
+					foreach (UEBuildModule DependencyModule in DependencyModules)
 					{
-						AddModuleToBinary(DependencyModule);
+						if (DependencyModule.Binary == null && DependencyModule is UEBuildModuleCPP DependencyModuleCpp) // Do type check here, much faster than in enumerable
+						{
+							AddModuleToBinary(DependencyModuleCpp);
+						}
 					}
 				}
-			}
 
 			// Add all the modules to the target if necessary.
 			if (Rules.bBuildAllModules)
@@ -4364,18 +4400,34 @@ namespace UnrealBuildTool
 			}
 
 			// Find a priority for each shared PCH, determined as the number of other shared PCHs it includes.
-			Dictionary<UEBuildModuleCPP, int> SharedPCHModuleToPriority = new Dictionary<UEBuildModuleCPP, int>();
-			foreach (UEBuildModuleCPP SharedPCHModule in SharedPCHModules)
+			SortedDictionary<string, (int, PrecompiledHeaderTemplate)> SharedPCHModuleToPriority = new ();
+			Parallel2.ForEach(SharedPCHModules, SharedPCHModule =>
 			{
-				HashSet<UEBuildModule> Dependencies = SharedPCHModule.GetAllDependencyModulesForPCH(false, false);
-				SharedPCHModuleToPriority.Add(SharedPCHModule, Dependencies.Count(x => SharedPCHModules.Contains(x)));
-			}
+				HashSet<UEBuildModule> Dependencies = SharedPCHModule.GetAllDependencyModulesForPCH(false);
+				Dependencies.IntersectWith(SharedPCHModules);
+
+				// TODO This is bad.. unfortunately does circular dependency logic not working when getting dependency modules
+				// If we have Engine with circular dependency on UnrealEd but then is included via some other module, then we're screwed
+				foreach (string CircularDependency in SharedPCHModule.Rules.CircularlyReferencedDependentModules)
+				{
+					if (Modules.TryGetValue(CircularDependency, out UEBuildModule? Module))
+					{
+						Dependencies.Remove(Module);
+					}
+				}
+				int DependencyCount = Dependencies.Count;
+				PrecompiledHeaderTemplate Template = SharedPCHModule.CreateSharedPCHTemplate(this, GlobalCompileEnvironment, Logger);
+				lock (SharedPCHModuleToPriority)
+				{
+					SharedPCHModuleToPriority.Add(SharedPCHModule.Name, (DependencyCount, Template));
+				}
+			});
 
 			// Create the shared PCH modules, in order
 			List<PrecompiledHeaderTemplate> OrderedSharedPCHModules = GlobalCompileEnvironment.SharedPCHs;
-			foreach (UEBuildModuleCPP Module in SharedPCHModuleToPriority.OrderByDescending(x => x.Value).Select(x => x.Key))
+			foreach (PrecompiledHeaderTemplate Template in SharedPCHModuleToPriority.OrderByDescending(x => x.Value.Item1).Select(v => v.Value.Item2))
 			{
-				OrderedSharedPCHModules.Add(Module.CreateSharedPCHTemplate(this, GlobalCompileEnvironment, Logger));
+				OrderedSharedPCHModules.Add(Template);
 			}
 
 			// Print the ordered list of shared PCHs
@@ -4412,7 +4464,7 @@ namespace UnrealBuildTool
 				ModuleInfos[SharedPCHIndex] = new ConcurrentBag<Tuple<UEBuildModuleCPP, CppCompileEnvironment>>();
 			}
 
-			Parallel.ForEach(OriginalBinaries, Binary =>
+			Parallel2.ForEach(OriginalBinaries, Binary =>
 			{
 				CppCompileEnvironment BinaryCompileEnvironment = Binary.CreateBinaryCompileEnvironment(GlobalCompileEnvironment);
 
@@ -6396,9 +6448,9 @@ namespace UnrealBuildTool
 				Module = InstantiateModule(RulesObject, GeneratedCodeDirectory, Logger);
 
 				// Check if this module conflicts with the ShortName of any other existing modules
-				UEBuildModule? ConflictingModule = Modules.Values.FirstOrDefault(x => (x.Rules.ShortName ?? x.Name) == (Module.Rules.ShortName ?? Module.Name));
-				if (ConflictingModule != null)
+				if (!UsedShortNames.TryAdd(Module.Rules.ShortName ?? Module.Name, Module))
 				{
+					UEBuildModule? ConflictingModule = UsedShortNames[Module.Rules.ShortName ?? Module.Name];
 					Logger.LogError("Conflicting module short names found, unable to create module {Name}", Module.Name);
 					Logger.LogInformation(" * Name={Name} ShortName={ShortName} Path={Path}", ConflictingModule.Name, ConflictingModule.Rules.ShortName ?? ConflictingModule.Name, ConflictingModule.RulesFile);
 					Logger.LogInformation(" * Name={Name} ShortName={ShortName} Path={Path}", Module.Name, Module.Rules.ShortName ?? Module.Name, Module.RulesFile);
diff --git a/Engine/Source/Programs/UnrealBuildTool/Modes/BuildMode.cs b/Engine/Source/Programs/UnrealBuildTool/Modes/BuildMode.cs
--- a/Engine/Source/Programs/UnrealBuildTool/Modes/BuildMode.cs
+++ b/Engine/Source/Programs/UnrealBuildTool/Modes/BuildMode.cs
@@ -389,6 +389,8 @@ namespace UnrealBuildTool
 		{
 			List<TargetMakefile> TargetMakefiles = new List<TargetMakefile>();
 
+			List<Task> SaveMakeFileTasks = new();
+
 			for (int Idx = 0; Idx < TargetDescriptors.Count; ++Idx)
 			{
 				using ITimelineEvent SetupScope = Timeline.ScopeEvent($"Creating MakeFile for {TargetDescriptors[Idx].Name}");
@@ -399,8 +401,18 @@ namespace UnrealBuildTool
 					Utils.ExecuteCustomBuildSteps(InitScripts, Logger);
 				}
 
-				TargetMakefile NewMakefile = await CreateMakefileAsync(BuildConfiguration, TargetDescriptors[Idx], WorkingSet, Logger);
+				(TargetMakefile NewMakefile, FileReference? SaveToLocation) = await CreateMakefileAsync2(BuildConfiguration, TargetDescriptors[Idx], WorkingSet, Logger);
 				TargetMakefiles.Add(NewMakefile);
+
+				if (SaveToLocation != null)
+				{
+					SaveMakeFileTasks.Add(Task.Run(() =>
+					{
+						using ITimelineEvent _ = Timeline.ScopeEvent($"SavingMakefile");
+						NewMakefile.Save(SaveToLocation);
+					}));
+				}
+
 				if (!bSkipPreBuildTargets)
 				{
 					foreach (TargetInfo PreBuildTarget in NewMakefile.PreBuildTargets)
@@ -415,6 +427,8 @@ namespace UnrealBuildTool
 			}
 
 			await BuildAsync(TargetMakefiles.ToArray(), TargetDescriptors, BuildConfiguration, Options, WriteOutdatedActionsFile, Logger, ActionTypeFilter);
+
+			Task.WaitAll(SaveMakeFileTasks.ToArray());
 		}
 
 		static Lazy<MemoryMappedFileCache> s_memoryMappedFileCache = new Lazy<MemoryMappedFileCache>(() => new MemoryMappedFileCache());
@@ -435,6 +449,53 @@ namespace UnrealBuildTool
 			// Execute the build
 			if ((Options & BuildOptions.SkipBuild) == 0)
 			{
+				Task<CppDependencyCache> CppDependenciesTask = Task.Run(() =>
+				{
+					// Create the C++ dependency cache
+					CppDependencyCache CppDependencies = new CppDependencyCache();
+					for (int TargetIdx = 0; TargetIdx < TargetDescriptors.Count; TargetIdx++)
+					{
+						using (GlobalTracer.Instance.BuildSpan("Reading dependency cache").StartActive())
+						{
+							TargetDescriptor TargetDescriptor = TargetDescriptors[TargetIdx];
+							CppDependencies.Mount(TargetDescriptor, Makefiles[TargetIdx].TargetType, Logger);
+						}
+					}
+					return CppDependencies;
+				});
+
+#if false // Doesn't seem to make things faster. Revisit since it feels like this should be faster
+				// Prepopulate directories
+				for (int Idx = 0; Idx < Makefiles.Length; Idx++)
+				{
+					TargetDescriptor TargetDescriptor = TargetDescriptors[Idx];
+					TargetMakefile Makefile = Makefiles[Idx];
+					//HashSet<DirectoryItem> handledDir = new();
+					//ParallelQueue queue = new();
+
+					Task throwaway = FileMetadataPrefetch.QueueDirectoryTrees(Makefile.UObjectModules.Select(m => m.GeneratedCodeDirectory));
+					throwaway.Wait();
+					/*
+					foreach (var module in Makefile.UObjectModules)
+					{
+						module.GeneratedCodeDirectory
+					foreach (IExternalAction action in Makefile.Actions)
+					{
+						foreach (FileItem prereq in action.PrerequisiteItems)
+						{
+							DirectoryItem dir = prereq.Directory;
+							if (handledDir.Add(dir))
+								queue.Enqueue(() =>
+								{
+									dir.CacheDirectories();
+								});
+						}
+					}
+					queue.Drain();
+					*/
+				}
+#endif
+
 				using ITimelineEvent ActionGraphTimer = Timeline.ScopeEvent($"Preparing ActionGraph");
 				// Make sure that none of the actions conflict with any other (producing output files differently, etc...)
 				using (GlobalTracer.Instance.BuildSpan("ActionGraph.CheckForConflicts").StartActive())
@@ -531,16 +592,8 @@ namespace UnrealBuildTool
 					}
 				}
 
-				// Create the C++ dependency cache
-				CppDependencyCache CppDependencies = new CppDependencyCache();
-				for (int TargetIdx = 0; TargetIdx < TargetDescriptors.Count; TargetIdx++)
-				{
-					using (GlobalTracer.Instance.BuildSpan("Reading dependency cache").StartActive())
-					{
-						TargetDescriptor TargetDescriptor = TargetDescriptors[TargetIdx];
-						CppDependencies.Mount(TargetDescriptor, Makefiles[TargetIdx].TargetType, Logger);
-					}
-				}
+				// Wait for cpp dependencies task to finish
+				CppDependencyCache CppDependencies = CppDependenciesTask.Result;
 
 				// Now that we have the dependencies object, optionally create the artifact cache
 				IActionArtifactCache? actionArtifactCache = null;
@@ -1084,6 +1137,22 @@ namespace UnrealBuildTool
 		/// <returns>Makefile for the given target</returns>
 		internal static async Task<TargetMakefile> CreateMakefileAsync(BuildConfiguration BuildConfiguration, TargetDescriptor TargetDescriptor, ISourceFileWorkingSet WorkingSet, ILogger Logger)
 		{
+			var (NewMakefile, SaveToLocation) = await CreateMakefileAsync2(BuildConfiguration, TargetDescriptor, WorkingSet, Logger);
+			if (SaveToLocation != null)
+			{
+				// Save the makefile for next time
+				using (GlobalTracer.Instance.BuildSpan("TargetMakefile.Save()").StartActive())
+				{
+					NewMakefile.Save(SaveToLocation);
+				}
+			}
+			return NewMakefile;
+		}
+
+		private static async Task<(TargetMakefile, FileReference?)> CreateMakefileAsync2(BuildConfiguration BuildConfiguration, TargetDescriptor TargetDescriptor, ISourceFileWorkingSet WorkingSet, ILogger Logger)
+		{
+			FileReference? SaveToLocation = null;
+
 			// Get the path to the makefile for this target
 			FileReference? MakefileLocation = null;
 			if (BuildConfiguration.bUseUBTMakefiles)
@@ -1171,10 +1240,7 @@ namespace UnrealBuildTool
 				// Save the makefile for next time
 				if (MakefileLocation != null)
 				{
-					using (GlobalTracer.Instance.BuildSpan("TargetMakefile.Save()").StartActive())
-					{
-						Makefile.Save(MakefileLocation);
-					}
+					SaveToLocation = MakefileLocation;
 				}
 			}
 			else
@@ -1203,7 +1269,7 @@ namespace UnrealBuildTool
 				VNITask.Wait();
 #endif
 			}
-			return Makefile;
+			return (Makefile, SaveToLocation);
 		}
 
 		/// <summary>
diff --git a/Engine/Source/Programs/UnrealBuildTool/System/ActionGraph.cs b/Engine/Source/Programs/UnrealBuildTool/System/ActionGraph.cs
--- a/Engine/Source/Programs/UnrealBuildTool/System/ActionGraph.cs
+++ b/Engine/Source/Programs/UnrealBuildTool/System/ActionGraph.cs
@@ -49,7 +49,7 @@ namespace UnrealBuildTool
 			using ITimelineEvent _ = Timeline.ScopeEvent($"LinkActions");
 
 			// Build a map from item to its producing action
-			Dictionary<FileItem, LinkedAction> ItemToProducingAction = new Dictionary<FileItem, LinkedAction>();
+			Dictionary<FileItem, LinkedAction> ItemToProducingAction = new (Actions.Count * 3);
 			foreach (LinkedAction Action in Actions)
 			{
 				foreach (FileItem ProducedItem in Action.ProducedItems)
diff --git a/Engine/Source/Programs/UnrealBuildTool/System/FileMetadataPrefetch.cs b/Engine/Source/Programs/UnrealBuildTool/System/FileMetadataPrefetch.cs
--- a/Engine/Source/Programs/UnrealBuildTool/System/FileMetadataPrefetch.cs
+++ b/Engine/Source/Programs/UnrealBuildTool/System/FileMetadataPrefetch.cs
@@ -1,5 +1,7 @@
 // Copyright Epic Games, Inc. All Rights Reserved.
 
+#define UBT_PREFETCH_ENABLED
+
 using System;
 using System.Collections.Concurrent;
 using System.Collections.Generic;
@@ -36,6 +38,7 @@ namespace UnrealBuildTool
 		/// </summary>
 		public static Task QueueEngineDirectory()
 		{
+#if UBT_PREFETCH_ENABLED
 			return QueuedDirectories.GetOrAdd(Unreal.EngineDirectory, new Lazy<Task>(() =>
 			{
 				return Task.Run(() =>
@@ -46,6 +49,9 @@ namespace UnrealBuildTool
 					queue.Drain();
 				});
 			})).Value;
+#else
+			return Task.CompletedTask;
+#endif
 		}
 
 		/// <summary>
@@ -54,6 +60,7 @@ namespace UnrealBuildTool
 		/// <param name="ProjectDirectory">The project directory to prefetch</param>
 		public static Task QueueProjectDirectory(DirectoryReference ProjectDirectory)
 		{
+#if UBT_PREFETCH_ENABLED
 			return QueuedDirectories.GetOrAdd(ProjectDirectory, new Lazy<Task>(() =>
 			{
 				return Task.Run(() =>
@@ -64,6 +71,9 @@ namespace UnrealBuildTool
 					queue.Drain();
 				});
 			})).Value;
+#else
+			return Task.CompletedTask;
+#endif
 		}
 
 		/// <summary>
@@ -72,6 +82,7 @@ namespace UnrealBuildTool
 		/// <param name="Directories">Directory to start searching from</param>
 		public static Task QueueDirectoryTrees(IEnumerable<DirectoryReference?> Directories)
 		{
+#if UBT_PREFETCH_ENABLED
 			return Task.Run(() =>
 			{
 				//using var _ = Timeline.ScopeEvent("PrefetchDir");
@@ -83,6 +94,9 @@ namespace UnrealBuildTool
 				}
 				queue.Drain();
 			});
+#else
+			return Task.CompletedTask;
+#endif
 		}
 
 		/// <summary>
@@ -90,11 +104,13 @@ namespace UnrealBuildTool
 		/// </summary>
 		public static void Stop()
 		{
+#if UBT_PREFETCH_ENABLED
 			foreach (var  directory in QueuedDirectories.Values)
 			{
 				directory.Value.Wait();
 			}
 			QueuedDirectories.Clear();
+#endif
 		}
 
 		/// <summary>
diff --git a/Engine/Source/Programs/UnrealBuildTool/System/TargetMakefile.cs b/Engine/Source/Programs/UnrealBuildTool/System/TargetMakefile.cs
--- a/Engine/Source/Programs/UnrealBuildTool/System/TargetMakefile.cs
+++ b/Engine/Source/Programs/UnrealBuildTool/System/TargetMakefile.cs
@@ -69,7 +69,7 @@ namespace UnrealBuildTool
 		private readonly ILogger Logger;
 		public TargetMakefile Makefile { get; }
 
-		private List<Task> WriteTasks = new();
+		private ConcurrentBag<Task> WriteTasks = new();
 
 		public TargetMakefileBuilder(TargetMakefile InMakefile, ILogger InLogger)
 		{
@@ -80,12 +80,21 @@ namespace UnrealBuildTool
 		/// <inheritdoc/>
 		public void AddAction(IExternalAction Action)
 		{
-			Makefile.Actions.Add(Action);
+			lock (Makefile.Actions)
+			{
+				Makefile.Actions.Add(Action);
+			}
 		}
 
 		/// <inheritdoc/>
 		public void CreateIntermediateTextFile(FileItem FileItem, string Contents, bool AllowAsync = true)
 		{
+			lock (Makefile.InternalDependencies)
+			{
+				if (!Makefile.InternalDependencies.Add(FileItem))
+					throw new Exception("Writing same file multiple times?");
+			}
+
 			if (AllowAsync)
 			{
 				WriteTasks.Add(Task.Run(() =>
@@ -97,13 +106,16 @@ namespace UnrealBuildTool
 			{
 				Utils.WriteFileIfChanged(FileItem.Location, Contents, Logger);
 			}
-
-			Makefile.InternalDependencies.Add(FileItem);
 		}
 
 		/// <inheritdoc/>
 		public void CreateIntermediateTextFile(FileItem FileItem, IEnumerable<string> ContentLines, bool AllowAsync = true)
 		{
+			lock (Makefile.InternalDependencies)
+			{
+				if (!Makefile.InternalDependencies.Add(FileItem))
+					throw new Exception("Writing same file multiple times?");
+			}
 			if (AllowAsync)
 			{
 				WriteTasks.Add(Task.Run(() =>
@@ -115,53 +127,72 @@ namespace UnrealBuildTool
 			{
 				Utils.WriteFileIfChanged(FileItem.Location, ContentLines, Logger);
 			}
-			// Reset the file info, in case it already knows about the old file
-			Makefile.InternalDependencies.Add(FileItem);
 		}
 
 		/// <inheritdoc/>
 		public void AddSourceDir(DirectoryItem SourceDir)
 		{
-			Makefile.SourceDirectories.Add(SourceDir);
+			lock (Makefile.SourceDirectories)
+			{
+				Makefile.SourceDirectories.Add(SourceDir);
+			}
 		}
 
 		/// <inheritdoc/>
 		public void AddSourceFiles(DirectoryItem SourceDir, FileItem[] SourceFiles)
 		{
-			Makefile.DirectoryToSourceFiles[SourceDir] = SourceFiles.Select(fi => new TargetMakefileSourceFileInfo(fi)).ToArray();
+			lock (Makefile.DirectoryToSourceFiles)
+			{
+				Makefile.DirectoryToSourceFiles[SourceDir] = SourceFiles.Select(fi => new TargetMakefileSourceFileInfo(fi)).ToArray();
+			}
 		}
 
 		/// <inheritdoc/>
 		public virtual void AddHeaderFiles(FileItem[] HeaderFiles)
 		{
-			Makefile.HeaderFiles.UnionWith(HeaderFiles);
+			lock (Makefile.HeaderFiles)
+			{
+				Makefile.HeaderFiles.UnionWith(HeaderFiles);
+			}
 		}
 
 		/// <inheritdoc/>
 		public void AddDiagnostic(string Message)
 		{
-			if (!Makefile.Diagnostics.Contains(Message))
+			lock (Makefile.Diagnostics)
 			{
-				Makefile.Diagnostics.Add(Message);
+				if (!Makefile.Diagnostics.Contains(Message))
+				{
+					Makefile.Diagnostics.Add(Message);
+				}
 			}
 		}
 
 		/// <inheritdoc/>
 		public void AddFileToWorkingSet(FileItem File)
 		{
-			Makefile.WorkingSet.Add(File);
+			lock (Makefile.WorkingSet)
+			{
+				Makefile.WorkingSet.Add(File);
+			}
 		}
 
 		/// <inheritdoc/>
 		public void AddCandidateForWorkingSet(FileItem File)
 		{
-			Makefile.CandidatesForWorkingSet.Add(File);
+			lock (Makefile.CandidatesForWorkingSet)
+			{
+				Makefile.CandidatesForWorkingSet.Add(File);
+			}
 		}
 
 		/// <inheritdoc/>
 		public void SetOutputItemsForModule(string ModuleName, FileItem[] OutputItems)
 		{
-			Makefile.ModuleNameToOutputItems[ModuleName] = OutputItems;
+			lock (Makefile.ModuleNameToOutputItems)
+			{
+				Makefile.ModuleNameToOutputItems[ModuleName] = OutputItems;
+			}
 		}
 
 		internal void WaitOnWriteTasks()
@@ -711,6 +742,7 @@ namespace UnrealBuildTool
 			TargetMakefile Makefile;
 			using (GlobalTracer.Instance.BuildSpan("Loading makefile").StartActive())
 			{
+				using var _ = Timeline.ScopeEvent("LoadMakefile");
 				try
 				{
 					using (BinaryArchiveReader Reader = new BinaryArchiveReader(MakefilePath))
@@ -791,6 +823,8 @@ namespace UnrealBuildTool
 
 				// Check to see if the number of plugins changed
 				{
+					using var _ = Timeline.ScopeEvent("FindPlugins");
+
 					FileReference[] Plugins = EnumPluginsTask.Result;
 
 					if ((Plugins != null) != (Makefile.UbtPlugins != null) ||
