[UBT]
* Fixed so PluginInfos are created in parallel

diff --git a/Engine/Source/Programs/UnrealBuildTool/System/Plugins.cs b/Engine/Source/Programs/UnrealBuildTool/System/Plugins.cs
--- a/Engine/Source/Programs/UnrealBuildTool/System/Plugins.cs
+++ b/Engine/Source/Programs/UnrealBuildTool/System/Plugins.cs
@@ -632,21 +632,18 @@ namespace UnrealBuildTool
 				Dictionary<PluginInfo, FileReference> ChildPlugins = new Dictionary<PluginInfo, FileReference>();
 				List<PluginInfo> AllParentPlugins = new List<PluginInfo>();
 
-				foreach (DirectoryReference Dir in RootDirectories)
-				{
-					if (!DirectoryReference.Exists(Dir))
-					{
-						continue;
-					}
 
-					List<FileReference> PluginsInDir = PluginsBase.EnumeratePlugins(Dir);
-					ConcurrentDictionary<FileReference, PluginInfo> CreatedInfos = [];
-					Parallel2.ForEach(PluginsInDir, (PluginFileName) =>
-					{
-						CreatedInfos.TryAdd(PluginFileName, new PluginInfo(PluginFileName, Type));
-					});
+				// Create all in parallel
+				ConcurrentDictionary<FileReference, PluginInfo> CreatedInfos = [];
+				Parallel2.ForEach(RootDirectories.SelectMany(Dir => PluginsBase.EnumeratePlugins(Dir)), (PluginFileName) =>
+				{
+					CreatedInfos.TryAdd(PluginFileName, new PluginInfo(PluginFileName, Type));
+				});
 
-					foreach (FileReference PluginFileName in PluginsInDir)
+				// Then add them properly sorted
+				foreach (DirectoryReference Dir in RootDirectories)
+				{
+					foreach (FileReference PluginFileName in PluginsBase.EnumeratePlugins(Dir))
 					{
 						PluginInfo Plugin = CreatedInfos[PluginFileName];
 
