[UBT]
* Do multiple tests in parallel when checking if makefile needs to be recreated

diff --git a/Engine/Source/Programs/UnrealBuildTool/System/TargetMakefile.cs b/Engine/Source/Programs/UnrealBuildTool/System/TargetMakefile.cs
--- a/Engine/Source/Programs/UnrealBuildTool/System/TargetMakefile.cs
+++ b/Engine/Source/Programs/UnrealBuildTool/System/TargetMakefile.cs
@@ -872,130 +872,206 @@ namespace UnrealBuildTool
 		/// <returns>True if the makefile is valid, false otherwise</returns>
 		public static bool IsValidForSourceFiles(TargetMakefile Makefile, FileReference? ProjectFile, UnrealTargetPlatform Platform, ISourceFileWorkingSet WorkingSet, ILogger Logger, out string? ReasonNotLoaded)
 		{
+			using var _ = Timeline.ScopeEvent("MakefileValid");
+
 			using (GlobalTracer.Instance.BuildSpan("TargetMakefile.IsValidForSourceFiles()").StartActive())
 			{
 				// Get the list of excluded folder names for this platform
 				IReadOnlySet<string> ExcludedFolderNames = UEBuildPlatform.GetBuildPlatform(Platform).GetExcludedFolderNames();
 
-				// Load the metadata cache
-				SourceFileMetadataCache MetadataCache = SourceFileMetadataCache.CreateHierarchy(ProjectFile, Logger);
+				string reasonNotLoaded1 = "";
+
+				SourceFileMetadataCache? MetadataCache = null;
 
-				// Check if any source files have been added or removed
-				foreach ((DirectoryItem InputDirectory, TargetMakefileSourceFileInfo[] SourceFileInfos) in Makefile.DirectoryToSourceFiles)
+				Task<bool> checkTask = Task.Run(() =>
 				{
-					if (!InputDirectory.Exists || InputDirectory.LastWriteTimeUtc > Makefile.CreateTimeUtc)
+
+					// Load the metadata cache
+					MetadataCache = SourceFileMetadataCache.CreateHierarchy(ProjectFile, Logger);
+
+					// Check if any source files have been added or removed
+					foreach ((DirectoryItem InputDirectory, TargetMakefileSourceFileInfo[] SourceFileInfos) in Makefile.DirectoryToSourceFiles)
 					{
-						FileItem[] SourceFiles = UEBuildModuleCPP.GetSourceFiles(InputDirectory, Logger);
-						if (SourceFiles.Length < SourceFileInfos.Length)
+						if (!InputDirectory.Exists || InputDirectory.LastWriteTimeUtc > Makefile.CreateTimeUtc)
 						{
-							ReasonNotLoaded = "source file removed";
-							return false;
+							FileItem[] SourceFiles = UEBuildModuleCPP.GetSourceFiles(InputDirectory, Logger);
+							if (SourceFiles.Length < SourceFileInfos.Length)
+							{
+								reasonNotLoaded1 = "source file removed";
+								return false;
+							}
+							else if (SourceFiles.Length > SourceFileInfos.Length)
+							{
+								reasonNotLoaded1 = "source file added";
+								return false;
+							}
+							else if (SourceFiles.Intersect(SourceFileInfos.Select((sfi) => sfi.SourceFileItem)).Count() != SourceFiles.Length)
+							{
+								reasonNotLoaded1 = "source file modified";
+								return false;
+							}
+							else if (!InputDirectory.Exists)
+							{
+								reasonNotLoaded1 = "source directory removed";
+								return false;
+							}
+
+							foreach (DirectoryItem Directory in InputDirectory.EnumerateDirectories())
+							{
+								if (!Makefile.DirectoryToSourceFiles.ContainsKey(Directory) && ContainsSourceFilesOrHeaders(Directory, ExcludedFolderNames, Logger))
+								{
+									reasonNotLoaded1 = "source directory added";
+									return false;
+								}
+							}
 						}
-						else if (SourceFiles.Length > SourceFileInfos.Length)
+
+						// Make sure the inlined gen.cpp files didn't change
+						foreach (TargetMakefileSourceFileInfo SourceFileInfo in SourceFileInfos)
 						{
-							ReasonNotLoaded = "source file added";
-							return false;
+							int InlineHash = TargetMakefileSourceFileInfo.CalculateInlineGenCppHash(MetadataCache.GetListOfInlinedGeneratedCppFiles(SourceFileInfo.SourceFileItem!));
+							if (SourceFileInfo.InlineGenCppHash != InlineHash)
+							{
+								reasonNotLoaded1 = "inlined gen.cpp files changed";
+								return false;
+							}
 						}
-						else if (SourceFiles.Intersect(SourceFileInfos.Select((sfi) => sfi.SourceFileItem)).Count() != SourceFiles.Length)
+					}
+
+					// Check if any external dependencies have changed. These comparisons are done against the makefile creation time.
+					foreach (FileItem ExternalDependency in Makefile.ExternalDependencies)
+					{
+						if (!ExternalDependency.Exists)
 						{
-							ReasonNotLoaded = "source file modified";
+							Logger.LogDebug("{File} has been deleted since makefile was built.", ExternalDependency.Location);
+							reasonNotLoaded1 = String.Format("{0} deleted", ExternalDependency.Location.GetFileName());
 							return false;
 						}
-						else if (!InputDirectory.Exists)
+						if (ExternalDependency.LastWriteTimeUtc > Makefile.CreateTimeUtc)
 						{
-							ReasonNotLoaded = "source directory removed";
+							Logger.LogDebug("{File} has been modified since makefile was built.", ExternalDependency.Location);
+							reasonNotLoaded1 = String.Format("{0} modified", ExternalDependency.Location.GetFileName());
 							return false;
 						}
+					}
+
+					return true;
+				});
+
 
-						foreach (DirectoryItem Directory in InputDirectory.EnumerateDirectories())
+				bool isValid2 = true;
+				object isValidLock2 = new();
+				string reasonNotLoaded2 = "";
+
+				// First collect file timestamps by caching directories. This is much faster than letting each file do a GetFileAttribute
+				List<FileItem> fileItems = new(Makefile.InternalDependencies);
+				fileItems.Sort(static (a, b) => a.Location.FullName.CompareTo(b.Location.FullName));
+				List<DirectoryItem> dirs = new();
+				DirectoryItem? lastDir = null;
+				foreach (FileItem item in fileItems)
+				{
+					DirectoryItem dir = item.Directory;
+					if (dir != lastDir)
+					{
+						dirs.Add(dir);
+						lastDir = dir;
+					}
+				}
+				Parallel2.ForEach(dirs, dir =>
+				{
+					dir.CacheFiles();
+				});
+
+				// Check if any internal dependencies has changed. These comparisons are done against the makefile modified time.
+				Parallel2.ForEach(Makefile.InternalDependencies, InternalDependency =>
+				{
+					if (!InternalDependency.Exists && isValid2)
+					{
+						lock (isValidLock2)
 						{
-							if (!Makefile.DirectoryToSourceFiles.ContainsKey(Directory) && ContainsSourceFilesOrHeaders(Directory, ExcludedFolderNames, Logger))
+							if (isValid2)
 							{
-								ReasonNotLoaded = "source directory added";
-								return false;
+								isValid2 = false;
+								Logger.LogDebug("{File} has been deleted since makefile was written.", InternalDependency.Location);
+								reasonNotLoaded2 = String.Format("{0} deleted", InternalDependency.Location.GetFileName());
 							}
 						}
 					}
-
-					// Make sure the inlined gen.cpp files didn't change
-					foreach (TargetMakefileSourceFileInfo SourceFileInfo in SourceFileInfos)
+					if (InternalDependency.LastWriteTimeUtc > Makefile.ModifiedTimeUtc && isValid2)
 					{
-						int InlineHash = TargetMakefileSourceFileInfo.CalculateInlineGenCppHash(MetadataCache.GetListOfInlinedGeneratedCppFiles(SourceFileInfo.SourceFileItem!));
-						if (SourceFileInfo.InlineGenCppHash != InlineHash)
+						lock (isValidLock2)
 						{
-							ReasonNotLoaded = "inlined gen.cpp files changed";
-							return false;
+							if (isValid2)
+							{
+								isValid2 = false;
+								Logger.LogDebug("{File} has been modified since makefile was written.", InternalDependency.Location);
+								reasonNotLoaded2 = String.Format("{0} modified", InternalDependency.Location.GetFileName());
+							}
 						}
 					}
-				}
+				});
 
-				// Check if any external dependencies have changed. These comparisons are done against the makefile creation time.
-				foreach (FileItem ExternalDependency in Makefile.ExternalDependencies)
+
+				string reasonNotLoaded3 = "";
+
+				bool Func3()
 				{
-					if (!ExternalDependency.Exists)
-					{
-						Logger.LogDebug("{File} has been deleted since makefile was built.", ExternalDependency.Location);
-						ReasonNotLoaded = String.Format("{0} deleted", ExternalDependency.Location.GetFileName());
-						return false;
-					}
-					if (ExternalDependency.LastWriteTimeUtc > Makefile.CreateTimeUtc)
+					// Check if available Platform SDKs have changed
+					foreach (KeyValuePair<string, string> Pair in Makefile.PlatformSDKInfo)
 					{
-						Logger.LogDebug("{File} has been modified since makefile was built.", ExternalDependency.Location);
-						ReasonNotLoaded = String.Format("{0} modified", ExternalDependency.Location.GetFileName());
-						return false;
+						string value = UEBuildPlatformSDK.GetSDKForPlatform(Pair.Key)?.GetAllSDKInfo().ToString() ?? String.Empty;
+						if (Pair.Value != value)
+						{
+							Logger.LogDebug("Platform SDK {Key} '{Value}' has been modified to '{NewValue}' since makefile was written.", Pair.Key, Pair.Value, !String.IsNullOrEmpty(value) ? value : "Missing");
+							reasonNotLoaded3 = String.Format("Platform SDK {0} modified", Pair.Key);
+							return false;
+						}
 					}
-				}
 
-				// Check if any internal dependencies has changed. These comparisons are done against the makefile modified time.
-				foreach (FileItem InternalDependency in Makefile.InternalDependencies)
-				{
-					if (!InternalDependency.Exists)
+					IEnumerable<DirectoryItem> UbaPaths = [.. Makefile.UbaDisallowedPaths, .. Makefile.UbaPathHashes.Keys, .. Makefile.UbaCrossArchitecturePaths.Keys, .. Makefile.UbaCrossArchitecturePaths.Values];
+					foreach (DirectoryItem item in UbaPaths.Distinct())
 					{
-						Logger.LogDebug("{File} has been deleted since makefile was written.", InternalDependency.Location);
-						ReasonNotLoaded = String.Format("{0} deleted", InternalDependency.Location.GetFileName());
-						return false;
+						if (!item.Exists)
+						{
+							Logger.LogDebug("{Directory} has been deleted since makefile was written.", item.Location);
+							reasonNotLoaded3 = String.Format("{0} deleted", item.Location);
+							return false;
+						}
 					}
-					if (InternalDependency.LastWriteTimeUtc > Makefile.ModifiedTimeUtc)
+
+					// Check that no new plugins have been added
+					foreach (FileReference PluginFile in PluginsBase.EnumeratePlugins(ProjectFile))
 					{
-						Logger.LogDebug("{File} has been modified since makefile was written.", InternalDependency.Location);
-						ReasonNotLoaded = String.Format("{0} modified", InternalDependency.Location.GetFileName());
-						return false;
+						FileItem PluginFileItem = FileItem.GetItemByFileReference(PluginFile);
+						if (!Makefile.PluginFiles.Contains(PluginFileItem))
+						{
+							Logger.LogDebug("{File} has been added", PluginFile.GetFileName());
+							reasonNotLoaded3 = String.Format("{0} has been added", PluginFile.GetFileName());
+							return false;
+						}
 					}
+					return true;
 				}
 
-				// Check if available Platform SDKs have changed
-				foreach (KeyValuePair<string, string> Pair in Makefile.PlatformSDKInfo)
+				bool isValid3 = Func3();
+
+				bool isValid1 = checkTask.Result;
+				if (!isValid1)
 				{
-					string value = UEBuildPlatformSDK.GetSDKForPlatform(Pair.Key)?.GetAllSDKInfo().ToString() ?? String.Empty;
-					if (Pair.Value != value)
-					{
-						Logger.LogDebug("Platform SDK {Key} '{Value}' has been modified to '{NewValue}' since makefile was written.", Pair.Key, Pair.Value, !String.IsNullOrEmpty(value) ? value : "Missing");
-						ReasonNotLoaded = String.Format("Platform SDK {0} modified", Pair.Key);
-						return false;
-					}
+					ReasonNotLoaded = reasonNotLoaded1;
+					return false;
 				}
 
-				IEnumerable<DirectoryItem> UbaPaths = [.. Makefile.UbaDisallowedPaths, .. Makefile.UbaPathHashes.Keys, .. Makefile.UbaCrossArchitecturePaths.Keys, .. Makefile.UbaCrossArchitecturePaths.Values];
-				foreach (DirectoryItem item in UbaPaths.Distinct())
+				if (!isValid2)
 				{
-					if (!item.Exists)
-					{
-						Logger.LogDebug("{Directory} has been deleted since makefile was written.", item.Location);
-						ReasonNotLoaded = String.Format("{0} deleted", item.Location);
-						return false;
-					}
+					ReasonNotLoaded = reasonNotLoaded2;
+					return false;
 				}
 
-				// Check that no new plugins have been added
-				foreach (FileReference PluginFile in PluginsBase.EnumeratePlugins(ProjectFile))
+				if (!isValid3)
 				{
-					FileItem PluginFileItem = FileItem.GetItemByFileReference(PluginFile);
-					if (!Makefile.PluginFiles.Contains(PluginFileItem))
-					{
-						Logger.LogDebug("{File} has been added", PluginFile.GetFileName());
-						ReasonNotLoaded = String.Format("{0} has been added", PluginFile.GetFileName());
-						return false;
-					}
+					ReasonNotLoaded = reasonNotLoaded3;
+					return false;
 				}
 
 				// Find the set of files that contain reflection markup
@@ -1004,7 +1080,7 @@ namespace UnrealBuildTool
 				{
 					foreach (DirectoryItem SourceDirectory in Makefile.SourceDirectories)
 					{
-						Queue.Enqueue(() => FindFilesWithMarkup(SourceDirectory, MetadataCache, ExcludedFolderNames, NewFilesWithMarkupBag, Queue));
+						Queue.Enqueue(() => FindFilesWithMarkup(SourceDirectory, MetadataCache!, ExcludedFolderNames, NewFilesWithMarkupBag, Queue));
 					}
 				}
 
