[UBT]
* Optimized UHT out-of-date checking by running things in parallel and also cache directories to reduce number of GetFileAttributes calls

diff --git a/Engine/Source/Programs/UnrealBuildTool/System/UHTExecution.cs b/Engine/Source/Programs/UnrealBuildTool/System/UHTExecution.cs
--- a/Engine/Source/Programs/UnrealBuildTool/System/UHTExecution.cs
+++ b/Engine/Source/Programs/UnrealBuildTool/System/UHTExecution.cs
@@ -11,6 +11,7 @@ using System.Linq;
 using System.Reflection;
 using System.Text;
 using System.Text.Json;
+using System.Threading;
 using System.Threading.Tasks;
 using EpicGames.Core;
 using EpicGames.UHT.Utils;
@@ -842,64 +843,91 @@ namespace UnrealBuildTool
 			// that check for the engine version
 			DateTime EngineVersionHeaderTimestampUtc = GetEngineVersionHeaderTimestampUtc();
 
-			foreach (UHTModuleInfo Module in UObjectModules)
+			int invalidCounter = 0;
+
+			Parallel2.ForEach(UObjectModules, Module =>
 			{
+				if (invalidCounter != 0)
+				{
+					return;
+				}
+
 				// If we're using a precompiled engine, skip checking timestamps for modules that are under the engine directory
 				if (Module.bIsReadOnly)
 				{
-					continue;
+					return;
 				}
 
 				// Make sure we have an existing folder for generated code.  If not, then we definitely need to generate code!
-				string GeneratedCodeDirectory = Path.GetDirectoryName(Module.GeneratedCPPFilenameBase)!;
-				FileSystemInfo TestDirectory = (FileSystemInfo)new DirectoryInfo(GeneratedCodeDirectory);
-				if (!TestDirectory.Exists)
+				//string GeneratedCodeDirectory = Path.GetDirectoryName(Module.GeneratedCPPFilenameBase)!;
+				DirectoryItem GeneratedCodeDirectory = Module.GeneratedCodeDirectory;
+
+				//FileSystemInfo TestDirectory = (FileSystemInfo)new DirectoryInfo(GeneratedCodeDirectory);
+				if (!GeneratedCodeDirectory.Exists)
 				{
 					// Generated code directory is missing entirely!
-					Logger.LogDebug("UnrealHeaderTool needs to run because no generated code directory was found for module {ModuleName}", Module.ModuleName);
-					return true;
+					if (Interlocked.Increment(ref invalidCounter) == 1)
+					{
+						Logger.LogDebug("UnrealHeaderTool needs to run because no generated code directory was found for module {ModuleName}", Module.ModuleName);
+					}
+					return;
 				}
+				GeneratedCodeDirectory.CacheFiles();
 
 				// Grab our special "Timestamp" file that we saved after the last set of headers were generated.  This file
 				// contains the list of source files which contained UObjects, so that we can compare to see if any
 				// UObject source files were deleted (or no longer contain UObjects), which means we need to run UHT even
 				// if no other source files were outdated
-				string TimestampFile = Path.Combine(GeneratedCodeDirectory, @"Timestamp");
-				FileSystemInfo SavedTimestampFileInfo = (FileSystemInfo)new FileInfo(TimestampFile);
-				if (!SavedTimestampFileInfo.Exists)
+				FileItem TimestampFile = FileItem.GetItemByFileReference(FileReference.Combine(GeneratedCodeDirectory.Location, "Timestamp"));
+				if (!TimestampFile.Exists)
 				{
 					// Timestamp file was missing (possibly deleted/cleaned), so headers are out of date
-					Logger.LogDebug("UnrealHeaderTool needs to run because UHT Timestamp file did not exist for module {ModuleName}", Module.ModuleName);
-					return true;
+					if (Interlocked.Increment(ref invalidCounter) == 1)
+					{
+						Logger.LogDebug("UnrealHeaderTool needs to run because UHT Timestamp file did not exist for module {ModuleName}", Module.ModuleName);
+					}
+					return;
 				}
 
 				// Make sure the last UHT run completed after UnrealHeaderTool.exe was compiled last, and after the CoreUObject headers were touched last.
-				DateTime SavedTimestampUtc = SavedTimestampFileInfo.LastWriteTimeUtc;
+				DateTime SavedTimestampUtc = TimestampFile.LastWriteTimeUtc;
 				if (HeaderToolTimestampUtc > SavedTimestampUtc)
 				{
 					// Generated code is older than UnrealHeaderTool.exe.  Out of date!
-					Logger.LogDebug("UnrealHeaderTool needs to run because UnrealHeaderTool timestamp ({Time}) is later than timestamp for module {ModuleName} ({ModuleTime})", HeaderToolTimestampUtc.ToLocalTime(), Module.ModuleName, SavedTimestampUtc.ToLocalTime());
-					return true;
+					if (Interlocked.Increment(ref invalidCounter) == 1)
+					{
+						Logger.LogDebug("UnrealHeaderTool needs to run because UnrealHeaderTool timestamp ({Time}) is later than timestamp for module {ModuleName} ({ModuleTime})", HeaderToolTimestampUtc.ToLocalTime(), Module.ModuleName, SavedTimestampUtc.ToLocalTime());
+					}
+					return;
 				}
 				if (CoreGeneratedTimestampUtc > SavedTimestampUtc)
 				{
 					// Generated code is older than CoreUObject headers.  Out of date!
-					Logger.LogDebug("UnrealHeaderTool needs to run because CoreUObject timestamp ({Time}) is newer than timestamp for module {ModuleName} ({ModuleTime})", CoreGeneratedTimestampUtc.Value.ToLocalTime(), Module.ModuleName, SavedTimestampUtc.ToLocalTime());
-					return true;
+					if (Interlocked.Increment(ref invalidCounter) == 1)
+					{
+						Logger.LogDebug("UnrealHeaderTool needs to run because CoreUObject timestamp ({Time}) is newer than timestamp for module {ModuleName} ({ModuleTime})", CoreGeneratedTimestampUtc.Value.ToLocalTime(), Module.ModuleName, SavedTimestampUtc.ToLocalTime());
+					}
+					return;
 				}
 				if (EngineVersionHeaderTimestampUtc > SavedTimestampUtc)
 				{
 					// Generated code is older than engine version header.  Out of date!
-					Logger.LogDebug("UnrealHeaderTool needs to run because engine version header timestamp ({Time}) is newer than timestamp for module {ModuleName} ({ModuleTime})", EngineVersionHeaderTimestampUtc.ToLocalTime(), Module.ModuleName, SavedTimestampUtc.ToLocalTime());
-					return true;
+					if (Interlocked.Increment(ref invalidCounter) == 1)
+					{
+						Logger.LogDebug("UnrealHeaderTool needs to run because engine version header timestamp ({Time}) is newer than timestamp for module {ModuleName} ({ModuleTime})", EngineVersionHeaderTimestampUtc.ToLocalTime(), Module.ModuleName, SavedTimestampUtc.ToLocalTime());
+					}
+					return;
 				}
 
 				// Has the .build.cs file changed since we last generated headers successfully?
 				FileInfo ModuleRulesFile = new(Module.ModuleRulesFile.FullName);
 				if (!ModuleRulesFile.Exists || ModuleRulesFile.LastWriteTimeUtc > SavedTimestampUtc)
 				{
-					Logger.LogDebug("UnrealHeaderTool needs to run because SavedTimestamp is older than the rules file ({ModuleModuleRulesFile}) for module {ModuleModuleName}", Module.ModuleRulesFile, Module.ModuleName);
-					return true;
+					if (Interlocked.Increment(ref invalidCounter) == 1)
+					{
+						Logger.LogDebug("UnrealHeaderTool needs to run because SavedTimestamp is older than the rules file ({ModuleModuleRulesFile}) for module {ModuleModuleName}", Module.ModuleRulesFile, Module.ModuleName);
+					}
+					return;
 				}
 
 				// Iterate over our UObjects headers and figure out if any of them have changed
@@ -913,11 +941,14 @@ namespace UnrealBuildTool
 
 				// Load up the old timestamp file and check to see if anything has changed
 				{
-					string[] UObjectFilesFromPreviousRun = File.ReadAllLines(TimestampFile);
+					string[] UObjectFilesFromPreviousRun = File.ReadAllLines(TimestampFile.FullName);
 					if (AllUObjectHeaders.Count != UObjectFilesFromPreviousRun.Length)
 					{
-						Logger.LogDebug("UnrealHeaderTool needs to run because there are a different number of UObject source files in module {ModuleModuleName}", Module.ModuleName);
-						return true;
+						if (Interlocked.Increment(ref invalidCounter) == 1)
+						{
+							Logger.LogDebug("UnrealHeaderTool needs to run because there are a different number of UObject source files in module {ModuleModuleName}", Module.ModuleName);
+						}
+						return;
 					}
 
 					// Iterate over our UObjects headers and figure out if any of them have changed
@@ -926,8 +957,11 @@ namespace UnrealBuildTool
 					{
 						if (!ObjectHeadersSet.Contains(FileName))
 						{
-							Logger.LogDebug("UnrealHeaderTool needs to run because the set of UObject source files in module {ModuleModuleName} has changed ({FileName})", Module.ModuleName, FileName);
-							return true;
+							if (Interlocked.Increment(ref invalidCounter) == 1)
+							{
+								Logger.LogDebug("UnrealHeaderTool needs to run because the set of UObject source files in module {ModuleModuleName} has changed ({FileName})", Module.ModuleName, FileName);
+							}
+							return;
 						}
 					}
 				}
@@ -939,13 +973,16 @@ namespace UnrealBuildTool
 					// Has the source header changed since we last generated headers successfully?
 					if (HeaderFileTimestampUtc > SavedTimestampUtc)
 					{
-						Logger.LogDebug("UnrealHeaderTool needs to run because SavedTimestamp is older than HeaderFileTimestamp ({HeaderFileAbsolutePath}) for module {ModuleModuleName}", HeaderFile.AbsolutePath, Module.ModuleName);
-						return true;
+						if (Interlocked.Increment(ref invalidCounter) == 1)
+						{
+							Logger.LogDebug("UnrealHeaderTool needs to run because SavedTimestamp is older than HeaderFileTimestamp ({HeaderFileAbsolutePath}) for module {ModuleModuleName}", HeaderFile.AbsolutePath, Module.ModuleName);
+						}
+						return;
 					}
 				}
-			}
+			});
 
-			return false;
+			return invalidCounter != 0;
 		}
 
 		/// <summary>
@@ -983,7 +1020,7 @@ namespace UnrealBuildTool
 		/// </summary>
 		private static void UpdateTimestamps(List<UHTModuleInfo> UObjectModules, bool resetDirectoryCaches)
 		{
-			Parallel.ForEach(UObjectModules, async Module =>
+			Parallel2.ForEach(UObjectModules, async Module =>
 			{
 				if (!Module.bIsReadOnly)
 				{
