[UBT]
* Fixed bug where logic adding generated files from uht/vni was using FileReference instead of FileItem to find files. Since we prepopulate FileItem and then create files in parallel it does not work to use Filereference (which goes directly on filesystem)

diff --git a/Engine/Source/Programs/UnrealBuildTool/Configuration/UEBuildBinary.cs b/Engine/Source/Programs/UnrealBuildTool/Configuration/UEBuildBinary.cs
--- a/Engine/Source/Programs/UnrealBuildTool/Configuration/UEBuildBinary.cs
+++ b/Engine/Source/Programs/UnrealBuildTool/Configuration/UEBuildBinary.cs
@@ -512,14 +512,28 @@ namespace UnrealBuildTool
 			{
 				if (Module.Rules.bPrecompile)
 				{
-					if (Module.GeneratedCodeDirectory != null && DirectoryReference.Exists(Module.GeneratedCodeDirectory))
+					DirectoryItem? GeneratedCodeDirectory = null;
+					if (Module.GeneratedCodeDirectory != null)
 					{
-						foreach (FileReference GeneratedCodeFile in DirectoryReference.EnumerateFiles(Module.GeneratedCodeDirectory, "*", SearchOption.AllDirectories))
+						GeneratedCodeDirectory = DirectoryItem.GetItemByDirectoryReference(Module.GeneratedCodeDirectory);
+					}
+					if (GeneratedCodeDirectory != null && GeneratedCodeDirectory.Exists)
+					{
+						Queue<DirectoryItem> DirQueue = new();
+						DirQueue.Enqueue(GeneratedCodeDirectory);
+						while (DirQueue.TryDequeue(out DirectoryItem? dir))
 						{
-							// Exclude timestamp files, since they're always updated and cause collisions between builds
-							if (!GeneratedCodeFile.GetFileName().Equals("Timestamp", StringComparison.OrdinalIgnoreCase) && !GeneratedCodeFile.HasExtension(".cpp"))
+							foreach (FileItem GeneratedCodeFile in dir.EnumerateFiles())
+							{
+								// Exclude timestamp files, since they're always updated and cause collisions between builds
+								if (!GeneratedCodeFile.Name.Equals("Timestamp", StringComparison.OrdinalIgnoreCase) && !GeneratedCodeFile.HasExtension(".cpp"))
+								{
+									BuildProducts.Add(GeneratedCodeFile.Location, BuildProductType.BuildResource);
+								}
+							}
+							foreach (DirectoryItem subdir in dir.EnumerateDirectories())
 							{
-								BuildProducts.Add(GeneratedCodeFile, BuildProductType.BuildResource);
+								DirQueue.Enqueue(subdir);
 							}
 						}
 					}
