UBT: Restore timeline printing (to the log) but with a MaxDepth of 4 to print ~130 lines in <10ms when generating projects.

#rnx

diff --git a/Engine/Source/Programs/UnrealBuildTool/System/Timeline.cs b/Engine/Source/Programs/UnrealBuildTool/System/Timeline.cs
--- a/Engine/Source/Programs/UnrealBuildTool/System/Timeline.cs
+++ b/Engine/Source/Programs/UnrealBuildTool/System/Timeline.cs
@@ -199,11 +199,12 @@ namespace UnrealBuildTool
 		/// <summary>
 		/// Prints this information to the log
 		/// </summary>
-		/// <param name="MaxUnknownTime">Non-instrumented time limit for inserting an "unknown" event</param>
 		/// <param name="MinKnownTime">Instrumented time limit for printing the labelled event at given verbosity. May be null.</param>
+		/// <param name="MaxUnknownTime">Non-instrumented time limit for inserting an "unknown" event</param>
+		/// <param name="MaxDepth">Maximum tree depth to print.</param>
 		/// <param name="Verbosity">LogLevel of the event.</param>
 		/// <param name="Logger">Logger.</param>
-		public static void Print(TimeSpan? MinKnownTime, TimeSpan MaxUnknownTime, LogLevel Verbosity, ILogger Logger)
+		public static void Print(TimeSpan? MinKnownTime, TimeSpan MaxUnknownTime, int? MaxDepth, LogLevel Verbosity, ILogger Logger)
 		{
 			// Print the start time
 			Logger.Log(Verbosity, "");
@@ -225,6 +226,17 @@ namespace UnrealBuildTool
 				{
 					Event Event = Events[EventIdx];
 
+					// Skip inner events when MaxDepth has been reached
+					if (MaxDepth != null && OuterEvents.Count == MaxDepth)
+					{
+						Event OuterEvent = OuterEvents.Last();
+						if (Event.StartTime < OuterEvent.FinishTime!.Value)
+						{
+							continue;
+						}
+						LastTime = OuterEvent.FinishTime.Value;
+					}
+
 					// Pop events off the stack
 					for (; OuterEvents.Count > 1; OuterEvents.RemoveAt(OuterEvents.Count - 1))
 					{
diff --git a/Engine/Source/Programs/UnrealBuildTool/UnrealBuildTool.cs b/Engine/Source/Programs/UnrealBuildTool/UnrealBuildTool.cs
--- a/Engine/Source/Programs/UnrealBuildTool/UnrealBuildTool.cs
+++ b/Engine/Source/Programs/UnrealBuildTool/UnrealBuildTool.cs
@@ -736,8 +736,7 @@ namespace UnrealBuildTool
 				LogLevel ExecutionLogLevel = bShowExecutionTime ? LogLevel.Information : LogLevel.Debug;
 				LogLevel TimelineLogLevel = bShowExecutionTime && bShowTimeline && ApplicationResult == CompilationResult.Succeeded ?  LogLevel.Information : LogLevel.Debug;
 				
-				// Commented out for now. Takes 1 minute to run when generating vs solution
-				//Timeline.Print(TimeSpan.FromMilliseconds(100.0), TimeSpan.FromMilliseconds(200.0), TimelineLogLevel, Logger);
+				Timeline.Print(TimeSpan.FromMilliseconds(100.0), TimeSpan.FromMilliseconds(200.0), /*MaxDepth*/ 4, TimelineLogLevel, Logger);
 				
 				Logger.Log(ExecutionLogLevel, "");
 				if (ApplicationResult == CompilationResult.Succeeded)
