Removing the UHT testing code since it wasn't ever really used and hard to maintain.

#rnx

diff --git a/Engine/Source/Programs/Shared/UnrealBuildTool.Tests/UnrealHeaderToolTests.cs b/Engine/Source/Programs/Shared/UnrealBuildTool.Tests/UnrealHeaderToolTests.cs
deleted file mode 100644
diff --git a/Engine/Source/Programs/UnrealBuildTool/Modes/UnrealHeaderToolMode.cs b/Engine/Source/Programs/UnrealBuildTool/Modes/UnrealHeaderToolMode.cs
--- a/Engine/Source/Programs/UnrealBuildTool/Modes/UnrealHeaderToolMode.cs
+++ b/Engine/Source/Programs/UnrealBuildTool/Modes/UnrealHeaderToolMode.cs
@@ -11,7 +11,6 @@ using System.Collections.Generic;
 using System.Diagnostics;
 using System.IO;
 using System.Reflection;
-using System.Text;
 using System.Threading.Tasks;
 using UnrealBuildBase;
 
@@ -526,9 +525,6 @@ namespace UnrealBuildTool.Modes
 		[CommandLine(Prefix = "-NoLog", Description = "Disable log file creation including the default log file")]
 		public bool bNoLog = false;
 
-		[CommandLine(Prefix = "-Test", Description = "Run testing scripts")]
-		public bool bTest = false;
-
 		[CommandLine("-WarningsAsErrors", Description = "Treat warnings as errors")]
 		public bool bWarningsAsErrors = false;
 
@@ -563,485 +559,6 @@ namespace UnrealBuildTool.Modes
 		}
 	}
 
-	/// <summary>
-	/// File manager for the test harness
-	/// </summary>
-	/// <remarks>
-	/// Construct a new instance of the test file manager
-	/// </remarks>
-	/// <param name="RootDirectory">Root directory of the UE</param>
-	public class UhtTestFileManager(string RootDirectory) : IUhtFileManager
-	{
-		/// <summary>
-		/// Collection of test fragments that can be read
-		/// </summary>
-		public Dictionary<string, UhtSourceFragment> SourceFragments = [];
-
-		/// <summary>
-		/// All output segments generated by code gen
-		/// </summary>
-		public SortedDictionary<string, string> Outputs = [];
-
-		private readonly UhtStdFileManager InnerManager = new();
-		private readonly string? RootDirectory = RootDirectory;
-
-		/// <inheritdoc/>
-		public string GetFullFilePath(string FilePath)
-		{
-			if (RootDirectory == null)
-			{
-				return FilePath;
-			}
-			else
-			{
-				return Path.Combine(RootDirectory, FilePath);
-			}
-		}
-
-		/// <inheritdoc/>
-		public bool ReadSource(string FilePath, out UhtSourceFragment Fragment)
-		{
-			if (SourceFragments.TryGetValue(FilePath, out Fragment))
-			{
-				return true;
-			}
-
-			return InnerManager.ReadSource(GetFullFilePath(FilePath), out Fragment);
-		}
-
-		/// <inheritdoc/>
-		[Obsolete("Use the new ReadOutput with UhtPoolBuffer")]
-		public UhtBuffer? ReadOutput(string FilePath)
-		{
-			return null;
-		}
-
-		/// <inheritdoc/>
-		public bool ReadOutput(string FilePath, out UhtPoolBuffer<char> Output)
-		{
-			Output = default;
-			return false;
-		}
-
-		/// <inheritdoc/>
-		public bool WriteOutput(string FilePath, ReadOnlySpan<char> Contents)
-		{
-			lock (Outputs)
-			{
-				Outputs.Add(FilePath, Contents.ToString());
-			}
-			return true;
-		}
-
-		/// <inheritdoc/>
-		public bool RenameOutput(string OldFilePath, string NewFilePath)
-		{
-			lock (Outputs)
-			{
-				if (Outputs.TryGetValue(OldFilePath, out string? Content))
-				{
-					Outputs.Remove(OldFilePath);
-					Outputs.Add(NewFilePath, Content);
-				}
-			}
-			return true;
-		}
-
-		/// <summary>
-		/// Add a source file fragment to the session.  When requests are made to read sources, the 
-		/// fragment list will be searched first.
-		/// </summary>
-		/// <param name="SourceFile">Source file</param>
-		/// <param name="FilePath">The relative path to add</param>
-		/// <param name="LineNumber">Starting line number</param>
-		/// <param name="Data">The data associated with the path</param>
-		public void AddSourceFragment(UhtSourceFile SourceFile, string FilePath, int LineNumber, StringView Data)
-		{
-			SourceFragments.Add(FilePath, new UhtSourceFragment { SourceFile = SourceFile, FilePath = FilePath, LineNumber = LineNumber, Data = Data });
-		}
-	}
-
-	/// <summary>
-	/// Testing harness to run the test scripts
-	/// </summary>
-	class UhtTestHarness
-	{
-		private enum ScriptFragmentType
-		{
-			Unknown,
-			Manifest,
-			Header,
-			Console,
-			Output,
-		}
-
-		private struct ScriptFragment
-		{
-			public ScriptFragmentType Type;
-			public string Name;
-			public int LineNumber;
-			public StringView Header;
-			public StringView Body;
-			public bool External;
-		}
-
-		private static bool RunScriptTest(UhtTables Tables, IUhtConfig Config, UhtGlobalOptions Options, string TestDirectory, string TestOutputDirectory, string Script, ILogger Logger)
-		{
-			string InPath = Path.Combine(TestDirectory, Script);
-			string OutPath = Path.Combine(TestOutputDirectory, Script);
-
-			UhtTestFileManager TestFileManager = new(TestDirectory);
-			UhtSession Session = new(Logger)
-			{
-				Tables = Tables,
-				Config = Config,
-				FileManager = TestFileManager,
-				RootDirectory = TestDirectory,
-				WarningsAsErrors = Options.bWarningsAsErrors,
-				RelativePathInLog = true,
-				GoWide = !Options.bNoGoWide,
-				NoOutput = false,
-				CullOutput = false,
-				CacheMessages = true,
-				IncludeDebugOutput = true,
-			};
-
-			// Read the testing script
-			List<ScriptFragment> ScriptFragments = [];
-			int ManifestIndex = -1;
-			int ConsoleIndex = -1;
-			UhtSourceFile ScriptSourceFile = new(Session, Script);
-			Dictionary<string, int> OutputFragments = [];
-			Session.Try(ScriptSourceFile, () =>
-			{
-				ScriptSourceFile.Read();
-				UhtTokenBufferReader Reader = new(ScriptSourceFile, ScriptSourceFile.Data.Memory);
-
-				bool done = false;
-				while (!done)
-				{
-
-					// Scan for the fragment header
-					ScriptFragmentType Type = ScriptFragmentType.Unknown;
-					string Name = "";
-					int HeaderStartPos = Reader.InputPos;
-					int HeaderEndPos = HeaderStartPos;
-					int LineNumber = 1;
-					while (true)
-					{
-						using UhtTokenSaveState SaveState = new(Reader);
-						UhtToken Token = Reader.GetLine();
-						if (Token.TokenType == UhtTokenType.EndOfFile)
-						{
-							break;
-						}
-						if (Token.Value.Span.Length == 0 || (Token.Value.Span.Length > 0 && Token.Value.Span[0] != '!'))
-						{
-							break;
-						}
-						HeaderEndPos = Reader.InputPos;
-
-						int EndCommandPos = Token.Value.Span.IndexOf(' ');
-						if (EndCommandPos == -1)
-						{
-							EndCommandPos = Token.Value.Span.Length;
-						}
-						string ScriptFragmentTypeString = Token.Value.Span[1..EndCommandPos].Trim().ToString();
-
-						if (!System.Enum.TryParse<ScriptFragmentType>(ScriptFragmentTypeString, true, out Type))
-						{
-							continue;
-						}
-						if (Type == ScriptFragmentType.Unknown)
-						{
-							continue;
-						}
-
-						Name = Token.Value.Span[EndCommandPos..].Trim().ToString();
-						LineNumber = Token.InputLine;
-						SaveState.AbandonState();
-						break;
-					}
-
-					// Scan for the fragment body
-					int BodyStartPos = Reader.InputPos;
-					int BodyEndPos = BodyStartPos;
-					while (true)
-					{
-						using UhtTokenSaveState SaveState = new(Reader);
-						UhtToken Token = Reader.GetLine();
-						if (Token.TokenType == UhtTokenType.EndOfFile)
-						{
-							done = true;
-							break;
-						}
-						if (Token.Value.Span.Length > 0 && Token.Value.Span[0] == '!')
-						{
-							break;
-						}
-						BodyEndPos = Reader.InputPos;
-						SaveState.AbandonState();
-					}
-
-					ScriptFragments.Add(new ScriptFragment
-					{
-						Type = Type,
-						Name = Name.Replace("\\\\", "\\"), // Be kind to people cut/copy/paste escaped strings around
-						LineNumber = LineNumber,
-						Header = new StringView(ScriptSourceFile.Data.Memory[HeaderStartPos..HeaderEndPos]),
-						Body = new StringView(ScriptSourceFile.Data.Memory[BodyStartPos..BodyEndPos]),
-						External = false,
-					});
-				}
-
-				// Search for the manifest and any output.  Add fragments to the session
-				for (int i = 0; i < ScriptFragments.Count; ++i)
-				{
-					switch (ScriptFragments[i].Type)
-					{
-						case ScriptFragmentType.Manifest:
-							if (ManifestIndex != -1)
-							{
-								ScriptSourceFile.LogError(ScriptFragments[i].LineNumber, "There can be only one manifest section in a test script");
-								break;
-							}
-							ManifestIndex = i;
-							if (ScriptFragments[i].Name.Length == 0)
-							{
-								ScriptSourceFile.LogError(ScriptFragments[i].LineNumber, "Manifest name can not be blank");
-								break;
-							}
-							TestFileManager.AddSourceFragment(ScriptSourceFile, ScriptFragments[i].Name, ScriptFragments[i].LineNumber, ScriptFragments[i].Body);
-							break;
-
-						case ScriptFragmentType.Console:
-							if (ConsoleIndex != -1)
-							{
-								ScriptSourceFile.LogError(ScriptFragments[i].LineNumber, "There can be only one console section in a test script");
-								break;
-							}
-							ConsoleIndex = i;
-							break;
-
-						case ScriptFragmentType.Header:
-							if (ScriptFragments[i].Name.Length == 0)
-							{
-								ScriptSourceFile.LogError(ScriptFragments[i].LineNumber, "Header name can not be blank");
-								break;
-							}
-							if (ScriptFragments[i].Body.Length == 0)
-							{
-								// Read the NoExportTypes.h file from the engine source so we don't have to keep a copy
-								if (Path.GetFileName(ScriptFragments[i].Name).Equals("NoExportTypes.h", StringComparison.OrdinalIgnoreCase))
-								{
-									string ExternalPath = Path.Combine(Unreal.EngineDirectory.FullName, ScriptFragments[i].Name);
-									if (File.Exists(ExternalPath))
-									{
-										ScriptFragment Copy = ScriptFragments[i];
-										Copy.Body = new StringView(File.ReadAllText(ExternalPath));
-										Copy.External = true;
-										ScriptFragments[i] = Copy;
-									}
-								}
-							}
-							TestFileManager.AddSourceFragment(ScriptSourceFile, ScriptFragments[i].Name, ScriptFragments[i].LineNumber, ScriptFragments[i].Body);
-							break;
-
-						case ScriptFragmentType.Output:
-							OutputFragments.Add(ScriptFragments[i].Name, i);
-							break;
-					}
-				}
-
-				if (ManifestIndex == -1)
-				{
-					ScriptSourceFile.LogError("There must be a manifest section in a test script");
-				}
-
-				if (ConsoleIndex == -1)
-				{
-					ScriptSourceFile.LogError("There must be a console section in a test script");
-				}
-			});
-
-			// Run UHT
-			if (!Session.HasErrors)
-			{
-				Session.Run(ScriptFragments[ManifestIndex].Name);
-			}
-
-			// If we have no console index, then there is nothing we can do.  This is a fatal error than can not be tested
-			bool bSuccess = true;
-			if (ConsoleIndex == -1)
-			{
-				ScriptSourceFile.LogError("Unable to do any verification without a console section");
-				Session.LogMessages();
-				File.Copy(InPath, OutPath, true);
-				bSuccess = false;
-			}
-			else
-			{
-
-				// Generate the console block
-				List<string> ConsoleLines = Session.CollectMessages();
-				StringBuilder SBConsole = new();
-				foreach (string Line in ConsoleLines)
-				{
-					SBConsole.AppendLine(Line);
-				}
-
-				// Verify the console block 
-				// We trim the ends because it is too easy to leave off the ending CRLF in the script file.
-				if (ScriptFragments[ConsoleIndex].Body.ToString().TrimEnd() != SBConsole.ToString().TrimEnd())
-				{
-					Logger.LogError("Console output failed to match");
-					bSuccess = false;
-				}
-
-				// Check the output
-				foreach (KeyValuePair<string, string> KVP in TestFileManager.Outputs)
-				{
-					if (OutputFragments.TryGetValue(KVP.Key, out int Index))
-					{
-						if (ScriptFragments[Index].Body.ToString().TrimEnd() != KVP.Value.TrimEnd())
-						{
-							Logger.LogError("Output \"{Key}\" failed to match", KVP.Key);
-							bSuccess = false;
-						}
-						OutputFragments.Remove(KVP.Key);
-					}
-					else
-					{
-						Logger.LogError("Output \"{Key}\" not found in test script", KVP.Key);
-					}
-				}
-				foreach (KeyValuePair<string, int> KVP in OutputFragments)
-				{
-					Logger.LogError("Output \"{Key}\" in test script but not generated", KVP.Key);
-				}
-
-				// Create the complete output.  Output includes all of the source fragments and console fragments
-				// and followed the output data sorted by file name.
-				StringBuilder SBTest = new();
-				for (int i = 0; i < ScriptFragments.Count; ++i)
-				{
-					if (ScriptFragments[i].Type != ScriptFragmentType.Output)
-					{
-						SBTest.Append(ScriptFragments[i].Header);
-						if (i == ConsoleIndex)
-						{
-							SBTest.Append(SBConsole);
-						}
-						else if (!ScriptFragments[i].External)
-						{
-							SBTest.Append(ScriptFragments[i].Body);
-						}
-					}
-				}
-
-				// Add the output
-				foreach (KeyValuePair<string, string> KVP in TestFileManager.Outputs)
-				{
-					SBTest.Append($"!output {KVP.Key}\r\n");
-					SBTest.Append(KVP.Value);
-				}
-
-				// Write the final content
-				try
-				{
-					File.WriteAllText(OutPath, SBTest.ToString());
-				}
-				catch (Exception E)
-				{
-					Logger.LogError(E, "Unable to write test result to \"{Ex}\"", E.Message);
-				}
-			}
-
-			if (bSuccess)
-			{
-				Logger.LogInformation("Test {InPath} succeeded", InPath);
-			}
-			else
-			{
-				Logger.LogError("Test {InPath} failed", InPath);
-			}
-			return bSuccess;
-		}
-
-		private static bool RunScriptTests(UhtTables Tables, IUhtConfig Config, UhtGlobalOptions Options, string TestDirectory, string TestOutputDirectory, List<string> Scripts, ILogger Logger)
-		{
-			bool bResult = true;
-			foreach (string Script in Scripts)
-			{
-				bResult &= RunScriptTest(Tables, Config, Options, TestDirectory, TestOutputDirectory, Script, Logger);
-			}
-			return bResult;
-		}
-
-		private static bool RunDirectoryTests(UhtTables Tables, IUhtConfig Config, UhtGlobalOptions Options, string TestDirectory, string TestOutputDirectory, List<string> Directories, ILogger Logger)
-		{
-			bool bResult = true;
-			foreach (string Directory in Directories)
-			{
-				bResult &= RunTests(Tables, Config, Options, Path.Combine(TestDirectory, Directory), Path.Combine(TestOutputDirectory, Directory), Logger);
-			}
-			return bResult;
-		}
-
-		private static bool RunTests(UhtTables Tables, IUhtConfig Config, UhtGlobalOptions Options, string TestDirectory, string TestOutputDirectory, ILogger Logger)
-		{
-			// Create output directory
-			Directory.CreateDirectory(TestOutputDirectory);
-
-			List<string> Scripts = [];
-			foreach (string Script in Directory.EnumerateFiles(TestDirectory, "*.uhttest"))
-			{
-				Scripts.Add(Path.GetFileName(Script));
-			}
-			Scripts.Sort(StringComparer.OrdinalIgnoreCase);
-
-			List<string> Directories = [];
-			foreach (string Directory in Directory.EnumerateDirectories(TestDirectory))
-			{
-				Directories.Add(Path.GetFileName(Directory));
-			}
-			Directories.Sort(StringComparer.OrdinalIgnoreCase);
-
-			List<string> Manifests = [];
-			foreach (string Manifest in Directory.EnumerateFiles(TestDirectory, "*.uhtmanifest"))
-			{
-				Manifests.Add(Path.GetFileName(Manifest));
-			}
-			Manifests.Sort(StringComparer.OrdinalIgnoreCase);
-
-			return
-				RunScriptTests(Tables, Config, Options, TestDirectory, TestOutputDirectory, Scripts, Logger) &&
-				RunDirectoryTests(Tables, Config, Options, TestDirectory, TestOutputDirectory, Directories, Logger);
-		}
-
-		public static bool RunTests(UhtTables Tables, IUhtConfig Config, UhtGlobalOptions Options, ILogger Logger)
-		{
-			DirectoryReference EngineSourceProgramDirectory = DirectoryReference.Combine(Unreal.EngineDirectory, "Source", "Programs");
-			string TestDirectory = FileReference.Combine(EngineSourceProgramDirectory, "UnrealBuildTool.Tests", "UHT").FullName;
-			string TestOutputDirectory = FileReference.Combine(EngineSourceProgramDirectory, "UnrealBuildTool.Tests", "UHT.Out").FullName;
-
-			// Clear the output directory
-			try
-			{
-				Directory.Delete(TestOutputDirectory, true);
-			}
-			catch (Exception)
-			{ }
-
-			// Collect a list of all the test scripts
-			Logger.LogInformation("Running tests in {TestDirectory}", TestDirectory);
-			Logger.LogInformation("Output can be compared in {TestOutputDirectory}", TestOutputDirectory);
-
-			// Run the tests on the directory
-			return RunTests(Tables, Config, Options, TestDirectory, TestOutputDirectory, Logger);
-		}
-	}
-
 	/// <summary>
 	/// Invoke UHT
 	/// </summary>
@@ -1158,7 +675,7 @@ namespace UnrealBuildTool.Modes
 						}
 					}
 				}
-				int RequiredArgCount = TargetArgumentIndex >= 0 || Options.bTest ? 0 : 2;
+				int RequiredArgCount = TargetArgumentIndex >= 0 ? 0 : 2;
 				if (Arguments.GetPositionalArgumentCount() != RequiredArgCount || Options.bGetHelp)
 				{
 					PrintUsage(Tables.ExporterTable, Config);
@@ -1192,12 +709,6 @@ namespace UnrealBuildTool.Modes
 					}
 				}
 
-				// If we are running test scripts
-				if (Options.bTest)
-				{
-					return UhtTestHarness.RunTests(Tables, Config, Options, Logger) ? (int)CompilationResult.Succeeded : (int)CompilationResult.OtherCompilationError;
-				}
-
 				string? ProjectFile = null;
 				string? ManifestPath = null;
 
