[UBT]
* Renamed a couple of scopes
* Changed VNI to be a long running task

diff --git a/Engine/Source/Programs/UnrealBuildTool/Configuration/UEBuildTarget.cs b/Engine/Source/Programs/UnrealBuildTool/Configuration/UEBuildTarget.cs
--- a/Engine/Source/Programs/UnrealBuildTool/Configuration/UEBuildTarget.cs
+++ b/Engine/Source/Programs/UnrealBuildTool/Configuration/UEBuildTarget.cs
@@ -2614,7 +2614,7 @@ namespace UnrealBuildTool
 			// Prepare cached data for UHT header generation
 			using (GlobalTracer.Instance.BuildSpan("UHTExecution.SetupUObjectModules()").StartActive())
 			{
-				using var _ = Timeline.ScopeEvent("SetupUObjectModules");
+				using var _ = Timeline.ScopeEvent("UHTModules");
 
 				UHTExecution.SetupUObjectModules(
 					ModulesToGenerateHeadersFor, 
@@ -2667,6 +2667,7 @@ namespace UnrealBuildTool
 			// Prepare cached data for VNI header generation
 			using (new ScopedTimer("UHTExecution.SetupVNIModules()", Logger, LogLevel.Debug, false))
 			{
+				using var _ = Timeline.ScopeEvent("VNIModules");
 				VNIExecution.SetupVNIModules(ModulesToGenerateHeadersFor, out Makefile.VNIModules);
 			}
 			VNIExecution.PrepopulateDirectoryCache(Makefile);
@@ -2676,13 +2677,14 @@ namespace UnrealBuildTool
 			Task? vniTask = null;
 			if (Makefile.VNIModules.Count > 0)
 			{
-				vniTask = Task.Run(() =>
+				vniTask = new Task(() =>
 				{
 					using (new ScopedTimer("ExternalExecution.ExecuteVNITool()", Logger, LogLevel.Debug, false))
 					{
 						VNIExecution.ExecuteVNITool(Makefile, TargetDescriptor, Logger, false);
 					}
-				});
+				}, TaskCreationOptions.LongRunning);
+				vniTask.Start();
 			}
 
 #endif
diff --git a/Engine/Source/Programs/UnrealBuildTool/Modes/BuildMode.cs b/Engine/Source/Programs/UnrealBuildTool/Modes/BuildMode.cs
--- a/Engine/Source/Programs/UnrealBuildTool/Modes/BuildMode.cs
+++ b/Engine/Source/Programs/UnrealBuildTool/Modes/BuildMode.cs
@@ -393,7 +393,7 @@ namespace UnrealBuildTool
 
 			for (int Idx = 0; Idx < TargetDescriptors.Count; ++Idx)
 			{
-				using ITimelineEvent SetupScope = Timeline.ScopeEvent($"Creating MakeFile for {TargetDescriptors[Idx].Name}");
+				using ITimelineEvent SetupScope = Timeline.ScopeEvent($"{TargetDescriptors[Idx].Name}");
 				// Create and execute the init scripts
 				FileReference[] InitScripts = CreateInitScripts(TargetDescriptors[Idx]);
 				if (InitScripts.Length > 0)
