[UBT]
* Start prefetching engine directory after config has been parsed. This schedules better with the things happening after config reads (initializing telemetry etc)
* Disconnect uba timeline even listener before print since print adds an extra event we don't want to see

diff --git a/Engine/Source/Programs/UnrealBuildTool/UnrealBuildTool.cs b/Engine/Source/Programs/UnrealBuildTool/UnrealBuildTool.cs
--- a/Engine/Source/Programs/UnrealBuildTool/UnrealBuildTool.cs
+++ b/Engine/Source/Programs/UnrealBuildTool/UnrealBuildTool.cs
@@ -440,8 +440,8 @@ namespace UnrealBuildTool
 			}
 
 			// Attach timeline to uba trace
-			using EpicGames.UBA.ITrace trace = EpicGames.UBA.ITrace.Create(traceName, traceFile, true, traceChannel);
-			using TimelineEventScope _ = new((name) => trace.TaskBegin(name, ""), (id) => trace.TaskEnd(id));
+			using EpicGames.UBA.ITrace ubaTrace = EpicGames.UBA.ITrace.Create(traceName, traceFile, true, traceChannel);
+			using TimelineEventScope ubaTimelineEventScope = new((name) => ubaTrace.TaskBegin(name, ""), (id) => ubaTrace.TaskEnd(id));
 
 			bool bShowTimeline = false;
 			try
@@ -535,15 +535,6 @@ namespace UnrealBuildTool
 					Log.RemoveTraceListener(StartupTrace);
 				}
 
-				// Start prefetching the contents of the engine folder
-				if ((ModeOptions & ToolModeOptions.StartPrefetchingEngine) != 0)
-				{
-					using (GlobalTracer.Instance.BuildSpan("FileMetadataPrefetch.QueueEngineDirectory()").StartActive())
-					{
-						FileMetadataPrefetch.QueueEngineDirectory();
-					}
-				}
-
 				// Read the XML configuration files
 				if ((ModeOptions & ToolModeOptions.XmlConfig) != 0)
 				{
@@ -563,6 +554,15 @@ namespace UnrealBuildTool
 					XmlConfig.ApplyTo(Options);
 				}
 
+				// Start prefetching the contents of the engine folder
+				if ((ModeOptions & ToolModeOptions.StartPrefetchingEngine) != 0)
+				{
+					using (GlobalTracer.Instance.BuildSpan("FileMetadataPrefetch.QueueEngineDirectory()").StartActive())
+					{
+						FileMetadataPrefetch.QueueEngineDirectory();
+					}
+				}
+
 				Log.BackupLogFiles = Options.bBackupLogFiles;
 				Log.LogFileBackupCount = Options.LogFileBackupCount;
 				bShowTimeline = Options.bShowTimeline;
@@ -730,6 +730,9 @@ namespace UnrealBuildTool
 
 				Utils.LogWriteFileIfChangedActivity(Logger);
 
+				// Stop listening to timeline events since Print adds an open <Root> event
+				ubaTimelineEventScope.Dispose();
+
 				// Print out all the performance info
 				Timeline.Stop();
 				bool bShowExecutionTime = ModeOptions.HasFlag(ToolModeOptions.ShowExecutionTime);
