UnrealBuildTool: Remove exception

#rnx

diff --git a/Engine/Source/Programs/UnrealBuildTool/System/TargetMakefile.cs b/Engine/Source/Programs/UnrealBuildTool/System/TargetMakefile.cs
--- a/Engine/Source/Programs/UnrealBuildTool/System/TargetMakefile.cs
+++ b/Engine/Source/Programs/UnrealBuildTool/System/TargetMakefile.cs
@@ -91,8 +91,7 @@ namespace UnrealBuildTool
 		{
 			lock (Makefile.InternalDependencies)
 			{
-				if (!Makefile.InternalDependencies.Add(FileItem))
-					throw new Exception("Writing same file multiple times?");
+				Makefile.InternalDependencies.Add(FileItem);
 			}
 
 			if (AllowAsync)
@@ -113,9 +112,9 @@ namespace UnrealBuildTool
 		{
 			lock (Makefile.InternalDependencies)
 			{
-				if (!Makefile.InternalDependencies.Add(FileItem))
-					throw new Exception("Writing same file multiple times?");
+				Makefile.InternalDependencies.Add(FileItem);
 			}
+
 			if (AllowAsync)
 			{
 				WriteTasks.Add(Task.Run(() =>
