[UBT]
* Changed DirectoryReference to DirectoryItem in a few places to reduce I/O

diff --git a/Engine/Source/Programs/UnrealBuildTool/Configuration/UEBuildModule.cs b/Engine/Source/Programs/UnrealBuildTool/Configuration/UEBuildModule.cs
--- a/Engine/Source/Programs/UnrealBuildTool/Configuration/UEBuildModule.cs
+++ b/Engine/Source/Programs/UnrealBuildTool/Configuration/UEBuildModule.cs
@@ -48,7 +48,7 @@ namespace UnrealBuildTool
 		/// <summary>
 		/// Paths to all potential module source directories (with platform extension directories added in)
 		/// </summary>
-		public DirectoryReference[] ModuleDirectories;
+		public DirectoryItem[] ModuleDirectories;
 
 		/// <summary>
 		/// The name of the .Build.cs file this module was created from, if any
@@ -318,7 +318,7 @@ namespace UnrealBuildTool
 			DependenciesToSkipPerArchitecture = Rules.DependenciesToSkipPerArchitecture;
 
 			// get the module directories from the module
-			ModuleDirectories = Rules.GetAllModuleDirectories();
+			ModuleDirectories = Rules.GetAllModuleDirectories().Select(dirRef => DirectoryItem.GetItemByDirectoryReference(dirRef)).ToArray();
 
 			// Add any additional debug visualizers 
 			foreach (string natVisPath in HashSetFromOptionalEnumerableStringParameter(Rules.PublicDebugVisualizerPaths))
@@ -339,7 +339,7 @@ namespace UnrealBuildTool
 			}
 
 			// Add any debug visualizers found in the module directories
-			foreach (DirectoryItem directory in ModuleDirectories.Select(x => DirectoryItem.GetItemByDirectoryReference(x)))
+			foreach (DirectoryItem directory in ModuleDirectories)
 			{
 				NatvisFiles.UnionWith(directory.EnumerateFiles().Where(x => x.HasExtension(".natvis") || x.HasExtension(".natstepfilter") || x.HasExtension(".natjmc")));
 			}
@@ -371,7 +371,7 @@ namespace UnrealBuildTool
 		/// <returns>True if the file is part of this module</returns>
 		public virtual bool ContainsFile(FileReference Location)
 		{
-			return ModuleDirectories.Any(x => Location.IsUnderDirectory(x));
+			return ModuleDirectories.Any(x => Location.IsUnderDirectory(x.Location));
 		}
 
 		/// <summary>
@@ -416,22 +416,19 @@ namespace UnrealBuildTool
 		/// <summary>
 		/// Converts an optional string list parameter to a well-defined hash set.
 		/// </summary>
-		protected HashSet<DirectoryReference> CreateDirectoryHashSet(IEnumerable<string> InEnumerableStrings)
+		protected HashSet<DirectoryReference> CreateDirectoryHashSet(List<string> InEnumerableStrings)
 		{
-			HashSet<DirectoryReference> Directories = new HashSet<DirectoryReference>();
-			if (InEnumerableStrings != null)
+			HashSet<DirectoryReference> Directories = new (InEnumerableStrings.Count);
+			foreach (string InputString in InEnumerableStrings)
 			{
-				foreach (string InputString in InEnumerableStrings)
+				DirectoryReference Dir = new DirectoryReference(ExpandPathVariables(InputString, null, null));
+				if (DirectoryLookupCache.DirectoryExists(Dir))
 				{
-					DirectoryReference Dir = new DirectoryReference(ExpandPathVariables(InputString, null, null));
-					if (DirectoryLookupCache.DirectoryExists(Dir))
-					{
-						Directories.Add(Dir);
-					}
-					else
-					{
-						Log.TraceWarningTask(RulesFile, $"Referenced directory '{Dir}' does not exist.");
-					}
+					Directories.Add(Dir);
+				}
+				else
+				{
+					Log.TraceWarningTask(RulesFile, $"Referenced directory '{Dir}' does not exist.");
 				}
 			}
 			return Directories;
diff --git a/Engine/Source/Programs/UnrealBuildTool/Configuration/UEBuildModuleCPP.cs b/Engine/Source/Programs/UnrealBuildTool/Configuration/UEBuildModuleCPP.cs
--- a/Engine/Source/Programs/UnrealBuildTool/Configuration/UEBuildModuleCPP.cs
+++ b/Engine/Source/Programs/UnrealBuildTool/Configuration/UEBuildModuleCPP.cs
@@ -106,7 +106,7 @@ namespace UnrealBuildTool
 		/// <summary>
 		/// Verse source code directory associated with this module
 		/// </summary>
-		DirectoryReference? AssociatedVerseDirectory;
+		DirectoryItem? AssociatedVerseDirectory;
 
 		/// <summary>
 		/// Create PerModuleInline source file. Defaults to PrimaryModule in a binary but could be another one if primary module does not depend on core
@@ -116,7 +116,7 @@ namespace UnrealBuildTool
 		/// <summary>
 		/// The Verse source code directory associated with this module if any
 		/// </summary>
-		public override DirectoryReference? VerseDirectory => AssociatedVerseDirectory;
+		public override DirectoryReference? VerseDirectory => AssociatedVerseDirectory?.Location;
 
 		protected override void GetReferencedDirectories(HashSet<DirectoryReference> Directories)
 		{
@@ -258,18 +258,23 @@ namespace UnrealBuildTool
 			this.GeneratedCodeDirectory = GeneratedCodeDirectory;
 
 			// Check for a Verse directory next to the rules file
-			DirectoryReference MaybeVerseDirectory = DirectoryReference.Combine(Rules.File.Directory, "Verse");
-			if (IsValidVerseDirectory(MaybeVerseDirectory))
+			DirectoryItem RulesDirectory = DirectoryItem.GetItemByDirectoryReference(Rules.File.Directory);
+			if (RulesDirectory.TryGetDirectory("Verse", out DirectoryItem? MaybeVerseDirectory))
 			{
-				AssociatedVerseDirectory = MaybeVerseDirectory;
-				bDependsOnVerse = true;
+				if (IsValidVerseDirectory(MaybeVerseDirectory.Location))
+				{
+					AssociatedVerseDirectory = MaybeVerseDirectory;
+					bDependsOnVerse = true;
+				}
 			}
-			
+
 			// Check for a Proto directory next to the rules file
-			DirectoryReference MaybeProtoDirectory = DirectoryReference.Combine(Rules.File.Directory, "Proto");
-			if (IsValidProtoDirectory(MaybeProtoDirectory))
+			if (RulesDirectory.TryGetDirectory("Proto", out DirectoryItem? MaybeProtoDirectory))
 			{
-				bHasProto = true;
+				if (IsValidProtoDirectory(MaybeProtoDirectory.Location))
+				{
+					bHasProto = true;
+				}
 			}
 
 			if (Rules.bAddDefaultIncludePaths)
@@ -355,43 +360,39 @@ namespace UnrealBuildTool
 		private void AddDefaultIncludePaths()
 		{
 			// Add the module's parent directory to the public include paths, so other modules may include headers from it explicitly.
-			foreach (DirectoryReference ModuleDir in ModuleDirectories)
+			foreach (DirectoryItem ModuleDir in ModuleDirectories)
 			{
 				// Add the parent directory to the legacy parent include paths.
-				LegacyParentIncludePaths.Add(ModuleDir.ParentDirectory!);
+				LegacyParentIncludePaths.Add(ModuleDir.Location.ParentDirectory!);
 
 				// Add the base directory to the legacy include paths.
-				LegacyPublicIncludePaths.Add(ModuleDir);
+				LegacyPublicIncludePaths.Add(ModuleDir.Location);
 
 				// Add the 'classes' directory, if it exists
-				DirectoryReference ClassesDirectory = DirectoryReference.Combine(ModuleDir, "Classes");
-				if (DirectoryLookupCache.DirectoryExistsAndContainsFiles(ClassesDirectory, SearchOption.AllDirectories))
+				if (ModuleDir.TryGetDirectory("Classes", out DirectoryItem? ClassesDirectory) && ClassesDirectory.ContainsFiles(SearchOption.AllDirectories))
 				{
-					PublicIncludePaths.Add(ClassesDirectory);
+					PublicIncludePaths.Add(ClassesDirectory.Location);
 				}
 
 				// Add all the public directories
-				DirectoryReference PublicDirectory = DirectoryReference.Combine(ModuleDir, "Public");
-				if (DirectoryLookupCache.DirectoryExistsAndContainsFiles(PublicDirectory, SearchOption.AllDirectories))
+				if (ModuleDir.TryGetDirectory("Public", out DirectoryItem? PublicDirectory) && PublicDirectory.ContainsFiles(SearchOption.AllDirectories))
 				{
-					PublicIncludePaths.Add(PublicDirectory);
+					PublicIncludePaths.Add(PublicDirectory.Location);
 
 					IReadOnlySet<string> ExcludeNames = UEBuildPlatform.GetBuildPlatform(Rules.Target.Platform).GetExcludedFolderNames();
-					EnumerateLegacyIncludePaths(DirectoryItem.GetItemByDirectoryReference(PublicDirectory), ExcludeNames, LegacyPublicIncludePaths);
+					EnumerateLegacyIncludePaths(PublicDirectory, ExcludeNames, LegacyPublicIncludePaths);
 				}
 
 				// Add the 'internal' directory, if it exists
-				DirectoryReference InternalDirectory = DirectoryReference.Combine(ModuleDir, "Internal");
-				if (DirectoryLookupCache.DirectoryExistsAndContainsFiles(InternalDirectory, SearchOption.AllDirectories))
+				if (ModuleDir.TryGetDirectory("Internal", out DirectoryItem? InternalDirectory) && InternalDirectory.ContainsFiles(SearchOption.AllDirectories))
 				{
-					InternalIncludePaths.Add(InternalDirectory);
+					InternalIncludePaths.Add(InternalDirectory.Location);
 				}
 
 				// Add the base private directory for this module
-				DirectoryReference PrivateDirectory = DirectoryReference.Combine(ModuleDir, "Private");
-				if (DirectoryLookupCache.DirectoryExistsAndContainsFiles(PrivateDirectory, SearchOption.AllDirectories))
+				if (ModuleDir.TryGetDirectory("Private", out DirectoryItem? PrivateDirectory) && PrivateDirectory.ContainsFiles(SearchOption.AllDirectories))
 				{
-					PrivateIncludePaths.Add(PrivateDirectory);
+					PrivateIncludePaths.Add(PrivateDirectory.Location);
 				}
 			}
 		}
@@ -504,10 +505,9 @@ namespace UnrealBuildTool
 			}
 
 			// Add all the module source directories to the makefile
-			foreach (DirectoryReference ModuleDirectory in ModuleDirectories)
+			foreach (DirectoryItem ModuleDirectory in ModuleDirectories)
 			{
-				DirectoryItem ModuleDirectoryItem = DirectoryItem.GetItemByDirectoryReference(ModuleDirectory);
-				Graph.AddSourceDir(ModuleDirectoryItem);
+				Graph.AddSourceDir(ModuleDirectory);
 			}
 
 			// Find all the input files
@@ -898,9 +898,9 @@ namespace UnrealBuildTool
 				CppCompileEnvironment SingleCompileEnvironment = new CppCompileEnvironment(CppCompileEnvironment);
 				SingleCompileEnvironment.PrecompiledHeaderAction = PrecompiledHeaderAction.None;
 				SingleCompileEnvironment.bUseUnity = false;
-				foreach (DirectoryReference Directory in ModuleDirectories)
+				foreach (DirectoryItem Directory in ModuleDirectories)
 				{
-					ToolChain.CreateSpecificFileAction(SingleCompileEnvironment, Directory, IntermediateDirectory, Graph);
+					ToolChain.CreateSpecificFileAction(SingleCompileEnvironment, Directory.Location, IntermediateDirectory, Graph);
 				}
 			}
 
@@ -2760,10 +2760,9 @@ namespace UnrealBuildTool
 			InputFileCollection InputFiles = new InputFileCollection();
 
 			SourceDirectories = new HashSet<DirectoryReference>();
-			foreach (DirectoryReference Dir in ModuleDirectories)
+			foreach (DirectoryItem Dir in ModuleDirectories)
 			{
-				DirectoryItem ModuleDirectoryItem = DirectoryItem.GetItemByDirectoryReference(Dir);
-				FindInputFilesFromDirectoryRecursive(ModuleDirectoryItem, ExcludedNames, SourceDirectories, DirectoryToSourceFiles, InputFiles, Logger);
+				FindInputFilesFromDirectoryRecursive(Dir, ExcludedNames, SourceDirectories, DirectoryToSourceFiles, InputFiles, Logger);
 			}
 
 			return InputFiles;
diff --git a/Engine/Source/Programs/UnrealBuildTool/Configuration/UEBuildTarget.cs b/Engine/Source/Programs/UnrealBuildTool/Configuration/UEBuildTarget.cs
--- a/Engine/Source/Programs/UnrealBuildTool/Configuration/UEBuildTarget.cs
+++ b/Engine/Source/Programs/UnrealBuildTool/Configuration/UEBuildTarget.cs
@@ -5184,10 +5184,12 @@ namespace UnrealBuildTool
 
 				if (ProjectDescriptor.Plugins != null)
 				{
+					HashSet<string> EnablePlugins = new(Rules.EnablePlugins, StringComparer.InvariantCultureIgnoreCase);
+					HashSet<string> DisablePlugins = new(Rules.DisablePlugins, StringComparer.InvariantCultureIgnoreCase);
 					string ProjectReferenceChain = ProjectFile!.GetFileName();
 					foreach (PluginReferenceDescriptor PluginReference in ProjectDescriptor.Plugins)
 					{
-						if (!Rules.EnablePlugins.Contains(PluginReference.Name, StringComparer.InvariantCultureIgnoreCase) && !Rules.DisablePlugins.Contains(PluginReference.Name, StringComparer.InvariantCultureIgnoreCase))
+						if (!EnablePlugins.Contains(PluginReference.Name, StringComparer.InvariantCultureIgnoreCase) && !DisablePlugins.Contains(PluginReference.Name, StringComparer.InvariantCultureIgnoreCase))
 						{
 							// Make sure we don't have multiple references to the same plugin
 							if (!ReferencedNames.Add(PluginReference.Name))
@@ -5505,12 +5507,14 @@ namespace UnrealBuildTool
 						UEBuildPlugin? NextInstance = AddPlugin(NextReference, PluginReferenceChain, ExcludeFolders, NameToInstance, NameToInfos, Logger);
 						if (NextInstance != null)
 						{
-							Dependencies.Add(NextInstance);
-							if (NextInstance.Dependencies == null)
+							if (Dependencies.Add(NextInstance))
 							{
-								throw new BuildException($"Found circular plugin dependency. {PluginReferenceChain} -> {NextReference.Name}");
+								if (NextInstance.Dependencies == null)
+								{
+									throw new BuildException($"Found circular plugin dependency. {PluginReferenceChain} -> {NextReference.Name}");
+								}
+								Dependencies.UnionWith(NextInstance.Dependencies);
 							}
-							Dependencies.UnionWith(NextInstance.Dependencies);
 						}
 					}
 				}
@@ -6408,7 +6412,7 @@ namespace UnrealBuildTool
 					// If it's a game module (plugin or otherwise), add the root source directory to the include paths.
 					if (RulesObject.File.IsUnderDirectory(BaseSourceDirectory) || (RulesObject.Plugin != null && RulesObject.Plugin.LoadedFrom == PluginLoadedFrom.Project))
 					{
-						if (DirectoryReference.Exists(BaseSourceDirectory))
+						if (DirectoryItem.GetItemByDirectoryReference(BaseSourceDirectory).Exists)
 						{
 							RulesObject.PublicIncludePaths.Add(NormalizeIncludePath(BaseSourceDirectory));
 						}
diff --git a/Engine/Source/Programs/UnrealBuildTool/Modes/IWYUMode.cs b/Engine/Source/Programs/UnrealBuildTool/Modes/IWYUMode.cs
--- a/Engine/Source/Programs/UnrealBuildTool/Modes/IWYUMode.cs
+++ b/Engine/Source/Programs/UnrealBuildTool/Modes/IWYUMode.cs
@@ -486,9 +486,9 @@ namespace UnrealBuildTool
 					{
 						continue;
 					}
-					foreach (DirectoryReference AdditionalDir in Module.ModuleDirectories)
+					foreach (DirectoryItem AdditionalDir in Module.ModuleDirectories)
 					{
-						if (AdditionalDir != Module.ModuleDirectory)
+						if (AdditionalDir.Location != Module.ModuleDirectory)
 						{
 							PathToModule.TryAdd(AdditionalDir.FullName.Replace('\\', '/'), Module);
 						}
@@ -563,9 +563,9 @@ namespace UnrealBuildTool
 					// When adding to ValidPaths, make sure the Module directory ends in a / so we match the exact folder later
 					ValidPaths.Add(AdjustModulePathForMatching(ModuleDir));
 
-					foreach (DirectoryReference AdditionalDir in Module.ModuleDirectories)
+					foreach (DirectoryItem AdditionalDir in Module.ModuleDirectories)
 					{
-						if (AdditionalDir != Module.ModuleDirectory)
+						if (AdditionalDir.Location != Module.ModuleDirectory)
 						{
 							ValidPaths.Add(AdjustModulePathForMatching(AdditionalDir.FullName));
 						}
diff --git a/Engine/Source/Programs/UnrealBuildTool/Modes/QueryMode.cs b/Engine/Source/Programs/UnrealBuildTool/Modes/QueryMode.cs
--- a/Engine/Source/Programs/UnrealBuildTool/Modes/QueryMode.cs
+++ b/Engine/Source/Programs/UnrealBuildTool/Modes/QueryMode.cs
@@ -444,9 +444,9 @@ namespace UnrealBuildTool
 								return ValueTask.CompletedTask;
 							}
 
-							foreach (DirectoryReference Dir in Module.ModuleDirectories)
+							foreach (DirectoryItem Dir in Module.ModuleDirectories)
 							{
-								CurrentTargetIntellisenseInfo.DirToModule.TryAdd(Dir, Module);
+								CurrentTargetIntellisenseInfo.DirToModule.TryAdd(Dir.Location, Module);
 							}
 							if (Module.GeneratedCodeDirectory != null)
 							{
diff --git a/Engine/Source/Programs/UnrealBuildTool/Modes/ServerMode.cs b/Engine/Source/Programs/UnrealBuildTool/Modes/ServerMode.cs
--- a/Engine/Source/Programs/UnrealBuildTool/Modes/ServerMode.cs
+++ b/Engine/Source/Programs/UnrealBuildTool/Modes/ServerMode.cs
@@ -358,10 +358,10 @@ namespace UnrealBuildTool
 						}
 
 						CppCompileEnvironment ModuleCompileEnvironment = Module.CreateModuleCompileEnvironment(CurrentTarget.Rules, BinaryCompileEnvironment, Logger);
-						foreach (DirectoryReference Dir in Module.ModuleDirectories)
+						foreach (DirectoryItem Dir in Module.ModuleDirectories)
 						{
 							BrowseConfigurationFolders.Add(Dir.ToString());
-							CurrentTargetIntellisenseInfo.DirToModule.TryAdd(Dir, Module);
+							CurrentTargetIntellisenseInfo.DirToModule.TryAdd(Dir.Location, Module);
 						}
 						if (Module.GeneratedCodeDirectory != null)
 						{
diff --git a/Engine/Source/Programs/UnrealBuildTool/ProjectFiles/ProjectFile.cs b/Engine/Source/Programs/UnrealBuildTool/ProjectFiles/ProjectFile.cs
--- a/Engine/Source/Programs/UnrealBuildTool/ProjectFiles/ProjectFile.cs
+++ b/Engine/Source/Programs/UnrealBuildTool/ProjectFiles/ProjectFile.cs
@@ -4,6 +4,7 @@ using System;
 using System.Collections.Generic;
 using EpicGames.Core;
 using Microsoft.Extensions.Logging;
+using UnrealBuildBase;
 
 namespace UnrealBuildTool
 {
@@ -260,13 +261,13 @@ namespace UnrealBuildTool
 			AddIntelliSenseIncludePaths(SystemIncludePaths, CompileEnvironment.SystemIncludePaths);
 			AddIntelliSenseIncludePaths(UserIncludePaths, CompileEnvironment.UserIncludePaths);
 
-			foreach (DirectoryReference BaseDir in Module.ModuleDirectories)
+			foreach (DirectoryItem BaseDir in Module.ModuleDirectories)
 			{
 				BuildEnvironment? BuildEnvironment;
-				if (!BaseDirToBuildEnvironment.TryGetValue(BaseDir, out BuildEnvironment))
+				if (!BaseDirToBuildEnvironment.TryGetValue(BaseDir.Location, out BuildEnvironment))
 				{
 					BuildEnvironment = new BuildEnvironment();
-					BaseDirToBuildEnvironment.Add(BaseDir, BuildEnvironment);
+					BaseDirToBuildEnvironment.Add(BaseDir.Location, BuildEnvironment);
 				}
 
 				AddIntelliSenseIncludePaths(BuildEnvironment.SystemIncludePaths, CompileEnvironment.SystemIncludePaths);
diff --git a/Engine/Source/Programs/UnrealBuildTool/ProjectFiles/VisualStudio/VCProject.cs b/Engine/Source/Programs/UnrealBuildTool/ProjectFiles/VisualStudio/VCProject.cs
--- a/Engine/Source/Programs/UnrealBuildTool/ProjectFiles/VisualStudio/VCProject.cs
+++ b/Engine/Source/Programs/UnrealBuildTool/ProjectFiles/VisualStudio/VCProject.cs
@@ -517,7 +517,7 @@ namespace UnrealBuildTool
 
 			if (Settings.bUsePerFileIntellisense)
 			{
-				foreach (DirectoryReference ModuleDirectory in Module.ModuleDirectories)
+				foreach (DirectoryItem ModuleDirectory in Module.ModuleDirectories)
 				{
 					List<string> ForceIncludePaths = new List<string>(CompileEnvironment.ForceIncludeFiles.Select(x => InsertPathVariables(x.Location)));
 					if (CompileEnvironment.PrecompiledHeaderIncludeFilename != null)
diff --git a/Engine/Source/Programs/UnrealBuildTool/ProjectFiles/VisualStudioWorkspace/VSWorkspaceProjectFile.cs b/Engine/Source/Programs/UnrealBuildTool/ProjectFiles/VisualStudioWorkspace/VSWorkspaceProjectFile.cs
--- a/Engine/Source/Programs/UnrealBuildTool/ProjectFiles/VisualStudioWorkspace/VSWorkspaceProjectFile.cs
+++ b/Engine/Source/Programs/UnrealBuildTool/ProjectFiles/VisualStudioWorkspace/VSWorkspaceProjectFile.cs
@@ -279,7 +279,7 @@ namespace UnrealBuildTool
 					CppCompileEnvironment ModuleCompileEnvironment = ModuleCpp.CreateCompileEnvironmentForIntellisense(Target.Rules, BinaryCompileEnvironment, Logger);
 					TargetInfo.ModuleToCompileSettings.Add(ModuleCpp.Name, ExportModule(ModuleCpp, TargetToolChain, ModuleCompileEnvironment));
 
-					foreach (DirectoryReference ModuleDirectory in ModuleCpp.ModuleDirectories)
+					foreach (DirectoryItem ModuleDirectory in ModuleCpp.ModuleDirectories)
 					{
 						TargetInfo.DirToModule.TryAdd(ModuleDirectory.FullName, ModuleCpp.Name);
 					}
diff --git a/Engine/Source/Programs/UnrealBuildTool/System/UHTExecution.cs b/Engine/Source/Programs/UnrealBuildTool/System/UHTExecution.cs
--- a/Engine/Source/Programs/UnrealBuildTool/System/UHTExecution.cs
+++ b/Engine/Source/Programs/UnrealBuildTool/System/UHTExecution.cs
@@ -98,7 +98,7 @@ namespace UnrealBuildTool
 		/// <summary>
 		/// Paths to all potential module source directories (with platform extension directories added in)
 		/// </summary>
-		public required DirectoryReference[] ModuleDirectories;
+		public required DirectoryItem[] ModuleDirectories;
 
 		/// <summary>
 		/// The include search path for generated headers to include other headers
@@ -229,7 +229,7 @@ namespace UnrealBuildTool
 		{
 			ModuleName = Reader.ReadString()!;
 			ModuleRulesFile = Reader.ReadFileReference();
-			ModuleDirectories = Reader.ReadArray<DirectoryReference>(Reader.ReadDirectoryReferenceNotNull)!;
+			ModuleDirectories = Reader.ReadArray<DirectoryItem>(() => Reader.ReadDirectoryItem()!)!;
 			ModuleIncludePaths = Reader.ReadArray<DirectoryReference>(Reader.ReadDirectoryReferenceNotNull)!;
 			ModuleType = Reader.ReadString()!;
 			OverrideModuleType = Reader.ReadString()!;
@@ -260,7 +260,7 @@ namespace UnrealBuildTool
 		{
 			Writer.WriteString(ModuleName);
 			Writer.WriteFileReference(ModuleRulesFile);
-			Writer.WriteArray<DirectoryReference>(ModuleDirectories, Writer.WriteDirectoryReference);
+			Writer.WriteArray<DirectoryItem>(ModuleDirectories, Writer.WriteDirectoryItem);
 			Writer.WriteArray<DirectoryReference>(ModuleIncludePaths, Writer.WriteDirectoryReference);
 			Writer.WriteString(ModuleType);
 			Writer.WriteString(OverrideModuleType);
@@ -613,7 +613,7 @@ namespace UnrealBuildTool
 				{
 					Queue.Enqueue(() =>
 					{
-						foreach (DirectoryItem ModuleDirectoryItem in Module.ModuleDirectories.Select(x => DirectoryItem.GetItemByDirectoryReference(x)))
+						foreach (DirectoryItem ModuleDirectoryItem in Module.ModuleDirectories)
 						{
 							List<FileItem> HeaderFiles = [];
 							FindHeaders(ModuleDirectoryItem, ExcludedFolders, HeaderFiles);
@@ -716,12 +716,10 @@ namespace UnrealBuildTool
 
 		static void SetupUObjectModule(UHTModuleInfo ModuleInfo, IReadOnlySet<string> ExcludedFolders, SourceFileMetadataCache MetadataCache, ParallelQueue Queue)
 		{
-			foreach (DirectoryReference ModuleDirectory in ModuleInfo.ModuleDirectories)
+			foreach (DirectoryItem ModuleDirectory in ModuleInfo.ModuleDirectories)
 			{
-				DirectoryItem ModuleDirectoryItem = DirectoryItem.GetItemByDirectoryReference(ModuleDirectory);
-
 				List<FileItem> HeaderFiles = [];
-				FindHeaders(ModuleDirectoryItem, ExcludedFolders, HeaderFiles);
+				FindHeaders(ModuleDirectory, ExcludedFolders, HeaderFiles);
 
 				foreach (FileItem HeaderFile in HeaderFiles)
 				{
@@ -739,19 +737,19 @@ namespace UnrealBuildTool
 				lock (ModuleInfo)
 				{
 					bool bFoundHeaderLocation = false;
-					foreach (DirectoryReference ModuleDirectory in ModuleInfo.ModuleDirectories)
+					foreach (DirectoryItem ModuleDirectory in ModuleInfo.ModuleDirectories)
 					{
-						if (HeaderFile.Location.IsUnderDirectory(DirectoryReference.Combine(ModuleDirectory, "Classes")))
+						if (HeaderFile.Location.IsUnderDirectory(DirectoryReference.Combine(ModuleDirectory.Location, "Classes")))
 						{
 							ModuleInfo.PublicUObjectClassesHeaders.Add(HeaderFile);
 							bFoundHeaderLocation = true;
 						}
-						else if (HeaderFile.Location.IsUnderDirectory(DirectoryReference.Combine(ModuleDirectory, "Public")))
+						else if (HeaderFile.Location.IsUnderDirectory(DirectoryReference.Combine(ModuleDirectory.Location, "Public")))
 						{
 							ModuleInfo.PublicUObjectHeaders.Add(HeaderFile);
 							bFoundHeaderLocation = true;
 						}
-						else if (HeaderFile.Location.IsUnderDirectory(DirectoryReference.Combine(ModuleDirectory, "Internal")))
+						else if (HeaderFile.Location.IsUnderDirectory(DirectoryReference.Combine(ModuleDirectory.Location, "Internal")))
 						{
 							ModuleInfo.InternalUObjectHeaders.Add(HeaderFile);
 							bFoundHeaderLocation = true;
