[UBT]
* Optimized parsing of syntax trees. Don't read bytes and then convert, use SourceText directly on filestream. Also, skip parsing documentation

diff --git a/Engine/Source/Programs/UnrealBuildTool/System/DynamicCompilation.cs b/Engine/Source/Programs/UnrealBuildTool/System/DynamicCompilation.cs
--- a/Engine/Source/Programs/UnrealBuildTool/System/DynamicCompilation.cs
+++ b/Engine/Source/Programs/UnrealBuildTool/System/DynamicCompilation.cs
@@ -8,11 +8,13 @@ using Microsoft.CodeAnalysis.Text;
 using Microsoft.Extensions.Logging;
 using OpenTracing.Util;
 using System;
+using System.Collections.Concurrent;
 using System.Collections.Generic;
 using System.IO;
 using System.Linq;
 using System.Reflection;
 using System.Runtime.InteropServices;
+using System.Text;
 using System.Threading;
 using System.Threading.Tasks;
 using UnrealBuildBase;
@@ -168,23 +170,31 @@ namespace UnrealBuildTool
 			}
 		}
 
-		private static async Task<SyntaxTree> ParseSyntaxTreeAsync(FileReference sourceFileName, CSharpParseOptions parseOptions, CancellationToken cancellationToken)
+		private static SyntaxTree ParseSyntaxTree(FileReference sourceFileName, CSharpParseOptions parseOptions, CancellationToken cancellationToken)
 		{
-			byte[] bytes = await FileReference.ReadAllBytesAsync(sourceFileName, cancellationToken);
-			SourceText source = SourceText.From(bytes, bytes.Length, FileUtils.GetEncoding(bytes));
-			return CSharpSyntaxTree.ParseText(source, parseOptions, sourceFileName.FullName, cancellationToken);
+			using var fs = File.OpenRead(sourceFileName.FullName);
+			var text = SourceText.From(fs, encoding: Encoding.UTF8, canBeEmbedded: true);
+			return CSharpSyntaxTree.ParseText(text, parseOptions, path: sourceFileName.FullName, cancellationToken: cancellationToken);
 		}
 
 		private static async Task<IEnumerable<SyntaxTree>> ParseSyntaxTreesAsync(IEnumerable<FileReference> sourceFileNames, IEnumerable<string>? preprocessorDefines, CancellationToken cancellationToken)
 		{
-			CSharpParseOptions parseOptions = new(
-				languageVersion: LanguageVersion.Latest,
-				kind: SourceCodeKind.Regular,
-				preprocessorSymbols: preprocessorDefines
-			);
+			var parseOptions = new CSharpParseOptions(
+				 languageVersion: LanguageVersion.CSharp12, // Prefer specific version, faster than using Latest due to version probing
+				 kind: SourceCodeKind.Regular,
+				 preprocessorSymbols: preprocessorDefines,
+				 documentationMode: DocumentationMode.None);
+
+			var trees = new ConcurrentBag<SyntaxTree>();
+			var dop = Math.Max(1, Environment.ProcessorCount - 1);
 
-			IEnumerable<Task<SyntaxTree>> tasks = sourceFileNames.Select(sourceFileName => ParseSyntaxTreeAsync(sourceFileName, parseOptions, cancellationToken));
-			return (await Task.WhenAll(tasks)).OrderBy(x => x.FilePath);
+			await Parallel.ForEachAsync(sourceFileNames, new ParallelOptions { MaxDegreeOfParallelism = dop, CancellationToken = cancellationToken },
+				(f, token) =>
+				{
+					trees.Add(ParseSyntaxTree(f, parseOptions, token));
+					return ValueTask.CompletedTask;
+				});
+			return trees.OrderBy(t => t.FilePath).ToArray();
 		}
 
 		private static async Task<Assembly?> CompileAssemblyAsync(FileReference OutputAssemblyPath, IEnumerable<FileReference> SourceFileNames, ILogger Logger, IEnumerable<string>? ReferencedAssembies, IEnumerable<string>? PreprocessorDefines = null, bool TreatWarningsAsErrors = false)
