[UBT]
* Changed so uht timestamp use prefix compression for paths. This makes them 4x smaller and hopefully will speed up the writing of timestamps

diff --git a/Engine/Source/Programs/UnrealBuildTool/System/UHTExecution.cs b/Engine/Source/Programs/UnrealBuildTool/System/UHTExecution.cs
--- a/Engine/Source/Programs/UnrealBuildTool/System/UHTExecution.cs
+++ b/Engine/Source/Programs/UnrealBuildTool/System/UHTExecution.cs
@@ -939,8 +939,8 @@ namespace UnrealBuildTool
 
 				// Load up the old timestamp file and check to see if anything has changed
 				{
-					string[] UObjectFilesFromPreviousRun = File.ReadAllLines(TimestampFile.FullName);
-					if (AllUObjectHeaders.Count != UObjectFilesFromPreviousRun.Length)
+					string[] TimestampLines = File.ReadAllLines(TimestampFile.FullName);
+					if (AllUObjectHeaders.Count != TimestampLines.Length)
 					{
 						if (Interlocked.Increment(ref invalidCounter) == 1)
 						{
@@ -951,13 +951,28 @@ namespace UnrealBuildTool
 
 					// Iterate over our UObjects headers and figure out if any of them have changed
 					HashSet<string> ObjectHeadersSet = new(AllUObjectHeaders.Select(x => x.AbsolutePath), FileReference.Comparer);
-					foreach (string FileName in UObjectFilesFromPreviousRun)
+
+					// Prefix compression, look at write
+					string lastFilename = "";
+					StringBuilder sb = new(256);
+					foreach (string line in TimestampLines)
 					{
-						if (!ObjectHeadersSet.Contains(FileName))
+						string fileName = line;
+
+						if (line[0] >= '0' && line[0] <= '9')
+						{
+							int space = line.IndexOf(' ');
+							int matchLength = int.Parse(line.AsSpan(0, space));
+							sb.Clear().Append(lastFilename, 0, matchLength).Append(line, space+1, line.Length - space - 1);
+							fileName = sb.ToString();
+						}
+						lastFilename = fileName;	
+
+						if (!ObjectHeadersSet.Contains(fileName))
 						{
 							if (Interlocked.Increment(ref invalidCounter) == 1)
 							{
-								Logger.LogDebug("UnrealHeaderTool needs to run because the set of UObject source files in module {ModuleModuleName} has changed ({FileName})", Module.ModuleName, FileName);
+								Logger.LogDebug("UnrealHeaderTool needs to run because the set of UObject source files in module {ModuleModuleName} has changed ({fileName})", Module.ModuleName, fileName);
 							}
 							return;
 						}
@@ -1031,19 +1046,48 @@ namespace UnrealBuildTool
 						{
 							FileReference TimestampFile = FileReference.Combine(new DirectoryReference(GeneratedCodeDirectoryInfo.FullName), "Timestamp");
 
+
+							using var fs = new FileStream(TimestampFile.FullName, FileMode.Create, FileAccess.Write, FileShare.None, bufferSize: 1 << 18, options: FileOptions.SequentialScan);
+							using var sw = new StreamWriter(fs, new UTF8Encoding(encoderShouldEmitUTF8Identifier: false), bufferSize: 1 << 18);
+
+							string lastFullName = "";
+
+							// Use prefix compression (If first part of the line is a number, then the line is lastLine[0..number] + line-after-number)
+
+							void WriteHeaders(List<FileItem> headers)
+							{
+								headers.Sort(FileItem.SortByPath);
+
+								foreach (FileItem header in headers)
+								{
+									ReadOnlySpan<char> last = lastFullName;
+									ReadOnlySpan<char> curr = header.FullName.AsSpan();
+
+									int matchLength = MemoryExtensions.CommonPrefixLength(last, curr);
+
+									if (matchLength == 0)
+									{
+										sw.WriteLine(curr);
+									}
+									else
+									{
+										sw.Write(matchLength);
+										sw.Write(' ');
+										sw.WriteLine(curr.Slice(matchLength));
+									}
+									lastFullName = header.FullName;
+								}
+							}
+
 							// Save all of the UObject files to a timestamp file.  We'll load these on the next run to see if any new
 							// files with UObject classes were deleted, so that we'll know to run UHT even if the timestamps of all
 							// of the other source files were unchanged
-							{
-								List<string> AllUObjectFiles =
-								[
-									.. Module.PublicUObjectClassesHeaders.ConvertAll(Item => Item.AbsolutePath),
-									.. Module.PublicUObjectHeaders.ConvertAll(Item => Item.AbsolutePath),
-									.. Module.InternalUObjectHeaders.ConvertAll(Item => Item.AbsolutePath),
-									.. Module.PrivateUObjectHeaders.ConvertAll(Item => Item.AbsolutePath),
-								];
-								await FileReference.WriteAllLinesAsync(TimestampFile, AllUObjectFiles);
-							}
+
+							WriteHeaders(Module.PublicUObjectClassesHeaders);
+							WriteHeaders(Module.PublicUObjectHeaders);
+							WriteHeaders(Module.InternalUObjectHeaders);
+							WriteHeaders(Module.PrivateUObjectHeaders);
+							await Task.CompletedTask;
 
 							if (resetDirectoryCaches)
 							{
