[UBT]
* Added platform and config to timeline scope for creating makefile

diff --git a/Engine/Source/Programs/UnrealBuildTool/Modes/BuildMode.cs b/Engine/Source/Programs/UnrealBuildTool/Modes/BuildMode.cs
--- a/Engine/Source/Programs/UnrealBuildTool/Modes/BuildMode.cs
+++ b/Engine/Source/Programs/UnrealBuildTool/Modes/BuildMode.cs
@@ -393,15 +393,16 @@ namespace UnrealBuildTool
 
 			for (int Idx = 0; Idx < TargetDescriptors.Count; ++Idx)
 			{
-				using ITimelineEvent SetupScope = Timeline.ScopeEvent($"{TargetDescriptors[Idx].Name}");
+				TargetDescriptor TargetDescriptor = TargetDescriptors[Idx];
+				using ITimelineEvent SetupScope = Timeline.ScopeEvent($"{TargetDescriptor.Name} {TargetDescriptor.Platform} {TargetDescriptor.Configuration}");
 				// Create and execute the init scripts
-				FileReference[] InitScripts = CreateInitScripts(TargetDescriptors[Idx]);
+				FileReference[] InitScripts = CreateInitScripts(TargetDescriptor);
 				if (InitScripts.Length > 0)
 				{
 					Utils.ExecuteCustomBuildSteps(InitScripts, Logger);
 				}
 
-				(TargetMakefile NewMakefile, FileReference? SaveToLocation) = await CreateMakefileAsync2(BuildConfiguration, TargetDescriptors[Idx], WorkingSet, Logger);
+				(TargetMakefile NewMakefile, FileReference? SaveToLocation) = await CreateMakefileAsync2(BuildConfiguration, TargetDescriptor, WorkingSet, Logger);
 				TargetMakefiles.Add(NewMakefile);
 
 				if (SaveToLocation != null)
