Add TargetRules flag that controls USE_ANDROID_JNI, fixes nonunity build failures

#rnx

diff --git a/Engine/Source/Programs/UnrealBuildTool/Configuration/ReadOnlyTargetRules.cs b/Engine/Source/Programs/UnrealBuildTool/Configuration/ReadOnlyTargetRules.cs
--- a/Engine/Source/Programs/UnrealBuildTool/Configuration/ReadOnlyTargetRules.cs
+++ b/Engine/Source/Programs/UnrealBuildTool/Configuration/ReadOnlyTargetRules.cs
@@ -267,6 +267,8 @@ namespace UnrealBuildTool
 
 		public bool bDefineForceInlineHintToInline => Inner.bDefineForceInlineHintToInline;
 
+		public bool bUseAndroidJni => Inner.bUseAndroidJni;
+
 		public FPSemanticsMode FPSemantics => Inner.FPSemantics;
 
 		public bool bRetainFramePointers => Inner.bRetainFramePointers;
diff --git a/Engine/Source/Programs/UnrealBuildTool/Configuration/TargetRules.cs b/Engine/Source/Programs/UnrealBuildTool/Configuration/TargetRules.cs
--- a/Engine/Source/Programs/UnrealBuildTool/Configuration/TargetRules.cs
+++ b/Engine/Source/Programs/UnrealBuildTool/Configuration/TargetRules.cs
@@ -1586,6 +1586,12 @@ namespace UnrealBuildTool
 		}
 		private bool? bDefineForceInlineHintToInlinePrivate;
 
+		/// <summary>
+		/// Whether to use Android Java Native Interface.
+		/// </summary>
+		[RequiresUniqueBuildEnvironment]
+		public bool bUseAndroidJni { get; set; } = true;
+
 		/// <summary>
 		/// Allows setting the FP semantics.
 		/// </summary>
diff --git a/Engine/Source/Programs/UnrealBuildTool/Configuration/TestTargetRules.cs b/Engine/Source/Programs/UnrealBuildTool/Configuration/TestTargetRules.cs
--- a/Engine/Source/Programs/UnrealBuildTool/Configuration/TestTargetRules.cs
+++ b/Engine/Source/Programs/UnrealBuildTool/Configuration/TestTargetRules.cs
@@ -344,7 +344,8 @@ namespace UnrealBuildTool
 				GlobalDefinitions.Add("USE_ANDROID_INPUT=0");
 				GlobalDefinitions.Add("USE_ANDROID_OPENGL=0");
 				GlobalDefinitions.Add("USE_ANDROID_LAUNCH=0");
-				GlobalDefinitions.Add("USE_ANDROID_JNI=0");
+
+				bUseAndroidJni = false;
 
 				// Addresses linker errors when building LowLevelTests for Android
 				// TODO: This should be written to the intermediate directory, somewhere else not when TargetRules are being created
diff --git a/Engine/Source/Runtime/Core/Core.Build.cs b/Engine/Source/Runtime/Core/Core.Build.cs
--- a/Engine/Source/Runtime/Core/Core.Build.cs
+++ b/Engine/Source/Runtime/Core/Core.Build.cs
@@ -31,6 +31,8 @@ public class Core : ModuleRules
 			AdditionalPropertiesForReceipt.Add("AndroidPlugin", Path.Combine(ModulePath, "Core_APL.xml"));
 		}
 
+		PublicDefinitions.Add(string.Format("USE_ANDROID_JNI={0}", System.Convert.ToInt32(Target.bUseAndroidJni)));
+
 		PrivateDependencyModuleNames.Add("AtomicQueue");
 		PrivateDependencyModuleNames.Add("AutoRTFM");
 		PrivateDependencyModuleNames.Add("BLAKE3");
diff --git a/Engine/Source/Runtime/Core/Private/Android/AndroidJavaEnv.cpp b/Engine/Source/Runtime/Core/Private/Android/AndroidJavaEnv.cpp
--- a/Engine/Source/Runtime/Core/Private/Android/AndroidJavaEnv.cpp
+++ b/Engine/Source/Runtime/Core/Private/Android/AndroidJavaEnv.cpp
@@ -1,7 +1,6 @@
 // Copyright Epic Games, Inc. All Rights Reserved.
 
-#if defined(USE_ANDROID_JNI) && USE_ANDROID_JNI
-
+#if USE_ANDROID_JNI
 #include "Android/AndroidJavaEnv.h"
 #include "Android/AndroidJniGameActivity.h"
 #include "CoreMinimal.h"
diff --git a/Engine/Source/Runtime/Core/Public/Android/AndroidJavaEnv.h b/Engine/Source/Runtime/Core/Public/Android/AndroidJavaEnv.h
--- a/Engine/Source/Runtime/Core/Public/Android/AndroidJavaEnv.h
+++ b/Engine/Source/Runtime/Core/Public/Android/AndroidJavaEnv.h
@@ -2,7 +2,7 @@
 
 #pragma once
 
-#if defined(USE_ANDROID_JNI) && USE_ANDROID_JNI
+#if USE_ANDROID_JNI
 
 #include "CoreTypes.h"
 #include "Containers/Array.h"
