[UBT]
* Added FileItem.IsPlaceholder that can be used to check if file item is a placeholder for a file that is about to be written
* Added exception if trying to access certain properties of file that is placeholder
* Fixed so Unity creation logic takes placeholder files into account and not query their size (which will now trigger an exception). Set defaul size

diff --git a/Engine/Source/Programs/Shared/EpicGames.Build/System/DirectoryItem.cs b/Engine/Source/Programs/Shared/EpicGames.Build/System/DirectoryItem.cs
--- a/Engine/Source/Programs/Shared/EpicGames.Build/System/DirectoryItem.cs
+++ b/Engine/Source/Programs/Shared/EpicGames.Build/System/DirectoryItem.cs
@@ -130,7 +130,7 @@ namespace UnrealBuildBase
 			{
 				FileReference fileRef = FileReference.Combine(Location, file);
 				FileItem fileItem = FileItem.GetItemByFileReference(fileRef);
-				fileItem.Update(new() { Attributes = FileAttributes.Normal }); 
+				fileItem.Update(new() { Attributes = FileAttributes.Normal, Size = long.MaxValue }); 
 				cacheFiles.Add(file, fileItem);
 			}
 
diff --git a/Engine/Source/Programs/Shared/EpicGames.Build/System/FileItem.cs b/Engine/Source/Programs/Shared/EpicGames.Build/System/FileItem.cs
--- a/Engine/Source/Programs/Shared/EpicGames.Build/System/FileItem.cs
+++ b/Engine/Source/Programs/Shared/EpicGames.Build/System/FileItem.cs
@@ -57,7 +57,7 @@ namespace UnrealBuildBase
 		/// <summary>
 		/// Size of the file if it exists, otherwise -1
 		/// </summary>
-		public long Length => Update().Length;
+		public long Length => UpdateCheckPlaceholder("Length").Length;
 
 		/// <summary>
 		/// The attributes for this file
@@ -67,12 +67,12 @@ namespace UnrealBuildBase
 		/// <summary>
 		/// The last write time of the file.
 		/// </summary>
-		public DateTime LastWriteTimeUtc => Update().LastWriteTimeUtc;
+		public DateTime LastWriteTimeUtc => UpdateCheckPlaceholder("LastWriteTimeUtc").LastWriteTimeUtc;
 
 		/// <summary>
 		/// The creation time of the file.
 		/// </summary>
-		public DateTime CreationTimeUtc => Update().CreationTimeUtc;
+		public DateTime CreationTimeUtc => UpdateCheckPlaceholder("CreationTimeUtc").CreationTimeUtc;
 
 		/// <summary>
 		/// Determines if the file has the given extension
@@ -100,6 +100,11 @@ namespace UnrealBuildBase
 			}
 		}
 
+		/// <summary>
+		/// Returns true if file is a placeholder file which can only be used to check if it exists. Length, create/write date are not valid
+		/// </summary>
+		public bool IsPlaceholder => Update().Length == long.MaxValue;
+
 		/// <summary>
 		/// Gets a FileItem corresponding to the given path
 		/// </summary>
@@ -247,6 +252,16 @@ namespace UnrealBuildBase
 			return newInfo;
 		}
 
+		private Info UpdateCheckPlaceholder(string property)
+		{
+			Info info = Update();
+			if (info.Length == long.MaxValue)
+			{
+				throw new Exception($"{property} can't be accessed, file {FullName} is placeholder");
+			}
+			return info;
+		}
+
 		internal void Update(DirectoryEnumerator.Entry entry)
 		{
 			if (_info == null)
diff --git a/Engine/Source/Programs/UnrealBuildTool/System/Unity.cs b/Engine/Source/Programs/UnrealBuildTool/System/Unity.cs
--- a/Engine/Source/Programs/UnrealBuildTool/System/Unity.cs
+++ b/Engine/Source/Programs/UnrealBuildTool/System/Unity.cs
@@ -16,6 +16,19 @@ namespace UnrealBuildTool
 		/// </summary>
 		public const string ModulePrefix = "Module.";
 
+		/// <summary>
+		/// Constant used to estimate file size for generated files
+		/// </summary>
+		public const long PlaceholderFileSize = 5000; // The generated files compile quite fast.. so we can pretend they are smaller than they are
+
+		/// <summary>
+		/// Safe way to get file length
+		/// </summary>
+		static long FileLength(FileItem file)
+		{
+			return file.IsPlaceholder ? PlaceholderFileSize : file.Length;
+		}
+
 		/// <summary>
 		/// A class which represents a list of files and the sum of their lengths.
 		/// </summary>
@@ -41,9 +54,9 @@ namespace UnrealBuildTool
 			{
 				Files.Add(File);
 
-				long FileLength = File.Length;
-				TotalLength += FileLength;
-				VirtualLength += FileLength;
+				long Length = FileLength(File);
+				TotalLength += Length;
+				VirtualLength += Length;
 			}
 
 			/// <summary>
@@ -88,7 +101,7 @@ namespace UnrealBuildTool
 			/// <param name="File">The file to add.</param>
 			public void AddFile(FileItem File)
 			{
-				if (SplitLength != -1 && File.Length > SplitLength)
+				if (SplitLength != -1 && FileLength(File) > SplitLength)
 				{
 					EndCurrentUnityFile();
 				}
@@ -106,7 +119,7 @@ namespace UnrealBuildTool
 			/// <param name="File">The file to add virtually.  Only the size of the file is tracked.</param>
 			public void AddVirtualFile(FileItem File)
 			{
-				if (SplitLength != -1 && File.Length > SplitLength)
+				if (SplitLength != -1 && FileLength(File) > SplitLength)
 				{
 					EndCurrentUnityFile();
 				}
@@ -183,7 +196,7 @@ namespace UnrealBuildTool
 			List<FileItem> NewSourceFiles = new List<FileItem>();
 
 			// Figure out size of all input files combined. We use this to determine whether to use larger unity threshold or not.
-			long TotalBytesInSourceFiles = SourceFiles.Sum(F => F.Length);
+			long TotalBytesInSourceFiles = SourceFiles.Sum(F => FileLength(F));
 
 			// We have an increased threshold for unity file size if, and only if, all files fit into the same unity file. This
 			// is beneficial when dealing with PCH files. The default PCH creation limit is X unity files so if we generate < X 
@@ -218,8 +231,8 @@ namespace UnrealBuildTool
 					// Sort oversized files to the end of the list so they will be placed into their own unity file
 					if (!bForceIntoSingleUnityFile)
 					{
-						bool bAIsOversized = A.Length > NumIncludedBytesPerUnitySource;
-						bool bBIsOversized = B.Length > NumIncludedBytesPerUnitySource;
+						bool bAIsOversized = FileLength(A) > NumIncludedBytesPerUnitySource;
+						bool bBIsOversized = FileLength(B) > NumIncludedBytesPerUnitySource;
 						if (bAIsOversized != bBIsOversized)
 						{
 							return !bAIsOversized && bBIsOversized ? -1 : 1;
