[UBT]
* Fixed so UHTTimestamp writing is not waited on by makefile creation.

diff --git a/Engine/Source/Programs/UnrealBuildTool/Configuration/UEBuildTarget.cs b/Engine/Source/Programs/UnrealBuildTool/Configuration/UEBuildTarget.cs
--- a/Engine/Source/Programs/UnrealBuildTool/Configuration/UEBuildTarget.cs
+++ b/Engine/Source/Programs/UnrealBuildTool/Configuration/UEBuildTarget.cs
@@ -2528,7 +2528,13 @@ namespace UnrealBuildTool
 		public async Task<TargetMakefile> BuildAsync(BuildConfiguration BuildConfiguration, ISourceFileWorkingSet WorkingSet, TargetDescriptor TargetDescriptor, ILogger Logger, bool bInitOnly = false)
 		{
 			await Task.CompletedTask;
+			(TargetMakefile Makefile, Task? UHTTimestampTask) = BuildMakefile(BuildConfiguration, WorkingSet, TargetDescriptor, Logger, bInitOnly);
+			UHTTimestampTask?.Wait();
+			return Makefile;
 
+		}
+		public (TargetMakefile, Task?) BuildMakefile(BuildConfiguration BuildConfiguration, ISourceFileWorkingSet WorkingSet, TargetDescriptor TargetDescriptor, ILogger Logger, bool bInitOnly = false)
+		{
 			CppConfiguration CppConfiguration = GetCppConfiguration(Configuration);
 
 			SourceFileMetadataCache MetadataCache = SourceFileMetadataCache.CreateHierarchy(ProjectFile, Logger);
@@ -2648,7 +2654,7 @@ namespace UnrealBuildTool
 			// UHT mode uses this to create the makefile at least to the point where UHT would have valid manifest information
 			if (bInitOnly)
 			{
-				return Makefile;
+				return (Makefile, null);
 			}
 
 			// We inject all files that we know UHT/VNI will create into the directory cache.. this way we can run UHT/VNI without waiting
@@ -2657,10 +2663,11 @@ namespace UnrealBuildTool
 
 			// NOTE: Even in Gather mode, we need to run UHT to make sure the files exist for the static action graph to be setup correctly.  This is because UHT generates .cpp
 			// files that are injected as top level prerequisites.  If UHT only emitted included header files, we wouldn't need to run it during the Gather phase at all.
-			Task? uhtTask = null;
+			Task? uhtRunTask = null;
+			Task? uhtTimestampTask = null;
 			if (Makefile.UObjectModules.Count > 0)
 			{
-				uhtTask = Task.Run(() => UHTExecution.ExecuteHeaderToolIfNecessaryAsync(BuildConfiguration, ProjectFile, Makefile, TargetName, Logger, false));
+				(uhtRunTask, uhtTimestampTask) = UHTExecution.ExecuteHeaderToolIfNecessary(BuildConfiguration, ProjectFile, Makefile, TargetName, Logger, false);
 			}
 
 #if __VPROJECT_AVAILABLE__
@@ -3124,15 +3131,15 @@ namespace UnrealBuildTool
 			}
 #endif
 
-			if (uhtTask != null)
+			if (uhtRunTask != null)
 			{
-				UHTExecution.WaitForExecution(uhtTask, Makefile);
+				UHTExecution.WaitForExecution(uhtRunTask, Makefile);
 
 				// Prefetch directory items for UHT folders since they are going to be used later for check out-of-date actions
 				Task task = FileMetadataPrefetch.QueueDirectoryTrees(ModulesToGenerateHeadersFor.Select(m => m.GeneratedCodeDirectoryUHT));
 			}
 
-			return Makefile;
+			return (Makefile, uhtTimestampTask);
 		}
 
 		void CreateStripExportsActions(UEToolChain TargetToolChain, TargetMakefileBuilder MakefileBuilder, TargetMakefile Makefile, CppCompileEnvironment CompileEnvironment)
diff --git a/Engine/Source/Programs/UnrealBuildTool/Modes/BuildMode.cs b/Engine/Source/Programs/UnrealBuildTool/Modes/BuildMode.cs
--- a/Engine/Source/Programs/UnrealBuildTool/Modes/BuildMode.cs
+++ b/Engine/Source/Programs/UnrealBuildTool/Modes/BuildMode.cs
@@ -389,7 +389,8 @@ namespace UnrealBuildTool
 		{
 			List<TargetMakefile> TargetMakefiles = new List<TargetMakefile>();
 
-			List<Task> SaveMakeFileTasks = new();
+			Task? LastSaveMakefileTask = null;
+			Task? LastUHTTimestampTask = null;
 
 			for (int Idx = 0; Idx < TargetDescriptors.Count; ++Idx)
 			{
@@ -402,16 +403,23 @@ namespace UnrealBuildTool
 					Utils.ExecuteCustomBuildSteps(InitScripts, Logger);
 				}
 
-				(TargetMakefile NewMakefile, FileReference? SaveToLocation) = await CreateMakefileAsync2(BuildConfiguration, TargetDescriptor, WorkingSet, Logger);
+				// Need to make sure to wait on previous run's timestamp since there are targets that write to same place..
+				// Code should be revisited
+				LastSaveMakefileTask?.Wait();
+				LastUHTTimestampTask?.Wait();
+
+				(TargetMakefile NewMakefile, FileReference? SaveToLocation, Task? UHTTimestampTask) = CreateMakefile2(BuildConfiguration, TargetDescriptor, WorkingSet, Logger);
 				TargetMakefiles.Add(NewMakefile);
 
+				LastUHTTimestampTask = UHTTimestampTask;
+
 				if (SaveToLocation != null)
 				{
-					SaveMakeFileTasks.Add(Task.Run(() =>
+					LastSaveMakefileTask = Task.Run(() =>
 					{
 						using ITimelineEvent _ = Timeline.ScopeEvent($"SavingMakefile");
 						NewMakefile.Save(SaveToLocation);
-					}));
+					});
 				}
 
 				if (!bSkipPreBuildTargets)
@@ -429,7 +437,8 @@ namespace UnrealBuildTool
 
 			await BuildAsync(TargetMakefiles.ToArray(), TargetDescriptors, BuildConfiguration, Options, WriteOutdatedActionsFile, Logger, ActionTypeFilter);
 
-			Task.WaitAll(SaveMakeFileTasks.ToArray());
+			LastSaveMakefileTask?.Wait();
+			LastUHTTimestampTask?.Wait();
 		}
 
 		static Lazy<MemoryMappedFileCache> s_memoryMappedFileCache = new Lazy<MemoryMappedFileCache>(() => new MemoryMappedFileCache());
@@ -1138,7 +1147,9 @@ namespace UnrealBuildTool
 		/// <returns>Makefile for the given target</returns>
 		internal static async Task<TargetMakefile> CreateMakefileAsync(BuildConfiguration BuildConfiguration, TargetDescriptor TargetDescriptor, ISourceFileWorkingSet WorkingSet, ILogger Logger)
 		{
-			var (NewMakefile, SaveToLocation) = await CreateMakefileAsync2(BuildConfiguration, TargetDescriptor, WorkingSet, Logger);
+			await Task.CompletedTask;
+
+			var (NewMakefile, SaveToLocation, UHTTimestampTask) = CreateMakefile2(BuildConfiguration, TargetDescriptor, WorkingSet, Logger);
 			if (SaveToLocation != null)
 			{
 				// Save the makefile for next time
@@ -1147,10 +1158,11 @@ namespace UnrealBuildTool
 					NewMakefile.Save(SaveToLocation);
 				}
 			}
+			UHTTimestampTask?.Wait();
 			return NewMakefile;
 		}
 
-		private static async Task<(TargetMakefile, FileReference?)> CreateMakefileAsync2(BuildConfiguration BuildConfiguration, TargetDescriptor TargetDescriptor, ISourceFileWorkingSet WorkingSet, ILogger Logger)
+		private static (TargetMakefile, FileReference?, Task?) CreateMakefile2(BuildConfiguration BuildConfiguration, TargetDescriptor TargetDescriptor, ISourceFileWorkingSet WorkingSet, ILogger Logger)
 		{
 			FileReference? SaveToLocation = null;
 
@@ -1195,6 +1207,8 @@ namespace UnrealBuildTool
 				}
 			}
 
+			Task? uhtTimeStampsTask = null;
+
 			// If we couldn't load a makefile, create a new one
 			if (Makefile == null)
 			{
@@ -1218,7 +1232,7 @@ namespace UnrealBuildTool
 				// Build the target
 				using (GlobalTracer.Instance.BuildSpan("UEBuildTarget.Build()").StartActive())
 				{
-					Makefile = await Target.BuildAsync(BuildConfiguration, WorkingSet, TargetDescriptor, Logger);
+					(Makefile, uhtTimeStampsTask) = Target.BuildMakefile(BuildConfiguration, WorkingSet, TargetDescriptor, Logger);
 				}
 
 				Makefile.MemoryPerActionGB = Target.Rules.MemoryPerActionGB;
@@ -1263,14 +1277,15 @@ namespace UnrealBuildTool
 				// If the target needs UHT to be run, we'll go ahead and do that now
 				if (Makefile.UObjectModules.Count > 0)
 				{
-					await UHTExecution.ExecuteHeaderToolIfNecessaryAsync(BuildConfiguration, TargetDescriptor.ProjectFile, Makefile, TargetDescriptor.Name, Logger);
+					(Task uhtRunTask, uhtTimeStampsTask) = UHTExecution.ExecuteHeaderToolIfNecessary(BuildConfiguration, TargetDescriptor.ProjectFile, Makefile, TargetDescriptor.Name, Logger, true);
+					uhtRunTask.Wait();
 				}
 
 #if __VPROJECT_AVAILABLE__
 				VNITask.Wait();
 #endif
 			}
-			return (Makefile, SaveToLocation);
+			return (Makefile, SaveToLocation, uhtTimeStampsTask);
 		}
 
 		/// <summary>
diff --git a/Engine/Source/Programs/UnrealBuildTool/System/UHTExecution.cs b/Engine/Source/Programs/UnrealBuildTool/System/UHTExecution.cs
--- a/Engine/Source/Programs/UnrealBuildTool/System/UHTExecution.cs
+++ b/Engine/Source/Programs/UnrealBuildTool/System/UHTExecution.cs
@@ -1033,7 +1033,7 @@ namespace UnrealBuildTool
 		/// </summary>
 		private static void UpdateTimestamps(List<UHTModuleInfo> UObjectModules, bool resetDirectoryCaches)
 		{
-			Parallel2.ForEach(UObjectModules, async Module =>
+			Parallel2.ForEach(UObjectModules, Module =>
 			{
 				if (!Module.bIsReadOnly)
 				{
@@ -1087,7 +1087,6 @@ namespace UnrealBuildTool
 							WriteHeaders(Module.PublicUObjectHeaders);
 							WriteHeaders(Module.InternalUObjectHeaders);
 							WriteHeaders(Module.PrivateUObjectHeaders);
-							await Task.CompletedTask;
 
 							if (resetDirectoryCaches)
 							{
@@ -1272,214 +1271,235 @@ namespace UnrealBuildTool
 			}
 		}
 
+		public static Task ExecuteHeaderToolIfNecessaryAsync(BuildConfiguration BuildConfiguration, FileReference? ProjectFile, TargetMakefile Makefile, string TargetName, ILogger Logger, bool resetDirectoryCaches = true)
+		{
+			return ExecuteHeaderToolIfNecessary(BuildConfiguration, ProjectFile, Makefile, TargetName, Logger, resetDirectoryCaches).timestamp;
+		}
+
 		/// <summary>
 		/// Builds and runs the header tool and touches the header directories.
 		/// Performs any early outs if headers need no changes, given the UObject modules, tool path, game name, and configuration
 		/// </summary>
-		public static async Task ExecuteHeaderToolIfNecessaryAsync(BuildConfiguration BuildConfiguration, FileReference? ProjectFile, TargetMakefile Makefile, string TargetName, ILogger Logger, bool resetDirectoryCaches = true)
+		public static (Task run, Task timestamp) ExecuteHeaderToolIfNecessary(BuildConfiguration BuildConfiguration, FileReference? ProjectFile, TargetMakefile Makefile, string TargetName, ILogger Logger, bool resetDirectoryCaches = true)
 		{
-			if (ProgressWriter.bWriteMarkup)
-			{
-				Logger.LogInformation("@progress push 5%");
-			}
-			using (ProgressWriter Progress = new("Generating code...", false, Logger))
+			bool runTimestamp = false;
+
+			Task runTask = new Task(() =>
 			{
-				using var uhtManifestScope = Timeline.ScopeEvent("UHTManifest");
+				if (ProgressWriter.bWriteMarkup)
+				{
+					Logger.LogInformation("@progress push 5%");
+				}
+				using (ProgressWriter Progress = new("Generating code...", false, Logger))
+				{
+					using var uhtManifestScope = Timeline.ScopeEvent("UHTManifest");
 
-				string RootLocalPath = Unreal.RootDirectory.FullName;
+					string RootLocalPath = Unreal.RootDirectory.FullName;
 
-				// Get the path to the assemblies we care about
-				string UhtAssemblyPath = typeof(UhtSession).Assembly.Location;
+					// Get the path to the assemblies we care about
+					string UhtAssemblyPath = typeof(UhtSession).Assembly.Location;
 
-				// Get a composite date/time for critical assemblies being used
-				DateTime CompositeTimestamp = LatestDateTime(new FileInfo(Assembly.GetExecutingAssembly().Location).LastWriteTimeUtc, new FileInfo(UhtAssemblyPath).LastWriteTimeUtc);
-				if (Makefile.EnabledUhtPlugins != null)
-				{
-					foreach (FileReference Plugin in Makefile.EnabledUhtPlugins)
+					// Get a composite date/time for critical assemblies being used
+					DateTime CompositeTimestamp = LatestDateTime(new FileInfo(Assembly.GetExecutingAssembly().Location).LastWriteTimeUtc, new FileInfo(UhtAssemblyPath).LastWriteTimeUtc);
+					if (Makefile.EnabledUhtPlugins != null)
 					{
-						CompositeTimestamp = LatestDateTime(CompositeTimestamp, new FileInfo(Plugin.FullName).LastWriteTimeUtc);
+						foreach (FileReference Plugin in Makefile.EnabledUhtPlugins)
+						{
+							CompositeTimestamp = LatestDateTime(CompositeTimestamp, new FileInfo(Plugin.FullName).LastWriteTimeUtc);
+						}
 					}
-				}
 
-				bool bUHTNeedsToRun = false;
+					bool bUHTNeedsToRun = false;
 
-				// Check explicit settings passed to UHT - this info is also saved in the manifest but we load just the settings here
-				FileReference UHTSettingsFileName = GetUHTSettingsFileName(Makefile, TargetName);
-				if (!FileReference.Exists(UHTSettingsFileName))
-				{
-					bUHTNeedsToRun = true;
-				}
-				else
-				{
-					using (BinaryArchiveReader reader = new(UHTSettingsFileName))
+					// Check explicit settings passed to UHT - this info is also saved in the manifest but we load just the settings here
+					FileReference UHTSettingsFileName = GetUHTSettingsFileName(Makefile, TargetName);
+					if (!FileReference.Exists(UHTSettingsFileName))
 					{
-						// Note that even if default/omitted values cause the settings to compare equal here, the UHT module version check below 
-						// will ensure that we rerun UHT when new properties are added to UHTTargetSettings
-						UHTTargetSettings settings = new(reader);
-						if (settings != Makefile.UHTTargetSettings)
+						bUHTNeedsToRun = true;
+					}
+					else
+					{
+						using (BinaryArchiveReader reader = new(UHTSettingsFileName))
 						{
-							bUHTNeedsToRun = true;
+							// Note that even if default/omitted values cause the settings to compare equal here, the UHT module version check below 
+							// will ensure that we rerun UHT when new properties are added to UHTTargetSettings
+							UHTTargetSettings settings = new(reader);
+							if (settings != Makefile.UHTTargetSettings)
+							{
+								bUHTNeedsToRun = true;
+							}
 						}
 					}
-				}
 
-				// ensure the headers are up to date
-				if (BuildConfiguration.bForceHeaderGeneration)
-				{
-					bUHTNeedsToRun = true;
-				}
-				else if (!bUHTNeedsToRun && AreGeneratedCodeFilesOutOfDate(Makefile.UObjectModules, CompositeTimestamp, Logger))
-				{
-					bUHTNeedsToRun = true;
-				}
-
-				// Check we're not using a different version of UHT
-				FileReference ModuleInfoFileName = GetUHTModuleInfoFileName(Makefile, TargetName);
-				FileReference ToolInfoFile = ModuleInfoFileName.ChangeExtension(".uhtpath");
-				if (!bUHTNeedsToRun)
-				{
-					if (!FileReference.Exists(ToolInfoFile))
+					// ensure the headers are up to date
+					if (BuildConfiguration.bForceHeaderGeneration)
 					{
 						bUHTNeedsToRun = true;
 					}
-					else if (FileReference.ReadAllText(ToolInfoFile) != UhtAssemblyPath)
+					else if (!bUHTNeedsToRun && AreGeneratedCodeFilesOutOfDate(Makefile.UObjectModules, CompositeTimestamp, Logger))
 					{
 						bUHTNeedsToRun = true;
 					}
-				}
 
-				// Get the file containing dependencies for the generated code
-				FileReference ExternalDependenciesFile = GetUHTDepsFileName(ModuleInfoFileName);
-				if (AreExternalDependenciesOutOfDate(ExternalDependenciesFile))
-				{
-					bUHTNeedsToRun = true;
-				}
+					// Check we're not using a different version of UHT
+					FileReference ModuleInfoFileName = GetUHTModuleInfoFileName(Makefile, TargetName);
+					FileReference ToolInfoFile = ModuleInfoFileName.ChangeExtension(".uhtpath");
+					if (!bUHTNeedsToRun)
+					{
+						if (!FileReference.Exists(ToolInfoFile))
+						{
+							bUHTNeedsToRun = true;
+						}
+						else if (FileReference.ReadAllText(ToolInfoFile) != UhtAssemblyPath)
+						{
+							bUHTNeedsToRun = true;
+						}
+					}
+
+					// Get the file containing dependencies for the generated code
+					FileReference ExternalDependenciesFile = GetUHTDepsFileName(ModuleInfoFileName);
+					if (AreExternalDependenciesOutOfDate(ExternalDependenciesFile))
+					{
+						bUHTNeedsToRun = true;
+					}
 
 
-				if (bUHTNeedsToRun)
-				{
-					Progress.Write(1, 3);
+					if (bUHTNeedsToRun)
+					{
+						Progress.Write(1, 3);
 
-					WriteUHTManifest(Makefile, TargetName, ModuleInfoFileName, ExternalDependenciesFile);
-					string ActualTargetName = String.IsNullOrEmpty(TargetName) ? "UE5" : TargetName;
-					Logger.LogInformation("Parsing headers for {ActualTargetName}", ActualTargetName);
+						WriteUHTManifest(Makefile, TargetName, ModuleInfoFileName, ExternalDependenciesFile);
+						string ActualTargetName = String.IsNullOrEmpty(TargetName) ? "UE5" : TargetName;
+						Logger.LogInformation("Parsing headers for {ActualTargetName}", ActualTargetName);
 
-					// Generate the command line
-					List<string> CmdArgs = [
-						(ProjectFile != null) ? ProjectFile.FullName : TargetName,
+						// Generate the command line
+						List<string> CmdArgs = [
+							(ProjectFile != null) ? ProjectFile.FullName : TargetName,
 						ModuleInfoFileName.FullName,
 						"-WarningsAsErrors",
 					];
 
-					if (Unreal.IsEngineInstalled())
-					{
-						CmdArgs.Add("-installed");
-					}
+						if (Unreal.IsEngineInstalled())
+						{
+							CmdArgs.Add("-installed");
+						}
 
-					if (BuildConfiguration.bFailIfGeneratedCodeChanges)
-					{
-						CmdArgs.Add("-FailIfGeneratedCodeChanges");
-					}
+						if (BuildConfiguration.bFailIfGeneratedCodeChanges)
+						{
+							CmdArgs.Add("-FailIfGeneratedCodeChanges");
+						}
 
-					if (Makefile.UHTAdditionalArguments != null)
-					{
-						foreach (string Arg in Makefile.UHTAdditionalArguments)
+						if (Makefile.UHTAdditionalArguments != null)
 						{
-							if (Arg[0] == '"' && Arg[^1] == '"')
+							foreach (string Arg in Makefile.UHTAdditionalArguments)
 							{
-								CmdArgs.Add(Arg[1..^1]);
-							}
-							else
-							{
-								CmdArgs.Add(Arg);
+								if (Arg[0] == '"' && Arg[^1] == '"')
+								{
+									CmdArgs.Add(Arg[1..^1]);
+								}
+								else
+								{
+									CmdArgs.Add(Arg);
+								}
 							}
 						}
-					}
-					CommandLineArguments Arguments = new([.. CmdArgs]);
+						CommandLineArguments Arguments = new([.. CmdArgs]);
 
-					StringBuilder CmdLine = new();
-					foreach (string Arg in CmdArgs)
-					{
-						CmdLine.AppendArgument(Arg);
-					}
+						StringBuilder CmdLine = new();
+						foreach (string Arg in CmdArgs)
+						{
+							CmdLine.AppendArgument(Arg);
+						}
 
-					uhtManifestScope.Finish();
+						uhtManifestScope.Finish();
 
-					Logger.LogInformation("  Running Internal UnrealHeaderTool {CmdLine}", CmdLine);
+						Logger.LogInformation("  Running Internal UnrealHeaderTool {CmdLine}", CmdLine);
 
-					using var uhtExecute = Timeline.ScopeEvent("UHT");
+						using var uhtExecute = Timeline.ScopeEvent("UHT");
 
-					// Run UHT
-					IScope Timer = GlobalTracer.Instance.BuildSpan("Executing UnrealHeaderTool").StartActive();
-					UnrealHeaderToolMode UHTTool = new();
+						// Run UHT
+						IScope Timer = GlobalTracer.Instance.BuildSpan("Executing UnrealHeaderTool").StartActive();
+						UnrealHeaderToolMode UHTTool = new();
 
-					CompilationResult UHTResult = (CompilationResult)await UHTTool.ExecuteAsync(Arguments, Logger);
-					Timer.Span.Finish();
+						CompilationResult UHTResult = (CompilationResult)UHTTool.ExecuteAsync(Arguments, Logger).Result;
+						Timer.Span.Finish();
 
-					uhtExecute.Dispose();
+						uhtExecute.Dispose();
 
-					if (UHTResult != CompilationResult.Succeeded)
-					{
-						// On Linux and Mac, the shell will return 128+signal number exit codes if UHT gets a signal (e.g. crashes or is interrupted)
-						if ((BuildHostPlatform.Current.Platform == UnrealTargetPlatform.Linux ||
-							BuildHostPlatform.Current.Platform == UnrealTargetPlatform.Mac) &&
-							(int)(UHTResult) >= 128
-							)
+						if (UHTResult != CompilationResult.Succeeded)
 						{
-							// SIGINT is 2, so 128 + SIGINT is 130
-							UHTResult = ((int)(UHTResult) == 130) ? CompilationResult.Canceled : CompilationResult.CrashOrAssert;
+							// On Linux and Mac, the shell will return 128+signal number exit codes if UHT gets a signal (e.g. crashes or is interrupted)
+							if ((BuildHostPlatform.Current.Platform == UnrealTargetPlatform.Linux ||
+								BuildHostPlatform.Current.Platform == UnrealTargetPlatform.Mac) &&
+								(int)(UHTResult) >= 128
+								)
+							{
+								// SIGINT is 2, so 128 + SIGINT is 130
+								UHTResult = ((int)(UHTResult) == 130) ? CompilationResult.Canceled : CompilationResult.CrashOrAssert;
+							}
+
+							if (BuildHostPlatform.Current.Platform == UnrealTargetPlatform.Win64 &&
+								(int)(UHTResult) < 0)
+							{
+								Logger.LogError("UnrealHeaderTool failed with exit code 0x{Result:X} - check that Unreal Engine prerequisites are installed.", (int)UHTResult);
+							}
+
+							throw new CompilationResultException(UHTResult);
 						}
 
-						if (BuildHostPlatform.Current.Platform == UnrealTargetPlatform.Win64 &&
-							(int)(UHTResult) < 0)
+						// Update the tool info file
+						DirectoryReference.CreateDirectory(ToolInfoFile.Directory);
+						FileReference.WriteAllText(ToolInfoFile, UhtAssemblyPath);
+
+						// Update saved UHT settings
+						using (BinaryArchiveWriter writer = new(UHTSettingsFileName))
 						{
-							Logger.LogError("UnrealHeaderTool failed with exit code 0x{Result:X} - check that Unreal Engine prerequisites are installed.", (int)UHTResult);
+							Makefile.UHTTargetSettings.Write(writer);
 						}
 
-						throw new CompilationResultException(UHTResult);
-					}
-
-					// Update the tool info file
-					DirectoryReference.CreateDirectory(ToolInfoFile.Directory);
-					FileReference.WriteAllText(ToolInfoFile, UhtAssemblyPath);
-					
-					// Update saved UHT settings
-					using (BinaryArchiveWriter writer = new(UHTSettingsFileName))
-					{
-						Makefile.UHTTargetSettings.Write(writer);
-					}
-					
-					if (resetDirectoryCaches)
-					{
-						// Now that UHT has successfully finished generating code, we need to update all cached FileItems in case their last write time has changed.
-						// Otherwise UBT might not detect changes UHT made.
-						using (GlobalTracer.Instance.BuildSpan("UHTExecution.ResetCachedHeaderInfo()").StartActive())
+						if (resetDirectoryCaches)
 						{
-							ResetCachedHeaderInfo(Makefile.UObjectModules);
+							// Now that UHT has successfully finished generating code, we need to update all cached FileItems in case their last write time has changed.
+							// Otherwise UBT might not detect changes UHT made.
+							using (GlobalTracer.Instance.BuildSpan("UHTExecution.ResetCachedHeaderInfo()").StartActive())
+							{
+								ResetCachedHeaderInfo(Makefile.UObjectModules);
+							}
 						}
 					}
-				}
-				else
-				{
-					uhtManifestScope.Finish();
+					else
+					{
+						uhtManifestScope.Finish();
 
-					Logger.LogDebug("Generated code is up to date.");
-				}
+						Logger.LogDebug("Generated code is up to date.");
+					}
+
+					Progress.Write(2, 3);
 
-				Progress.Write(2, 3);
+					runTimestamp = true;
 
-				using (GlobalTracer.Instance.BuildSpan("UHTExecution.UpdateTimestamps()").StartActive())
+					Progress.Write(3, 3);
+				}
+				if (ProgressWriter.bWriteMarkup)
 				{
-					using var uhtExecute = Timeline.ScopeEvent("UHTTimestamps");
-					UpdateTimestamps(Makefile.UObjectModules, resetDirectoryCaches);
+					Logger.LogInformation("@progress pop");
 				}
-
-				Progress.Write(3, 3);
-			}
-			if (ProgressWriter.bWriteMarkup)
+			});
+			Task timestampTask = runTask.ContinueWith((prev) =>
 			{
-				Logger.LogInformation("@progress pop");
-			}
+				if (runTimestamp)
+				{
+					using (GlobalTracer.Instance.BuildSpan("UHTExecution.UpdateTimestamps()").StartActive())
+					{
+						using var uhtExecute = Timeline.ScopeEvent("UHTTimestamps");
+						UpdateTimestamps(Makefile.UObjectModules, resetDirectoryCaches);
+					}
+				}
+			});
+
+			runTask.Start();
+
+			return (runTask, timestampTask);
 		}
 
 		/// <summary>
