[UBT] Reinstate old UEBuildTarget private API
- Old CTOR API is required for a 5.7 change

diff --git a/Engine/Source/Programs/UnrealBuildTool/Configuration/UEBuildTarget.cs b/Engine/Source/Programs/UnrealBuildTool/Configuration/UEBuildTarget.cs
--- a/Engine/Source/Programs/UnrealBuildTool/Configuration/UEBuildTarget.cs
+++ b/Engine/Source/Programs/UnrealBuildTool/Configuration/UEBuildTarget.cs
@@ -1793,11 +1793,9 @@ namespace UnrealBuildTool
 		/// <param name="InDescriptor">Target descriptor</param>
 		/// <param name="InRules">The target rules, as created by RulesCompiler.</param>
 		/// <param name="InRulesAssembly">The chain of rules assemblies that this target was created with</param>
-		/// <param name="InMetadataCache"></param>
 		/// <param name="Logger">Logger for output</param>
-		private UEBuildTarget(TargetDescriptor InDescriptor, ReadOnlyTargetRules InRules, RulesAssembly InRulesAssembly, SourceFileMetadataCache InMetadataCache, ILogger Logger)
+		private UEBuildTarget(TargetDescriptor InDescriptor, ReadOnlyTargetRules InRules, RulesAssembly InRulesAssembly, ILogger Logger)
 		{
-			MetadataCache = InMetadataCache;
 			ProjectFile = InDescriptor.ProjectFile;
 			TargetName = InRules.Name;
 			AppName = InRules.Name;
@@ -1880,6 +1878,21 @@ namespace UnrealBuildTool
 			{
 				ProjectDescriptor = ProjectDescriptor.FromFile(ProjectFile);
 			}
+
+			MetadataCache = SourceFileMetadataCache.CreateHierarchy(ProjectFile, Logger);
+		}
+
+		/// <summary>
+		/// Constructor.
+		/// </summary>
+		/// <param name="InDescriptor">Target descriptor</param>
+		/// <param name="InRules">The target rules, as created by RulesCompiler.</param>
+		/// <param name="InRulesAssembly">The chain of rules assemblies that this target was created with</param>
+		/// <param name="InMetadataCache"></param>
+		/// <param name="Logger">Logger for output</param>
+		private UEBuildTarget(TargetDescriptor InDescriptor, ReadOnlyTargetRules InRules, RulesAssembly InRulesAssembly, SourceFileMetadataCache InMetadataCache, ILogger Logger) : this(InDescriptor, InRules, InRulesAssembly, Logger)
+		{
+			MetadataCache = InMetadataCache;
 		}
 
 		/// <summary>
