[UBT]
* Fixed UEBuildTarget so old ctor is calling new ctor

diff --git a/Engine/Source/Programs/UnrealBuildTool/Configuration/UEBuildTarget.cs b/Engine/Source/Programs/UnrealBuildTool/Configuration/UEBuildTarget.cs
--- a/Engine/Source/Programs/UnrealBuildTool/Configuration/UEBuildTarget.cs
+++ b/Engine/Source/Programs/UnrealBuildTool/Configuration/UEBuildTarget.cs
@@ -1788,14 +1788,28 @@ namespace UnrealBuildTool
 		}
 
 		/// <summary>
-		/// Constructor.
+		/// OLD Constructor.
 		/// </summary>
 		/// <param name="InDescriptor">Target descriptor</param>
 		/// <param name="InRules">The target rules, as created by RulesCompiler.</param>
 		/// <param name="InRulesAssembly">The chain of rules assemblies that this target was created with</param>
 		/// <param name="Logger">Logger for output</param>
 		private UEBuildTarget(TargetDescriptor InDescriptor, ReadOnlyTargetRules InRules, RulesAssembly InRulesAssembly, ILogger Logger)
+			: this(InDescriptor, InRules, InRulesAssembly, SourceFileMetadataCache.CreateHierarchy(InDescriptor.ProjectFile, Logger), Logger)
 		{
+		}
+
+		/// <summary>
+		/// Constructor.
+		/// </summary>
+		/// <param name="InDescriptor">Target descriptor</param>
+		/// <param name="InRules">The target rules, as created by RulesCompiler.</param>
+		/// <param name="InRulesAssembly">The chain of rules assemblies that this target was created with</param>
+		/// <param name="InMetadataCache"></param>
+		/// <param name="Logger">Logger for output</param>
+		private UEBuildTarget(TargetDescriptor InDescriptor, ReadOnlyTargetRules InRules, RulesAssembly InRulesAssembly, SourceFileMetadataCache InMetadataCache, ILogger Logger)
+		{
+			MetadataCache = InMetadataCache;
 			ProjectFile = InDescriptor.ProjectFile;
 			TargetName = InRules.Name;
 			AppName = InRules.Name;
@@ -1878,21 +1892,6 @@ namespace UnrealBuildTool
 			{
 				ProjectDescriptor = ProjectDescriptor.FromFile(ProjectFile);
 			}
-
-			MetadataCache = SourceFileMetadataCache.CreateHierarchy(ProjectFile, Logger);
-		}
-
-		/// <summary>
-		/// Constructor.
-		/// </summary>
-		/// <param name="InDescriptor">Target descriptor</param>
-		/// <param name="InRules">The target rules, as created by RulesCompiler.</param>
-		/// <param name="InRulesAssembly">The chain of rules assemblies that this target was created with</param>
-		/// <param name="InMetadataCache"></param>
-		/// <param name="Logger">Logger for output</param>
-		private UEBuildTarget(TargetDescriptor InDescriptor, ReadOnlyTargetRules InRules, RulesAssembly InRulesAssembly, SourceFileMetadataCache InMetadataCache, ILogger Logger) : this(InDescriptor, InRules, InRulesAssembly, Logger)
-		{
-			MetadataCache = InMetadataCache;
 		}
 
 		/// <summary>
