[UBT]
* Changed so BuildMode CreateMakefile is now CreateMakefiles that takes in an array of target descriptors. First step towards being able to run targets in parallel (there is logic needed to see so two targets don't read/write same things.. and can't be done on the outside)

diff --git a/Engine/Source/Programs/UnrealBuildTool/Modes/BuildMode.cs b/Engine/Source/Programs/UnrealBuildTool/Modes/BuildMode.cs
--- a/Engine/Source/Programs/UnrealBuildTool/Modes/BuildMode.cs
+++ b/Engine/Source/Programs/UnrealBuildTool/Modes/BuildMode.cs
@@ -387,58 +387,11 @@ namespace UnrealBuildTool
 		/// <returns>Result from the compilation</returns>
 		public static async Task BuildAsync(List<TargetDescriptor> TargetDescriptors, BuildConfiguration BuildConfiguration, ISourceFileWorkingSet WorkingSet, BuildOptions Options, FileReference? WriteOutdatedActionsFile, ILogger Logger, bool bSkipPreBuildTargets = false, string? ActionTypeFilter = null)
 		{
-			List<TargetMakefile> TargetMakefiles = new List<TargetMakefile>();
+			(List<TargetMakefile> Makefiles, List<Task> Tasks) = CreateMakefiles(BuildConfiguration, TargetDescriptors, WorkingSet, Logger, true, true);
 
-			Task? LastSaveMakefileTask = null;
-			Task? LastUHTTimestampTask = null;
+			await BuildAsync(Makefiles.ToArray(), TargetDescriptors, BuildConfiguration, Options, WriteOutdatedActionsFile, Logger, ActionTypeFilter);
 
-			for (int Idx = 0; Idx < TargetDescriptors.Count; ++Idx)
-			{
-				TargetDescriptor TargetDescriptor = TargetDescriptors[Idx];
-				using ITimelineEvent SetupScope = Timeline.ScopeEvent($"{TargetDescriptor.Name} {TargetDescriptor.Platform} {TargetDescriptor.Configuration}");
-				// Create and execute the init scripts
-				FileReference[] InitScripts = CreateInitScripts(TargetDescriptor);
-				if (InitScripts.Length > 0)
-				{
-					Utils.ExecuteCustomBuildSteps(InitScripts, Logger);
-				}
-
-				// Need to make sure to wait on previous run's timestamp since there are targets that write to same place..
-				// Code should be revisited
-				LastSaveMakefileTask?.Wait();
-				LastUHTTimestampTask?.Wait();
-
-				(TargetMakefile NewMakefile, FileReference? SaveToLocation, Task? UHTTimestampTask) = CreateMakefile2(BuildConfiguration, TargetDescriptor, WorkingSet, Logger);
-				TargetMakefiles.Add(NewMakefile);
-
-				LastUHTTimestampTask = UHTTimestampTask;
-
-				if (SaveToLocation != null)
-				{
-					LastSaveMakefileTask = Task.Run(() =>
-					{
-						using ITimelineEvent _ = Timeline.ScopeEvent($"SavingMakefile");
-						NewMakefile.Save(SaveToLocation);
-					});
-				}
-
-				if (!bSkipPreBuildTargets)
-				{
-					foreach (TargetInfo PreBuildTarget in NewMakefile.PreBuildTargets)
-					{
-						TargetDescriptor NewTarget = TargetDescriptor.FromTargetInfo(PreBuildTarget);
-						if (!TargetDescriptors.Contains(NewTarget))
-						{
-							TargetDescriptors.Add(NewTarget);
-						}
-					}
-				}
-			}
-
-			await BuildAsync(TargetMakefiles.ToArray(), TargetDescriptors, BuildConfiguration, Options, WriteOutdatedActionsFile, Logger, ActionTypeFilter);
-
-			LastSaveMakefileTask?.Wait();
-			LastUHTTimestampTask?.Wait();
+			Task.WaitAll(Tasks.ToArray());
 		}
 
 		static Lazy<MemoryMappedFileCache> s_memoryMappedFileCache = new Lazy<MemoryMappedFileCache>(() => new MemoryMappedFileCache());
@@ -1149,143 +1102,205 @@ namespace UnrealBuildTool
 		{
 			await Task.CompletedTask;
 
-			var (NewMakefile, SaveToLocation, UHTTimestampTask) = CreateMakefile2(BuildConfiguration, TargetDescriptor, WorkingSet, Logger);
-			if (SaveToLocation != null)
-			{
-				// Save the makefile for next time
-				using (GlobalTracer.Instance.BuildSpan("TargetMakefile.Save()").StartActive())
-				{
-					NewMakefile.Save(SaveToLocation);
-				}
-			}
-			UHTTimestampTask?.Wait();
-			return NewMakefile;
+			(List<TargetMakefile> Makefiles, List<Task> Tasks) = CreateMakefiles(BuildConfiguration, [TargetDescriptor], WorkingSet, Logger, false, false);
+
+			Task.WaitAll(Tasks.ToArray());
+			return Makefiles[0];
 		}
 
-		private static (TargetMakefile, FileReference?, Task?) CreateMakefile2(BuildConfiguration BuildConfiguration, TargetDescriptor TargetDescriptor, ISourceFileWorkingSet WorkingSet, ILogger Logger)
+		private static (List<TargetMakefile>, List<Task>) CreateMakefiles(BuildConfiguration BuildConfiguration, List<TargetDescriptor> TargetDescriptors, ISourceFileWorkingSet WorkingSet, ILogger Logger, bool runInitScripts, bool bSkipPreBuildTargets)
 		{
-			FileReference? SaveToLocation = null;
+			List<TargetMakefile> Makefiles = [];
 
-			// Get the path to the makefile for this target
-			FileReference? MakefileLocation = null;
-			if (BuildConfiguration.bUseUBTMakefiles)
-			{
-				MakefileLocation = TargetMakefile.GetLocation(TargetDescriptor.ProjectFile, TargetDescriptor.Name, TargetDescriptor.Platform, TargetDescriptor.Architectures, TargetDescriptor.Configuration, TargetDescriptor.IntermediateEnvironment);
-			}
+			Task? LastSaveMakefileTask = null;
+			Task? LastUHTTimestampTask = null;
+
+			bool skipRulesCompile = BuildConfiguration.bSkipRulesCompile;
+			bool forceRulesCompile = BuildConfiguration.bForceRulesCompile;
+			bool usePrecompiled = BuildConfiguration.bUsePrecompiled;
 
-			// Try to load an existing makefile
-			TargetMakefile? Makefile = null;
-			if (MakefileLocation != null)
+			for (int Idx = 0; Idx < TargetDescriptors.Count; ++Idx)
 			{
-				using (GlobalTracer.Instance.BuildSpan("TargetMakefile.Load()").StartActive())
+				TargetDescriptor TargetDescriptor = TargetDescriptors[Idx];
+
+
+				if (runInitScripts)
 				{
-					string? ReasonNotLoaded;
-					Makefile = TargetMakefile.Load(MakefileLocation, TargetDescriptor.ProjectFile, TargetDescriptor.Platform, TargetDescriptor.AdditionalArguments.GetRawArray(), Logger, out ReasonNotLoaded);
-					if (Makefile == null)
+					using ITimelineEvent SetupScope = Timeline.ScopeEvent($"{TargetDescriptor.Name} {TargetDescriptor.Platform} {TargetDescriptor.Configuration}");
+					// Create and execute the init scripts
+					FileReference[] InitScripts = CreateInitScripts(TargetDescriptor);
+					if (InitScripts.Length > 0)
 					{
-						Logger.LogInformation("Creating makefile for {TargetDescriptorName} ({ReasonNotLoaded})", TargetDescriptor.Name, ReasonNotLoaded);
+						Utils.ExecuteCustomBuildSteps(InitScripts, Logger);
 					}
 				}
-			}
-
-			// If we have a makefile, execute the pre-build steps and check it's still valid
-			bool bHasRunPreBuildScripts = false;
-			if (Makefile != null)
-			{
-				// Execute the scripts.
-				Utils.ExecuteCustomBuildSteps(Makefile.PreBuildScripts, Logger);
 
-				// Don't run the pre-build steps again, even if we invalidate the makefile.
-				bHasRunPreBuildScripts = true;
+				FileReference? SaveToLocation = null;
 
-				// Check that the makefile is still valid
-				string? Reason;
-				if (!TargetMakefile.IsValidForSourceFiles(Makefile, TargetDescriptor.ProjectFile, TargetDescriptor.Platform, WorkingSet, Logger, out Reason))
+				// Get the path to the makefile for this target
+				FileReference? MakefileLocation = null;
+				if (BuildConfiguration.bUseUBTMakefiles)
 				{
-					Logger.LogInformation("Invalidating makefile for {TargetDescriptorName} ({Reason})", TargetDescriptor.Name, Reason);
-					Makefile = null;
+					MakefileLocation = TargetMakefile.GetLocation(TargetDescriptor.ProjectFile, TargetDescriptor.Name, TargetDescriptor.Platform, TargetDescriptor.Architectures, TargetDescriptor.Configuration, TargetDescriptor.IntermediateEnvironment);
 				}
-			}
 
-			Task? uhtTimeStampsTask = null;
+				// Need to make sure to wait on previous run's makefile
+				LastSaveMakefileTask?.Wait();
 
-			// If we couldn't load a makefile, create a new one
-			if (Makefile == null)
-			{
-				// Create the target
-				UEBuildTarget Target;
-				using (GlobalTracer.Instance.BuildSpan("UEBuildTarget.Create()").StartActive())
+				// Try to load an existing makefile
+				TargetMakefile? Makefile = null;
+				if (MakefileLocation != null)
 				{
-					Target = UEBuildTarget.Create(TargetDescriptor, BuildConfiguration, Logger);
+					using (GlobalTracer.Instance.BuildSpan("TargetMakefile.Load()").StartActive())
+					{
+						string? ReasonNotLoaded;
+						Makefile = TargetMakefile.Load(MakefileLocation, TargetDescriptor.ProjectFile, TargetDescriptor.Platform, TargetDescriptor.AdditionalArguments.GetRawArray(), Logger, out ReasonNotLoaded);
+						if (Makefile == null)
+						{
+							Logger.LogInformation("Creating makefile for {TargetDescriptorName} ({ReasonNotLoaded})", TargetDescriptor.Name, ReasonNotLoaded);
+						}
+					}
 				}
 
-				// Create the pre-build scripts
-				FileReference[] PreBuildScripts = Target.CreatePreBuildScripts();
-
-				// Execute the pre-build scripts
-				if (!bHasRunPreBuildScripts)
+				// If we have a makefile, execute the pre-build steps and check it's still valid
+				bool bHasRunPreBuildScripts = false;
+				if (Makefile != null)
 				{
-					Utils.ExecuteCustomBuildSteps(PreBuildScripts, Logger);
+					// Execute the scripts.
+					Utils.ExecuteCustomBuildSteps(Makefile.PreBuildScripts, Logger);
+
+					// Don't run the pre-build steps again, even if we invalidate the makefile.
 					bHasRunPreBuildScripts = true;
+
+					// Check that the makefile is still valid
+					string? Reason;
+					if (!TargetMakefile.IsValidForSourceFiles(Makefile, TargetDescriptor.ProjectFile, TargetDescriptor.Platform, WorkingSet, Logger, out Reason))
+					{
+						Logger.LogInformation("Invalidating makefile for {TargetDescriptorName} ({Reason})", TargetDescriptor.Name, Reason);
+						Makefile = null;
+					}
 				}
 
-				// Build the target
-				using (GlobalTracer.Instance.BuildSpan("UEBuildTarget.Build()").StartActive())
+				// Need to make sure to wait on previous run's timestamp since there are targets that write to same place..
+				LastUHTTimestampTask?.Wait();
+
+				Task? uhtTimeStampsTask = null;
+
+				// If we couldn't load a makefile, create a new one
+				if (Makefile == null)
 				{
-					(Makefile, uhtTimeStampsTask) = Target.BuildMakefile(BuildConfiguration, WorkingSet, TargetDescriptor, Logger);
-				}
+					// Create the target
+					UEBuildTarget Target;
+					using (GlobalTracer.Instance.BuildSpan("UEBuildTarget.Create()").StartActive())
+					{
+						Target = UEBuildTarget.Create(TargetDescriptor, skipRulesCompile, forceRulesCompile, usePrecompiled, Logger);
+					}
+
+					// Create the pre-build scripts
+					FileReference[] PreBuildScripts = Target.CreatePreBuildScripts();
 
-				Makefile.MemoryPerActionGB = Target.Rules.MemoryPerActionGB;
+					// Execute the pre-build scripts
+					if (!bHasRunPreBuildScripts)
+					{
+						Utils.ExecuteCustomBuildSteps(PreBuildScripts, Logger);
+						bHasRunPreBuildScripts = true;
+					}
+
+					// Build the target
+					using (GlobalTracer.Instance.BuildSpan("UEBuildTarget.Build()").StartActive())
+					{
+						(Makefile, uhtTimeStampsTask) = Target.BuildMakefile(BuildConfiguration, WorkingSet, TargetDescriptor, Logger);
+					}
 
-				// Save the pre-build scripts onto the makefile
-				Makefile.PreBuildScripts = PreBuildScripts;
+					Makefile.MemoryPerActionGB = Target.Rules.MemoryPerActionGB;
 
-				// Save the additional command line arguments
-				Makefile.AdditionalArguments = TargetDescriptor.AdditionalArguments.GetRawArray();
+					// Save the pre-build scripts onto the makefile
+					Makefile.PreBuildScripts = PreBuildScripts;
 
-				// Save the environment variables
-				foreach (System.Collections.DictionaryEntry? EnvironmentVariable in Environment.GetEnvironmentVariables())
-				{
-					if (EnvironmentVariable != null)
+					// Save the additional command line arguments
+					Makefile.AdditionalArguments = TargetDescriptor.AdditionalArguments.GetRawArray();
+
+					// Save the environment variables
+					foreach (System.Collections.DictionaryEntry? EnvironmentVariable in Environment.GetEnvironmentVariables())
 					{
-						Makefile.EnvironmentVariables.Add(Tuple.Create((string)EnvironmentVariable.Value.Key, (string)EnvironmentVariable.Value.Value!));
+						if (EnvironmentVariable != null)
+						{
+							Makefile.EnvironmentVariables.Add(Tuple.Create((string)EnvironmentVariable.Value.Key, (string)EnvironmentVariable.Value.Value!));
+						}
 					}
-				}
 
-				// Save the makefile for next time
-				if (MakefileLocation != null)
-				{
-					SaveToLocation = MakefileLocation;
+					// Save the makefile for next time
+					if (MakefileLocation != null)
+					{
+						SaveToLocation = MakefileLocation;
+					}
 				}
-			}
-			else
-			{
-				// Restore the environment variables
-				foreach (Tuple<string, string> EnvironmentVariable in Makefile.EnvironmentVariables)
+				else
 				{
-					Environment.SetEnvironmentVariable(EnvironmentVariable.Item1, EnvironmentVariable.Item2);
-				}
+					// Restore the environment variables
+					foreach (Tuple<string, string> EnvironmentVariable in Makefile.EnvironmentVariables)
+					{
+						Environment.SetEnvironmentVariable(EnvironmentVariable.Item1, EnvironmentVariable.Item2);
+					}
+
+#if __VPROJECT_AVAILABLE__
+					Task VNITask = Task.CompletedTask;
+					// Same for VNI
+					if (Makefile.VNIModules.Count > 0)
+					{
+						VNITask = Task.Run(() => VNIExecution.ExecuteVNITool(Makefile, TargetDescriptor, Logger));
+					}
+#endif
+					// If the target needs UHT to be run, we'll go ahead and do that now
+					if (Makefile.UObjectModules.Count > 0)
+					{
+						(Task uhtRunTask, uhtTimeStampsTask) = UHTExecution.ExecuteHeaderToolIfNecessary(BuildConfiguration, TargetDescriptor.ProjectFile, Makefile, TargetDescriptor.Name, Logger, true);
+						uhtRunTask.Wait();
+					}
 
 #if __VPROJECT_AVAILABLE__
-				Task VNITask = Task.CompletedTask;
-				// Same for VNI
-				if (Makefile.VNIModules.Count > 0)
+					VNITask.Wait();
+#endif
+				}
+
+				Makefiles.Add(Makefile);
+
+				LastUHTTimestampTask = uhtTimeStampsTask;
+
+				if (SaveToLocation != null)
 				{
-					VNITask = Task.Run(() => VNIExecution.ExecuteVNITool(Makefile, TargetDescriptor, Logger));
+					LastSaveMakefileTask = Task.Run(() =>
+					{
+						using ITimelineEvent _ = Timeline.ScopeEvent($"SavingMakefile");
+						Makefile.Save(SaveToLocation);
+					});
 				}
-#endif
-				// If the target needs UHT to be run, we'll go ahead and do that now
-				if (Makefile.UObjectModules.Count > 0)
+
+				if (!bSkipPreBuildTargets)
 				{
-					(Task uhtRunTask, uhtTimeStampsTask) = UHTExecution.ExecuteHeaderToolIfNecessary(BuildConfiguration, TargetDescriptor.ProjectFile, Makefile, TargetDescriptor.Name, Logger, true);
-					uhtRunTask.Wait();
+					foreach (TargetInfo PreBuildTarget in Makefile.PreBuildTargets)
+					{
+						TargetDescriptor NewTarget = TargetDescriptor.FromTargetInfo(PreBuildTarget);
+						if (!TargetDescriptors.Contains(NewTarget))
+						{
+							TargetDescriptors.Add(NewTarget);
+						}
+					}
 				}
+			}
 
-#if __VPROJECT_AVAILABLE__
-				VNITask.Wait();
-#endif
+			List<Task> tasks = [];
+
+			if (LastSaveMakefileTask != null)
+			{
+				tasks.Add(LastSaveMakefileTask);
+			}
+			if (LastUHTTimestampTask != null)
+			{
+				tasks.Add(LastUHTTimestampTask);
 			}
-			return (Makefile, SaveToLocation, uhtTimeStampsTask);
+
+			return (Makefiles, tasks);
 		}
 
 		/// <summary>
