[UBT]
* Fixed so targets that has different appname can build in parallel. Note, all editors have appname UnrealEditor.

diff --git a/Engine/Source/Programs/UnrealBuildTool/Configuration/UEBuildTarget.cs b/Engine/Source/Programs/UnrealBuildTool/Configuration/UEBuildTarget.cs
--- a/Engine/Source/Programs/UnrealBuildTool/Configuration/UEBuildTarget.cs
+++ b/Engine/Source/Programs/UnrealBuildTool/Configuration/UEBuildTarget.cs
@@ -15,8 +15,6 @@ using System.Text.Json.Serialization;
 using System.Threading;
 using System.Threading.Tasks;
 using EpicGames.Core;
-using EpicGames.UHT.Tables;
-using EpicGames.UHT.Types;
 using EpicGames.UHT.Utils;
 using Microsoft.Extensions.Logging;
 using OpenTracing.Util;
@@ -1257,6 +1255,26 @@ namespace UnrealBuildTool
 		/// <param name="Logger">Logger for output</param>
 		/// <returns>The build target object for the specified build rules source file</returns>
 		public static UEBuildTarget Create(TargetDescriptor Descriptor, bool bSkipRulesCompile, bool bForceRulesCompile, bool bUsePrecompiled, UnrealIntermediateEnvironment IntermediateEnvironment, ILogger Logger)
+		{
+			Task<SourceFileMetadataCache> metaDataTask = Task.Run(() => SourceFileMetadataCache.CreateHierarchy(Descriptor.ProjectFile, Logger));
+
+			var (RulesObject, RulesAssembly) = CreateRulesObject(Descriptor, bSkipRulesCompile, bForceRulesCompile, bUsePrecompiled, IntermediateEnvironment, Logger);
+
+			// Generate a build target from this rules module
+			UEBuildTarget Target;
+			using (GlobalTracer.Instance.BuildSpan("UEBuildTarget constructor").StartActive())
+			{
+				Target = new UEBuildTarget(Descriptor, new ReadOnlyTargetRules(RulesObject), RulesAssembly, metaDataTask.Result, Logger);
+			}
+			using (GlobalTracer.Instance.BuildSpan("UEBuildTarget.PreBuildSetup()").StartActive())
+			{
+				Target.PreBuildSetup(Logger);
+			}
+
+			return Target;
+		}
+
+		internal static (TargetRules, RulesAssembly) CreateRulesObject(TargetDescriptor Descriptor, bool bSkipRulesCompile, bool bForceRulesCompile, bool bUsePrecompiled, UnrealIntermediateEnvironment IntermediateEnvironment, ILogger Logger)
 		{
 			// make sure we are allowed to build this platform
 			if (!UEBuildPlatform.IsPlatformAvailable(Descriptor.Platform))
@@ -1264,8 +1282,6 @@ namespace UnrealBuildTool
 				throw new BuildException("Platform {0} is not a valid platform to build. Check that the SDK is installed properly and that you have the necessary platform support files (DataDrivenPlatformInfo.ini, SDK.json, etc).", Descriptor.Platform);
 			}
 
-			Task<SourceFileMetadataCache> metaDataTask = Task.Run(() => SourceFileMetadataCache.CreateHierarchy(Descriptor.ProjectFile, Logger));
-
 			RulesAssembly RulesAssembly;
 			using (GlobalTracer.Instance.BuildSpan("RulesCompiler.CreateTargetRulesAssembly()").StartActive())
 			{
@@ -1396,18 +1412,7 @@ namespace UnrealBuildTool
 			{
 				RulesObject.bIgnoreBuildOutputs = true;
 			}
-
-			// Generate a build target from this rules module
-			UEBuildTarget Target;
-			using (GlobalTracer.Instance.BuildSpan("UEBuildTarget constructor").StartActive())
-			{
-				Target = new UEBuildTarget(Descriptor, new ReadOnlyTargetRules(RulesObject), RulesAssembly, metaDataTask.Result, Logger);
-			}
-			using (GlobalTracer.Instance.BuildSpan("UEBuildTarget.PreBuildSetup()").StartActive())
-			{
-				Target.PreBuildSetup(Logger);
-			}
-			return Target;
+			return (RulesObject, RulesAssembly);
 		}
 
 		/// <summary>
@@ -1810,7 +1815,7 @@ namespace UnrealBuildTool
 		/// <param name="InRulesAssembly">The chain of rules assemblies that this target was created with</param>
 		/// <param name="InMetadataCache"></param>
 		/// <param name="Logger">Logger for output</param>
-		private UEBuildTarget(TargetDescriptor InDescriptor, ReadOnlyTargetRules InRules, RulesAssembly InRulesAssembly, SourceFileMetadataCache InMetadataCache, ILogger Logger)
+		internal UEBuildTarget(TargetDescriptor InDescriptor, ReadOnlyTargetRules InRules, RulesAssembly InRulesAssembly, SourceFileMetadataCache InMetadataCache, ILogger Logger)
 		{
 			MetadataCache = InMetadataCache;
 			ProjectFile = InDescriptor.ProjectFile;
@@ -1970,6 +1975,15 @@ namespace UnrealBuildTool
 			}
 		}
 
+		public static string GetAppName(TargetRules Rules)
+		{
+			if (Rules.BuildEnvironment == TargetBuildEnvironment.Shared && Rules.Type != TargetType.Program)
+			{
+				return GetAppNameForTargetType(Rules.Type);
+			}
+			return Rules.Name;
+		}
+
 		/// <summary>
 		/// Writes a list of all the externally referenced files required to use the precompiled data for this target
 		/// </summary>
diff --git a/Engine/Source/Programs/UnrealBuildTool/Modes/BuildMode.cs b/Engine/Source/Programs/UnrealBuildTool/Modes/BuildMode.cs
--- a/Engine/Source/Programs/UnrealBuildTool/Modes/BuildMode.cs
+++ b/Engine/Source/Programs/UnrealBuildTool/Modes/BuildMode.cs
@@ -387,9 +387,9 @@ namespace UnrealBuildTool
 		/// <returns>Result from the compilation</returns>
 		public static async Task BuildAsync(List<TargetDescriptor> TargetDescriptors, BuildConfiguration BuildConfiguration, ISourceFileWorkingSet WorkingSet, BuildOptions Options, FileReference? WriteOutdatedActionsFile, ILogger Logger, bool bSkipPreBuildTargets = false, string? ActionTypeFilter = null)
 		{
-			(List<TargetMakefile> Makefiles, List<Task> Tasks) = CreateMakefiles(BuildConfiguration, TargetDescriptors, WorkingSet, Logger, true, true);
+			(List<(TargetDescriptor, TargetMakefile)> Targets, List<Task> Tasks) = CreateMakefiles(BuildConfiguration, TargetDescriptors, WorkingSet, Logger, true, true);
 
-			await BuildAsync(Makefiles.ToArray(), TargetDescriptors, BuildConfiguration, Options, WriteOutdatedActionsFile, Logger, ActionTypeFilter);
+			await BuildAsync(Targets, BuildConfiguration, Options, WriteOutdatedActionsFile, Logger, ActionTypeFilter);
 
 			Task.WaitAll(Tasks.ToArray());
 		}
@@ -399,16 +399,18 @@ namespace UnrealBuildTool
 		/// <summary>
 		/// Build a list of targets with a given set of makefiles.
 		/// </summary>
-		/// <param name="Makefiles">Makefiles created with CreateMakefiles</param>
-		/// <param name="TargetDescriptors">Target descriptors</param>
+		/// <param name="Targets">Targets</param>
 		/// <param name="BuildConfiguration">Current build configuration</param>
 		/// <param name="Options">Additional options for the build</param>
 		/// <param name="WriteOutdatedActionsFile">Files to write the list of outdated actions to (rather than building them)</param>
 		/// <param name="Logger">Logger for output</param>
 		/// <param name="ActionTypeFilter">If specified, only actions of this type will be run</param>
 		/// <returns>Result from the compilation</returns>
-		internal static async Task BuildAsync(TargetMakefile[] Makefiles, List<TargetDescriptor> TargetDescriptors, BuildConfiguration BuildConfiguration, BuildOptions Options, FileReference? WriteOutdatedActionsFile, ILogger Logger, string? ActionTypeFilter=null)
+		internal static async Task BuildAsync(List<(TargetDescriptor Descriptor, TargetMakefile Makefile)> Targets, BuildConfiguration BuildConfiguration, BuildOptions Options, FileReference? WriteOutdatedActionsFile, ILogger Logger, string? ActionTypeFilter=null)
 		{
+			List<TargetDescriptor> TargetDescriptors = Targets.Select(T => T.Descriptor).ToList();
+			TargetMakefile[] Makefiles = Targets.Select(T => T.Makefile).ToArray();
+
 			// Execute the build
 			if ((Options & BuildOptions.SkipBuild) == 0)
 			{
@@ -416,12 +418,12 @@ namespace UnrealBuildTool
 				{
 					// Create the C++ dependency cache
 					CppDependencyCache CppDependencies = new CppDependencyCache();
-					for (int TargetIdx = 0; TargetIdx < TargetDescriptors.Count; TargetIdx++)
+					for (int TargetIdx = 0; TargetIdx < Targets.Count; TargetIdx++)
 					{
 						using (GlobalTracer.Instance.BuildSpan("Reading dependency cache").StartActive())
 						{
-							TargetDescriptor TargetDescriptor = TargetDescriptors[TargetIdx];
-							CppDependencies.Mount(TargetDescriptor, Makefiles[TargetIdx].TargetType, Logger);
+							TargetDescriptor TargetDescriptor = Targets[TargetIdx].Descriptor;
+							CppDependencies.Mount(TargetDescriptor, Targets[TargetIdx].Makefile.TargetType, Logger);
 						}
 					}
 					return CppDependencies;
@@ -466,26 +468,26 @@ namespace UnrealBuildTool
 					// TODO: Skipping conflicts for WriteMetadata is a hack to maintain parity with the BuildGraph executor which allowed
 					// exporting mulitple conflicing WriteMetadata actions.
 					// This is technically buggy behavior and should be properly fixed in UEBuildTarget.cs
-					IEnumerable<IExternalAction> CheckActions = Makefiles.SelectMany(x => x.Actions).Where(x => !x.IgnoreConflicts());
+					IEnumerable<IExternalAction> CheckActions = Targets.Select(T => T.Makefile).SelectMany(x => x.Actions).Where(x => !x.IgnoreConflicts());
 					ActionGraph.CheckForConflicts(CheckActions, Logger);
 				}
 
 				// Check we don't exceed the nominal max path length
 				using (GlobalTracer.Instance.BuildSpan("ActionGraph.CheckPathLengths").StartActive())
 				{
-					ActionGraph.CheckPathLengths(BuildConfiguration, Makefiles.SelectMany(x => x.Actions), Logger);
+					ActionGraph.CheckPathLengths(BuildConfiguration, Targets.Select(T => T.Makefile).SelectMany(x => x.Actions), Logger);
 				}
 
 				HashSet<FileItem> MergedOutputItems = new HashSet<FileItem>();
 
 				// Create a QueuedAction instance for each action in the makefiles
-				List<LinkedAction>[] QueuedActions = new List<LinkedAction>[Makefiles.Length];
-				for (int Idx = 0; Idx < Makefiles.Length; Idx++)
+				List<LinkedAction>[] QueuedActions = new List<LinkedAction>[Targets.Count];
+				for (int Idx = 0; Idx < Targets.Count; Idx++)
 				{
-					TargetDescriptor TargetDescriptor = TargetDescriptors[Idx];
-					TargetMakefile Makefile = Makefiles[Idx];
+					TargetDescriptor TargetDescriptor = Targets[Idx].Descriptor;
+					TargetMakefile Makefile = Targets[Idx].Makefile;
 
-					QueuedActions[Idx] = Makefile.Actions.ConvertAll(x => new LinkedAction(x, TargetDescriptors[Idx]));
+					QueuedActions[Idx] = Makefile.Actions.ConvertAll(x => new LinkedAction(x, TargetDescriptor));
 
 					if (TargetDescriptor.SpecificFilesToCompile.Count != 0)
 					{
@@ -513,9 +515,9 @@ namespace UnrealBuildTool
 
 				// Clean up any previous hot reload runs, and reapply the current state if it's already active
 				Dictionary<FileReference, FileReference>? InitialPatchedOldLocationToNewLocation = null;
-				for (int TargetIdx = 0; TargetIdx < TargetDescriptors.Count; TargetIdx++)
+				for (int TargetIdx = 0; TargetIdx < Targets.Count; TargetIdx++)
 				{
-					InitialPatchedOldLocationToNewLocation = HotReload.Setup(TargetDescriptors[TargetIdx], Makefiles[TargetIdx], QueuedActions[TargetIdx], BuildConfiguration, Logger);
+					InitialPatchedOldLocationToNewLocation = HotReload.Setup(Targets[TargetIdx].Descriptor, Targets[TargetIdx].Makefile, QueuedActions[TargetIdx], BuildConfiguration, Logger);
 				}
 
 				// Merge the action graphs together
@@ -1102,205 +1104,252 @@ namespace UnrealBuildTool
 		{
 			await Task.CompletedTask;
 
-			(List<TargetMakefile> Makefiles, List<Task> Tasks) = CreateMakefiles(BuildConfiguration, [TargetDescriptor], WorkingSet, Logger, false, false);
+			(List<(TargetDescriptor, TargetMakefile Makefile)> Targets, List<Task> Tasks) = CreateMakefiles(BuildConfiguration, [TargetDescriptor], WorkingSet, Logger, false, false);
 
 			Task.WaitAll(Tasks.ToArray());
-			return Makefiles[0];
+			return Targets[0].Makefile;
 		}
 
-		private static (List<TargetMakefile>, List<Task>) CreateMakefiles(BuildConfiguration BuildConfiguration, List<TargetDescriptor> TargetDescriptors, ISourceFileWorkingSet WorkingSet, ILogger Logger, bool runInitScripts, bool bSkipPreBuildTargets)
+		private static (List<(TargetDescriptor, TargetMakefile)>, List<Task>) CreateMakefiles(BuildConfiguration BuildConfiguration, List<TargetDescriptor> TargetDescriptors, ISourceFileWorkingSet WorkingSet, ILogger Logger, bool runInitScripts, bool bSkipPreBuildTargets)
 		{
-			List<TargetMakefile> Makefiles = [];
+			List<(TargetDescriptor descriptor, TargetMakefile makefile)> targets = new();
+			List<Task> tasks = [];
 
-			Task? LastSaveMakefileTask = null;
-			Task? LastUHTTimestampTask = null;
+			while (TargetDescriptors.Count != 0) // TargetDescriptors will be reallocated for PrebuildTargets
+			{
+				bool skipRulesCompile = BuildConfiguration.bSkipRulesCompile;
+				bool forceRulesCompile = BuildConfiguration.bForceRulesCompile;
+				bool usePrecompiled = BuildConfiguration.bUsePrecompiled;
 
-			bool skipRulesCompile = BuildConfiguration.bSkipRulesCompile;
-			bool forceRulesCompile = BuildConfiguration.bForceRulesCompile;
-			bool usePrecompiled = BuildConfiguration.bUsePrecompiled;
+				(TargetRules rules, RulesAssembly assembly)[] rules = new (TargetRules, RulesAssembly)[TargetDescriptors.Count];
+				Task<SourceFileMetadataCache>[] metadataCaches = new Task<SourceFileMetadataCache>[TargetDescriptors.Count];
 
-			for (int Idx = 0; Idx < TargetDescriptors.Count; ++Idx)
-			{
-				TargetDescriptor TargetDescriptor = TargetDescriptors[Idx];
+
+				// Build all rules objects
+				Parallel2.For(0, TargetDescriptors.Count, (index) =>
+				{
+					TargetDescriptor TargetDescriptor = TargetDescriptors[index];
+					metadataCaches[index] = Task.Run(() => SourceFileMetadataCache.CreateHierarchy(TargetDescriptor.ProjectFile, Logger));
+					rules[index] = UEBuildTarget.CreateRulesObject(TargetDescriptor, skipRulesCompile, forceRulesCompile, usePrecompiled, TargetDescriptor.IntermediateEnvironment, Logger);
+				});
 
 
-				if (runInitScripts)
+				// Use rules objects to figure out how to group work. We can't run targets that write to same folders in parallel
+				Dictionary<string, List<int>> buildGroups = [];
+
+				for (int i = 0; i < TargetDescriptors.Count; ++i)
 				{
-					using ITimelineEvent SetupScope = Timeline.ScopeEvent($"{TargetDescriptor.Name} {TargetDescriptor.Platform} {TargetDescriptor.Configuration}");
-					// Create and execute the init scripts
-					FileReference[] InitScripts = CreateInitScripts(TargetDescriptor);
-					if (InitScripts.Length > 0)
+					TargetDescriptor targetDescriptor = TargetDescriptors[i];
+					string platformIntermediateFolderNoArch = UEBuildTarget.GetPlatformIntermediateFolder(targetDescriptor.Platform, null, false);
+					string appName = UEBuildTarget.GetAppName(rules[i].rules);
+
+					if (!buildGroups.TryGetValue(appName, out var group))
 					{
-						Utils.ExecuteCustomBuildSteps(InitScripts, Logger);
+						group = new();
+						buildGroups.Add(appName, group);
 					}
+
+					group.Add(i);
 				}
 
-				FileReference? SaveToLocation = null;
+				List<TargetDescriptor> newTargetDescriptors = [];
 
-				// Get the path to the makefile for this target
-				FileReference? MakefileLocation = null;
-				if (BuildConfiguration.bUseUBTMakefiles)
-				{
-					MakefileLocation = TargetMakefile.GetLocation(TargetDescriptor.ProjectFile, TargetDescriptor.Name, TargetDescriptor.Platform, TargetDescriptor.Architectures, TargetDescriptor.Configuration, TargetDescriptor.IntermediateEnvironment);
-				}
+				// TODO: We also want to handle different configs.. they can't run UHT etc at the same time but can do the other things..
+				// ... so ne need to get that support in here too
 
-				// Need to make sure to wait on previous run's makefile
-				LastSaveMakefileTask?.Wait();
 
-				// Try to load an existing makefile
-				TargetMakefile? Makefile = null;
-				if (MakefileLocation != null)
+				Parallel.ForEach(buildGroups, kv =>
 				{
-					using (GlobalTracer.Instance.BuildSpan("TargetMakefile.Load()").StartActive())
+					Task? LastSaveMakefileTask = null;
+					Task? LastUHTTimestampTask = null;
+
+					List<int> buildGroup = kv.Value;
+					for (int i = 0; i != buildGroup.Count; ++i)
 					{
-						string? ReasonNotLoaded;
-						Makefile = TargetMakefile.Load(MakefileLocation, TargetDescriptor.ProjectFile, TargetDescriptor.Platform, TargetDescriptor.AdditionalArguments.GetRawArray(), Logger, out ReasonNotLoaded);
-						if (Makefile == null)
+						int index = buildGroup[i];
+						TargetDescriptor TargetDescriptor = TargetDescriptors[index];
+
+						if (runInitScripts)
 						{
-							Logger.LogInformation("Creating makefile for {TargetDescriptorName} ({ReasonNotLoaded})", TargetDescriptor.Name, ReasonNotLoaded);
+							using ITimelineEvent SetupScope = Timeline.ScopeEvent($"{TargetDescriptor.Name} {TargetDescriptor.Platform} {TargetDescriptor.Configuration}");
+							// Create and execute the init scripts
+							FileReference[] InitScripts = CreateInitScripts(TargetDescriptor);
+							if (InitScripts.Length > 0)
+							{
+								Utils.ExecuteCustomBuildSteps(InitScripts, Logger);
+							}
 						}
-					}
-				}
 
-				// If we have a makefile, execute the pre-build steps and check it's still valid
-				bool bHasRunPreBuildScripts = false;
-				if (Makefile != null)
-				{
-					// Execute the scripts.
-					Utils.ExecuteCustomBuildSteps(Makefile.PreBuildScripts, Logger);
+						FileReference? SaveToLocation = null;
 
-					// Don't run the pre-build steps again, even if we invalidate the makefile.
-					bHasRunPreBuildScripts = true;
+						// Get the path to the makefile for this target
+						FileReference? MakefileLocation = null;
+						if (BuildConfiguration.bUseUBTMakefiles)
+						{
+							MakefileLocation = TargetMakefile.GetLocation(TargetDescriptor.ProjectFile, TargetDescriptor.Name, TargetDescriptor.Platform, TargetDescriptor.Architectures, TargetDescriptor.Configuration, TargetDescriptor.IntermediateEnvironment);
+						}
 
-					// Check that the makefile is still valid
-					string? Reason;
-					if (!TargetMakefile.IsValidForSourceFiles(Makefile, TargetDescriptor.ProjectFile, TargetDescriptor.Platform, WorkingSet, Logger, out Reason))
-					{
-						Logger.LogInformation("Invalidating makefile for {TargetDescriptorName} ({Reason})", TargetDescriptor.Name, Reason);
-						Makefile = null;
-					}
-				}
+						// Need to make sure to wait on previous run's makefile
+						LastSaveMakefileTask?.Wait();
 
-				// Need to make sure to wait on previous run's timestamp since there are targets that write to same place..
-				LastUHTTimestampTask?.Wait();
+						// Try to load an existing makefile
+						TargetMakefile? Makefile = null;
+						if (MakefileLocation != null)
+						{
+							using (GlobalTracer.Instance.BuildSpan("TargetMakefile.Load()").StartActive())
+							{
+								string? ReasonNotLoaded;
+								Makefile = TargetMakefile.Load(MakefileLocation, TargetDescriptor.ProjectFile, TargetDescriptor.Platform, TargetDescriptor.AdditionalArguments.GetRawArray(), Logger, out ReasonNotLoaded);
+								if (Makefile == null)
+								{
+									Logger.LogInformation("Creating makefile for {TargetDescriptorName} ({ReasonNotLoaded})", TargetDescriptor.Name, ReasonNotLoaded);
+								}
+							}
+						}
 
-				Task? uhtTimeStampsTask = null;
+						// If we have a makefile, execute the pre-build steps and check it's still valid
+						bool bHasRunPreBuildScripts = false;
+						if (Makefile != null)
+						{
+							// Execute the scripts.
+							Utils.ExecuteCustomBuildSteps(Makefile.PreBuildScripts, Logger);
 
-				// If we couldn't load a makefile, create a new one
-				if (Makefile == null)
-				{
-					// Create the target
-					UEBuildTarget Target;
-					using (GlobalTracer.Instance.BuildSpan("UEBuildTarget.Create()").StartActive())
-					{
-						Target = UEBuildTarget.Create(TargetDescriptor, skipRulesCompile, forceRulesCompile, usePrecompiled, Logger);
-					}
+							// Don't run the pre-build steps again, even if we invalidate the makefile.
+							bHasRunPreBuildScripts = true;
 
-					// Create the pre-build scripts
-					FileReference[] PreBuildScripts = Target.CreatePreBuildScripts();
+							// Check that the makefile is still valid
+							string? Reason;
+							if (!TargetMakefile.IsValidForSourceFiles(Makefile, TargetDescriptor.ProjectFile, TargetDescriptor.Platform, WorkingSet, Logger, out Reason))
+							{
+								Logger.LogInformation("Invalidating makefile for {TargetDescriptorName} ({Reason})", TargetDescriptor.Name, Reason);
+								Makefile = null;
+							}
+						}
 
-					// Execute the pre-build scripts
-					if (!bHasRunPreBuildScripts)
-					{
-						Utils.ExecuteCustomBuildSteps(PreBuildScripts, Logger);
-						bHasRunPreBuildScripts = true;
-					}
+						// Need to make sure to wait on previous run's timestamp since there are targets that write to same place..
+						LastUHTTimestampTask?.Wait();
 
-					// Build the target
-					using (GlobalTracer.Instance.BuildSpan("UEBuildTarget.Build()").StartActive())
-					{
-						(Makefile, uhtTimeStampsTask) = Target.BuildMakefile(BuildConfiguration, WorkingSet, TargetDescriptor, Logger);
-					}
+						Task? uhtTimeStampsTask = null;
 
-					Makefile.MemoryPerActionGB = Target.Rules.MemoryPerActionGB;
+						// If we couldn't load a makefile, create a new one
+						if (Makefile == null)
+						{
+							// Create the target
+							UEBuildTarget Target = new UEBuildTarget(TargetDescriptor, new ReadOnlyTargetRules(rules[index].rules), rules[index].assembly, metadataCaches[index].Result, Logger);
+							Target.PreBuildSetup(Logger);
 
-					// Save the pre-build scripts onto the makefile
-					Makefile.PreBuildScripts = PreBuildScripts;
+							// Create the pre-build scripts
+							FileReference[] PreBuildScripts = Target.CreatePreBuildScripts();
 
-					// Save the additional command line arguments
-					Makefile.AdditionalArguments = TargetDescriptor.AdditionalArguments.GetRawArray();
+							// Execute the pre-build scripts
+							if (!bHasRunPreBuildScripts)
+							{
+								Utils.ExecuteCustomBuildSteps(PreBuildScripts, Logger);
+								bHasRunPreBuildScripts = true;
+							}
 
-					// Save the environment variables
-					foreach (System.Collections.DictionaryEntry? EnvironmentVariable in Environment.GetEnvironmentVariables())
-					{
-						if (EnvironmentVariable != null)
-						{
-							Makefile.EnvironmentVariables.Add(Tuple.Create((string)EnvironmentVariable.Value.Key, (string)EnvironmentVariable.Value.Value!));
-						}
-					}
+							// Build the target
+							using (GlobalTracer.Instance.BuildSpan("UEBuildTarget.Build()").StartActive())
+							{
+								(Makefile, uhtTimeStampsTask) = Target.BuildMakefile(BuildConfiguration, WorkingSet, TargetDescriptor, Logger);
+							}
 
-					// Save the makefile for next time
-					if (MakefileLocation != null)
-					{
-						SaveToLocation = MakefileLocation;
-					}
-				}
-				else
-				{
-					// Restore the environment variables
-					foreach (Tuple<string, string> EnvironmentVariable in Makefile.EnvironmentVariables)
-					{
-						Environment.SetEnvironmentVariable(EnvironmentVariable.Item1, EnvironmentVariable.Item2);
-					}
+							Makefile.MemoryPerActionGB = Target.Rules.MemoryPerActionGB;
+
+							// Save the pre-build scripts onto the makefile
+							Makefile.PreBuildScripts = PreBuildScripts;
+
+							// Save the additional command line arguments
+							Makefile.AdditionalArguments = TargetDescriptor.AdditionalArguments.GetRawArray();
+
+							// Save the environment variables
+							foreach (System.Collections.DictionaryEntry? EnvironmentVariable in Environment.GetEnvironmentVariables())
+							{
+								if (EnvironmentVariable != null)
+								{
+									Makefile.EnvironmentVariables.Add(Tuple.Create((string)EnvironmentVariable.Value.Key, (string)EnvironmentVariable.Value.Value!));
+								}
+							}
+
+							// Save the makefile for next time
+							if (MakefileLocation != null)
+							{
+								SaveToLocation = MakefileLocation;
+							}
+						}
+						else
+						{
+							// Restore the environment variables
+							foreach (Tuple<string, string> EnvironmentVariable in Makefile.EnvironmentVariables)
+							{
+								Environment.SetEnvironmentVariable(EnvironmentVariable.Item1, EnvironmentVariable.Item2);
+							}
 
 #if __VPROJECT_AVAILABLE__
-					Task VNITask = Task.CompletedTask;
-					// Same for VNI
-					if (Makefile.VNIModules.Count > 0)
-					{
-						VNITask = Task.Run(() => VNIExecution.ExecuteVNITool(Makefile, TargetDescriptor, Logger));
-					}
+							Task VNITask = Task.CompletedTask;
+							// Same for VNI
+							if (Makefile.VNIModules.Count > 0)
+							{
+								VNITask = Task.Run(() => VNIExecution.ExecuteVNITool(Makefile, TargetDescriptor, Logger));
+							}
 #endif
-					// If the target needs UHT to be run, we'll go ahead and do that now
-					if (Makefile.UObjectModules.Count > 0)
-					{
-						(Task uhtRunTask, uhtTimeStampsTask) = UHTExecution.ExecuteHeaderToolIfNecessary(BuildConfiguration, TargetDescriptor.ProjectFile, Makefile, TargetDescriptor.Name, Logger, true);
-						uhtRunTask.Wait();
-					}
+							// If the target needs UHT to be run, we'll go ahead and do that now
+							if (Makefile.UObjectModules.Count > 0)
+							{
+								(Task uhtRunTask, uhtTimeStampsTask) = UHTExecution.ExecuteHeaderToolIfNecessary(BuildConfiguration, TargetDescriptor.ProjectFile, Makefile, TargetDescriptor.Name, Logger, true);
+								uhtRunTask.Wait();
+							}
 
 #if __VPROJECT_AVAILABLE__
-					VNITask.Wait();
+							VNITask.Wait();
 #endif
-				}
+						}
 
-				Makefiles.Add(Makefile);
+						targets.Add((TargetDescriptor, Makefile));
 
-				LastUHTTimestampTask = uhtTimeStampsTask;
+						LastUHTTimestampTask = uhtTimeStampsTask;
 
-				if (SaveToLocation != null)
-				{
-					LastSaveMakefileTask = Task.Run(() =>
-					{
-						using ITimelineEvent _ = Timeline.ScopeEvent($"SavingMakefile");
-						Makefile.Save(SaveToLocation);
-					});
-				}
+						if (SaveToLocation != null)
+						{
+							LastSaveMakefileTask = Task.Run(() =>
+							{
+								using ITimelineEvent _ = Timeline.ScopeEvent($"SavingMakefile");
+								Makefile.Save(SaveToLocation);
+							});
+						}
 
-				if (!bSkipPreBuildTargets)
-				{
-					foreach (TargetInfo PreBuildTarget in Makefile.PreBuildTargets)
-					{
-						TargetDescriptor NewTarget = TargetDescriptor.FromTargetInfo(PreBuildTarget);
-						if (!TargetDescriptors.Contains(NewTarget))
+						if (!bSkipPreBuildTargets)
 						{
-							TargetDescriptors.Add(NewTarget);
+							foreach (TargetInfo PreBuildTarget in Makefile.PreBuildTargets)
+							{
+								lock (newTargetDescriptors)
+								{
+									newTargetDescriptors.Add(TargetDescriptor.FromTargetInfo(PreBuildTarget));
+								}
+							}
 						}
 					}
-				}
-			}
 
-			List<Task> tasks = [];
+					if (LastSaveMakefileTask != null)
+					{
+						lock (tasks)
+						{
+							tasks.Add(LastSaveMakefileTask);
+						}
+					}
+					if (LastUHTTimestampTask != null)
+					{
+						lock (tasks)
+						{
+							tasks.Add(LastUHTTimestampTask);
+						}
+					}
+				});
 
-			if (LastSaveMakefileTask != null)
-			{
-				tasks.Add(LastSaveMakefileTask);
-			}
-			if (LastUHTTimestampTask != null)
-			{
-				tasks.Add(LastUHTTimestampTask);
+				TargetDescriptors = newTargetDescriptors;
 			}
 
-			return (Makefiles, tasks);
+			return (targets, tasks);
 		}
 
 		/// <summary>
diff --git a/Engine/Source/Programs/UnrealBuildTool/Modes/IWYUMode.cs b/Engine/Source/Programs/UnrealBuildTool/Modes/IWYUMode.cs
--- a/Engine/Source/Programs/UnrealBuildTool/Modes/IWYUMode.cs
+++ b/Engine/Source/Programs/UnrealBuildTool/Modes/IWYUMode.cs
@@ -627,7 +627,7 @@ namespace UnrealBuildTool
 					{
 						try
 						{
-							await BuildMode.BuildAsync(new TargetMakefile[] { Makefile }, new List<TargetDescriptor>() { Descriptor }, BuildConfiguration, BuildOptions.None, null, Logger);
+							await BuildMode.BuildAsync([(Descriptor, Makefile)], BuildConfiguration, BuildOptions.None, null, Logger);
 						}
 						finally
 						{
