[UBT]
* Disabled multithreaded makefile generation because of race condition reported by farm

diff --git a/Engine/Source/Programs/UnrealBuildTool/Modes/BuildMode.cs b/Engine/Source/Programs/UnrealBuildTool/Modes/BuildMode.cs
--- a/Engine/Source/Programs/UnrealBuildTool/Modes/BuildMode.cs
+++ b/Engine/Source/Programs/UnrealBuildTool/Modes/BuildMode.cs
@@ -1133,7 +1133,6 @@ namespace UnrealBuildTool
 					rules[index] = UEBuildTarget.CreateRulesObject(TargetDescriptor, skipRulesCompile, forceRulesCompile, usePrecompiled, TargetDescriptor.IntermediateEnvironment, Logger);
 				});
 
-
 				// Use rules objects to figure out how to group work. We can't run targets that write to same folders in parallel
 				Dictionary<string, List<int>> buildGroups = [];
 
@@ -1158,7 +1157,8 @@ namespace UnrealBuildTool
 				// ... so ne need to get that support in here too
 
 
-				Parallel.ForEach(buildGroups, kv =>
+				//Parallel2.ForEach(buildGroups, kv =>
+				foreach (KeyValuePair<string, List<int>> kv in buildGroups)
 				{
 					Task? LastSaveMakefileTask = null;
 					Task? LastUHTTimestampTask = null;
@@ -1345,7 +1345,7 @@ namespace UnrealBuildTool
 							tasks.Add(LastUHTTimestampTask);
 						}
 					}
-				});
+				}//);
 
 				TargetDescriptors = newTargetDescriptors;
 			}
