UBT wouldn't find the temporary targets and modules it may generate for content-only projects.

UBT will create temporary .Build.cs and .Target.cs files in Intermediate/Source if a content-only project uses plugins with source that need to be compiled. Rules.FindDirectory() wouldn't find these targets and modules because Intermediate was missing from Rules.Root. This would cause UBT to error out due to missing module dependencies.

#rnx

diff --git a/Engine/Source/Programs/UnrealBuildTool/System/RulesCompiler.cs b/Engine/Source/Programs/UnrealBuildTool/System/RulesCompiler.cs
--- a/Engine/Source/Programs/UnrealBuildTool/System/RulesCompiler.cs
+++ b/Engine/Source/Programs/UnrealBuildTool/System/RulesCompiler.cs
@@ -335,6 +335,7 @@ namespace UnrealBuildTool
 					foreach (DirectoryReference RootDirectory in ProjectDirectories)
 					{
 						populateDirs.Add(DirectoryReference.Combine(RootDirectory, "Source"));
+						populateDirs.Add(DirectoryReference.Combine(RootDirectory, "Intermediate", "Source"));
 					}
 					Rules.Populate(populateDirs);
 				}
