[UBT]
* Fixed race conditions in MacToolChain

diff --git a/Engine/Source/Programs/UnrealBuildTool/Platform/Mac/MacToolChain.cs b/Engine/Source/Programs/UnrealBuildTool/Platform/Mac/MacToolChain.cs
--- a/Engine/Source/Programs/UnrealBuildTool/Platform/Mac/MacToolChain.cs
+++ b/Engine/Source/Programs/UnrealBuildTool/Platform/Mac/MacToolChain.cs
@@ -1150,20 +1150,24 @@ namespace UnrealBuildTool
 				return;
 			}
 
-			if (!BundleContentsDirectories.ContainsKey(Target) && Binary.Type == UEBuildBinaryType.Executable)
+			DirectoryReference? BundleContentsDirectory;
+			lock (BundleContentsDirectories)
 			{
-				// For Mac binary executables, we may build it outside of app, but expect it to be inside of .app after modern Xcode does its thing
-				// We still keep the executable outside of app as RequiredResource, so that Horde will copy in-between agents
-				FileReference FinalBinaryPath = Binary.OutputFilePath;
-				if (!FinalBinaryPath.Directory.FullName.EndsWith(".app/Contents/MacOS", StringComparison.OrdinalIgnoreCase))
+				if (!BundleContentsDirectories.ContainsKey(Target) && Binary.Type == UEBuildBinaryType.Executable)
 				{
-					FinalBinaryPath = FileReference.Combine(FinalBinaryPath.Directory, FinalBinaryPath.GetFileName() + ".app", "Contents", "MacOS", FinalBinaryPath.GetFileName());
-					BuildProducts[Binary.OutputFilePath] = BuildProductType.RequiredResource;
-					BuildProducts.Add(FinalBinaryPath, BuildProductType.Executable);
+					// For Mac binary executables, we may build it outside of app, but expect it to be inside of .app after modern Xcode does its thing
+					// We still keep the executable outside of app as RequiredResource, so that Horde will copy in-between agents
+					FileReference FinalBinaryPath = Binary.OutputFilePath;
+					if (!FinalBinaryPath.Directory.FullName.EndsWith(".app/Contents/MacOS", StringComparison.OrdinalIgnoreCase))
+					{
+						FinalBinaryPath = FileReference.Combine(FinalBinaryPath.Directory, FinalBinaryPath.GetFileName() + ".app", "Contents", "MacOS", FinalBinaryPath.GetFileName());
+						BuildProducts[Binary.OutputFilePath] = BuildProductType.RequiredResource;
+						BuildProducts.Add(FinalBinaryPath, BuildProductType.Executable);
+					}
+					BundleContentsDirectories.Add(Target, FinalBinaryPath.Directory.ParentDirectory!);
 				}
-				BundleContentsDirectories.Add(Target, FinalBinaryPath.Directory.ParentDirectory!);
+				BundleContentsDirectory = BundleContentsDirectories.GetValueOrDefault(Target);
 			}
-			DirectoryReference? BundleContentsDirectory = BundleContentsDirectories.GetValueOrDefault(Target);
 
 			// We need to know what third party dylibs would be copied to the bundle
 			if (Binary.Type != UEBuildBinaryType.StaticLibrary)
