Generate better errors when generating project files when the uproject file is invalid

diff --git a/Engine/Source/Programs/UnrealBuildTool/System/ProjectDescriptor.cs b/Engine/Source/Programs/UnrealBuildTool/System/ProjectDescriptor.cs
--- a/Engine/Source/Programs/UnrealBuildTool/System/ProjectDescriptor.cs
+++ b/Engine/Source/Programs/UnrealBuildTool/System/ProjectDescriptor.cs
@@ -230,6 +230,10 @@ namespace UnrealBuildTool
 			{
 				throw new JsonException($"{ex.Message} (in {FileName})", FileName.FullName, ex.LineNumber, ex.BytePositionInLine, ex);
 			}
+			catch (Exception ex)
+			{
+				throw new Exception($"{ex.Message} (while reading project file {FileName})", ex);
+			}
 		}
 
 		/// <summary>
