[UBT]
* Improved log entry for makefile creation

diff --git a/Engine/Source/Programs/UnrealBuildTool/System/TargetMakefile.cs b/Engine/Source/Programs/UnrealBuildTool/System/TargetMakefile.cs
--- a/Engine/Source/Programs/UnrealBuildTool/System/TargetMakefile.cs
+++ b/Engine/Source/Programs/UnrealBuildTool/System/TargetMakefile.cs
@@ -767,7 +767,7 @@ namespace UnrealBuildTool
 				}
 				catch (Exception Ex)
 				{
-					Logger.LogWarning("Failed to read makefile: {ExMessage}", Ex.Message);
+					Logger.LogWarning("Failed to read makefile {File}: {ExMessage}", MakefilePath.FullName, Ex.Message);
 					Logger.LogDebug(Ex, "Exception: {Ex}", Ex.ToString());
 					ReasonNotLoaded = "couldn't read existing makefile";
 					return null;
