[UBT]
* Improved timeline scopes in rules compiler

diff --git a/Engine/Source/Programs/UnrealBuildTool/System/RulesCompiler.cs b/Engine/Source/Programs/UnrealBuildTool/System/RulesCompiler.cs
--- a/Engine/Source/Programs/UnrealBuildTool/System/RulesCompiler.cs
+++ b/Engine/Source/Programs/UnrealBuildTool/System/RulesCompiler.cs
@@ -73,6 +73,8 @@ namespace UnrealBuildTool
 			{
 				if (EngineRulesAssembly == null)
 				{
+					using var ite = Timeline.ScopeEvent("EngineRulesAssembly", Logger);
+
 					// Prevent multiple conflicting processes building rules assembly at the same time
 					using GlobalSingleInstanceMutex _ = new(GlobalSingleInstanceMutex.GetUniqueMutexForPath("UnrealBuildTool_CreateTargetRulesAssembly", Unreal.EngineDirectory), true);
 					{
@@ -118,8 +120,6 @@ namespace UnrealBuildTool
 		/// <returns>New rules assembly</returns>
 		private static RulesAssembly CreateEngineRulesAssemblyInternal(List<DirectoryReference> RootDirectories, string AssemblyPrefix, IReadOnlyList<PluginInfo> Plugins, bool bReadOnly, bool bSkipCompile, bool bForceCompile, RulesAssembly? Parent, ILogger Logger)
 		{
-			using ITimelineEvent _ = Timeline.ScopeEvent("EngineRulesAssembly", Logger);
-
 			// Scope hierarchy
 			RulesScope Scope = new RulesScope("Engine", null);
 			RulesScope PluginsScope = new RulesScope("Engine Plugins", Scope);
@@ -282,6 +282,8 @@ namespace UnrealBuildTool
 				// Create the parent assembly
 				RulesAssembly Parent = CreateEngineRulesAssembly(bUsePrecompiled, bSkipCompile, bForceCompile, Logger);
 
+				using var ite = Timeline.ScopeEvent("ProjectRulesAssembly");
+
 				DirectoryReference MainProjectDirectory = ProjectFileName.Directory;
 				//DirectoryReference MainProjectSourceDirectory = DirectoryReference.Combine(MainProjectDirectory, "Source");
 
@@ -466,7 +468,6 @@ namespace UnrealBuildTool
 		/// <returns>The compiled rules assembly</returns>
 		public static RulesAssembly CreateTargetRulesAssembly(FileReference? ProjectFile, string TargetName, bool bSkipRulesCompile, bool bForceRulesCompile, bool bUsePrecompiled, FileReference? ForeignPlugin, bool bBuildPluginAsLocal, ILogger Logger)
 		{
-			using var _ = Timeline.ScopeEvent("CreateRulesAssemly");
 			RulesAssembly RulesAssembly;
 			if (ProjectFile != null)
 			{
