[UBT]
* Minor optimization in actiongraph logic for finding outdated actions and delete old files

diff --git a/Engine/Source/Programs/UnrealBuildTool/System/ActionGraph.cs b/Engine/Source/Programs/UnrealBuildTool/System/ActionGraph.cs
--- a/Engine/Source/Programs/UnrealBuildTool/System/ActionGraph.cs
+++ b/Engine/Source/Programs/UnrealBuildTool/System/ActionGraph.cs
@@ -956,20 +956,6 @@ namespace UnrealBuildTool
 		{
 			using ITimelineEvent __ = Timeline.ScopeEvent($"GatherOutdatedActions");
 
-			using (GlobalTracer.Instance.BuildSpan("Prefetching include dependencies").StartActive())
-			{
-				List<FileItem> Dependencies = new List<FileItem>();
-				foreach (LinkedAction Action in Actions)
-				{
-					if (Action.DependencyListFile != null)
-					{
-						Dependencies.Add(Action.DependencyListFile);
-					}
-				}
-
-				Parallel2.ForEach(Dependencies, File => { CppDependencies.TryGetDependencies(File, Logger, out _); });
-			}
-
 			using (GlobalTracer.Instance.BuildSpan("Cache individual outdated actions").StartActive())
 			{
 				Parallel2.ForEach(Actions,
@@ -995,15 +981,12 @@ namespace UnrealBuildTool
 		{
 			using ITimelineEvent __ = Timeline.ScopeEvent($"DeleteOutdatedProducedItems");
 
-			Parallel2.ForEach(OutdatedActions, OutdatedAction =>
+			Parallel2.ForEach(OutdatedActions.SelectMany(oa => oa.DeleteItems), DeleteItem =>
 			{
-				foreach (FileItem DeleteItem in OutdatedAction.DeleteItems)
+				if (DeleteItem.Exists)
 				{
-					if (DeleteItem.Exists)
-					{
-						Logger.LogDebug("Deleting outdated item: {DeleteItem}", DeleteItem.Location);
-						DeleteItem.Delete(Logger);
-					}
+					Logger.LogDebug("Deleting outdated item: {DeleteItem}", DeleteItem.Location);
+					DeleteItem.Delete(Logger);
 				}
 			});
 		}
