[UBT]
* Implemented parallel makefile generation. UEBuildTargets will now run in parallel as long as they don't share codegen and intermediate folders..
* Fixed so codegen is "merged" when multiple targets share codegen. Since codegen is expensive we want to run this as little as possible
* Fixed so each CppDependencyCache partition is now based on path + platform/config/appname in order to make sure cache is isolated to target
* Added BuildConfiguration.bParallelMakefileGeneration that defaults to true but can be set to false to make sure makefile generation is singlethreaded

diff --git a/Engine/Source/Programs/UnrealBuildTool/Configuration/BuildConfiguration.cs b/Engine/Source/Programs/UnrealBuildTool/Configuration/BuildConfiguration.cs
--- a/Engine/Source/Programs/UnrealBuildTool/Configuration/BuildConfiguration.cs
+++ b/Engine/Source/Programs/UnrealBuildTool/Configuration/BuildConfiguration.cs
@@ -264,6 +264,9 @@ namespace UnrealBuildTool
 		[XmlConfigFile(Category = "BuildConfiguration")]
 		public bool? bForceUnityBuild { get; set; }
 
+		[XmlConfigFile(Category = "BuildConfiguration")]
+		public bool bParallelMakefileGeneration { get; set; } = true;
+
 		/// <summary>
 		/// Intermediate environment. Determines if the intermediates end up in a different folder than normal.
 		/// </summary>
diff --git a/Engine/Source/Programs/UnrealBuildTool/Configuration/UEBuildTarget.cs b/Engine/Source/Programs/UnrealBuildTool/Configuration/UEBuildTarget.cs
--- a/Engine/Source/Programs/UnrealBuildTool/Configuration/UEBuildTarget.cs
+++ b/Engine/Source/Programs/UnrealBuildTool/Configuration/UEBuildTarget.cs
@@ -1256,15 +1256,13 @@ namespace UnrealBuildTool
 		/// <returns>The build target object for the specified build rules source file</returns>
 		public static UEBuildTarget Create(TargetDescriptor Descriptor, bool bSkipRulesCompile, bool bForceRulesCompile, bool bUsePrecompiled, UnrealIntermediateEnvironment IntermediateEnvironment, ILogger Logger)
 		{
-			Task<SourceFileMetadataCache> metaDataTask = Task.Run(() => SourceFileMetadataCache.CreateHierarchy(Descriptor.ProjectFile, Logger));
-
 			var (RulesObject, RulesAssembly) = CreateRulesObject(Descriptor, bSkipRulesCompile, bForceRulesCompile, bUsePrecompiled, IntermediateEnvironment, Logger);
 
 			// Generate a build target from this rules module
 			UEBuildTarget Target;
 			using (GlobalTracer.Instance.BuildSpan("UEBuildTarget constructor").StartActive())
 			{
-				Target = new UEBuildTarget(Descriptor, new ReadOnlyTargetRules(RulesObject), RulesAssembly, metaDataTask.Result, Logger);
+				Target = new UEBuildTarget(Descriptor, new ReadOnlyTargetRules(RulesObject), RulesAssembly, Logger);
 			}
 			using (GlobalTracer.Instance.BuildSpan("UEBuildTarget.PreBuildSetup()").StartActive())
 			{
@@ -1291,7 +1289,7 @@ namespace UnrealBuildTool
 			TargetRules RulesObject;
 			using (GlobalTracer.Instance.BuildSpan("RulesAssembly.CreateTargetRules()").StartActive())
 			{
-				using var _ = Timeline.ScopeEvent("CreateTarget");
+				using var _ = Timeline.ScopeEvent("CreateTargetRules");
 				RulesObject = RulesAssembly.CreateTargetRules(Descriptor.Name, Descriptor.Platform, Descriptor.Configuration, Descriptor.Architectures, Descriptor.ProjectFile, Descriptor.AdditionalArguments, Logger, Descriptor.IsTestsTarget, IntermediateEnvironment: IntermediateEnvironment);
 			}
 			if ((ProjectFileGenerator.bGenerateProjectFiles == false) && !RulesObject.GetSupportedPlatforms().Contains(Descriptor.Platform))
@@ -1616,11 +1614,6 @@ namespace UnrealBuildTool
 		/// </summary>
 		public RulesAssembly RulesAssembly;
 
-		/// <summary>
-		/// Cache of source file metadata for this target
-		/// </summary>
-		public SourceFileMetadataCache MetadataCache;
-
 		/// <summary>
 		/// The project file for this target
 		/// </summary>
@@ -1754,7 +1747,7 @@ namespace UnrealBuildTool
 		/// <summary>
 		/// Used to keep track of all modules by name.
 		/// </summary>
-		private Dictionary<string, UEBuildModule> Modules = new (StringComparer.InvariantCultureIgnoreCase);
+		private Dictionary<string, UEBuildModule> Modules = new(StringComparer.InvariantCultureIgnoreCase);
 
 		private Dictionary<string, UEBuildModule> UsedShortNames = new(StringComparer.InvariantCultureIgnoreCase);
 
@@ -1795,29 +1788,15 @@ namespace UnrealBuildTool
 			return UEBuildPlatformSDK.bHasAnySDKOverride;
 		}
 
-		/// <summary>
-		/// OLD Constructor.
-		/// </summary>
-		/// <param name="InDescriptor">Target descriptor</param>
-		/// <param name="InRules">The target rules, as created by RulesCompiler.</param>
-		/// <param name="InRulesAssembly">The chain of rules assemblies that this target was created with</param>
-		/// <param name="Logger">Logger for output</param>
-		private UEBuildTarget(TargetDescriptor InDescriptor, ReadOnlyTargetRules InRules, RulesAssembly InRulesAssembly, ILogger Logger)
-			: this(InDescriptor, InRules, InRulesAssembly, SourceFileMetadataCache.CreateHierarchy(InDescriptor.ProjectFile, Logger), Logger)
-		{
-		}
-
 		/// <summary>
 		/// Constructor.
 		/// </summary>
 		/// <param name="InDescriptor">Target descriptor</param>
 		/// <param name="InRules">The target rules, as created by RulesCompiler.</param>
 		/// <param name="InRulesAssembly">The chain of rules assemblies that this target was created with</param>
-		/// <param name="InMetadataCache"></param>
 		/// <param name="Logger">Logger for output</param>
-		internal UEBuildTarget(TargetDescriptor InDescriptor, ReadOnlyTargetRules InRules, RulesAssembly InRulesAssembly, SourceFileMetadataCache InMetadataCache, ILogger Logger)
+		internal UEBuildTarget(TargetDescriptor InDescriptor, ReadOnlyTargetRules InRules, RulesAssembly InRulesAssembly, ILogger Logger)
 		{
-			MetadataCache = InMetadataCache;
 			ProjectFile = InDescriptor.ProjectFile;
 			TargetName = InRules.Name;
 			AppName = InRules.Name;
@@ -1881,7 +1860,14 @@ namespace UnrealBuildTool
 			// Build the engine intermediate directory. If we're building agnostic engine binaries, we can use the engine intermediates folder. Otherwise we need to use the project intermediates directory.
 			if (!bUseSharedBuildEnvironment)
 			{
-				EngineIntermediateDirectory = ProjectIntermediateDirectory;
+				if (Configuration == UnrealTargetConfiguration.DebugGame)
+				{
+					EngineIntermediateDirectory = DirectoryReference.Combine(OutputRootDirectory, PlatformIntermediateFolder, GetTargetIntermediateFolderName(TargetName, IntermediateEnvironment), UnrealTargetConfiguration.Development.ToString());
+				}
+				else
+				{
+					EngineIntermediateDirectory = ProjectIntermediateDirectory;
+				}
 			}
 			else if (Configuration == UnrealTargetConfiguration.DebugGame)
 			{
@@ -2534,7 +2520,7 @@ namespace UnrealBuildTool
 			{
 				Logger.LogWarning("Failed to write partial receipt for target {Descriptor}: {Message}", Descriptor.Name, innerEx.Message);
 				Logger.LogDebug("{InnerException}", innerEx);
-				
+
 				return null;
 			}
 		}
@@ -2603,11 +2589,61 @@ namespace UnrealBuildTool
 			return Makefile;
 
 		}
+
 		public (TargetMakefile, Task?) BuildMakefile(BuildConfiguration BuildConfiguration, ISourceFileWorkingSet WorkingSet, TargetDescriptor TargetDescriptor, ILogger Logger, bool bInitOnly = false)
+		{
+			BuildContext context = BuildMakefile1(TargetDescriptor, Logger, null);
+
+			// UHT mode uses this to create the makefile at least to the point where UHT would have valid manifest information
+			if (bInitOnly)
+			{
+				return (context.Makefile, null);
+			}
+
+			CodeGenContext codeGenContext = new CodeGenContext()
+			{
+				Descriptor = TargetDescriptor,
+				Makefile = context.Makefile,
+				TargetName = TargetName,
+				GenerateHeaderFuncs = context.GenerateHeaderFuncs,
+			};
+
+			RunCodeGen(BuildConfiguration, [codeGenContext], Logger);
+
+			BuildMakefile2(WorkingSet, Logger, context);
+
+			codeGenContext.uhtRun?.Wait();
+
+			return (context.Makefile, codeGenContext.uhtTimestamp);
+		}
+
+		internal struct GenerateHeaderFunc
+		{
+			internal DirectoryReference Dir;
+			internal Action<ILogger, DirectoryReference> Func;
+		}
+
+		internal class BuildContext
+		{
+			required internal SourceFileMetadataCache MetadataCache;
+			required internal TargetDescriptor Descriptor;
+			required internal TargetMakefile Makefile;
+			required internal List<UEBuildBinary> OriginalBinaries;
+			required internal HashSet<UEBuildModuleCPP> ModulesToGenerateHeadersFor;
+			required internal CppCompileEnvironment GlobalCompileEnvironment;
+			required internal UEToolChain TargetToolChain;
+			required internal LinkEnvironment GlobalLinkEnvironment;
+			required internal List<GenerateHeaderFunc> GenerateHeaderFuncs;
+		}
+
+		internal BuildContext BuildMakefile1(TargetDescriptor TargetDescriptor, ILogger Logger, SourceFileMetadataCache? MetadataCache)
 		{
 			CppConfiguration CppConfiguration = GetCppConfiguration(Configuration);
 
-			SourceFileMetadataCache MetadataCache = SourceFileMetadataCache.CreateHierarchy(ProjectFile, Logger);
+			if (MetadataCache == null)
+			{
+				MetadataCache = SourceFileMetadataCache.CreateHierarchy(ProjectFile, Logger);
+			}
 
 			CppCompileEnvironment GlobalCompileEnvironment = new CppCompileEnvironment(Platform, CppConfiguration, Architectures, MetadataCache);
 			LinkEnvironment GlobalLinkEnvironment = new LinkEnvironment(GlobalCompileEnvironment.Platform, GlobalCompileEnvironment.Configuration, GlobalCompileEnvironment.Architectures);
@@ -2647,7 +2683,6 @@ namespace UnrealBuildTool
 				Rules.ConfigValueTracker, bDeployAfterCompile, UbtPlugins?.ToArray(), EnabledUbtPlugins?.ToArray(),
 				EnabledUhtPlugins?.ToArray());
 			Makefile.IsTestTarget = Rules.IsTestTarget;
-			TargetMakefileBuilder MakefileBuilder = new TargetMakefileBuilder(Makefile, Logger);
 
 			// Get diagnostic info to be printed before each build
 			TargetToolChain.GetVersionInfo(Makefile.Diagnostics);
@@ -2693,13 +2728,13 @@ namespace UnrealBuildTool
 				using var _ = Timeline.ScopeEvent("UHTModules");
 
 				UHTExecution.SetupUObjectModules(
-					ModulesToGenerateHeadersFor, 
+					ModulesToGenerateHeadersFor,
 					Rules.Platform,
-					ProjectDescriptor, 
-					Makefile.UObjectModules, 
-					Makefile.UObjectModuleHeaders, 
-					Makefile.ExecutableFile.Directory, 
-					Rules.GeneratedCodeVersion, 
+					ProjectDescriptor,
+					Makefile.UObjectModules,
+					Makefile.UObjectModuleHeaders,
+					Makefile.ExecutableFile.Directory,
+					Rules.GeneratedCodeVersion,
 					MetadataCache);
 			}
 
@@ -2721,50 +2756,186 @@ namespace UnrealBuildTool
 			}
 			Logger.LogInformation($"UHT compiled-in object format {Makefile.UHTTargetSettings.CompiledInObjectFormat}");
 
-			// UHT mode uses this to create the makefile at least to the point where UHT would have valid manifest information
-			if (bInitOnly)
+#if __VPROJECT_AVAILABLE__
 			{
-				return (Makefile, null);
-			}
+				using var _ = Timeline.ScopeEvent("VNIModules");
+				VNIExecution.SetupVNIModules(ModulesToGenerateHeadersFor, out Makefile.VNIModules);
 
-			// We inject all files that we know UHT/VNI will create into the directory cache.. this way we can run UHT/VNI without waiting
-			// .. and when we are about to test if actions are out-of-date we just reset directory cache so times are re-fetched
-			UHTExecution.PrepopulateDirectoryCache(Makefile.UObjectModules);
+				// Copy Verse BPVM usage flag
+				Makefile.bUseVerseBPVM = Rules.bUseVerseBPVM;
+			}
+#endif
 
-			// NOTE: Even in Gather mode, we need to run UHT to make sure the files exist for the static action graph to be setup correctly.  This is because UHT generates .cpp
-			// files that are injected as top level prerequisites.  If UHT only emitted included header files, we wouldn't need to run it during the Gather phase at all.
-			Task? uhtRunTask = null;
-			Task? uhtTimestampTask = null;
-			if (Makefile.UObjectModules.Count > 0)
+			// Collect generate header funcs
+			List<GenerateHeaderFunc> GenerateHeaderFuncs = [];
+			foreach (UEBuildModuleCPP Module in Modules.Values.OfType<UEBuildModuleCPP>())
 			{
-				(uhtRunTask, uhtTimestampTask) = UHTExecution.ExecuteHeaderToolIfNecessary(BuildConfiguration, ProjectFile, Makefile, TargetName, Logger, false);
+				foreach ((string Subdirectory, Action<ILogger, DirectoryReference> Func) in Module.Rules.GenerateHeaderFuncs)
+				{
+					DirectoryReference GenDirectory = DirectoryReference.Combine(Module.GeneratedCodeDirectory!, Subdirectory);
+					GenerateHeaderFuncs.Add(new() { Dir = GenDirectory, Func = Func });
+				}
 			}
 
-#if __VPROJECT_AVAILABLE__
-			// Prepare cached data for VNI header generation
-			using (new ScopedTimer("UHTExecution.SetupVNIModules()", Logger, LogLevel.Debug, false))
+
+			return new BuildContext()
 			{
-				using var _ = Timeline.ScopeEvent("VNIModules");
-				VNIExecution.SetupVNIModules(ModulesToGenerateHeadersFor, out Makefile.VNIModules);
+				MetadataCache = MetadataCache,
+				Descriptor = TargetDescriptor,
+				Makefile = Makefile,
+				OriginalBinaries = OriginalBinaries,
+				ModulesToGenerateHeadersFor = ModulesToGenerateHeadersFor,
+				GlobalCompileEnvironment = GlobalCompileEnvironment,
+				TargetToolChain = TargetToolChain,
+				GlobalLinkEnvironment = GlobalLinkEnvironment,
+				GenerateHeaderFuncs = GenerateHeaderFuncs,
+			};
+		}
+
+		internal class CodeGenContext
+		{
+			required internal TargetDescriptor Descriptor;
+			required internal TargetMakefile Makefile;
+			required internal string TargetName;
+			required internal List<GenerateHeaderFunc> GenerateHeaderFuncs;
+			internal Task? uhtRun;
+			internal Task? uhtTimestamp;
+			internal Task? vniRun;
+		}
+
+		internal static void RunCodeGen(BuildConfiguration BuildConfiguration, List<CodeGenContext> contexts, ILogger Logger)
+		{
+			if (contexts.Count == 0)
+			{
+				return;
 			}
-			VNIExecution.PrepopulateDirectoryCache(Makefile);
 
-			// Copy Verse BPVM usage flag
-			Makefile.bUseVerseBPVM = Rules.bUseVerseBPVM;
-			Task? vniTask = null;
-			if (Makefile.VNIModules.Count > 0)
+			CodeGenContext? lastRunContext = null;
+			HashSet<string>? lastRunUObjectModules = null;
+			HashSet<string>? lastRunVNIModules = null;
+
+			string GetTargetName(TargetDescriptor desc) => $"{desc.Name}{desc.Platform}{desc.Configuration}";
+
+			foreach (CodeGenContext context in contexts)
 			{
-				vniTask = new Task(() =>
+				TargetMakefile makefile = context.Makefile;
+
+				foreach (var entry in context.GenerateHeaderFuncs)
 				{
-					using (new ScopedTimer("ExternalExecution.ExecuteVNITool()", Logger, LogLevel.Debug, false))
+					DirectoryReference.CreateDirectory(entry.Dir);
+					entry.Func(Logger, entry.Dir);
+				}
+
+				if (lastRunContext != null)
+				{
+					if (context.TargetName != lastRunContext.TargetName)
+					{
+						throw new Exception("Different targetname!.. run after each other!");
+					}
+
+					if (lastRunUObjectModules == null)
 					{
-						VNIExecution.ExecuteVNITool(Makefile, TargetDescriptor, Logger, false);
+						lastRunUObjectModules = new(lastRunContext.Makefile.UObjectModules.Select(m => m.ModuleName));
+					}
+
+					foreach (var module in makefile.UObjectModules)
+					{
+						if (!lastRunUObjectModules.Contains(module.ModuleName))
+						{
+							throw new Exception($"Target {GetTargetName(context.Descriptor)} has uht modules that previous target {GetTargetName(lastRunContext.Descriptor)} did not have!.. run after each other!");
+						}
+					}
+
+					if (lastRunVNIModules == null)
+					{
+						lastRunVNIModules = new(lastRunContext.Makefile.VNIModules.Select(m => m.ModuleName));
+					}
+
+					foreach (var module in makefile.VNIModules)
+					{
+						if (!lastRunVNIModules.Contains(module.ModuleName))
+						{
+							throw new Exception($"Target {GetTargetName(context.Descriptor)} has vni modules that previous target {GetTargetName(lastRunContext.Descriptor)} did not have!.. run after each other!");
+						}
 					}
-				}, TaskCreationOptions.LongRunning);
-				vniTask.Start();
+
+					//if (context.Makefile.EnabledUhtPlugins!.Length != lastRunContext.Makefile.EnabledUhtPlugins!.Length)
+					//	throw new Exception("DIFFERENT THIS NEEDS TO BE HANDLED!.. run after each other!");
+
+					context.uhtRun = lastRunContext.uhtRun;
+					context.uhtTimestamp = lastRunContext.uhtTimestamp;
+					context.vniRun = lastRunContext.vniRun;
+					continue;
+
+				}
+
+				TargetDescriptor descriptor = context.Descriptor;
+
+				if (makefile.UObjectModules.Count > 0)
+				{
+					UHTExecution.PrepopulateDirectoryCache(makefile.UObjectModules);
+					(Task? uhtRun, Task? uhtTimestamp) = UHTExecution.ExecuteHeaderToolIfNecessary(BuildConfiguration, descriptor.ProjectFile, makefile, context.TargetName!, Logger, false);
+					context.uhtRun = uhtRun;
+					context.uhtTimestamp = uhtTimestamp;
+				}
+
+#if __VPROJECT_AVAILABLE__
+				if (makefile.VNIModules.Count > 0)
+				{
+					VNIExecution.PrepopulateDirectoryCache(makefile);
+					context.vniRun = Task.Run(() => VNIExecution.ExecuteVNITool(makefile, descriptor, Logger));
+				}
+#endif
+
+				lastRunContext = context;
 			}
+		}
+
+		internal static void WaitCodeGen(List<CodeGenContext> contexts)
+		{
+			foreach (var context in contexts)
+			{
+				if (context.uhtRun != null)
+				{
+					UHTExecution.WaitForExecution(context.uhtRun, context.Makefile);
 
+					// TODO: Kick prefetch of folders
+					//Task task = FileMetadataPrefetch.QueueDirectoryTrees(ModulesToGenerateHeadersFor.Select(m => m.GeneratedCodeDirectoryUHT));
+				}
+#if __VPROJECT_AVAILABLE__
+				if (context.vniRun != null)
+				{
+					VNIExecution.WaitForExecution(context.vniRun, context.Makefile);
+				}
 #endif
+			}
+		}
+
+		internal static List<Task> GetTimestampTasks(List<CodeGenContext> contexts)
+		{
+			List<Task> timestampTasks = [];
+			foreach (var codeGenContext in contexts)
+			{
+				if (codeGenContext.uhtTimestamp != null)
+				{
+					timestampTasks.Add(codeGenContext.uhtTimestamp);
+				}
+			}
+			return timestampTasks;
+		}
+
+		internal void BuildMakefile2(ISourceFileWorkingSet WorkingSet, ILogger Logger, BuildContext context)
+		{
+			SourceFileMetadataCache MetadataCache = context.MetadataCache;
+			TargetDescriptor TargetDescriptor = context.Descriptor;
+			TargetMakefile Makefile = context.Makefile;
+			List<UEBuildBinary> OriginalBinaries = context.OriginalBinaries;
+			HashSet<UEBuildModuleCPP> ModulesToGenerateHeadersFor = context.ModulesToGenerateHeadersFor;
+			CppCompileEnvironment GlobalCompileEnvironment = context.GlobalCompileEnvironment;
+			UEToolChain TargetToolChain = context.TargetToolChain;
+			LinkEnvironment GlobalLinkEnvironment = context.GlobalLinkEnvironment;
+
+			TargetMakefileBuilder MakefileBuilder = new TargetMakefileBuilder(Makefile, Logger);
 
 			if (Rules.bUseSharedPCHs)
 			{
@@ -2775,16 +2946,6 @@ namespace UnrealBuildTool
 				CreateSharedPCHInstances(Rules, TargetToolChain, OriginalBinaries, GlobalCompileEnvironment, MakefileBuilder, Logger);
 			}
 
-			foreach (UEBuildModuleCPP Module in Modules.Values.OfType<UEBuildModuleCPP>())
-			{
-				foreach ((string Subdirectory, Action<ILogger, DirectoryReference> Func) in Module.Rules.GenerateHeaderFuncs)
-				{
-					DirectoryReference GenDirectory = DirectoryReference.Combine(Module.GeneratedCodeDirectory!, Subdirectory);
-					DirectoryReference.CreateDirectory(GenDirectory);
-					Func(Logger, GenDirectory);
-				}
-			}
-
 			// Compile the resource files common to all DLLs on Windows
 			if (!ShouldCompileMonolithic() & !Rules.bFormalBuild && Platform.IsInGroup(UnrealPlatformGroup.Windows))
 			{
@@ -3193,23 +3354,6 @@ namespace UnrealBuildTool
 
 			// Clean any stale modules which exist in multiple output directories. This can lead to the wrong DLL being loaded on Windows.
 			CleanStaleModules(Logger);
-
-#if __VPROJECT_AVAILABLE__
-			if (vniTask != null)
-			{
-				VNIExecution.WaitForExecution(vniTask, Makefile);
-			}
-#endif
-
-			if (uhtRunTask != null)
-			{
-				UHTExecution.WaitForExecution(uhtRunTask, Makefile);
-
-				// Prefetch directory items for UHT folders since they are going to be used later for check out-of-date actions
-				Task task = FileMetadataPrefetch.QueueDirectoryTrees(ModulesToGenerateHeadersFor.Select(m => m.GeneratedCodeDirectoryUHT));
-			}
-
-			return (Makefile, uhtTimestampTask);
 		}
 
 		void CreateStripExportsActions(UEToolChain TargetToolChain, TargetMakefileBuilder MakefileBuilder, TargetMakefile Makefile, CppCompileEnvironment CompileEnvironment)
@@ -5362,7 +5506,7 @@ namespace UnrealBuildTool
 				{
 					// Compile in the build plugins list for the runtime plugin manager to filter optional plugin references against. Only plugins that were actually built
 					// should be enabled when another plugin has an optional dependency on the plugin.
-					IEnumerable<string> BuildPluginStrings = BuildPlugins?.Select(x => $"TEXT(\"{x}\")") ?? Array.Empty<string>();
+					IEnumerable<string> BuildPluginStrings = BuildPlugins?.OrderBy(x => x.Name)?.Select(x => $"TEXT(\"{x}\")") ?? Array.Empty<string>();
 					moduleRules.PrivateDefinitions.Add($"UBT_TARGET_BUILD_PLUGINS={String.Join(", ", BuildPluginStrings)}");
 				}
 			}
diff --git a/Engine/Source/Programs/UnrealBuildTool/Modes/BuildMode.cs b/Engine/Source/Programs/UnrealBuildTool/Modes/BuildMode.cs
--- a/Engine/Source/Programs/UnrealBuildTool/Modes/BuildMode.cs
+++ b/Engine/Source/Programs/UnrealBuildTool/Modes/BuildMode.cs
@@ -8,6 +8,7 @@ using System.Linq;
 using System.Text;
 using System.Threading.Tasks;
 using EpicGames.Core;
+using Microsoft.CodeAnalysis.CSharp.Syntax;
 using Microsoft.Extensions.Logging;
 using OpenTracing.Util;
 using UnrealBuildBase;
@@ -387,9 +388,16 @@ namespace UnrealBuildTool
 		/// <returns>Result from the compilation</returns>
 		public static async Task BuildAsync(List<TargetDescriptor> TargetDescriptors, BuildConfiguration BuildConfiguration, ISourceFileWorkingSet WorkingSet, BuildOptions Options, FileReference? WriteOutdatedActionsFile, ILogger Logger, bool bSkipPreBuildTargets = false, string? ActionTypeFilter = null)
 		{
-			(List<(TargetDescriptor, TargetMakefile)> Targets, List<Task> Tasks) = CreateMakefiles(BuildConfiguration, TargetDescriptors, WorkingSet, Logger, true, true);
+			// Create cpp dependency cache in order to be able to preload them as we discover them
+			CppDependencyCache? cppDependencyCache = null;
+			if ((Options & BuildOptions.SkipBuild) == 0)
+			{
+				cppDependencyCache = new();
+			}
+
+			(List<(TargetDescriptor, TargetMakefile)> Targets, List<Task> Tasks) = CreateMakefiles(BuildConfiguration, TargetDescriptors, WorkingSet, Logger, true, true, cppDependencyCache);
 
-			await BuildAsync(Targets, BuildConfiguration, Options, WriteOutdatedActionsFile, Logger, ActionTypeFilter);
+			await BuildAsync(Targets, BuildConfiguration, Options, WriteOutdatedActionsFile, Logger, ActionTypeFilter, cppDependencyCache);
 
 			Task.WaitAll(Tasks.ToArray());
 		}
@@ -405,8 +413,9 @@ namespace UnrealBuildTool
 		/// <param name="WriteOutdatedActionsFile">Files to write the list of outdated actions to (rather than building them)</param>
 		/// <param name="Logger">Logger for output</param>
 		/// <param name="ActionTypeFilter">If specified, only actions of this type will be run</param>
+		/// <param name="cppDependencyCache"></param>
 		/// <returns>Result from the compilation</returns>
-		internal static async Task BuildAsync(List<(TargetDescriptor Descriptor, TargetMakefile Makefile)> Targets, BuildConfiguration BuildConfiguration, BuildOptions Options, FileReference? WriteOutdatedActionsFile, ILogger Logger, string? ActionTypeFilter=null)
+		internal static async Task BuildAsync(List<(TargetDescriptor Descriptor, TargetMakefile Makefile)> Targets, BuildConfiguration BuildConfiguration, BuildOptions Options, FileReference? WriteOutdatedActionsFile, ILogger Logger, string? ActionTypeFilter=null, CppDependencyCache? cppDependencyCache = null)
 		{
 			List<TargetDescriptor> TargetDescriptors = Targets.Select(T => T.Descriptor).ToList();
 			TargetMakefile[] Makefiles = Targets.Select(T => T.Makefile).ToArray();
@@ -414,20 +423,16 @@ namespace UnrealBuildTool
 			// Execute the build
 			if ((Options & BuildOptions.SkipBuild) == 0)
 			{
-				Task<CppDependencyCache> CppDependenciesTask = Task.Run(() =>
+				if (cppDependencyCache == null)
 				{
-					// Create the C++ dependency cache
-					CppDependencyCache CppDependencies = new CppDependencyCache();
-					for (int TargetIdx = 0; TargetIdx < Targets.Count; TargetIdx++)
-					{
-						using (GlobalTracer.Instance.BuildSpan("Reading dependency cache").StartActive())
-						{
-							TargetDescriptor TargetDescriptor = Targets[TargetIdx].Descriptor;
-							CppDependencies.Mount(TargetDescriptor, Targets[TargetIdx].Makefile.TargetType, Logger);
-						}
-					}
-					return CppDependencies;
-				});
+					cppDependencyCache = new();
+				}
+
+				// Create the C++ dependency cache
+				for (int TargetIdx = 0; TargetIdx < Targets.Count; TargetIdx++)
+				{
+					cppDependencyCache.MountAsync(Targets[TargetIdx].Descriptor, "", Targets[TargetIdx].Makefile.TargetType, Logger);
+				}
 
 #if false // Doesn't seem to make things faster. Revisit since it feels like this should be faster
 				// Prepopulate directories
@@ -558,7 +563,8 @@ namespace UnrealBuildTool
 				}
 
 				// Wait for cpp dependencies task to finish
-				CppDependencyCache CppDependencies = CppDependenciesTask.Result;
+				cppDependencyCache.WaitAsync();
+				CppDependencyCache CppDependencies = cppDependencyCache;
 
 				// Now that we have the dependencies object, optionally create the artifact cache
 				IActionArtifactCache? actionArtifactCache = null;
@@ -1104,253 +1110,380 @@ namespace UnrealBuildTool
 		{
 			await Task.CompletedTask;
 
-			(List<(TargetDescriptor, TargetMakefile Makefile)> Targets, List<Task> Tasks) = CreateMakefiles(BuildConfiguration, [TargetDescriptor], WorkingSet, Logger, false, false);
+			(List<(TargetDescriptor, TargetMakefile Makefile)> Targets, List<Task> Tasks) = CreateMakefiles(BuildConfiguration, [TargetDescriptor], WorkingSet, Logger, false, false, null);
 
 			Task.WaitAll(Tasks.ToArray());
 			return Targets[0].Makefile;
 		}
 
-		private static (List<(TargetDescriptor, TargetMakefile)>, List<Task>) CreateMakefiles(BuildConfiguration BuildConfiguration, List<TargetDescriptor> TargetDescriptors, ISourceFileWorkingSet WorkingSet, ILogger Logger, bool runInitScripts, bool bSkipPreBuildTargets)
+		// Function that will load makefile and check if it is valid. (First round of targets will have their makefiles loaded in parallel with other work)
+		private static (TargetMakefile? makefile, bool ranPrebuildScripts) LoadMakefile(FileReference location, TargetDescriptor descriptor, ISourceFileWorkingSet workingSet, ILogger logger)
 		{
-			List<(TargetDescriptor descriptor, TargetMakefile makefile)> targets = new();
-			List<Task> tasks = [];
+			string? ReasonNotLoaded;
+			TargetMakefile? makefile = TargetMakefile.Load(location, descriptor.ProjectFile, descriptor.Platform, descriptor.AdditionalArguments.GetRawArray(), logger, out ReasonNotLoaded);
+			if (makefile == null)
+			{
+				logger.LogInformation("Creating makefile for {DescriptorName} ({ReasonNotLoaded})", descriptor.Name, ReasonNotLoaded);
+				return (null, false);
+			}
+
+			// Execute the scripts.
+			Utils.ExecuteCustomBuildSteps(makefile.PreBuildScripts, logger);
 
-			while (TargetDescriptors.Count != 0) // TargetDescriptors will be reallocated for PrebuildTargets
+			// Check that the makefile is still valid
+			string? Reason;
+			if (!TargetMakefile.IsValidForSourceFiles(makefile, descriptor.ProjectFile, descriptor.Platform, workingSet, logger, out Reason))
 			{
-				bool skipRulesCompile = BuildConfiguration.bSkipRulesCompile;
-				bool forceRulesCompile = BuildConfiguration.bForceRulesCompile;
-				bool usePrecompiled = BuildConfiguration.bUsePrecompiled;
+				logger.LogInformation("Invalidating makefile for {DescriptorName} ({Reason})", descriptor.Name, Reason);
+				return (null, true);
+			}
 
-				(TargetRules rules, RulesAssembly assembly)[] rules = new (TargetRules, RulesAssembly)[TargetDescriptors.Count];
-				Task<SourceFileMetadataCache>[] metadataCaches = new Task<SourceFileMetadataCache>[TargetDescriptors.Count];
+			return (makefile, true);
+		}
 
+		private static (List<(TargetDescriptor, TargetMakefile)>, List<Task>) CreateMakefiles(BuildConfiguration buildConfiguration, List<TargetDescriptor> descriptors, ISourceFileWorkingSet workingSet, ILogger logger, bool runInitScripts, bool skipPreBuildTargets, CppDependencyCache? cppDependencyCache)
+		{
+			List<(TargetDescriptor descriptor, TargetMakefile makefile)> outTargets = new();
+			List<Task> outTasks = [];
 
-				// Build all rules objects
-				Parallel2.For(0, TargetDescriptors.Count, (index) =>
-				{
-					TargetDescriptor TargetDescriptor = TargetDescriptors[index];
-					metadataCaches[index] = Task.Run(() => SourceFileMetadataCache.CreateHierarchy(TargetDescriptor.ProjectFile, Logger));
-					rules[index] = UEBuildTarget.CreateRulesObject(TargetDescriptor, skipRulesCompile, forceRulesCompile, usePrecompiled, TargetDescriptor.IntermediateEnvironment, Logger);
-				});
+			int helperCount = buildConfiguration.bParallelMakefileGeneration ? -1 : 0;
 
-				// Use rules objects to figure out how to group work. We can't run targets that write to same folders in parallel
-				Dictionary<string, List<int>> buildGroups = [];
+			Dictionary<FileReference, Task<SourceFileMetadataCache>> sourceFileMetadataCaches = [];
 
-				for (int i = 0; i < TargetDescriptors.Count; ++i)
-				{
-					TargetDescriptor targetDescriptor = TargetDescriptors[i];
-					string platformIntermediateFolderNoArch = UEBuildTarget.GetPlatformIntermediateFolder(targetDescriptor.Platform, null, false);
-					string appName = UEBuildTarget.GetAppName(rules[i].rules);
+			while (descriptors.Count != 0) // descriptors list will be reallocated for PrebuildTargets. hence the while-loop
+			{
+				int descriptorCount = descriptors.Count;
+				bool skipRulesCompile = buildConfiguration.bSkipRulesCompile;
+				bool forceRulesCompile = buildConfiguration.bForceRulesCompile;
+				bool usePrecompiled = buildConfiguration.bUsePrecompiled;
 
-					if (!buildGroups.TryGetValue(appName, out var group))
+				FileReference?[] makefileLocations = new FileReference[descriptorCount];
+				UEBuildTarget[] buildTargets = new UEBuildTarget[descriptorCount];
+
+				// Start loading unique makefiles
+				Dictionary<FileReference, Task<(TargetMakefile?, bool)>> preLoadedMakefiles = [];
+				if (buildConfiguration.bUseUBTMakefiles)
+				{
+					for (int index = 0; index != descriptorCount; ++index)
 					{
-						group = new();
-						buildGroups.Add(appName, group);
+						TargetDescriptor descriptor = descriptors[index];
+						FileReference makefileLocation = TargetMakefile.GetLocation(descriptor.ProjectFile, descriptor.Name, descriptor.Platform, descriptor.Architectures, descriptor.Configuration, descriptor.IntermediateEnvironment);
+						makefileLocations[index] = makefileLocation;
+						if (!preLoadedMakefiles.ContainsKey(makefileLocation))
+						{
+							preLoadedMakefiles.Add(makefileLocation, Task.Run(() => { return LoadMakefile(makefileLocation, descriptor, workingSet, logger); }));
+						}
 					}
-
-					group.Add(i);
 				}
 
-				List<TargetDescriptor> newTargetDescriptors = [];
-
-				// TODO: We also want to handle different configs.. they can't run UHT etc at the same time but can do the other things..
-				// ... so ne need to get that support in here too
+				// All appnames have unique folders. So if appnames are different we can run them in parallel.
+				// The reason we can't run different targets with same appname in parallel is because of codegen (since codegen is not in target folder)
+				// LyraEditor, QAGameEditor etc has the same appname (UnrealEditor) but same codegen folder (for engine part at least)
+				Dictionary<string, List<int>> appNameTargetGroups = [];
 
-
-				//Parallel2.ForEach(buildGroups, kv =>
-				foreach (KeyValuePair<string, List<int>> kv in buildGroups)
+				// Build all rules objects and populate groups. Start loading makefiles
+				Parallel2.For(0, descriptorCount, (index) =>
 				{
-					Task? LastSaveMakefileTask = null;
-					Task? LastUHTTimestampTask = null;
+					TargetDescriptor descriptor = descriptors[index];
 
-					List<int> buildGroup = kv.Value;
-					for (int i = 0; i != buildGroup.Count; ++i)
+					// Create task to load source file metadata cache already here, it will be used later on
+					if (descriptor.ProjectFile != null)
 					{
-						int index = buildGroup[i];
-						TargetDescriptor TargetDescriptor = TargetDescriptors[index];
-
-						using ITimelineEvent SetupScope = Timeline.ScopeEvent($"{TargetDescriptor.Name} {TargetDescriptor.Platform} {TargetDescriptor.Configuration}");
-
-						if (runInitScripts)
+						lock (sourceFileMetadataCaches)
 						{
-							// Create and execute the init scripts
-							FileReference[] InitScripts = CreateInitScripts(TargetDescriptor);
-							if (InitScripts.Length > 0)
+							if (!sourceFileMetadataCaches.ContainsKey(descriptor.ProjectFile))
 							{
-								Utils.ExecuteCustomBuildSteps(InitScripts, Logger);
+								sourceFileMetadataCaches.Add(descriptor.ProjectFile, Task.Run(() => SourceFileMetadataCache.CreateHierarchy(descriptor.ProjectFile, logger)));
 							}
 						}
+					}
 
-						FileReference? SaveToLocation = null;
+					// Compile/load assembly and create target rules
+					(TargetRules rules, RulesAssembly assembly) = UEBuildTarget.CreateRulesObject(descriptor, skipRulesCompile, forceRulesCompile, usePrecompiled, descriptor.IntermediateEnvironment, logger);
 
-						// Get the path to the makefile for this target
-						FileReference? MakefileLocation = null;
-						if (BuildConfiguration.bUseUBTMakefiles)
-						{
-							MakefileLocation = TargetMakefile.GetLocation(TargetDescriptor.ProjectFile, TargetDescriptor.Name, TargetDescriptor.Platform, TargetDescriptor.Architectures, TargetDescriptor.Configuration, TargetDescriptor.IntermediateEnvironment);
-						}
+					// Create the target. Ctor does not do much so we can create it regardless if we are building it or not
+					UEBuildTarget buildTarget = new(descriptor, new ReadOnlyTargetRules(rules), assembly, logger);
+					buildTargets[index] = buildTarget;
 
-						// Need to make sure to wait on previous run's makefile
-						LastSaveMakefileTask?.Wait();
+					// Queue up cppdependencies loading
+					if (cppDependencyCache != null)
+					{
+						cppDependencyCache.MountAsync(descriptor, buildTarget.AppName, buildTarget.TargetType, logger);
+					}
 
-						// Try to load an existing makefile
-						TargetMakefile? Makefile = null;
-						if (MakefileLocation != null)
+					// Add target to proper app name group
+					string appName = $"{buildTarget.AppName}_{buildTarget.Platform}";
+					lock (appNameTargetGroups)
+					{
+						if (!appNameTargetGroups.TryGetValue(appName, out List<int>? appNameTarget))
 						{
-							using (GlobalTracer.Instance.BuildSpan("TargetMakefile.Load()").StartActive())
-							{
-								string? ReasonNotLoaded;
-								Makefile = TargetMakefile.Load(MakefileLocation, TargetDescriptor.ProjectFile, TargetDescriptor.Platform, TargetDescriptor.AdditionalArguments.GetRawArray(), Logger, out ReasonNotLoaded);
-								if (Makefile == null)
-								{
-									Logger.LogInformation("Creating makefile for {TargetDescriptorName} ({ReasonNotLoaded})", TargetDescriptor.Name, ReasonNotLoaded);
-								}
-							}
+							appNameTarget = new();
+							appNameTargetGroups.Add(appName, appNameTarget);
 						}
 
-						// If we have a makefile, execute the pre-build steps and check it's still valid
-						bool bHasRunPreBuildScripts = false;
-						if (Makefile != null)
-						{
-							// Execute the scripts.
-							Utils.ExecuteCustomBuildSteps(Makefile.PreBuildScripts, Logger);
+						appNameTarget.Add(index);
+					}
+				}, helperCount);
 
-							// Don't run the pre-build steps again, even if we invalidate the makefile.
-							bHasRunPreBuildScripts = true;
+				UEBuildTarget.BuildContext?[] buildContexts = new UEBuildTarget.BuildContext?[descriptorCount];
+				TargetMakefile?[] makefiles = new TargetMakefile?[descriptorCount];
+				ITimelineEvent?[] targetTimelines = new ITimelineEvent?[descriptorCount];
 
-							// Check that the makefile is still valid
-							string? Reason;
-							if (!TargetMakefile.IsValidForSourceFiles(Makefile, TargetDescriptor.ProjectFile, TargetDescriptor.Platform, WorkingSet, Logger, out Reason))
-							{
-								Logger.LogInformation("Invalidating makefile for {TargetDescriptorName} ({Reason})", TargetDescriptor.Name, Reason);
-								Makefile = null;
-							}
-						}
+				// Will be populated with prebuild targets if there are any
+				List<TargetDescriptor> newDescriptors = [];
+
+				// Run all appname groups in parallel
+				Parallel2.ForEach(appNameTargetGroups, kv =>
+				{
+					List<int> appNameTargets = kv.Value;
+					appNameTargets.Sort(); // To build in same order as registered
 
-						// Need to make sure to wait on previous run's timestamp since there are targets that write to same place..
-						LastUHTTimestampTask?.Wait();
+					// All targets with same code gen folder but different engine intermediate dir can run in parallel
+					// but they need to sync up for when code gen happens (UHT/VNI)
 
-						Task? uhtTimeStampsTask = null;
+					Dictionary<DirectoryReference, List<int>> intermediateFolderGroups = [];
+					int maxTargetsInIntermediateCount = 0;
 
-						// If we couldn't load a makefile, create a new one
-						if (Makefile == null)
+					for (int i = 0; i != appNameTargets.Count; ++i)
+					{
+						int index = appNameTargets[i];
+						UEBuildTarget buildTarget = buildTargets[index];
+						DirectoryReference dir = buildTarget.EngineIntermediateDirectory;
+						if (!intermediateFolderGroups.TryGetValue(dir, out List<int>? intermediateFolderGroup))
 						{
-							// Create the target
-							UEBuildTarget Target = new UEBuildTarget(TargetDescriptor, new ReadOnlyTargetRules(rules[index].rules), rules[index].assembly, metadataCaches[index].Result, Logger);
-							Target.PreBuildSetup(Logger);
+							intermediateFolderGroup = new();
+							intermediateFolderGroups.Add(dir, intermediateFolderGroup);
+						}
+						intermediateFolderGroup.Add(index);
+						maxTargetsInIntermediateCount = Math.Max(maxTargetsInIntermediateCount, intermediateFolderGroup.Count);
+					}
+
+					// We need to track save makefile task so we can wait for it in case we need to load it again.
+					// TODO: Just reuse previous one. we don't need to save until we are on the last one.
+					Task?[] lastSaveMakefileTask = new Task?[descriptorCount]; // This must be allocated in here because it is iterated when all passes are done
 
-							// Create the pre-build scripts
-							FileReference[] PreBuildScripts = Target.CreatePreBuildScripts();
+					List<UEBuildTarget.CodeGenContext>? lastContextsForCodeGen = null;
 
-							// Execute the pre-build scripts
-							if (!bHasRunPreBuildScripts)
+					// Split up in passes. We run one target from each intermediate folder.
+					for (int passIt = 0; passIt != maxTargetsInIntermediateCount; ++passIt)
+					{
+						List<UEBuildTarget.CodeGenContext> contextsForCodeGen = [];
+
+						// Do everyting up to code gen in parallel
+						Parallel2.ForEach(intermediateFolderGroups.Values, intermediateFolderGroup =>
+						{
+							if (passIt >= intermediateFolderGroup.Count)
 							{
-								Utils.ExecuteCustomBuildSteps(PreBuildScripts, Logger);
-								bHasRunPreBuildScripts = true;
+								return; // This group is done
 							}
 
-							// Build the target
-							using (GlobalTracer.Instance.BuildSpan("UEBuildTarget.Build()").StartActive())
+							int index = intermediateFolderGroup[passIt];
+							TargetDescriptor descriptor = descriptors[index];
+
+							targetTimelines[index] = Timeline.ScopeEvent($"{descriptor.Name} {descriptor.Platform} {descriptor.Configuration}");
+
+							if (runInitScripts)
 							{
-								(Makefile, uhtTimeStampsTask) = Target.BuildMakefile(BuildConfiguration, WorkingSet, TargetDescriptor, Logger);
+								// Create and execute the init scripts
+								FileReference[] InitScripts = CreateInitScripts(descriptor);
+								if (InitScripts.Length > 0)
+								{
+									Utils.ExecuteCustomBuildSteps(InitScripts, logger);
+								}
 							}
 
-							Makefile.MemoryPerActionGB = Target.Rules.MemoryPerActionGB;
+							TargetMakefile? makefile = null;
+							bool ranPrebuildScripts = false;
 
-							// Save the pre-build scripts onto the makefile
-							Makefile.PreBuildScripts = PreBuildScripts;
+							FileReference? makefileLocation = makefileLocations[index];
+							if (makefileLocation != null)
+							{
+								Task<(TargetMakefile?, bool)>? makefileTask = null;
+								lock (preLoadedMakefiles)
+								{
+									if (preLoadedMakefiles.TryGetValue(makefileLocation, out makefileTask))
+									{
+										preLoadedMakefiles.Remove(makefileLocation);
+									}
+								}
+								if (makefileTask != null)
+								{
+									(makefile, ranPrebuildScripts) = makefileTask.Result;
+								}
+								else
+								{
+									// Need to make sure to wait on previous run's makefile
+									lastSaveMakefileTask[index]?.Wait();
+									lastSaveMakefileTask[index] = null;
+									(makefile, ranPrebuildScripts) = LoadMakefile(makefileLocation, descriptor, workingSet, logger);
+								}
+							}
+
+							UEBuildTarget buildTarget = buildTargets[index];
 
-							// Save the additional command line arguments
-							Makefile.AdditionalArguments = TargetDescriptor.AdditionalArguments.GetRawArray();
+							List<UEBuildTarget.GenerateHeaderFunc> generatedHeaderFuncs = [];
 
-							// Save the environment variables
-							foreach (System.Collections.DictionaryEntry? EnvironmentVariable in Environment.GetEnvironmentVariables())
+							// If we couldn't load a valid makefile, create a new one
+							if (makefile == null)
 							{
-								if (EnvironmentVariable != null)
+								buildTarget.PreBuildSetup(logger);
+
+								// Create the pre-build scripts
+								FileReference[] PreBuildScripts = buildTarget.CreatePreBuildScripts();
+
+								// Execute the pre-build scripts
+								if (!ranPrebuildScripts)
 								{
-									Makefile.EnvironmentVariables.Add(Tuple.Create((string)EnvironmentVariable.Value.Key, (string)EnvironmentVariable.Value.Value!));
+									Utils.ExecuteCustomBuildSteps(PreBuildScripts, logger);
+									ranPrebuildScripts = true;
 								}
+
+								SourceFileMetadataCache? sourceFileMetaDataCache = null;
+								if (descriptor.ProjectFile != null)
+								{
+									sourceFileMetaDataCache = sourceFileMetadataCaches[descriptor.ProjectFile].Result;
+								}
+
+								// Build the target
+								UEBuildTarget.BuildContext context = buildTarget.BuildMakefile1(descriptor, logger, sourceFileMetaDataCache);
+								buildContexts[index] = context;
+
+								makefile = context.Makefile;
+
+								makefile.PreBuildScripts = PreBuildScripts;
+
+								makefile.MemoryPerActionGB = buildTarget.Rules.MemoryPerActionGB;
+
+								// Save the additional command line arguments
+								makefile.AdditionalArguments = descriptor.AdditionalArguments.GetRawArray();
+
+								generatedHeaderFuncs = context.GenerateHeaderFuncs;
 							}
 
-							// Save the makefile for next time
-							if (MakefileLocation != null)
+							contextsForCodeGen.Add(new UEBuildTarget.CodeGenContext()
 							{
-								SaveToLocation = MakefileLocation;
-							}
+								Descriptor = descriptor,
+								Makefile = makefile!,
+								TargetName = buildTarget.TargetName,
+								GenerateHeaderFuncs = generatedHeaderFuncs,
+							});
+
+							makefiles[index] = makefile;
+						});
+
+						// Make sure to wait for previous pass codegen and timestamp writes
+						if (lastContextsForCodeGen != null)
+						{
+							UEBuildTarget.WaitCodeGen(lastContextsForCodeGen);
+							Task.WaitAll(UEBuildTarget.GetTimestampTasks(lastContextsForCodeGen).ToArray());
 						}
-						else
+
+						// Sort codegen based on config with Shipping last.
+						// Since shipping often has fewer modules, that means previous configs will include later ones
+						contextsForCodeGen.SortBy(context => context.Descriptor.Configuration);
+
+						// Run codegen for all targets
+						UEBuildTarget.RunCodeGen(buildConfiguration, contextsForCodeGen, logger);
+
+
+						// Go wide again and finish the work
+						Parallel2.ForEach(intermediateFolderGroups.Values, intermediateFolderGroup =>
 						{
-							// Restore the environment variables
-							foreach (Tuple<string, string> EnvironmentVariable in Makefile.EnvironmentVariables)
+							if (passIt >= intermediateFolderGroup.Count)
 							{
-								Environment.SetEnvironmentVariable(EnvironmentVariable.Item1, EnvironmentVariable.Item2);
+								return; // This group is done
 							}
 
-#if __VPROJECT_AVAILABLE__
-							Task VNITask = Task.CompletedTask;
-							// Same for VNI
-							if (Makefile.VNIModules.Count > 0)
-							{
-								VNITask = Task.Run(() => VNIExecution.ExecuteVNITool(Makefile, TargetDescriptor, Logger));
-							}
-#endif
-							// If the target needs UHT to be run, we'll go ahead and do that now
-							if (Makefile.UObjectModules.Count > 0)
-							{
-								(Task uhtRunTask, uhtTimeStampsTask) = UHTExecution.ExecuteHeaderToolIfNecessary(BuildConfiguration, TargetDescriptor.ProjectFile, Makefile, TargetDescriptor.Name, Logger, true);
-								uhtRunTask.Wait();
-							}
+							int index = intermediateFolderGroup[passIt];
+							TargetDescriptor descriptor = descriptors[index];
 
-#if __VPROJECT_AVAILABLE__
-							VNITask.Wait();
-#endif
-						}
+							TargetMakefile makefile = makefiles[index]!;
+							UEBuildTarget.BuildContext? buildContext = buildContexts[index];
 
-						targets.Add((TargetDescriptor, Makefile));
+							if (buildContext != null)
+							{
+								UEBuildTarget buildTarget = buildTargets[index];
 
-						LastUHTTimestampTask = uhtTimeStampsTask;
+								buildTarget.BuildMakefile2(workingSet, logger, buildContext);
 
-						if (SaveToLocation != null)
-						{
-							LastSaveMakefileTask = Task.Run(() =>
+								// Save the makefile for next time
+								if (makefileLocations[index] != null)
+								{
+									// Save the environment variables. Toolchains etc might set environment variables
+									// and we want to be able to restore them on runs where makefile is not invalidated.
+									// We know they can contain other target's environment variables too which is bad but by-design atm.
+									foreach (System.Collections.DictionaryEntry? EnvironmentVariable in Environment.GetEnvironmentVariables())
+									{
+										if (EnvironmentVariable != null)
+										{
+											makefile.EnvironmentVariables.Add(Tuple.Create((string)EnvironmentVariable.Value.Key, (string)EnvironmentVariable.Value.Value!));
+										}
+									}
+
+									lastSaveMakefileTask[index] = Task.Run(() =>
+									{
+										using ITimelineEvent _ = Timeline.ScopeEvent($"SavingMakefile");
+										makefile.Save(makefileLocations[index]!);
+									});
+								}
+							}
+							else
 							{
-								using ITimelineEvent _ = Timeline.ScopeEvent($"SavingMakefile");
-								Makefile.Save(SaveToLocation);
-							});
-						}
+								// Restore the environment variables
+								foreach (Tuple<string, string> EnvironmentVariable in makefile.EnvironmentVariables)
+								{
+									Environment.SetEnvironmentVariable(EnvironmentVariable.Item1, EnvironmentVariable.Item2);
+								}
+							}
 
-						if (!bSkipPreBuildTargets)
-						{
-							foreach (TargetInfo PreBuildTarget in Makefile.PreBuildTargets)
+							outTargets.Add((descriptor, makefile));
+
+							if (!skipPreBuildTargets)
 							{
-								lock (newTargetDescriptors)
+								foreach (TargetInfo PreBuildTarget in makefile.PreBuildTargets)
 								{
-									newTargetDescriptors.Add(TargetDescriptor.FromTargetInfo(PreBuildTarget));
+									lock (newDescriptors)
+									{
+										newDescriptors.Add(TargetDescriptor.FromTargetInfo(PreBuildTarget));
+									}
 								}
 							}
-						}
+
+							targetTimelines[index]?.Finish();
+						});
+
+						lastContextsForCodeGen = contextsForCodeGen;
 					}
 
-					if (LastSaveMakefileTask != null)
+					// All passes are done.
+
+					if (lastContextsForCodeGen != null)
 					{
-						lock (tasks)
+						List<Task> timestampTasks = UEBuildTarget.GetTimestampTasks(lastContextsForCodeGen);
+						UEBuildTarget.WaitCodeGen(lastContextsForCodeGen);
+
+						lock (outTasks)
 						{
-							tasks.Add(LastSaveMakefileTask);
+							outTasks.AddRange(timestampTasks);
 						}
 					}
-					if (LastUHTTimestampTask != null)
+
+					foreach (Task? task in lastSaveMakefileTask)
 					{
-						lock (tasks)
+						if (task != null)
 						{
-							tasks.Add(LastUHTTimestampTask);
+							lock (outTasks)
+							{
+								outTasks.Add(task);
+							}
 						}
 					}
-				}//);
+				}, helperCount);
 
-				TargetDescriptors = newTargetDescriptors;
+				descriptors = newDescriptors;
 			}
 
-			return (targets, tasks);
+			return (outTargets, outTasks);
 		}
 
 		/// <summary>
diff --git a/Engine/Source/Programs/UnrealBuildTool/System/CppDependencyCache.cs b/Engine/Source/Programs/UnrealBuildTool/System/CppDependencyCache.cs
--- a/Engine/Source/Programs/UnrealBuildTool/System/CppDependencyCache.cs
+++ b/Engine/Source/Programs/UnrealBuildTool/System/CppDependencyCache.cs
@@ -74,11 +74,6 @@ namespace UnrealBuildTool
 			/// </summary>
 			public FileReference Location;
 
-			/// <summary>
-			/// Directory for files to cache dependencies for.
-			/// </summary>
-			public DirectoryReference BaseDir;
-
 			/// <summary>
 			/// Map from file item to dependency info
 			/// </summary>
@@ -89,21 +84,18 @@ namespace UnrealBuildTool
 			/// </summary>
 			public bool bModified;
 
+			/// <summary>
+			/// Whether the partition has been loaded or not
+			/// </summary>
+			public bool Loaded;
+
 			/// <summary>
 			/// Constructs a dependency cache. This method is private; call CppDependencyCache.Create() to create a cache hierarchy for a given project.
 			/// </summary>
 			/// <param name="Location">File to store the cache</param>
-			/// <param name="BaseDir">Base directory for files that this cache should store data for</param>
-			/// <param name="Logger">Logger for output</param>
-			public CachePartition(FileReference Location, DirectoryReference BaseDir, ILogger Logger)
+			public CachePartition(FileReference Location)
 			{
 				this.Location = Location;
-				this.BaseDir = BaseDir;
-
-				if (FileReference.Exists(Location))
-				{
-					Read(Logger);
-				}
 			}
 
 			/// <summary>
@@ -164,12 +156,9 @@ namespace UnrealBuildTool
 		/// <summary>
 		/// List of partitions
 		/// </summary>
-		List<CachePartition> Partitions = new List<CachePartition>();
+		static List<(DirectoryReference Root, List<(string Filter, CachePartition Partition)>)> GlobalPartitions = [];
 
-		/// <summary>
-		/// Static cache of all constructed dependency caches
-		/// </summary>
-		static Dictionary<FileReference, CachePartition> GlobalPartitions = new Dictionary<FileReference, CachePartition>();
+		static List<Task> GlobalTasks = [];
 
 		/// <summary>
 		/// Minimum version of a MSVC source dependency json
@@ -254,11 +243,13 @@ namespace UnrealBuildTool
 			}
 		}
 
+		static readonly string SearchString = $"Intermediate{Path.DirectorySeparatorChar}Build";
+
 		/// <summary>
-		/// Attempts to read the dependencies from the given input file
+		/// Attempts to read dependencies from the given file.
 		/// </summary>
 		/// <param name="InputFile">File to be read</param>
-		/// <param name="Logger">Logger for output</param>
+		/// <param name="Logger">Logger</param>
 		/// <param name="OutInfo">The dependency info</param>
 		/// <returns>True if the input file exists and the dependencies were read</returns>
 		private bool TryGetDependencyInfo(FileItem InputFile, ILogger Logger, [NotNullWhen(true)] out DependencyInfo? OutInfo)
@@ -269,45 +260,65 @@ namespace UnrealBuildTool
 				return false;
 			}
 
-			try
+			ReadOnlySpan<char> FullName = InputFile.FullName;
+			int index = FullName.IndexOf(SearchString, FileReference.Comparison);
+			if (index == -1)
 			{
-				return TryGetDependencyInfoInternal(InputFile, out OutInfo);
-			}
-			catch (Exception Ex)
-			{
-				Logger.LogDebug("Unable to read {File}:\n{Ex}", InputFile, ExceptionUtils.FormatExceptionDetails(Ex));
+				Logger.LogWarning($"Path {InputFile} is not under intermediate folder");
 				OutInfo = null;
 				return false;
 			}
-		}
 
-		/// <summary>
-		/// Attempts to read dependencies from the given file.
-		/// </summary>
-		/// <param name="InputFile">File to be read</param>
-		/// <param name="OutInfo">The dependency info</param>
-		/// <returns>True if the input file exists and the dependencies were read</returns>
-		private bool TryGetDependencyInfoInternal(FileItem InputFile, [NotNullWhen(true)] out DependencyInfo? OutInfo)
-		{
-			foreach (CachePartition Partition in Partitions)
+			int start = index + SearchString.Length + 1;
+			ReadOnlySpan<char> filterPath = FullName.Slice(index + SearchString.Length + 1);
+
+			CachePartition? partition = null;
+
+			foreach (var root in GlobalPartitions)
 			{
-				if (InputFile.Location.IsUnderDirectory(Partition.BaseDir))
+				if (!InputFile.Location.IsUnderDirectory(root.Root))
+				{
+					continue;
+				}
+
+				foreach (var entry in root.Item2)
 				{
-					DependencyInfo? Info;
-					if (!Partition.DependencyFileToInfo.TryGetValue(InputFile, out Info) || InputFile.LastWriteTimeUtc.Ticks > Info.LastWriteTimeUtc)
+					if (filterPath.StartsWith(entry.Filter, FileReference.Comparison))
 					{
-						Info = ReadDependencyInfo(InputFile);
-						Partition.DependencyFileToInfo.AddOrUpdate(InputFile, Info, (k, v) => Info);
-						Partition.bModified = true;
+						partition = entry.Partition;
+						break;
 					}
+				}
+				break;
 
-					OutInfo = Info;
-					return true;
+			}
+
+			if (partition == null)
+			{
+				Logger.LogWarning($"There is no partition covering path to {InputFile}!");
+				OutInfo = null;
+				return false;
+			}
+
+			DependencyInfo? Info;
+			if (!partition.DependencyFileToInfo.TryGetValue(InputFile, out Info) || InputFile.LastWriteTimeUtc.Ticks > Info.LastWriteTimeUtc)
+			{
+				try
+				{
+					Info = ReadDependencyInfo(InputFile);
+					partition.DependencyFileToInfo.AddOrUpdate(InputFile, Info, (k, v) => Info);
+					partition.bModified = true;
+				}
+				catch (Exception Ex)
+				{
+					Logger.LogDebug("Unable to read {File}:\n{Ex}", InputFile, ExceptionUtils.FormatExceptionDetails(Ex));
+					OutInfo = null;
+					return false;
 				}
 			}
 
-			OutInfo = null;
-			return false;
+			OutInfo = Info;
+			return true;
 		}
 
 		/// <summary>
@@ -337,7 +348,34 @@ namespace UnrealBuildTool
 		/// <returns>Dependency cache hierarchy for the given project</returns>
 		public void Mount(TargetDescriptor TargetDescriptor, TargetType TargetType, ILogger Logger)
 		{
-			Mount(TargetDescriptor.ProjectFile, TargetDescriptor.Name, TargetDescriptor.Platform, TargetDescriptor.Configuration, TargetType, TargetDescriptor.Architectures, TargetDescriptor.IntermediateEnvironment, Logger);
+			Mount(TargetDescriptor.ProjectFile, TargetDescriptor.Name, "", TargetDescriptor.Platform, TargetDescriptor.Configuration, TargetType, TargetDescriptor.Architectures, TargetDescriptor.IntermediateEnvironment, Logger);
+		}
+
+		/// <summary>
+		/// Same as Mount but in a task
+		/// </summary>
+		public void MountAsync(TargetDescriptor TargetDescriptor, string AppName, TargetType TargetType, ILogger Logger)
+		{
+			lock (GlobalTasks)
+			{
+				GlobalTasks.Add(Task.Run(() =>
+				{
+					Mount(TargetDescriptor.ProjectFile, TargetDescriptor.Name, AppName, TargetDescriptor.Platform, TargetDescriptor.Configuration, TargetType, TargetDescriptor.Architectures, TargetDescriptor.IntermediateEnvironment, Logger);
+				}));
+			}
+		}
+
+		public void WaitAsync()
+		{
+			Task[] tasks;
+			lock (GlobalTasks)
+			{
+				tasks = GlobalTasks.ToArray();
+				GlobalTasks.Clear();
+			}
+
+			using ITimelineEvent ActionGraphTimer = Timeline.ScopeEvent($"CppDependencyCacheWait");
+			Task.WaitAll(tasks);
 		}
 
 		/// <summary>
@@ -345,6 +383,7 @@ namespace UnrealBuildTool
 		/// </summary>
 		/// <param name="ProjectFile">Project file for the target being built</param>
 		/// <param name="TargetName">Name of the target</param>
+		/// <param name="AppName">Name of the app.</param>
 		/// <param name="Platform">Platform being built</param>
 		/// <param name="Configuration">Configuration being built</param>
 		/// <param name="TargetType">The target type</param>
@@ -352,52 +391,81 @@ namespace UnrealBuildTool
 		/// <param name="IntermediateEnvironment">Intermediate environment to use</param>
 		/// <param name="Logger">Logger for output</param>
 		/// <returns>Dependency cache hierarchy for the given project</returns>
-		public void Mount(FileReference? ProjectFile, string TargetName, UnrealTargetPlatform Platform, UnrealTargetConfiguration Configuration, TargetType TargetType, UnrealArchitectures Architectures, UnrealIntermediateEnvironment IntermediateEnvironment, ILogger Logger)
+		public void Mount(FileReference? ProjectFile, string TargetName, string AppName, UnrealTargetPlatform Platform, UnrealTargetConfiguration Configuration, TargetType TargetType, UnrealArchitectures Architectures, UnrealIntermediateEnvironment IntermediateEnvironment, ILogger Logger)
 		{
-			string AppName = TargetName;
-			if (ProjectFile == null || !Unreal.IsEngineInstalled())
+			if (AppName == string.Empty)
 			{
-				if (TargetType == TargetType.Program)
-				{
-					AppName = TargetName;
-				}
-				else
+				AppName = TargetName;
+				if (ProjectFile == null || !Unreal.IsEngineInstalled())
 				{
-					AppName = UEBuildTarget.GetAppNameForTargetType(TargetType);
+					if (TargetType == TargetType.Program)
+					{
+						AppName = TargetName;
+					}
+					else
+					{
+						AppName = UEBuildTarget.GetAppNameForTargetType(TargetType);
+					}
 				}
 			}
 
-			FileReference EngineCacheLocation = FileReference.Combine(Unreal.WritableEngineDirectory, UEBuildTarget.GetPlatformIntermediateFolder(Platform, Architectures, false), UEBuildTarget.GetTargetIntermediateFolderName(AppName, IntermediateEnvironment), Configuration.ToString(), "DependencyCache.bin");
-			FindOrAddPartition(EngineCacheLocation, Unreal.EngineDirectory, Logger);
+			using ITimelineEvent ActionGraphTimer = Timeline.ScopeEvent($"CppDependencyCache");
 
-			if (ProjectFile != null)
+			AppName = UEBuildTarget.GetTargetIntermediateFolderName(AppName, IntermediateEnvironment);
+
+			string GetFilterPath(UnrealTargetConfiguration Configuration)
 			{
-				FileReference ProjectCacheLocation = FileReference.Combine(ProjectFile.Directory, UEBuildTarget.GetPlatformIntermediateFolder(Platform, Architectures, false), UEBuildTarget.GetTargetIntermediateFolderName(TargetName, IntermediateEnvironment), Configuration.ToString(), "DependencyCache.bin");
-				FindOrAddPartition(ProjectCacheLocation, ProjectFile.Directory, Logger);
+				StringBuilder sb = new();
+				sb.Append(Platform).Append(Path.DirectorySeparatorChar).Append(UnrealArchitectureConfig.ForPlatform(Platform).GetFolderNameForArchitectures(Architectures)).Append(Path.DirectorySeparatorChar);
+				sb.Append(AppName).Append(Path.DirectorySeparatorChar).Append(Configuration);
+				return sb.ToString();
 			}
-		}
 
-		/// <summary>
-		/// Reads a cache from the given location, or creates it with the given settings
-		/// </summary>
-		/// <param name="Location">File to store the cache</param>
-		/// <param name="BaseDir">Base directory for files that this cache should store data for</param>
-		/// <param name="Logger">Logger for output</param>
-		/// <returns>Reference to a dependency cache with the given settings</returns>
-		void FindOrAddPartition(FileReference Location, DirectoryReference BaseDir, ILogger Logger)
-		{
-			lock (GlobalPartitions)
+			void AddPartition(DirectoryReference rootPath, string filterPath)
 			{
-				if (!Partitions.Any(x => x.Location == Location))
+				CachePartition? partition = null;
+				lock (GlobalPartitions)
 				{
-					CachePartition? Partition;
-					if (!GlobalPartitions.TryGetValue(Location, out Partition))
+					var root = GlobalPartitions.FirstOrDefault(e => e.Root == rootPath);
+					if (root.Root == null)
 					{
-						Partition = new CachePartition(Location, BaseDir, Logger);
-						GlobalPartitions.Add(Location, Partition);
+						root = new(rootPath, new());
+						GlobalPartitions.Add(root);
+					}
+					var filters = root.Item2;
+					var filter = filters.FirstOrDefault(f => f.Filter == filterPath);
+					partition = filter.Partition;
+					if (partition == null)
+					{
+						FileReference location = FileReference.Combine(rootPath, "Intermediate", "Build", filterPath, "DependencyCache.bin");
+						partition = new CachePartition(location);
+						filters.Add((filterPath, partition));
 					}
-					Partitions.Add(Partition);
 				}
+
+				lock (partition)
+				{
+					if (!partition.Loaded)
+					{
+						partition.Loaded = true;
+						if (FileReference.Exists(partition.Location))
+						{
+							partition.Read(Logger);
+						}
+					}
+				}
+
+			}
+
+			AddPartition(Unreal.WritableEngineDirectory, GetFilterPath(Configuration));
+			if (Configuration == UnrealTargetConfiguration.DebugGame)
+			{
+				AddPartition(Unreal.WritableEngineDirectory, GetFilterPath(UnrealTargetConfiguration.Development));
+			}
+
+			if (ProjectFile != null)
+			{
+				AddPartition(ProjectFile.Directory, GetFilterPath(Configuration));
 			}
 		}
 
@@ -406,7 +474,7 @@ namespace UnrealBuildTool
 		/// </summary>
 		public static void SaveAll()
 		{
-			Parallel.ForEach(GlobalPartitions.Values, Cache =>
+			Parallel2.ForEach(GlobalPartitions.SelectMany(p => p.Item2).Select(f => f.Partition), Cache =>
 			{
 				if (Cache.bModified)
 				{
