[UBT]
* Fixed race condition when collecting configs codegen info
* Added more validation when checking so codegen can run once
* Changed so CppDependencyCache is not a singleton

diff --git a/Engine/Source/Programs/UnrealBuildTool/Configuration/UEBuildTarget.cs b/Engine/Source/Programs/UnrealBuildTool/Configuration/UEBuildTarget.cs
--- a/Engine/Source/Programs/UnrealBuildTool/Configuration/UEBuildTarget.cs
+++ b/Engine/Source/Programs/UnrealBuildTool/Configuration/UEBuildTarget.cs
@@ -2592,7 +2592,8 @@ namespace UnrealBuildTool
 
 		public (TargetMakefile, Task?) BuildMakefile(BuildConfiguration BuildConfiguration, ISourceFileWorkingSet WorkingSet, TargetDescriptor TargetDescriptor, ILogger Logger, bool bInitOnly = false)
 		{
-			BuildContext context = BuildMakefile1(TargetDescriptor, Logger, null);
+			SourceFileMetadataCache cache = SourceFileMetadataCache.CreateHierarchy(ProjectFile, Logger);
+			BuildContext context = BuildMakefile1(TargetDescriptor, Logger, cache);
 
 			// UHT mode uses this to create the makefile at least to the point where UHT would have valid manifest information
 			if (bInitOnly)
@@ -2636,15 +2637,9 @@ namespace UnrealBuildTool
 			required internal List<GenerateHeaderFunc> GenerateHeaderFuncs;
 		}
 
-		internal BuildContext BuildMakefile1(TargetDescriptor TargetDescriptor, ILogger Logger, SourceFileMetadataCache? MetadataCache)
+		internal BuildContext BuildMakefile1(TargetDescriptor TargetDescriptor, ILogger Logger, SourceFileMetadataCache MetadataCache)
 		{
 			CppConfiguration CppConfiguration = GetCppConfiguration(Configuration);
-
-			if (MetadataCache == null)
-			{
-				MetadataCache = SourceFileMetadataCache.CreateHierarchy(ProjectFile, Logger);
-			}
-
 			CppCompileEnvironment GlobalCompileEnvironment = new CppCompileEnvironment(Platform, CppConfiguration, Architectures, MetadataCache);
 			LinkEnvironment GlobalLinkEnvironment = new LinkEnvironment(GlobalCompileEnvironment.Platform, GlobalCompileEnvironment.Configuration, GlobalCompileEnvironment.Architectures);
 
@@ -2811,10 +2806,10 @@ namespace UnrealBuildTool
 			}
 
 			CodeGenContext? lastRunContext = null;
-			HashSet<string>? lastRunUObjectModules = null;
+			Dictionary<string, UHTModuleInfo>? lastRunUObjectModules = null;
 			HashSet<string>? lastRunVNIModules = null;
 
-			string GetTargetName(TargetDescriptor desc) => $"{desc.Name}{desc.Platform}{desc.Configuration}";
+			string GetTargetName(TargetDescriptor desc) => $"{desc.Name}-{desc.Platform}-{desc.Configuration}";
 
 			foreach (CodeGenContext context in contexts)
 			{
@@ -2835,14 +2830,22 @@ namespace UnrealBuildTool
 
 					if (lastRunUObjectModules == null)
 					{
-						lastRunUObjectModules = new(lastRunContext.Makefile.UObjectModules.Select(m => m.ModuleName));
+						lastRunUObjectModules = lastRunContext.Makefile.UObjectModules.ToDictionary(k => k.ModuleName, k => k);
 					}
 
 					foreach (var module in makefile.UObjectModules)
 					{
-						if (!lastRunUObjectModules.Contains(module.ModuleName))
+						if (lastRunUObjectModules.TryGetValue(module.ModuleName, out UHTModuleInfo? lastModule))
+						{
+							int CountHeaders(UHTModuleInfo info) => info.PublicUObjectClassesHeaders.Count + info.PublicUObjectHeaders.Count + info.InternalUObjectHeaders.Count + info.PrivateUObjectHeaders.Count;
+							if (CountHeaders(module) != CountHeaders(lastModule))
+							{
+								throw new Exception($"Target {GetTargetName(context.Descriptor)} module {module.ModuleName} has different number of headers than previous target {GetTargetName(lastRunContext.Descriptor)} did not have!.. run after each other!");
+							}
+						}
+						else
 						{
-							throw new Exception($"Target {GetTargetName(context.Descriptor)} has uht modules that previous target {GetTargetName(lastRunContext.Descriptor)} did not have!.. run after each other!");
+							throw new Exception($"Target {GetTargetName(context.Descriptor)} has uht module {module.ModuleName} that previous target {GetTargetName(lastRunContext.Descriptor)} did not have!.. run after each other!");
 						}
 					}
 
diff --git a/Engine/Source/Programs/UnrealBuildTool/Modes/BuildMode.cs b/Engine/Source/Programs/UnrealBuildTool/Modes/BuildMode.cs
--- a/Engine/Source/Programs/UnrealBuildTool/Modes/BuildMode.cs
+++ b/Engine/Source/Programs/UnrealBuildTool/Modes/BuildMode.cs
@@ -8,7 +8,6 @@ using System.Linq;
 using System.Text;
 using System.Threading.Tasks;
 using EpicGames.Core;
-using Microsoft.CodeAnalysis.CSharp.Syntax;
 using Microsoft.Extensions.Logging;
 using OpenTracing.Util;
 using UnrealBuildBase;
@@ -166,6 +165,8 @@ namespace UnrealBuildTool
 				}
 			}
 
+			CppDependencyCache dependencyCache = new();
+
 			// Parse and build the targets
 			try
 			{
@@ -251,7 +252,7 @@ namespace UnrealBuildTool
 					// Create the working set provider per group.
 					using (ISourceFileWorkingSet WorkingSet = SourceFileWorkingSet.Create(Unreal.RootDirectory, ProjectDirs, Logger))
 					{
-						await BuildAsync(TargetDescriptors, BuildConfiguration, WorkingSet, Options, WriteOutdatedActionsFile, Logger, bSkipPreBuildTargets, ActionTypeFilter:ActionTypeFilter);
+						await BuildAsync(TargetDescriptors, BuildConfiguration, WorkingSet, Options, WriteOutdatedActionsFile, dependencyCache, Logger, bSkipPreBuildTargets, ActionTypeFilter:ActionTypeFilter);
 					}
 				}
 			}
@@ -264,7 +265,7 @@ namespace UnrealBuildTool
 				using (Timeline.ScopeEvent("Saving caches"))
 				{
 					SourceFileMetadataCache.SaveAll();
-					CppDependencyCache.SaveAll();
+					dependencyCache.SaveAll();
 				}
 			}
 			return 0;
@@ -382,22 +383,16 @@ namespace UnrealBuildTool
 		/// <param name="WorkingSet">The source file working set</param>
 		/// <param name="Options">Additional options for the build</param>
 		/// <param name="WriteOutdatedActionsFile">Files to write the list of outdated actions to (rather than building them)</param>
+		/// <param name="cppDependencyCache"></param>
 		/// <param name="Logger">Logger for output</param>
 		/// <param name="bSkipPreBuildTargets">If true then only the current target descriptors will be built.</param>
 		/// <param name="ActionTypeFilter">If specified, only actions of this type will be run</param>
 		/// <returns>Result from the compilation</returns>
-		public static async Task BuildAsync(List<TargetDescriptor> TargetDescriptors, BuildConfiguration BuildConfiguration, ISourceFileWorkingSet WorkingSet, BuildOptions Options, FileReference? WriteOutdatedActionsFile, ILogger Logger, bool bSkipPreBuildTargets = false, string? ActionTypeFilter = null)
+		public static async Task BuildAsync(List<TargetDescriptor> TargetDescriptors, BuildConfiguration BuildConfiguration, ISourceFileWorkingSet WorkingSet, BuildOptions Options, FileReference? WriteOutdatedActionsFile, CppDependencyCache cppDependencyCache, ILogger Logger, bool bSkipPreBuildTargets = false, string? ActionTypeFilter = null)
 		{
-			// Create cpp dependency cache in order to be able to preload them as we discover them
-			CppDependencyCache? cppDependencyCache = null;
-			if ((Options & BuildOptions.SkipBuild) == 0)
-			{
-				cppDependencyCache = new();
-			}
-
 			(List<(TargetDescriptor, TargetMakefile)> Targets, List<Task> Tasks) = CreateMakefiles(BuildConfiguration, TargetDescriptors, WorkingSet, Logger, true, true, cppDependencyCache);
 
-			await BuildAsync(Targets, BuildConfiguration, Options, WriteOutdatedActionsFile, Logger, ActionTypeFilter, cppDependencyCache);
+			await BuildAsync(Targets, BuildConfiguration, Options, WriteOutdatedActionsFile, cppDependencyCache, Logger, ActionTypeFilter);
 
 			Task.WaitAll(Tasks.ToArray());
 		}
@@ -415,7 +410,7 @@ namespace UnrealBuildTool
 		/// <param name="ActionTypeFilter">If specified, only actions of this type will be run</param>
 		/// <param name="cppDependencyCache"></param>
 		/// <returns>Result from the compilation</returns>
-		internal static async Task BuildAsync(List<(TargetDescriptor Descriptor, TargetMakefile Makefile)> Targets, BuildConfiguration BuildConfiguration, BuildOptions Options, FileReference? WriteOutdatedActionsFile, ILogger Logger, string? ActionTypeFilter=null, CppDependencyCache? cppDependencyCache = null)
+		internal static async Task BuildAsync(List<(TargetDescriptor Descriptor, TargetMakefile Makefile)> Targets, BuildConfiguration BuildConfiguration, BuildOptions Options, FileReference? WriteOutdatedActionsFile, CppDependencyCache cppDependencyCache, ILogger Logger, string? ActionTypeFilter=null)
 		{
 			List<TargetDescriptor> TargetDescriptors = Targets.Select(T => T.Descriptor).ToList();
 			TargetMakefile[] Makefiles = Targets.Select(T => T.Makefile).ToArray();
@@ -423,11 +418,6 @@ namespace UnrealBuildTool
 			// Execute the build
 			if ((Options & BuildOptions.SkipBuild) == 0)
 			{
-				if (cppDependencyCache == null)
-				{
-					cppDependencyCache = new();
-				}
-
 				// Create the C++ dependency cache
 				for (int TargetIdx = 0; TargetIdx < Targets.Count; TargetIdx++)
 				{
@@ -564,7 +554,6 @@ namespace UnrealBuildTool
 
 				// Wait for cpp dependencies task to finish
 				cppDependencyCache.WaitAsync();
-				CppDependencyCache CppDependencies = cppDependencyCache;
 
 				// Now that we have the dependencies object, optionally create the artifact cache
 				IActionArtifactCache? actionArtifactCache = null;
@@ -582,7 +571,7 @@ namespace UnrealBuildTool
 
 					Logger.LogInformation("Experimental artifact system using '{Directory}' in {Mode} mode", artifactDirectory.FullName, mode);
 
-					actionArtifactCache = ActionArtifactCache.CreateHordeFileCache(artifactDirectory, CppDependencies, s_memoryMappedFileCache.Value, Logger);
+					actionArtifactCache = ActionArtifactCache.CreateHordeFileCache(artifactDirectory, cppDependencyCache, s_memoryMappedFileCache.Value, Logger);
 					actionArtifactCache.EnableReads = BuildConfiguration.bArtifactRead;
 					actionArtifactCache.EnableWrites = BuildConfiguration.bArtifactWrites;
 					actionArtifactCache.LogCacheMisses = BuildConfiguration.bLogArtifactCacheMisses;
@@ -594,7 +583,7 @@ namespace UnrealBuildTool
 				if (ModuleDependencyActions.Count > 0)
 				{
 					ConcurrentDictionary<LinkedAction, bool> PreprocessActionToOutdatedFlag = new ConcurrentDictionary<LinkedAction, bool>();
-					ActionGraph.GatherAllOutdatedActions(ModuleDependencyActions, History, PreprocessActionToOutdatedFlag, CppDependencies, BuildConfiguration.bIgnoreOutdatedImportLibraries, Logger);
+					ActionGraph.GatherAllOutdatedActions(ModuleDependencyActions, History, PreprocessActionToOutdatedFlag, cppDependencyCache, BuildConfiguration.bIgnoreOutdatedImportLibraries, Logger);
 
 					List<LinkedAction> PreprocessActions = PreprocessActionToOutdatedFlag.Where(x => x.Value).Select(x => x.Key).ToList();
 					if (PreprocessActions.Count > 0)
@@ -628,13 +617,13 @@ namespace UnrealBuildTool
 							if (CppModulesAction != null && CppModulesAction.CompiledModuleInterfaceFile != null)
 							{
 								string? ProducedModule;
-								if (CppDependencies.TryGetProducedModule(ModuleDependencyAction.DependencyListFile!, Logger, out ProducedModule))
+								if (cppDependencyCache.TryGetProducedModule(ModuleDependencyAction.DependencyListFile!, Logger, out ProducedModule))
 								{
 									ModuleOutputs[ProducedModule] = CppModulesAction.CompiledModuleInterfaceFile;
 								}
 
 								List<(string Name, string BMI)>? ImportedModules;
-								if (CppDependencies.TryGetImportedModules(ModuleDependencyAction.DependencyListFile!, Logger, out ImportedModules))
+								if (cppDependencyCache.TryGetImportedModules(ModuleDependencyAction.DependencyListFile!, Logger, out ImportedModules))
 								{
 									ModuleImports[CppModulesAction.CompiledModuleInterfaceFile] = ImportedModules.Select(x => x.Name).ToList();
 								}
@@ -696,7 +685,7 @@ namespace UnrealBuildTool
 					}
 
 					// Plan the actions to execute for the build. For single file compiles, always rebuild the source file regardless of whether it's out of date.
-					ActionGraph.GatherAllOutdatedActions(PrerequisiteActions, History, ActionToOutdatedFlag, CppDependencies, BuildConfiguration.bIgnoreOutdatedImportLibraries, Logger);
+					ActionGraph.GatherAllOutdatedActions(PrerequisiteActions, History, ActionToOutdatedFlag, cppDependencyCache, BuildConfiguration.bIgnoreOutdatedImportLibraries, Logger);
 				}
 
 				// Link the action graph again to sort it
@@ -745,7 +734,7 @@ namespace UnrealBuildTool
 					Logger.LogDebug("Re-evaluating action graph");
 					// Re-check the graph to remove any LiveCoding actions added by PatchActionsForTarget() that are already up to date.
 					ConcurrentDictionary<LinkedAction, bool> LiveActionToOutdatedFlag = new ConcurrentDictionary<LinkedAction, bool>(Environment.ProcessorCount, MergedActionsToExecute.Count);
-					ActionGraph.GatherAllOutdatedActions(MergedActionsToExecute, History, LiveActionToOutdatedFlag, CppDependencies, BuildConfiguration.bIgnoreOutdatedImportLibraries, Logger);
+					ActionGraph.GatherAllOutdatedActions(MergedActionsToExecute, History, LiveActionToOutdatedFlag, cppDependencyCache, BuildConfiguration.bIgnoreOutdatedImportLibraries, Logger);
 					List<LinkedAction> LiveCodingActionsToExecute = LiveActionToOutdatedFlag.Where(x => x.Value).Select(x => x.Key).ToList();
 					ActionGraph.Link(LiveCodingActionsToExecute, Logger);
 					MergedActionsToExecute = LiveCodingActionsToExecute;
@@ -1338,11 +1327,15 @@ namespace UnrealBuildTool
 									ranPrebuildScripts = true;
 								}
 
-								SourceFileMetadataCache? sourceFileMetaDataCache = null;
+								SourceFileMetadataCache sourceFileMetaDataCache;
 								if (descriptor.ProjectFile != null)
 								{
 									sourceFileMetaDataCache = sourceFileMetadataCaches[descriptor.ProjectFile].Result;
 								}
+								else
+								{
+									sourceFileMetaDataCache = SourceFileMetadataCache.CreateHierarchy(null, logger);
+								}
 
 								// Build the target
 								UEBuildTarget.BuildContext context = buildTarget.BuildMakefile1(descriptor, logger, sourceFileMetaDataCache);
@@ -1360,14 +1353,16 @@ namespace UnrealBuildTool
 								generatedHeaderFuncs = context.GenerateHeaderFuncs;
 							}
 
-							contextsForCodeGen.Add(new UEBuildTarget.CodeGenContext()
+							lock (contextsForCodeGen)
 							{
-								Descriptor = descriptor,
-								Makefile = makefile!,
-								TargetName = buildTarget.TargetName,
-								GenerateHeaderFuncs = generatedHeaderFuncs,
-							});
-
+								contextsForCodeGen.Add(new UEBuildTarget.CodeGenContext()
+								{
+									Descriptor = descriptor,
+									Makefile = makefile!,
+									TargetName = buildTarget.TargetName,
+									GenerateHeaderFuncs = generatedHeaderFuncs,
+								});
+							}
 							makefiles[index] = makefile;
 						});
 
diff --git a/Engine/Source/Programs/UnrealBuildTool/Modes/IWYUMode.cs b/Engine/Source/Programs/UnrealBuildTool/Modes/IWYUMode.cs
--- a/Engine/Source/Programs/UnrealBuildTool/Modes/IWYUMode.cs
+++ b/Engine/Source/Programs/UnrealBuildTool/Modes/IWYUMode.cs
@@ -625,13 +625,14 @@ namespace UnrealBuildTool
 					// Use the make file to build unless -NoCompile is set.
 					if (!bNoCompile)
 					{
+						CppDependencyCache dependencyCache = new();
 						try
 						{
-							await BuildMode.BuildAsync([(Descriptor, Makefile)], BuildConfiguration, BuildOptions.None, null, Logger);
+							await BuildMode.BuildAsync([(Descriptor, Makefile)], BuildConfiguration, BuildOptions.None, null, dependencyCache, Logger);
 						}
 						finally
 						{
-							CppDependencyCache.SaveAll();
+							dependencyCache.SaveAll();
 						}
 					}
 
diff --git a/Engine/Source/Programs/UnrealBuildTool/Modes/ProfileUnitySizesMode.cs b/Engine/Source/Programs/UnrealBuildTool/Modes/ProfileUnitySizesMode.cs
--- a/Engine/Source/Programs/UnrealBuildTool/Modes/ProfileUnitySizesMode.cs
+++ b/Engine/Source/Programs/UnrealBuildTool/Modes/ProfileUnitySizesMode.cs
@@ -338,7 +338,7 @@ namespace UnrealBuildTool
 					}
 
 					TimingLogger NewTimingLogger = new(Logger);
-					await BuildMode.BuildAsync(new List<TargetDescriptor>() { TargetDescriptor }, BuildConfiguration, WorkingSet, BuildOptions.None, null, NewTimingLogger);
+					await BuildMode.BuildAsync(new List<TargetDescriptor>() { TargetDescriptor }, BuildConfiguration, WorkingSet, BuildOptions.None, null, new(), NewTimingLogger);
 					NewTimingData = NewTimingLogger.TimingData;
 				}
 			}
diff --git a/Engine/Source/Programs/UnrealBuildTool/Modes/TestMode.cs b/Engine/Source/Programs/UnrealBuildTool/Modes/TestMode.cs
--- a/Engine/Source/Programs/UnrealBuildTool/Modes/TestMode.cs
+++ b/Engine/Source/Programs/UnrealBuildTool/Modes/TestMode.cs
@@ -122,11 +122,11 @@ namespace UnrealBuildTool
 				else if (bRebuildTests)
 				{
 					new CleanMode().Clean(TargetDescriptors, BuildConfiguration, Logger);
-					await BuildMode.BuildAsync(TargetDescriptors, BuildConfiguration, WorkingSet, BuildOptions.None, null, Logger);
+					await BuildMode.BuildAsync(TargetDescriptors, BuildConfiguration, WorkingSet, BuildOptions.None, null, new(), Logger);
 				}
 				else
 				{
-					await BuildMode.BuildAsync(TargetDescriptors, BuildConfiguration, WorkingSet, BuildOptions.None, null, Logger);
+					await BuildMode.BuildAsync(TargetDescriptors, BuildConfiguration, WorkingSet, BuildOptions.None, null, new(), Logger);
 				}
 			}
 			return 0;
diff --git a/Engine/Source/Programs/UnrealBuildTool/System/CppDependencyCache.cs b/Engine/Source/Programs/UnrealBuildTool/System/CppDependencyCache.cs
--- a/Engine/Source/Programs/UnrealBuildTool/System/CppDependencyCache.cs
+++ b/Engine/Source/Programs/UnrealBuildTool/System/CppDependencyCache.cs
@@ -156,9 +156,9 @@ namespace UnrealBuildTool
 		/// <summary>
 		/// List of partitions
 		/// </summary>
-		static List<(DirectoryReference Root, List<(string Filter, CachePartition Partition)>)> GlobalPartitions = [];
+		List<(DirectoryReference Root, List<(string Filter, CachePartition Partition)>)> GlobalPartitions = [];
 
-		static List<Task> GlobalTasks = [];
+		List<Task> GlobalTasks = [];
 
 		/// <summary>
 		/// Minimum version of a MSVC source dependency json
@@ -243,7 +243,7 @@ namespace UnrealBuildTool
 			}
 		}
 
-		static readonly string SearchString = $"Intermediate{Path.DirectorySeparatorChar}Build";
+		readonly string SearchString = $"Intermediate{Path.DirectorySeparatorChar}Build";
 
 		/// <summary>
 		/// Attempts to read dependencies from the given file.
@@ -327,7 +327,7 @@ namespace UnrealBuildTool
 		/// <param name="InputFile">File to be read</param>
 		/// <param name="OutDependencyItems">Receives a list of output items</param>
 		/// <returns>True if the input file exists and the dependencies were read</returns>
-		public static bool TryGetDependenciesUncached(FileItem InputFile, [NotNullWhen(true)] out List<FileItem>? OutDependencyItems)
+		public bool TryGetDependenciesUncached(FileItem InputFile, [NotNullWhen(true)] out List<FileItem>? OutDependencyItems)
 		{
 			if (!InputFile.Exists)
 			{
@@ -472,7 +472,7 @@ namespace UnrealBuildTool
 		/// <summary>
 		/// Save all the caches that have been modified
 		/// </summary>
-		public static void SaveAll()
+		public void SaveAll()
 		{
 			Parallel2.ForEach(GlobalPartitions.SelectMany(p => p.Item2).Select(f => f.Partition), Cache =>
 			{
