[UBT]
* Fixed bug in CppDependencyCache when building targets that have special intermediate folder

diff --git a/Engine/Source/Programs/UnrealBuildTool/Modes/BuildMode.cs b/Engine/Source/Programs/UnrealBuildTool/Modes/BuildMode.cs
--- a/Engine/Source/Programs/UnrealBuildTool/Modes/BuildMode.cs
+++ b/Engine/Source/Programs/UnrealBuildTool/Modes/BuildMode.cs
@@ -421,7 +421,7 @@ namespace UnrealBuildTool
 				// Create the C++ dependency cache
 				for (int TargetIdx = 0; TargetIdx < Targets.Count; TargetIdx++)
 				{
-					cppDependencyCache.MountAsync(Targets[TargetIdx].Descriptor, "", Targets[TargetIdx].Makefile.TargetType, Logger);
+					cppDependencyCache.Mount(Targets[TargetIdx].Descriptor, Makefiles[TargetIdx].TargetType, Logger);
 				}
 
 #if false // Doesn't seem to make things faster. Revisit since it feels like this should be faster
@@ -1197,7 +1197,7 @@ namespace UnrealBuildTool
 					// Queue up cppdependencies loading
 					if (cppDependencyCache != null)
 					{
-						cppDependencyCache.MountAsync(descriptor, buildTarget.AppName, buildTarget.TargetType, logger);
+						cppDependencyCache.MountAsync(descriptor, buildTarget.AppName, buildTarget.IntermediateEnvironment, logger);
 					}
 
 					// Add target to proper app name group
diff --git a/Engine/Source/Programs/UnrealBuildTool/System/CppDependencyCache.cs b/Engine/Source/Programs/UnrealBuildTool/System/CppDependencyCache.cs
--- a/Engine/Source/Programs/UnrealBuildTool/System/CppDependencyCache.cs
+++ b/Engine/Source/Programs/UnrealBuildTool/System/CppDependencyCache.cs
@@ -348,19 +348,32 @@ namespace UnrealBuildTool
 		/// <returns>Dependency cache hierarchy for the given project</returns>
 		public void Mount(TargetDescriptor TargetDescriptor, TargetType TargetType, ILogger Logger)
 		{
-			Mount(TargetDescriptor.ProjectFile, TargetDescriptor.Name, "", TargetDescriptor.Platform, TargetDescriptor.Configuration, TargetType, TargetDescriptor.Architectures, TargetDescriptor.IntermediateEnvironment, Logger);
+
+			string AppName = TargetDescriptor.Name;
+			if (TargetDescriptor.ProjectFile == null || !Unreal.IsEngineInstalled())
+			{
+				if (TargetType == TargetType.Program)
+				{
+					AppName = TargetDescriptor.Name;
+				}
+				else
+				{
+					AppName = UEBuildTarget.GetAppNameForTargetType(TargetType);
+				}
+			}
+			Mount(TargetDescriptor.ProjectFile, AppName, TargetDescriptor.Platform, TargetDescriptor.Configuration, TargetDescriptor.Architectures, TargetDescriptor.IntermediateEnvironment, Logger);
 		}
 
 		/// <summary>
 		/// Same as Mount but in a task
 		/// </summary>
-		public void MountAsync(TargetDescriptor TargetDescriptor, string AppName, TargetType TargetType, ILogger Logger)
+		public void MountAsync(TargetDescriptor descriptor, string AppName, UnrealIntermediateEnvironment intermediateEnvironment, ILogger logger)
 		{
 			lock (GlobalTasks)
 			{
 				GlobalTasks.Add(Task.Run(() =>
 				{
-					Mount(TargetDescriptor.ProjectFile, TargetDescriptor.Name, AppName, TargetDescriptor.Platform, TargetDescriptor.Configuration, TargetType, TargetDescriptor.Architectures, TargetDescriptor.IntermediateEnvironment, Logger);
+					Mount(descriptor.ProjectFile, AppName, descriptor.Platform, descriptor.Configuration, descriptor.Architectures, intermediateEnvironment, logger);
 				}));
 			}
 		}
@@ -382,33 +395,15 @@ namespace UnrealBuildTool
 		/// Creates a cache hierarchy for a particular target
 		/// </summary>
 		/// <param name="ProjectFile">Project file for the target being built</param>
-		/// <param name="TargetName">Name of the target</param>
 		/// <param name="AppName">Name of the app.</param>
 		/// <param name="Platform">Platform being built</param>
 		/// <param name="Configuration">Configuration being built</param>
-		/// <param name="TargetType">The target type</param>
 		/// <param name="Architectures">The target architectures</param>
 		/// <param name="IntermediateEnvironment">Intermediate environment to use</param>
 		/// <param name="Logger">Logger for output</param>
 		/// <returns>Dependency cache hierarchy for the given project</returns>
-		public void Mount(FileReference? ProjectFile, string TargetName, string AppName, UnrealTargetPlatform Platform, UnrealTargetConfiguration Configuration, TargetType TargetType, UnrealArchitectures Architectures, UnrealIntermediateEnvironment IntermediateEnvironment, ILogger Logger)
+		public void Mount(FileReference? ProjectFile, string AppName, UnrealTargetPlatform Platform, UnrealTargetConfiguration Configuration, UnrealArchitectures Architectures, UnrealIntermediateEnvironment IntermediateEnvironment, ILogger Logger)
 		{
-			if (AppName == string.Empty)
-			{
-				AppName = TargetName;
-				if (ProjectFile == null || !Unreal.IsEngineInstalled())
-				{
-					if (TargetType == TargetType.Program)
-					{
-						AppName = TargetName;
-					}
-					else
-					{
-						AppName = UEBuildTarget.GetAppNameForTargetType(TargetType);
-					}
-				}
-			}
-
 			using ITimelineEvent ActionGraphTimer = Timeline.ScopeEvent($"CppDependencyCache");
 
 			AppName = UEBuildTarget.GetTargetIntermediateFolderName(AppName, IntermediateEnvironment);
