[Backout] - CL45883954
#rnx
Original CL Desc
-----------------------------------------------------------------
[UBT]
* Implemented parallel makefile generation. UEBuildTargets will now run in parallel as long as they don't share codegen and intermediate folders..
* Fixed so codegen is "merged" when multiple targets share codegen. Since codegen is expensive we want to run this as little as possible
* Fixed so each CppDependencyCache partition is now based on path + platform/config/appname in order to make sure cache is isolated to target
* Added BuildConfiguration.bParallelMakefileGeneration that defaults to true but can be set to false to make sure makefile generation is singlethreaded

diff --git a/Engine/Source/Programs/UnrealBuildTool/Configuration/BuildConfiguration.cs b/Engine/Source/Programs/UnrealBuildTool/Configuration/BuildConfiguration.cs
--- a/Engine/Source/Programs/UnrealBuildTool/Configuration/BuildConfiguration.cs
+++ b/Engine/Source/Programs/UnrealBuildTool/Configuration/BuildConfiguration.cs
@@ -264,9 +264,6 @@ namespace UnrealBuildTool
 		[XmlConfigFile(Category = "BuildConfiguration")]
 		public bool? bForceUnityBuild { get; set; }
 
-		[XmlConfigFile(Category = "BuildConfiguration")]
-		public bool bParallelMakefileGeneration { get; set; } = true;
-
 		/// <summary>
 		/// Intermediate environment. Determines if the intermediates end up in a different folder than normal.
 		/// </summary>
diff --git a/Engine/Source/Programs/UnrealBuildTool/Configuration/UEBuildTarget.cs b/Engine/Source/Programs/UnrealBuildTool/Configuration/UEBuildTarget.cs
--- a/Engine/Source/Programs/UnrealBuildTool/Configuration/UEBuildTarget.cs
+++ b/Engine/Source/Programs/UnrealBuildTool/Configuration/UEBuildTarget.cs
@@ -1256,13 +1256,15 @@ namespace UnrealBuildTool
 		/// <returns>The build target object for the specified build rules source file</returns>
 		public static UEBuildTarget Create(TargetDescriptor Descriptor, bool bSkipRulesCompile, bool bForceRulesCompile, bool bUsePrecompiled, UnrealIntermediateEnvironment IntermediateEnvironment, ILogger Logger)
 		{
+			Task<SourceFileMetadataCache> metaDataTask = Task.Run(() => SourceFileMetadataCache.CreateHierarchy(Descriptor.ProjectFile, Logger));
+
 			var (RulesObject, RulesAssembly) = CreateRulesObject(Descriptor, bSkipRulesCompile, bForceRulesCompile, bUsePrecompiled, IntermediateEnvironment, Logger);
 
 			// Generate a build target from this rules module
 			UEBuildTarget Target;
 			using (GlobalTracer.Instance.BuildSpan("UEBuildTarget constructor").StartActive())
 			{
-				Target = new UEBuildTarget(Descriptor, new ReadOnlyTargetRules(RulesObject), RulesAssembly, Logger);
+				Target = new UEBuildTarget(Descriptor, new ReadOnlyTargetRules(RulesObject), RulesAssembly, metaDataTask.Result, Logger);
 			}
 			using (GlobalTracer.Instance.BuildSpan("UEBuildTarget.PreBuildSetup()").StartActive())
 			{
@@ -1289,7 +1291,7 @@ namespace UnrealBuildTool
 			TargetRules RulesObject;
 			using (GlobalTracer.Instance.BuildSpan("RulesAssembly.CreateTargetRules()").StartActive())
 			{
-				using var _ = Timeline.ScopeEvent("CreateTargetRules");
+				using var _ = Timeline.ScopeEvent("CreateTarget");
 				RulesObject = RulesAssembly.CreateTargetRules(Descriptor.Name, Descriptor.Platform, Descriptor.Configuration, Descriptor.Architectures, Descriptor.ProjectFile, Descriptor.AdditionalArguments, Logger, Descriptor.IsTestsTarget, IntermediateEnvironment: IntermediateEnvironment);
 			}
 			if ((ProjectFileGenerator.bGenerateProjectFiles == false) && !RulesObject.GetSupportedPlatforms().Contains(Descriptor.Platform))
@@ -1614,6 +1616,11 @@ namespace UnrealBuildTool
 		/// </summary>
 		public RulesAssembly RulesAssembly;
 
+		/// <summary>
+		/// Cache of source file metadata for this target
+		/// </summary>
+		public SourceFileMetadataCache MetadataCache;
+
 		/// <summary>
 		/// The project file for this target
 		/// </summary>
@@ -1747,7 +1754,7 @@ namespace UnrealBuildTool
 		/// <summary>
 		/// Used to keep track of all modules by name.
 		/// </summary>
-		private Dictionary<string, UEBuildModule> Modules = new(StringComparer.InvariantCultureIgnoreCase);
+		private Dictionary<string, UEBuildModule> Modules = new (StringComparer.InvariantCultureIgnoreCase);
 
 		private Dictionary<string, UEBuildModule> UsedShortNames = new(StringComparer.InvariantCultureIgnoreCase);
 
@@ -1788,15 +1795,29 @@ namespace UnrealBuildTool
 			return UEBuildPlatformSDK.bHasAnySDKOverride;
 		}
 
+		/// <summary>
+		/// OLD Constructor.
+		/// </summary>
+		/// <param name="InDescriptor">Target descriptor</param>
+		/// <param name="InRules">The target rules, as created by RulesCompiler.</param>
+		/// <param name="InRulesAssembly">The chain of rules assemblies that this target was created with</param>
+		/// <param name="Logger">Logger for output</param>
+		private UEBuildTarget(TargetDescriptor InDescriptor, ReadOnlyTargetRules InRules, RulesAssembly InRulesAssembly, ILogger Logger)
+			: this(InDescriptor, InRules, InRulesAssembly, SourceFileMetadataCache.CreateHierarchy(InDescriptor.ProjectFile, Logger), Logger)
+		{
+		}
+
 		/// <summary>
 		/// Constructor.
 		/// </summary>
 		/// <param name="InDescriptor">Target descriptor</param>
 		/// <param name="InRules">The target rules, as created by RulesCompiler.</param>
 		/// <param name="InRulesAssembly">The chain of rules assemblies that this target was created with</param>
+		/// <param name="InMetadataCache"></param>
 		/// <param name="Logger">Logger for output</param>
-		internal UEBuildTarget(TargetDescriptor InDescriptor, ReadOnlyTargetRules InRules, RulesAssembly InRulesAssembly, ILogger Logger)
+		internal UEBuildTarget(TargetDescriptor InDescriptor, ReadOnlyTargetRules InRules, RulesAssembly InRulesAssembly, SourceFileMetadataCache InMetadataCache, ILogger Logger)
 		{
+			MetadataCache = InMetadataCache;
 			ProjectFile = InDescriptor.ProjectFile;
 			TargetName = InRules.Name;
 			AppName = InRules.Name;
@@ -1860,14 +1881,7 @@ namespace UnrealBuildTool
 			// Build the engine intermediate directory. If we're building agnostic engine binaries, we can use the engine intermediates folder. Otherwise we need to use the project intermediates directory.
 			if (!bUseSharedBuildEnvironment)
 			{
-				if (Configuration == UnrealTargetConfiguration.DebugGame)
-				{
-					EngineIntermediateDirectory = DirectoryReference.Combine(OutputRootDirectory, PlatformIntermediateFolder, GetTargetIntermediateFolderName(TargetName, IntermediateEnvironment), UnrealTargetConfiguration.Development.ToString());
-				}
-				else
-				{
-					EngineIntermediateDirectory = ProjectIntermediateDirectory;
-				}
+				EngineIntermediateDirectory = ProjectIntermediateDirectory;
 			}
 			else if (Configuration == UnrealTargetConfiguration.DebugGame)
 			{
@@ -2520,7 +2534,7 @@ namespace UnrealBuildTool
 			{
 				Logger.LogWarning("Failed to write partial receipt for target {Descriptor}: {Message}", Descriptor.Name, innerEx.Message);
 				Logger.LogDebug("{InnerException}", innerEx);
-
+				
 				return null;
 			}
 		}
@@ -2589,61 +2603,11 @@ namespace UnrealBuildTool
 			return Makefile;
 
 		}
-
 		public (TargetMakefile, Task?) BuildMakefile(BuildConfiguration BuildConfiguration, ISourceFileWorkingSet WorkingSet, TargetDescriptor TargetDescriptor, ILogger Logger, bool bInitOnly = false)
-		{
-			BuildContext context = BuildMakefile1(TargetDescriptor, Logger, null);
-
-			// UHT mode uses this to create the makefile at least to the point where UHT would have valid manifest information
-			if (bInitOnly)
-			{
-				return (context.Makefile, null);
-			}
-
-			CodeGenContext codeGenContext = new CodeGenContext()
-			{
-				Descriptor = TargetDescriptor,
-				Makefile = context.Makefile,
-				TargetName = TargetName,
-				GenerateHeaderFuncs = context.GenerateHeaderFuncs,
-			};
-
-			RunCodeGen(BuildConfiguration, [codeGenContext], Logger);
-
-			BuildMakefile2(WorkingSet, Logger, context);
-
-			codeGenContext.uhtRun?.Wait();
-
-			return (context.Makefile, codeGenContext.uhtTimestamp);
-		}
-
-		internal struct GenerateHeaderFunc
-		{
-			internal DirectoryReference Dir;
-			internal Action<ILogger, DirectoryReference> Func;
-		}
-
-		internal class BuildContext
-		{
-			required internal SourceFileMetadataCache MetadataCache;
-			required internal TargetDescriptor Descriptor;
-			required internal TargetMakefile Makefile;
-			required internal List<UEBuildBinary> OriginalBinaries;
-			required internal HashSet<UEBuildModuleCPP> ModulesToGenerateHeadersFor;
-			required internal CppCompileEnvironment GlobalCompileEnvironment;
-			required internal UEToolChain TargetToolChain;
-			required internal LinkEnvironment GlobalLinkEnvironment;
-			required internal List<GenerateHeaderFunc> GenerateHeaderFuncs;
-		}
-
-		internal BuildContext BuildMakefile1(TargetDescriptor TargetDescriptor, ILogger Logger, SourceFileMetadataCache? MetadataCache)
 		{
 			CppConfiguration CppConfiguration = GetCppConfiguration(Configuration);
 
-			if (MetadataCache == null)
-			{
-				MetadataCache = SourceFileMetadataCache.CreateHierarchy(ProjectFile, Logger);
-			}
+			SourceFileMetadataCache MetadataCache = SourceFileMetadataCache.CreateHierarchy(ProjectFile, Logger);
 
 			CppCompileEnvironment GlobalCompileEnvironment = new CppCompileEnvironment(Platform, CppConfiguration, Architectures, MetadataCache);
 			LinkEnvironment GlobalLinkEnvironment = new LinkEnvironment(GlobalCompileEnvironment.Platform, GlobalCompileEnvironment.Configuration, GlobalCompileEnvironment.Architectures);
@@ -2683,6 +2647,7 @@ namespace UnrealBuildTool
 				Rules.ConfigValueTracker, bDeployAfterCompile, UbtPlugins?.ToArray(), EnabledUbtPlugins?.ToArray(),
 				EnabledUhtPlugins?.ToArray());
 			Makefile.IsTestTarget = Rules.IsTestTarget;
+			TargetMakefileBuilder MakefileBuilder = new TargetMakefileBuilder(Makefile, Logger);
 
 			// Get diagnostic info to be printed before each build
 			TargetToolChain.GetVersionInfo(Makefile.Diagnostics);
@@ -2728,13 +2693,13 @@ namespace UnrealBuildTool
 				using var _ = Timeline.ScopeEvent("UHTModules");
 
 				UHTExecution.SetupUObjectModules(
-					ModulesToGenerateHeadersFor,
+					ModulesToGenerateHeadersFor, 
 					Rules.Platform,
-					ProjectDescriptor,
-					Makefile.UObjectModules,
-					Makefile.UObjectModuleHeaders,
-					Makefile.ExecutableFile.Directory,
-					Rules.GeneratedCodeVersion,
+					ProjectDescriptor, 
+					Makefile.UObjectModules, 
+					Makefile.UObjectModuleHeaders, 
+					Makefile.ExecutableFile.Directory, 
+					Rules.GeneratedCodeVersion, 
 					MetadataCache);
 			}
 
@@ -2756,186 +2721,50 @@ namespace UnrealBuildTool
 			}
 			Logger.LogInformation($"UHT compiled-in object format {Makefile.UHTTargetSettings.CompiledInObjectFormat}");
 
-#if __VPROJECT_AVAILABLE__
+			// UHT mode uses this to create the makefile at least to the point where UHT would have valid manifest information
+			if (bInitOnly)
 			{
-				using var _ = Timeline.ScopeEvent("VNIModules");
-				VNIExecution.SetupVNIModules(ModulesToGenerateHeadersFor, out Makefile.VNIModules);
-
-				// Copy Verse BPVM usage flag
-				Makefile.bUseVerseBPVM = Rules.bUseVerseBPVM;
+				return (Makefile, null);
 			}
-#endif
 
-			// Collect generate header funcs
-			List<GenerateHeaderFunc> GenerateHeaderFuncs = [];
-			foreach (UEBuildModuleCPP Module in Modules.Values.OfType<UEBuildModuleCPP>())
-			{
-				foreach ((string Subdirectory, Action<ILogger, DirectoryReference> Func) in Module.Rules.GenerateHeaderFuncs)
-				{
-					DirectoryReference GenDirectory = DirectoryReference.Combine(Module.GeneratedCodeDirectory!, Subdirectory);
-					GenerateHeaderFuncs.Add(new() { Dir = GenDirectory, Func = Func });
-				}
-			}
+			// We inject all files that we know UHT/VNI will create into the directory cache.. this way we can run UHT/VNI without waiting
+			// .. and when we are about to test if actions are out-of-date we just reset directory cache so times are re-fetched
+			UHTExecution.PrepopulateDirectoryCache(Makefile.UObjectModules);
 
-
-			return new BuildContext()
+			// NOTE: Even in Gather mode, we need to run UHT to make sure the files exist for the static action graph to be setup correctly.  This is because UHT generates .cpp
+			// files that are injected as top level prerequisites.  If UHT only emitted included header files, we wouldn't need to run it during the Gather phase at all.
+			Task? uhtRunTask = null;
+			Task? uhtTimestampTask = null;
+			if (Makefile.UObjectModules.Count > 0)
 			{
-				MetadataCache = MetadataCache,
-				Descriptor = TargetDescriptor,
-				Makefile = Makefile,
-				OriginalBinaries = OriginalBinaries,
-				ModulesToGenerateHeadersFor = ModulesToGenerateHeadersFor,
-				GlobalCompileEnvironment = GlobalCompileEnvironment,
-				TargetToolChain = TargetToolChain,
-				GlobalLinkEnvironment = GlobalLinkEnvironment,
-				GenerateHeaderFuncs = GenerateHeaderFuncs,
-			};
-		}
-
-		internal class CodeGenContext
-		{
-			required internal TargetDescriptor Descriptor;
-			required internal TargetMakefile Makefile;
-			required internal string TargetName;
-			required internal List<GenerateHeaderFunc> GenerateHeaderFuncs;
-			internal Task? uhtRun;
-			internal Task? uhtTimestamp;
-			internal Task? vniRun;
-		}
+				(uhtRunTask, uhtTimestampTask) = UHTExecution.ExecuteHeaderToolIfNecessary(BuildConfiguration, ProjectFile, Makefile, TargetName, Logger, false);
+			}
 
-		internal static void RunCodeGen(BuildConfiguration BuildConfiguration, List<CodeGenContext> contexts, ILogger Logger)
-		{
-			if (contexts.Count == 0)
+#if __VPROJECT_AVAILABLE__
+			// Prepare cached data for VNI header generation
+			using (new ScopedTimer("UHTExecution.SetupVNIModules()", Logger, LogLevel.Debug, false))
 			{
-				return;
+				using var _ = Timeline.ScopeEvent("VNIModules");
+				VNIExecution.SetupVNIModules(ModulesToGenerateHeadersFor, out Makefile.VNIModules);
 			}
+			VNIExecution.PrepopulateDirectoryCache(Makefile);
 
-			CodeGenContext? lastRunContext = null;
-			HashSet<string>? lastRunUObjectModules = null;
-			HashSet<string>? lastRunVNIModules = null;
-
-			string GetTargetName(TargetDescriptor desc) => $"{desc.Name}{desc.Platform}{desc.Configuration}";
-
-			foreach (CodeGenContext context in contexts)
+			// Copy Verse BPVM usage flag
+			Makefile.bUseVerseBPVM = Rules.bUseVerseBPVM;
+			Task? vniTask = null;
+			if (Makefile.VNIModules.Count > 0)
 			{
-				TargetMakefile makefile = context.Makefile;
-
-				foreach (var entry in context.GenerateHeaderFuncs)
+				vniTask = new Task(() =>
 				{
-					DirectoryReference.CreateDirectory(entry.Dir);
-					entry.Func(Logger, entry.Dir);
-				}
-
-				if (lastRunContext != null)
-				{
-					if (context.TargetName != lastRunContext.TargetName)
-					{
-						throw new Exception("Different targetname!.. run after each other!");
-					}
-
-					if (lastRunUObjectModules == null)
+					using (new ScopedTimer("ExternalExecution.ExecuteVNITool()", Logger, LogLevel.Debug, false))
 					{
-						lastRunUObjectModules = new(lastRunContext.Makefile.UObjectModules.Select(m => m.ModuleName));
-					}
-
-					foreach (var module in makefile.UObjectModules)
-					{
-						if (!lastRunUObjectModules.Contains(module.ModuleName))
-						{
-							throw new Exception($"Target {GetTargetName(context.Descriptor)} has uht modules that previous target {GetTargetName(lastRunContext.Descriptor)} did not have!.. run after each other!");
-						}
-					}
-
-					if (lastRunVNIModules == null)
-					{
-						lastRunVNIModules = new(lastRunContext.Makefile.VNIModules.Select(m => m.ModuleName));
-					}
-
-					foreach (var module in makefile.VNIModules)
-					{
-						if (!lastRunVNIModules.Contains(module.ModuleName))
-						{
-							throw new Exception($"Target {GetTargetName(context.Descriptor)} has vni modules that previous target {GetTargetName(lastRunContext.Descriptor)} did not have!.. run after each other!");
-						}
+						VNIExecution.ExecuteVNITool(Makefile, TargetDescriptor, Logger, false);
 					}
-
-					//if (context.Makefile.EnabledUhtPlugins!.Length != lastRunContext.Makefile.EnabledUhtPlugins!.Length)
-					//	throw new Exception("DIFFERENT THIS NEEDS TO BE HANDLED!.. run after each other!");
-
-					context.uhtRun = lastRunContext.uhtRun;
-					context.uhtTimestamp = lastRunContext.uhtTimestamp;
-					context.vniRun = lastRunContext.vniRun;
-					continue;
-
-				}
-
-				TargetDescriptor descriptor = context.Descriptor;
-
-				if (makefile.UObjectModules.Count > 0)
-				{
-					UHTExecution.PrepopulateDirectoryCache(makefile.UObjectModules);
-					(Task? uhtRun, Task? uhtTimestamp) = UHTExecution.ExecuteHeaderToolIfNecessary(BuildConfiguration, descriptor.ProjectFile, makefile, context.TargetName!, Logger, false);
-					context.uhtRun = uhtRun;
-					context.uhtTimestamp = uhtTimestamp;
-				}
-
-#if __VPROJECT_AVAILABLE__
-				if (makefile.VNIModules.Count > 0)
-				{
-					VNIExecution.PrepopulateDirectoryCache(makefile);
-					context.vniRun = Task.Run(() => VNIExecution.ExecuteVNITool(makefile, descriptor, Logger));
-				}
-#endif
-
-				lastRunContext = context;
+				}, TaskCreationOptions.LongRunning);
+				vniTask.Start();
 			}
-		}
-
-		internal static void WaitCodeGen(List<CodeGenContext> contexts)
-		{
-			foreach (var context in contexts)
-			{
-				if (context.uhtRun != null)
-				{
-					UHTExecution.WaitForExecution(context.uhtRun, context.Makefile);
 
-					// TODO: Kick prefetch of folders
-					//Task task = FileMetadataPrefetch.QueueDirectoryTrees(ModulesToGenerateHeadersFor.Select(m => m.GeneratedCodeDirectoryUHT));
-				}
-#if __VPROJECT_AVAILABLE__
-				if (context.vniRun != null)
-				{
-					VNIExecution.WaitForExecution(context.vniRun, context.Makefile);
-				}
 #endif
-			}
-		}
-
-		internal static List<Task> GetTimestampTasks(List<CodeGenContext> contexts)
-		{
-			List<Task> timestampTasks = [];
-			foreach (var codeGenContext in contexts)
-			{
-				if (codeGenContext.uhtTimestamp != null)
-				{
-					timestampTasks.Add(codeGenContext.uhtTimestamp);
-				}
-			}
-			return timestampTasks;
-		}
-
-		internal void BuildMakefile2(ISourceFileWorkingSet WorkingSet, ILogger Logger, BuildContext context)
-		{
-			SourceFileMetadataCache MetadataCache = context.MetadataCache;
-			TargetDescriptor TargetDescriptor = context.Descriptor;
-			TargetMakefile Makefile = context.Makefile;
-			List<UEBuildBinary> OriginalBinaries = context.OriginalBinaries;
-			HashSet<UEBuildModuleCPP> ModulesToGenerateHeadersFor = context.ModulesToGenerateHeadersFor;
-			CppCompileEnvironment GlobalCompileEnvironment = context.GlobalCompileEnvironment;
-			UEToolChain TargetToolChain = context.TargetToolChain;
-			LinkEnvironment GlobalLinkEnvironment = context.GlobalLinkEnvironment;
-
-			TargetMakefileBuilder MakefileBuilder = new TargetMakefileBuilder(Makefile, Logger);
 
 			if (Rules.bUseSharedPCHs)
 			{
@@ -2946,6 +2775,16 @@ namespace UnrealBuildTool
 				CreateSharedPCHInstances(Rules, TargetToolChain, OriginalBinaries, GlobalCompileEnvironment, MakefileBuilder, Logger);
 			}
 
+			foreach (UEBuildModuleCPP Module in Modules.Values.OfType<UEBuildModuleCPP>())
+			{
+				foreach ((string Subdirectory, Action<ILogger, DirectoryReference> Func) in Module.Rules.GenerateHeaderFuncs)
+				{
+					DirectoryReference GenDirectory = DirectoryReference.Combine(Module.GeneratedCodeDirectory!, Subdirectory);
+					DirectoryReference.CreateDirectory(GenDirectory);
+					Func(Logger, GenDirectory);
+				}
+			}
+
 			// Compile the resource files common to all DLLs on Windows
 			if (!ShouldCompileMonolithic() & !Rules.bFormalBuild && Platform.IsInGroup(UnrealPlatformGroup.Windows))
 			{
@@ -3354,6 +3193,23 @@ namespace UnrealBuildTool
 
 			// Clean any stale modules which exist in multiple output directories. This can lead to the wrong DLL being loaded on Windows.
 			CleanStaleModules(Logger);
+
+#if __VPROJECT_AVAILABLE__
+			if (vniTask != null)
+			{
+				VNIExecution.WaitForExecution(vniTask, Makefile);
+			}
+#endif
+
+			if (uhtRunTask != null)
+			{
+				UHTExecution.WaitForExecution(uhtRunTask, Makefile);
+
+				// Prefetch directory items for UHT folders since they are going to be used later for check out-of-date actions
+				Task task = FileMetadataPrefetch.QueueDirectoryTrees(ModulesToGenerateHeadersFor.Select(m => m.GeneratedCodeDirectoryUHT));
+			}
+
+			return (Makefile, uhtTimestampTask);
 		}
 
 		void CreateStripExportsActions(UEToolChain TargetToolChain, TargetMakefileBuilder MakefileBuilder, TargetMakefile Makefile, CppCompileEnvironment CompileEnvironment)
@@ -5506,7 +5362,7 @@ namespace UnrealBuildTool
 				{
 					// Compile in the build plugins list for the runtime plugin manager to filter optional plugin references against. Only plugins that were actually built
 					// should be enabled when another plugin has an optional dependency on the plugin.
-					IEnumerable<string> BuildPluginStrings = BuildPlugins?.OrderBy(x => x.Name)?.Select(x => $"TEXT(\"{x}\")") ?? Array.Empty<string>();
+					IEnumerable<string> BuildPluginStrings = BuildPlugins?.Select(x => $"TEXT(\"{x}\")") ?? Array.Empty<string>();
 					moduleRules.PrivateDefinitions.Add($"UBT_TARGET_BUILD_PLUGINS={String.Join(", ", BuildPluginStrings)}");
 				}
 			}
diff --git a/Engine/Source/Programs/UnrealBuildTool/Modes/BuildMode.cs b/Engine/Source/Programs/UnrealBuildTool/Modes/BuildMode.cs
--- a/Engine/Source/Programs/UnrealBuildTool/Modes/BuildMode.cs
+++ b/Engine/Source/Programs/UnrealBuildTool/Modes/BuildMode.cs
@@ -8,7 +8,6 @@ using System.Linq;
 using System.Text;
 using System.Threading.Tasks;
 using EpicGames.Core;
-using Microsoft.CodeAnalysis.CSharp.Syntax;
 using Microsoft.Extensions.Logging;
 using OpenTracing.Util;
 using UnrealBuildBase;
@@ -388,16 +387,9 @@ namespace UnrealBuildTool
 		/// <returns>Result from the compilation</returns>
 		public static async Task BuildAsync(List<TargetDescriptor> TargetDescriptors, BuildConfiguration BuildConfiguration, ISourceFileWorkingSet WorkingSet, BuildOptions Options, FileReference? WriteOutdatedActionsFile, ILogger Logger, bool bSkipPreBuildTargets = false, string? ActionTypeFilter = null)
 		{
-			// Create cpp dependency cache in order to be able to preload them as we discover them
-			CppDependencyCache? cppDependencyCache = null;
-			if ((Options & BuildOptions.SkipBuild) == 0)
-			{
-				cppDependencyCache = new();
-			}
-
-			(List<(TargetDescriptor, TargetMakefile)> Targets, List<Task> Tasks) = CreateMakefiles(BuildConfiguration, TargetDescriptors, WorkingSet, Logger, true, true, cppDependencyCache);
+			(List<(TargetDescriptor, TargetMakefile)> Targets, List<Task> Tasks) = CreateMakefiles(BuildConfiguration, TargetDescriptors, WorkingSet, Logger, true, true);
 
-			await BuildAsync(Targets, BuildConfiguration, Options, WriteOutdatedActionsFile, Logger, ActionTypeFilter, cppDependencyCache);
+			await BuildAsync(Targets, BuildConfiguration, Options, WriteOutdatedActionsFile, Logger, ActionTypeFilter);
 
 			Task.WaitAll(Tasks.ToArray());
 		}
@@ -413,9 +405,8 @@ namespace UnrealBuildTool
 		/// <param name="WriteOutdatedActionsFile">Files to write the list of outdated actions to (rather than building them)</param>
 		/// <param name="Logger">Logger for output</param>
 		/// <param name="ActionTypeFilter">If specified, only actions of this type will be run</param>
-		/// <param name="cppDependencyCache"></param>
 		/// <returns>Result from the compilation</returns>
-		internal static async Task BuildAsync(List<(TargetDescriptor Descriptor, TargetMakefile Makefile)> Targets, BuildConfiguration BuildConfiguration, BuildOptions Options, FileReference? WriteOutdatedActionsFile, ILogger Logger, string? ActionTypeFilter=null, CppDependencyCache? cppDependencyCache = null)
+		internal static async Task BuildAsync(List<(TargetDescriptor Descriptor, TargetMakefile Makefile)> Targets, BuildConfiguration BuildConfiguration, BuildOptions Options, FileReference? WriteOutdatedActionsFile, ILogger Logger, string? ActionTypeFilter=null)
 		{
 			List<TargetDescriptor> TargetDescriptors = Targets.Select(T => T.Descriptor).ToList();
 			TargetMakefile[] Makefiles = Targets.Select(T => T.Makefile).ToArray();
@@ -423,16 +414,20 @@ namespace UnrealBuildTool
 			// Execute the build
 			if ((Options & BuildOptions.SkipBuild) == 0)
 			{
-				if (cppDependencyCache == null)
-				{
-					cppDependencyCache = new();
-				}
-
-				// Create the C++ dependency cache
-				for (int TargetIdx = 0; TargetIdx < Targets.Count; TargetIdx++)
+				Task<CppDependencyCache> CppDependenciesTask = Task.Run(() =>
 				{
-					cppDependencyCache.MountAsync(Targets[TargetIdx].Descriptor, "", Targets[TargetIdx].Makefile.TargetType, Logger);
-				}
+					// Create the C++ dependency cache
+					CppDependencyCache CppDependencies = new CppDependencyCache();
+					for (int TargetIdx = 0; TargetIdx < Targets.Count; TargetIdx++)
+					{
+						using (GlobalTracer.Instance.BuildSpan("Reading dependency cache").StartActive())
+						{
+							TargetDescriptor TargetDescriptor = Targets[TargetIdx].Descriptor;
+							CppDependencies.Mount(TargetDescriptor, Targets[TargetIdx].Makefile.TargetType, Logger);
+						}
+					}
+					return CppDependencies;
+				});
 
 #if false // Doesn't seem to make things faster. Revisit since it feels like this should be faster
 				// Prepopulate directories
@@ -563,8 +558,7 @@ namespace UnrealBuildTool
 				}
 
 				// Wait for cpp dependencies task to finish
-				cppDependencyCache.WaitAsync();
-				CppDependencyCache CppDependencies = cppDependencyCache;
+				CppDependencyCache CppDependencies = CppDependenciesTask.Result;
 
 				// Now that we have the dependencies object, optionally create the artifact cache
 				IActionArtifactCache? actionArtifactCache = null;
@@ -1110,380 +1104,253 @@ namespace UnrealBuildTool
 		{
 			await Task.CompletedTask;
 
-			(List<(TargetDescriptor, TargetMakefile Makefile)> Targets, List<Task> Tasks) = CreateMakefiles(BuildConfiguration, [TargetDescriptor], WorkingSet, Logger, false, false, null);
+			(List<(TargetDescriptor, TargetMakefile Makefile)> Targets, List<Task> Tasks) = CreateMakefiles(BuildConfiguration, [TargetDescriptor], WorkingSet, Logger, false, false);
 
 			Task.WaitAll(Tasks.ToArray());
 			return Targets[0].Makefile;
 		}
 
-		// Function that will load makefile and check if it is valid. (First round of targets will have their makefiles loaded in parallel with other work)
-		private static (TargetMakefile? makefile, bool ranPrebuildScripts) LoadMakefile(FileReference location, TargetDescriptor descriptor, ISourceFileWorkingSet workingSet, ILogger logger)
+		private static (List<(TargetDescriptor, TargetMakefile)>, List<Task>) CreateMakefiles(BuildConfiguration BuildConfiguration, List<TargetDescriptor> TargetDescriptors, ISourceFileWorkingSet WorkingSet, ILogger Logger, bool runInitScripts, bool bSkipPreBuildTargets)
 		{
-			string? ReasonNotLoaded;
-			TargetMakefile? makefile = TargetMakefile.Load(location, descriptor.ProjectFile, descriptor.Platform, descriptor.AdditionalArguments.GetRawArray(), logger, out ReasonNotLoaded);
-			if (makefile == null)
-			{
-				logger.LogInformation("Creating makefile for {DescriptorName} ({ReasonNotLoaded})", descriptor.Name, ReasonNotLoaded);
-				return (null, false);
-			}
-
-			// Execute the scripts.
-			Utils.ExecuteCustomBuildSteps(makefile.PreBuildScripts, logger);
+			List<(TargetDescriptor descriptor, TargetMakefile makefile)> targets = new();
+			List<Task> tasks = [];
 
-			// Check that the makefile is still valid
-			string? Reason;
-			if (!TargetMakefile.IsValidForSourceFiles(makefile, descriptor.ProjectFile, descriptor.Platform, workingSet, logger, out Reason))
+			while (TargetDescriptors.Count != 0) // TargetDescriptors will be reallocated for PrebuildTargets
 			{
-				logger.LogInformation("Invalidating makefile for {DescriptorName} ({Reason})", descriptor.Name, Reason);
-				return (null, true);
-			}
-
-			return (makefile, true);
-		}
-
-		private static (List<(TargetDescriptor, TargetMakefile)>, List<Task>) CreateMakefiles(BuildConfiguration buildConfiguration, List<TargetDescriptor> descriptors, ISourceFileWorkingSet workingSet, ILogger logger, bool runInitScripts, bool skipPreBuildTargets, CppDependencyCache? cppDependencyCache)
-		{
-			List<(TargetDescriptor descriptor, TargetMakefile makefile)> outTargets = new();
-			List<Task> outTasks = [];
+				bool skipRulesCompile = BuildConfiguration.bSkipRulesCompile;
+				bool forceRulesCompile = BuildConfiguration.bForceRulesCompile;
+				bool usePrecompiled = BuildConfiguration.bUsePrecompiled;
 
-			int helperCount = buildConfiguration.bParallelMakefileGeneration ? -1 : 0;
+				(TargetRules rules, RulesAssembly assembly)[] rules = new (TargetRules, RulesAssembly)[TargetDescriptors.Count];
+				Task<SourceFileMetadataCache>[] metadataCaches = new Task<SourceFileMetadataCache>[TargetDescriptors.Count];
 
-			Dictionary<FileReference, Task<SourceFileMetadataCache>> sourceFileMetadataCaches = [];
 
-			while (descriptors.Count != 0) // descriptors list will be reallocated for PrebuildTargets. hence the while-loop
-			{
-				int descriptorCount = descriptors.Count;
-				bool skipRulesCompile = buildConfiguration.bSkipRulesCompile;
-				bool forceRulesCompile = buildConfiguration.bForceRulesCompile;
-				bool usePrecompiled = buildConfiguration.bUsePrecompiled;
+				// Build all rules objects
+				Parallel2.For(0, TargetDescriptors.Count, (index) =>
+				{
+					TargetDescriptor TargetDescriptor = TargetDescriptors[index];
+					metadataCaches[index] = Task.Run(() => SourceFileMetadataCache.CreateHierarchy(TargetDescriptor.ProjectFile, Logger));
+					rules[index] = UEBuildTarget.CreateRulesObject(TargetDescriptor, skipRulesCompile, forceRulesCompile, usePrecompiled, TargetDescriptor.IntermediateEnvironment, Logger);
+				});
 
-				FileReference?[] makefileLocations = new FileReference[descriptorCount];
-				UEBuildTarget[] buildTargets = new UEBuildTarget[descriptorCount];
+				// Use rules objects to figure out how to group work. We can't run targets that write to same folders in parallel
+				Dictionary<string, List<int>> buildGroups = [];
 
-				// Start loading unique makefiles
-				Dictionary<FileReference, Task<(TargetMakefile?, bool)>> preLoadedMakefiles = [];
-				if (buildConfiguration.bUseUBTMakefiles)
+				for (int i = 0; i < TargetDescriptors.Count; ++i)
 				{
-					for (int index = 0; index != descriptorCount; ++index)
+					TargetDescriptor targetDescriptor = TargetDescriptors[i];
+					string platformIntermediateFolderNoArch = UEBuildTarget.GetPlatformIntermediateFolder(targetDescriptor.Platform, null, false);
+					string appName = UEBuildTarget.GetAppName(rules[i].rules);
+
+					if (!buildGroups.TryGetValue(appName, out var group))
 					{
-						TargetDescriptor descriptor = descriptors[index];
-						FileReference makefileLocation = TargetMakefile.GetLocation(descriptor.ProjectFile, descriptor.Name, descriptor.Platform, descriptor.Architectures, descriptor.Configuration, descriptor.IntermediateEnvironment);
-						makefileLocations[index] = makefileLocation;
-						if (!preLoadedMakefiles.ContainsKey(makefileLocation))
-						{
-							preLoadedMakefiles.Add(makefileLocation, Task.Run(() => { return LoadMakefile(makefileLocation, descriptor, workingSet, logger); }));
-						}
+						group = new();
+						buildGroups.Add(appName, group);
 					}
+
+					group.Add(i);
 				}
 
-				// All appnames have unique folders. So if appnames are different we can run them in parallel.
-				// The reason we can't run different targets with same appname in parallel is because of codegen (since codegen is not in target folder)
-				// LyraEditor, QAGameEditor etc has the same appname (UnrealEditor) but same codegen folder (for engine part at least)
-				Dictionary<string, List<int>> appNameTargetGroups = [];
+				List<TargetDescriptor> newTargetDescriptors = [];
+
+				// TODO: We also want to handle different configs.. they can't run UHT etc at the same time but can do the other things..
+				// ... so ne need to get that support in here too
 
-				// Build all rules objects and populate groups. Start loading makefiles
-				Parallel2.For(0, descriptorCount, (index) =>
+
+				//Parallel2.ForEach(buildGroups, kv =>
+				foreach (KeyValuePair<string, List<int>> kv in buildGroups)
 				{
-					TargetDescriptor descriptor = descriptors[index];
+					Task? LastSaveMakefileTask = null;
+					Task? LastUHTTimestampTask = null;
 
-					// Create task to load source file metadata cache already here, it will be used later on
-					if (descriptor.ProjectFile != null)
+					List<int> buildGroup = kv.Value;
+					for (int i = 0; i != buildGroup.Count; ++i)
 					{
-						lock (sourceFileMetadataCaches)
+						int index = buildGroup[i];
+						TargetDescriptor TargetDescriptor = TargetDescriptors[index];
+
+						using ITimelineEvent SetupScope = Timeline.ScopeEvent($"{TargetDescriptor.Name} {TargetDescriptor.Platform} {TargetDescriptor.Configuration}");
+
+						if (runInitScripts)
 						{
-							if (!sourceFileMetadataCaches.ContainsKey(descriptor.ProjectFile))
+							// Create and execute the init scripts
+							FileReference[] InitScripts = CreateInitScripts(TargetDescriptor);
+							if (InitScripts.Length > 0)
 							{
-								sourceFileMetadataCaches.Add(descriptor.ProjectFile, Task.Run(() => SourceFileMetadataCache.CreateHierarchy(descriptor.ProjectFile, logger)));
+								Utils.ExecuteCustomBuildSteps(InitScripts, Logger);
 							}
 						}
-					}
 
-					// Compile/load assembly and create target rules
-					(TargetRules rules, RulesAssembly assembly) = UEBuildTarget.CreateRulesObject(descriptor, skipRulesCompile, forceRulesCompile, usePrecompiled, descriptor.IntermediateEnvironment, logger);
+						FileReference? SaveToLocation = null;
 
-					// Create the target. Ctor does not do much so we can create it regardless if we are building it or not
-					UEBuildTarget buildTarget = new(descriptor, new ReadOnlyTargetRules(rules), assembly, logger);
-					buildTargets[index] = buildTarget;
-
-					// Queue up cppdependencies loading
-					if (cppDependencyCache != null)
-					{
-						cppDependencyCache.MountAsync(descriptor, buildTarget.AppName, buildTarget.TargetType, logger);
-					}
-
-					// Add target to proper app name group
-					string appName = $"{buildTarget.AppName}_{buildTarget.Platform}";
-					lock (appNameTargetGroups)
-					{
-						if (!appNameTargetGroups.TryGetValue(appName, out List<int>? appNameTarget))
+						// Get the path to the makefile for this target
+						FileReference? MakefileLocation = null;
+						if (BuildConfiguration.bUseUBTMakefiles)
 						{
-							appNameTarget = new();
-							appNameTargetGroups.Add(appName, appNameTarget);
+							MakefileLocation = TargetMakefile.GetLocation(TargetDescriptor.ProjectFile, TargetDescriptor.Name, TargetDescriptor.Platform, TargetDescriptor.Architectures, TargetDescriptor.Configuration, TargetDescriptor.IntermediateEnvironment);
 						}
 
-						appNameTarget.Add(index);
-					}
-				}, helperCount);
-
-				UEBuildTarget.BuildContext?[] buildContexts = new UEBuildTarget.BuildContext?[descriptorCount];
-				TargetMakefile?[] makefiles = new TargetMakefile?[descriptorCount];
-				ITimelineEvent?[] targetTimelines = new ITimelineEvent?[descriptorCount];
-
-				// Will be populated with prebuild targets if there are any
-				List<TargetDescriptor> newDescriptors = [];
+						// Need to make sure to wait on previous run's makefile
+						LastSaveMakefileTask?.Wait();
 
-				// Run all appname groups in parallel
-				Parallel2.ForEach(appNameTargetGroups, kv =>
-				{
-					List<int> appNameTargets = kv.Value;
-					appNameTargets.Sort(); // To build in same order as registered
-
-					// All targets with same code gen folder but different engine intermediate dir can run in parallel
-					// but they need to sync up for when code gen happens (UHT/VNI)
-
-					Dictionary<DirectoryReference, List<int>> intermediateFolderGroups = [];
-					int maxTargetsInIntermediateCount = 0;
-
-					for (int i = 0; i != appNameTargets.Count; ++i)
-					{
-						int index = appNameTargets[i];
-						UEBuildTarget buildTarget = buildTargets[index];
-						DirectoryReference dir = buildTarget.EngineIntermediateDirectory;
-						if (!intermediateFolderGroups.TryGetValue(dir, out List<int>? intermediateFolderGroup))
+						// Try to load an existing makefile
+						TargetMakefile? Makefile = null;
+						if (MakefileLocation != null)
 						{
-							intermediateFolderGroup = new();
-							intermediateFolderGroups.Add(dir, intermediateFolderGroup);
+							using (GlobalTracer.Instance.BuildSpan("TargetMakefile.Load()").StartActive())
+							{
+								string? ReasonNotLoaded;
+								Makefile = TargetMakefile.Load(MakefileLocation, TargetDescriptor.ProjectFile, TargetDescriptor.Platform, TargetDescriptor.AdditionalArguments.GetRawArray(), Logger, out ReasonNotLoaded);
+								if (Makefile == null)
+								{
+									Logger.LogInformation("Creating makefile for {TargetDescriptorName} ({ReasonNotLoaded})", TargetDescriptor.Name, ReasonNotLoaded);
+								}
+							}
 						}
-						intermediateFolderGroup.Add(index);
-						maxTargetsInIntermediateCount = Math.Max(maxTargetsInIntermediateCount, intermediateFolderGroup.Count);
-					}
 
-					// We need to track save makefile task so we can wait for it in case we need to load it again.
-					// TODO: Just reuse previous one. we don't need to save until we are on the last one.
-					Task?[] lastSaveMakefileTask = new Task?[descriptorCount]; // This must be allocated in here because it is iterated when all passes are done
-
-					List<UEBuildTarget.CodeGenContext>? lastContextsForCodeGen = null;
+						// If we have a makefile, execute the pre-build steps and check it's still valid
+						bool bHasRunPreBuildScripts = false;
+						if (Makefile != null)
+						{
+							// Execute the scripts.
+							Utils.ExecuteCustomBuildSteps(Makefile.PreBuildScripts, Logger);
 
-					// Split up in passes. We run one target from each intermediate folder.
-					for (int passIt = 0; passIt != maxTargetsInIntermediateCount; ++passIt)
-					{
-						List<UEBuildTarget.CodeGenContext> contextsForCodeGen = [];
+							// Don't run the pre-build steps again, even if we invalidate the makefile.
+							bHasRunPreBuildScripts = true;
 
-						// Do everyting up to code gen in parallel
-						Parallel2.ForEach(intermediateFolderGroups.Values, intermediateFolderGroup =>
-						{
-							if (passIt >= intermediateFolderGroup.Count)
+							// Check that the makefile is still valid
+							string? Reason;
+							if (!TargetMakefile.IsValidForSourceFiles(Makefile, TargetDescriptor.ProjectFile, TargetDescriptor.Platform, WorkingSet, Logger, out Reason))
 							{
-								return; // This group is done
+								Logger.LogInformation("Invalidating makefile for {TargetDescriptorName} ({Reason})", TargetDescriptor.Name, Reason);
+								Makefile = null;
 							}
+						}
 
-							int index = intermediateFolderGroup[passIt];
-							TargetDescriptor descriptor = descriptors[index];
+						// Need to make sure to wait on previous run's timestamp since there are targets that write to same place..
+						LastUHTTimestampTask?.Wait();
 
-							targetTimelines[index] = Timeline.ScopeEvent($"{descriptor.Name} {descriptor.Platform} {descriptor.Configuration}");
+						Task? uhtTimeStampsTask = null;
 
-							if (runInitScripts)
-							{
-								// Create and execute the init scripts
-								FileReference[] InitScripts = CreateInitScripts(descriptor);
-								if (InitScripts.Length > 0)
-								{
-									Utils.ExecuteCustomBuildSteps(InitScripts, logger);
-								}
-							}
+						// If we couldn't load a makefile, create a new one
+						if (Makefile == null)
+						{
+							// Create the target
+							UEBuildTarget Target = new UEBuildTarget(TargetDescriptor, new ReadOnlyTargetRules(rules[index].rules), rules[index].assembly, metadataCaches[index].Result, Logger);
+							Target.PreBuildSetup(Logger);
 
-							TargetMakefile? makefile = null;
-							bool ranPrebuildScripts = false;
+							// Create the pre-build scripts
+							FileReference[] PreBuildScripts = Target.CreatePreBuildScripts();
 
-							FileReference? makefileLocation = makefileLocations[index];
-							if (makefileLocation != null)
+							// Execute the pre-build scripts
+							if (!bHasRunPreBuildScripts)
 							{
-								Task<(TargetMakefile?, bool)>? makefileTask = null;
-								lock (preLoadedMakefiles)
-								{
-									if (preLoadedMakefiles.TryGetValue(makefileLocation, out makefileTask))
-									{
-										preLoadedMakefiles.Remove(makefileLocation);
-									}
-								}
-								if (makefileTask != null)
-								{
-									(makefile, ranPrebuildScripts) = makefileTask.Result;
-								}
-								else
-								{
-									// Need to make sure to wait on previous run's makefile
-									lastSaveMakefileTask[index]?.Wait();
-									lastSaveMakefileTask[index] = null;
-									(makefile, ranPrebuildScripts) = LoadMakefile(makefileLocation, descriptor, workingSet, logger);
-								}
+								Utils.ExecuteCustomBuildSteps(PreBuildScripts, Logger);
+								bHasRunPreBuildScripts = true;
 							}
 
-							UEBuildTarget buildTarget = buildTargets[index];
-
-							List<UEBuildTarget.GenerateHeaderFunc> generatedHeaderFuncs = [];
-
-							// If we couldn't load a valid makefile, create a new one
-							if (makefile == null)
+							// Build the target
+							using (GlobalTracer.Instance.BuildSpan("UEBuildTarget.Build()").StartActive())
 							{
-								buildTarget.PreBuildSetup(logger);
+								(Makefile, uhtTimeStampsTask) = Target.BuildMakefile(BuildConfiguration, WorkingSet, TargetDescriptor, Logger);
+							}
 
-								// Create the pre-build scripts
-								FileReference[] PreBuildScripts = buildTarget.CreatePreBuildScripts();
+							Makefile.MemoryPerActionGB = Target.Rules.MemoryPerActionGB;
 
-								// Execute the pre-build scripts
-								if (!ranPrebuildScripts)
-								{
-									Utils.ExecuteCustomBuildSteps(PreBuildScripts, logger);
-									ranPrebuildScripts = true;
-								}
+							// Save the pre-build scripts onto the makefile
+							Makefile.PreBuildScripts = PreBuildScripts;
+
+							// Save the additional command line arguments
+							Makefile.AdditionalArguments = TargetDescriptor.AdditionalArguments.GetRawArray();
 
-								SourceFileMetadataCache? sourceFileMetaDataCache = null;
-								if (descriptor.ProjectFile != null)
+							// Save the environment variables
+							foreach (System.Collections.DictionaryEntry? EnvironmentVariable in Environment.GetEnvironmentVariables())
+							{
+								if (EnvironmentVariable != null)
 								{
-									sourceFileMetaDataCache = sourceFileMetadataCaches[descriptor.ProjectFile].Result;
+									Makefile.EnvironmentVariables.Add(Tuple.Create((string)EnvironmentVariable.Value.Key, (string)EnvironmentVariable.Value.Value!));
 								}
-
-								// Build the target
-								UEBuildTarget.BuildContext context = buildTarget.BuildMakefile1(descriptor, logger, sourceFileMetaDataCache);
-								buildContexts[index] = context;
-
-								makefile = context.Makefile;
-
-								makefile.PreBuildScripts = PreBuildScripts;
-
-								makefile.MemoryPerActionGB = buildTarget.Rules.MemoryPerActionGB;
-
-								// Save the additional command line arguments
-								makefile.AdditionalArguments = descriptor.AdditionalArguments.GetRawArray();
-
-								generatedHeaderFuncs = context.GenerateHeaderFuncs;
 							}
 
-							contextsForCodeGen.Add(new UEBuildTarget.CodeGenContext()
+							// Save the makefile for next time
+							if (MakefileLocation != null)
 							{
-								Descriptor = descriptor,
-								Makefile = makefile!,
-								TargetName = buildTarget.TargetName,
-								GenerateHeaderFuncs = generatedHeaderFuncs,
-							});
-
-							makefiles[index] = makefile;
-						});
-
-						// Make sure to wait for previous pass codegen and timestamp writes
-						if (lastContextsForCodeGen != null)
-						{
-							UEBuildTarget.WaitCodeGen(lastContextsForCodeGen);
-							Task.WaitAll(UEBuildTarget.GetTimestampTasks(lastContextsForCodeGen).ToArray());
+								SaveToLocation = MakefileLocation;
+							}
 						}
-
-						// Sort codegen based on config with Shipping last.
-						// Since shipping often has fewer modules, that means previous configs will include later ones
-						contextsForCodeGen.SortBy(context => context.Descriptor.Configuration);
-
-						// Run codegen for all targets
-						UEBuildTarget.RunCodeGen(buildConfiguration, contextsForCodeGen, logger);
-
-
-						// Go wide again and finish the work
-						Parallel2.ForEach(intermediateFolderGroups.Values, intermediateFolderGroup =>
+						else
 						{
-							if (passIt >= intermediateFolderGroup.Count)
+							// Restore the environment variables
+							foreach (Tuple<string, string> EnvironmentVariable in Makefile.EnvironmentVariables)
 							{
-								return; // This group is done
+								Environment.SetEnvironmentVariable(EnvironmentVariable.Item1, EnvironmentVariable.Item2);
 							}
 
-							int index = intermediateFolderGroup[passIt];
-							TargetDescriptor descriptor = descriptors[index];
-
-							TargetMakefile makefile = makefiles[index]!;
-							UEBuildTarget.BuildContext? buildContext = buildContexts[index];
-
-							if (buildContext != null)
+#if __VPROJECT_AVAILABLE__
+							Task VNITask = Task.CompletedTask;
+							// Same for VNI
+							if (Makefile.VNIModules.Count > 0)
 							{
-								UEBuildTarget buildTarget = buildTargets[index];
-
-								buildTarget.BuildMakefile2(workingSet, logger, buildContext);
-
-								// Save the makefile for next time
-								if (makefileLocations[index] != null)
-								{
-									// Save the environment variables. Toolchains etc might set environment variables
-									// and we want to be able to restore them on runs where makefile is not invalidated.
-									// We know they can contain other target's environment variables too which is bad but by-design atm.
-									foreach (System.Collections.DictionaryEntry? EnvironmentVariable in Environment.GetEnvironmentVariables())
-									{
-										if (EnvironmentVariable != null)
-										{
-											makefile.EnvironmentVariables.Add(Tuple.Create((string)EnvironmentVariable.Value.Key, (string)EnvironmentVariable.Value.Value!));
-										}
-									}
-
-									lastSaveMakefileTask[index] = Task.Run(() =>
-									{
-										using ITimelineEvent _ = Timeline.ScopeEvent($"SavingMakefile");
-										makefile.Save(makefileLocations[index]!);
-									});
-								}
+								VNITask = Task.Run(() => VNIExecution.ExecuteVNITool(Makefile, TargetDescriptor, Logger));
 							}
-							else
+#endif
+							// If the target needs UHT to be run, we'll go ahead and do that now
+							if (Makefile.UObjectModules.Count > 0)
 							{
-								// Restore the environment variables
-								foreach (Tuple<string, string> EnvironmentVariable in makefile.EnvironmentVariables)
-								{
-									Environment.SetEnvironmentVariable(EnvironmentVariable.Item1, EnvironmentVariable.Item2);
-								}
+								(Task uhtRunTask, uhtTimeStampsTask) = UHTExecution.ExecuteHeaderToolIfNecessary(BuildConfiguration, TargetDescriptor.ProjectFile, Makefile, TargetDescriptor.Name, Logger, true);
+								uhtRunTask.Wait();
 							}
 
-							outTargets.Add((descriptor, makefile));
+#if __VPROJECT_AVAILABLE__
+							VNITask.Wait();
+#endif
+						}
+
+						targets.Add((TargetDescriptor, Makefile));
 
-							if (!skipPreBuildTargets)
+						LastUHTTimestampTask = uhtTimeStampsTask;
+
+						if (SaveToLocation != null)
+						{
+							LastSaveMakefileTask = Task.Run(() =>
 							{
-								foreach (TargetInfo PreBuildTarget in makefile.PreBuildTargets)
+								using ITimelineEvent _ = Timeline.ScopeEvent($"SavingMakefile");
+								Makefile.Save(SaveToLocation);
+							});
+						}
+
+						if (!bSkipPreBuildTargets)
+						{
+							foreach (TargetInfo PreBuildTarget in Makefile.PreBuildTargets)
+							{
+								lock (newTargetDescriptors)
 								{
-									lock (newDescriptors)
-									{
-										newDescriptors.Add(TargetDescriptor.FromTargetInfo(PreBuildTarget));
-									}
+									newTargetDescriptors.Add(TargetDescriptor.FromTargetInfo(PreBuildTarget));
 								}
 							}
-
-							targetTimelines[index]?.Finish();
-						});
-
-						lastContextsForCodeGen = contextsForCodeGen;
+						}
 					}
 
-					// All passes are done.
-
-					if (lastContextsForCodeGen != null)
+					if (LastSaveMakefileTask != null)
 					{
-						List<Task> timestampTasks = UEBuildTarget.GetTimestampTasks(lastContextsForCodeGen);
-						UEBuildTarget.WaitCodeGen(lastContextsForCodeGen);
-
-						lock (outTasks)
+						lock (tasks)
 						{
-							outTasks.AddRange(timestampTasks);
+							tasks.Add(LastSaveMakefileTask);
 						}
 					}
-
-					foreach (Task? task in lastSaveMakefileTask)
+					if (LastUHTTimestampTask != null)
 					{
-						if (task != null)
+						lock (tasks)
 						{
-							lock (outTasks)
-							{
-								outTasks.Add(task);
-							}
+							tasks.Add(LastUHTTimestampTask);
 						}
 					}
-				}, helperCount);
+				}//);
 
-				descriptors = newDescriptors;
+				TargetDescriptors = newTargetDescriptors;
 			}
 
-			return (outTargets, outTasks);
+			return (targets, tasks);
 		}
 
 		/// <summary>
diff --git a/Engine/Source/Programs/UnrealBuildTool/System/CppDependencyCache.cs b/Engine/Source/Programs/UnrealBuildTool/System/CppDependencyCache.cs
--- a/Engine/Source/Programs/UnrealBuildTool/System/CppDependencyCache.cs
+++ b/Engine/Source/Programs/UnrealBuildTool/System/CppDependencyCache.cs
@@ -74,6 +74,11 @@ namespace UnrealBuildTool
 			/// </summary>
 			public FileReference Location;
 
+			/// <summary>
+			/// Directory for files to cache dependencies for.
+			/// </summary>
+			public DirectoryReference BaseDir;
+
 			/// <summary>
 			/// Map from file item to dependency info
 			/// </summary>
@@ -84,18 +89,21 @@ namespace UnrealBuildTool
 			/// </summary>
 			public bool bModified;
 
-			/// <summary>
-			/// Whether the partition has been loaded or not
-			/// </summary>
-			public bool Loaded;
-
 			/// <summary>
 			/// Constructs a dependency cache. This method is private; call CppDependencyCache.Create() to create a cache hierarchy for a given project.
 			/// </summary>
 			/// <param name="Location">File to store the cache</param>
-			public CachePartition(FileReference Location)
+			/// <param name="BaseDir">Base directory for files that this cache should store data for</param>
+			/// <param name="Logger">Logger for output</param>
+			public CachePartition(FileReference Location, DirectoryReference BaseDir, ILogger Logger)
 			{
 				this.Location = Location;
+				this.BaseDir = BaseDir;
+
+				if (FileReference.Exists(Location))
+				{
+					Read(Logger);
+				}
 			}
 
 			/// <summary>
@@ -156,9 +164,12 @@ namespace UnrealBuildTool
 		/// <summary>
 		/// List of partitions
 		/// </summary>
-		static List<(DirectoryReference Root, List<(string Filter, CachePartition Partition)>)> GlobalPartitions = [];
+		List<CachePartition> Partitions = new List<CachePartition>();
 
-		static List<Task> GlobalTasks = [];
+		/// <summary>
+		/// Static cache of all constructed dependency caches
+		/// </summary>
+		static Dictionary<FileReference, CachePartition> GlobalPartitions = new Dictionary<FileReference, CachePartition>();
 
 		/// <summary>
 		/// Minimum version of a MSVC source dependency json
@@ -243,13 +254,11 @@ namespace UnrealBuildTool
 			}
 		}
 
-		static readonly string SearchString = $"Intermediate{Path.DirectorySeparatorChar}Build";
-
 		/// <summary>
-		/// Attempts to read dependencies from the given file.
+		/// Attempts to read the dependencies from the given input file
 		/// </summary>
 		/// <param name="InputFile">File to be read</param>
-		/// <param name="Logger">Logger</param>
+		/// <param name="Logger">Logger for output</param>
 		/// <param name="OutInfo">The dependency info</param>
 		/// <returns>True if the input file exists and the dependencies were read</returns>
 		private bool TryGetDependencyInfo(FileItem InputFile, ILogger Logger, [NotNullWhen(true)] out DependencyInfo? OutInfo)
@@ -260,65 +269,45 @@ namespace UnrealBuildTool
 				return false;
 			}
 
-			ReadOnlySpan<char> FullName = InputFile.FullName;
-			int index = FullName.IndexOf(SearchString, FileReference.Comparison);
-			if (index == -1)
+			try
 			{
-				Logger.LogWarning($"Path {InputFile} is not under intermediate folder");
+				return TryGetDependencyInfoInternal(InputFile, out OutInfo);
+			}
+			catch (Exception Ex)
+			{
+				Logger.LogDebug("Unable to read {File}:\n{Ex}", InputFile, ExceptionUtils.FormatExceptionDetails(Ex));
 				OutInfo = null;
 				return false;
 			}
+		}
 
-			int start = index + SearchString.Length + 1;
-			ReadOnlySpan<char> filterPath = FullName.Slice(index + SearchString.Length + 1);
-
-			CachePartition? partition = null;
-
-			foreach (var root in GlobalPartitions)
+		/// <summary>
+		/// Attempts to read dependencies from the given file.
+		/// </summary>
+		/// <param name="InputFile">File to be read</param>
+		/// <param name="OutInfo">The dependency info</param>
+		/// <returns>True if the input file exists and the dependencies were read</returns>
+		private bool TryGetDependencyInfoInternal(FileItem InputFile, [NotNullWhen(true)] out DependencyInfo? OutInfo)
+		{
+			foreach (CachePartition Partition in Partitions)
 			{
-				if (!InputFile.Location.IsUnderDirectory(root.Root))
-				{
-					continue;
-				}
-
-				foreach (var entry in root.Item2)
+				if (InputFile.Location.IsUnderDirectory(Partition.BaseDir))
 				{
-					if (filterPath.StartsWith(entry.Filter, FileReference.Comparison))
+					DependencyInfo? Info;
+					if (!Partition.DependencyFileToInfo.TryGetValue(InputFile, out Info) || InputFile.LastWriteTimeUtc.Ticks > Info.LastWriteTimeUtc)
 					{
-						partition = entry.Partition;
-						break;
+						Info = ReadDependencyInfo(InputFile);
+						Partition.DependencyFileToInfo.AddOrUpdate(InputFile, Info, (k, v) => Info);
+						Partition.bModified = true;
 					}
-				}
-				break;
-
-			}
 
-			if (partition == null)
-			{
-				Logger.LogWarning($"There is no partition covering path to {InputFile}!");
-				OutInfo = null;
-				return false;
-			}
-
-			DependencyInfo? Info;
-			if (!partition.DependencyFileToInfo.TryGetValue(InputFile, out Info) || InputFile.LastWriteTimeUtc.Ticks > Info.LastWriteTimeUtc)
-			{
-				try
-				{
-					Info = ReadDependencyInfo(InputFile);
-					partition.DependencyFileToInfo.AddOrUpdate(InputFile, Info, (k, v) => Info);
-					partition.bModified = true;
-				}
-				catch (Exception Ex)
-				{
-					Logger.LogDebug("Unable to read {File}:\n{Ex}", InputFile, ExceptionUtils.FormatExceptionDetails(Ex));
-					OutInfo = null;
-					return false;
+					OutInfo = Info;
+					return true;
 				}
 			}
 
-			OutInfo = Info;
-			return true;
+			OutInfo = null;
+			return false;
 		}
 
 		/// <summary>
@@ -348,34 +337,7 @@ namespace UnrealBuildTool
 		/// <returns>Dependency cache hierarchy for the given project</returns>
 		public void Mount(TargetDescriptor TargetDescriptor, TargetType TargetType, ILogger Logger)
 		{
-			Mount(TargetDescriptor.ProjectFile, TargetDescriptor.Name, "", TargetDescriptor.Platform, TargetDescriptor.Configuration, TargetType, TargetDescriptor.Architectures, TargetDescriptor.IntermediateEnvironment, Logger);
-		}
-
-		/// <summary>
-		/// Same as Mount but in a task
-		/// </summary>
-		public void MountAsync(TargetDescriptor TargetDescriptor, string AppName, TargetType TargetType, ILogger Logger)
-		{
-			lock (GlobalTasks)
-			{
-				GlobalTasks.Add(Task.Run(() =>
-				{
-					Mount(TargetDescriptor.ProjectFile, TargetDescriptor.Name, AppName, TargetDescriptor.Platform, TargetDescriptor.Configuration, TargetType, TargetDescriptor.Architectures, TargetDescriptor.IntermediateEnvironment, Logger);
-				}));
-			}
-		}
-
-		public void WaitAsync()
-		{
-			Task[] tasks;
-			lock (GlobalTasks)
-			{
-				tasks = GlobalTasks.ToArray();
-				GlobalTasks.Clear();
-			}
-
-			using ITimelineEvent ActionGraphTimer = Timeline.ScopeEvent($"CppDependencyCacheWait");
-			Task.WaitAll(tasks);
+			Mount(TargetDescriptor.ProjectFile, TargetDescriptor.Name, TargetDescriptor.Platform, TargetDescriptor.Configuration, TargetType, TargetDescriptor.Architectures, TargetDescriptor.IntermediateEnvironment, Logger);
 		}
 
 		/// <summary>
@@ -383,7 +345,6 @@ namespace UnrealBuildTool
 		/// </summary>
 		/// <param name="ProjectFile">Project file for the target being built</param>
 		/// <param name="TargetName">Name of the target</param>
-		/// <param name="AppName">Name of the app.</param>
 		/// <param name="Platform">Platform being built</param>
 		/// <param name="Configuration">Configuration being built</param>
 		/// <param name="TargetType">The target type</param>
@@ -391,81 +352,52 @@ namespace UnrealBuildTool
 		/// <param name="IntermediateEnvironment">Intermediate environment to use</param>
 		/// <param name="Logger">Logger for output</param>
 		/// <returns>Dependency cache hierarchy for the given project</returns>
-		public void Mount(FileReference? ProjectFile, string TargetName, string AppName, UnrealTargetPlatform Platform, UnrealTargetConfiguration Configuration, TargetType TargetType, UnrealArchitectures Architectures, UnrealIntermediateEnvironment IntermediateEnvironment, ILogger Logger)
+		public void Mount(FileReference? ProjectFile, string TargetName, UnrealTargetPlatform Platform, UnrealTargetConfiguration Configuration, TargetType TargetType, UnrealArchitectures Architectures, UnrealIntermediateEnvironment IntermediateEnvironment, ILogger Logger)
 		{
-			if (AppName == string.Empty)
+			string AppName = TargetName;
+			if (ProjectFile == null || !Unreal.IsEngineInstalled())
 			{
-				AppName = TargetName;
-				if (ProjectFile == null || !Unreal.IsEngineInstalled())
+				if (TargetType == TargetType.Program)
 				{
-					if (TargetType == TargetType.Program)
-					{
-						AppName = TargetName;
-					}
-					else
-					{
-						AppName = UEBuildTarget.GetAppNameForTargetType(TargetType);
-					}
+					AppName = TargetName;
+				}
+				else
+				{
+					AppName = UEBuildTarget.GetAppNameForTargetType(TargetType);
 				}
 			}
 
-			using ITimelineEvent ActionGraphTimer = Timeline.ScopeEvent($"CppDependencyCache");
-
-			AppName = UEBuildTarget.GetTargetIntermediateFolderName(AppName, IntermediateEnvironment);
+			FileReference EngineCacheLocation = FileReference.Combine(Unreal.WritableEngineDirectory, UEBuildTarget.GetPlatformIntermediateFolder(Platform, Architectures, false), UEBuildTarget.GetTargetIntermediateFolderName(AppName, IntermediateEnvironment), Configuration.ToString(), "DependencyCache.bin");
+			FindOrAddPartition(EngineCacheLocation, Unreal.EngineDirectory, Logger);
 
-			string GetFilterPath(UnrealTargetConfiguration Configuration)
+			if (ProjectFile != null)
 			{
-				StringBuilder sb = new();
-				sb.Append(Platform).Append(Path.DirectorySeparatorChar).Append(UnrealArchitectureConfig.ForPlatform(Platform).GetFolderNameForArchitectures(Architectures)).Append(Path.DirectorySeparatorChar);
-				sb.Append(AppName).Append(Path.DirectorySeparatorChar).Append(Configuration);
-				return sb.ToString();
+				FileReference ProjectCacheLocation = FileReference.Combine(ProjectFile.Directory, UEBuildTarget.GetPlatformIntermediateFolder(Platform, Architectures, false), UEBuildTarget.GetTargetIntermediateFolderName(TargetName, IntermediateEnvironment), Configuration.ToString(), "DependencyCache.bin");
+				FindOrAddPartition(ProjectCacheLocation, ProjectFile.Directory, Logger);
 			}
+		}
 
-			void AddPartition(DirectoryReference rootPath, string filterPath)
+		/// <summary>
+		/// Reads a cache from the given location, or creates it with the given settings
+		/// </summary>
+		/// <param name="Location">File to store the cache</param>
+		/// <param name="BaseDir">Base directory for files that this cache should store data for</param>
+		/// <param name="Logger">Logger for output</param>
+		/// <returns>Reference to a dependency cache with the given settings</returns>
+		void FindOrAddPartition(FileReference Location, DirectoryReference BaseDir, ILogger Logger)
+		{
+			lock (GlobalPartitions)
 			{
-				CachePartition? partition = null;
-				lock (GlobalPartitions)
+				if (!Partitions.Any(x => x.Location == Location))
 				{
-					var root = GlobalPartitions.FirstOrDefault(e => e.Root == rootPath);
-					if (root.Root == null)
+					CachePartition? Partition;
+					if (!GlobalPartitions.TryGetValue(Location, out Partition))
 					{
-						root = new(rootPath, new());
-						GlobalPartitions.Add(root);
-					}
-					var filters = root.Item2;
-					var filter = filters.FirstOrDefault(f => f.Filter == filterPath);
-					partition = filter.Partition;
-					if (partition == null)
-					{
-						FileReference location = FileReference.Combine(rootPath, "Intermediate", "Build", filterPath, "DependencyCache.bin");
-						partition = new CachePartition(location);
-						filters.Add((filterPath, partition));
+						Partition = new CachePartition(Location, BaseDir, Logger);
+						GlobalPartitions.Add(Location, Partition);
 					}
+					Partitions.Add(Partition);
 				}
-
-				lock (partition)
-				{
-					if (!partition.Loaded)
-					{
-						partition.Loaded = true;
-						if (FileReference.Exists(partition.Location))
-						{
-							partition.Read(Logger);
-						}
-					}
-				}
-
-			}
-
-			AddPartition(Unreal.WritableEngineDirectory, GetFilterPath(Configuration));
-			if (Configuration == UnrealTargetConfiguration.DebugGame)
-			{
-				AddPartition(Unreal.WritableEngineDirectory, GetFilterPath(UnrealTargetConfiguration.Development));
-			}
-
-			if (ProjectFile != null)
-			{
-				AddPartition(ProjectFile.Directory, GetFilterPath(Configuration));
 			}
 		}
 
@@ -474,7 +406,7 @@ namespace UnrealBuildTool
 		/// </summary>
 		public static void SaveAll()
 		{
-			Parallel2.ForEach(GlobalPartitions.SelectMany(p => p.Item2).Select(f => f.Partition), Cache =>
+			Parallel.ForEach(GlobalPartitions.Values, Cache =>
 			{
 				if (Cache.bModified)
 				{
