[UBT]
* Added lock around fetching visual studio installations. Got a crash on the farm and only theory is that that the native code is not thread safe.

diff --git a/Engine/Source/Programs/UnrealBuildTool/Platform/Windows/MicrosoftPlatformSDK.cs b/Engine/Source/Programs/UnrealBuildTool/Platform/Windows/MicrosoftPlatformSDK.cs
--- a/Engine/Source/Programs/UnrealBuildTool/Platform/Windows/MicrosoftPlatformSDK.cs
+++ b/Engine/Source/Programs/UnrealBuildTool/Platform/Windows/MicrosoftPlatformSDK.cs
@@ -107,6 +107,8 @@ namespace UnrealBuildTool
 		/// </summary>
 		private static IReadOnlyList<VisualStudioInstallation>? CachedVisualStudioInstallations;
 
+		private static object CachedVisualStudioInstallationsLock = new();
+
 		/// <summary>
 		/// Cache of DIA SDK installation directories
 		/// </summary>
@@ -951,86 +953,89 @@ namespace UnrealBuildTool
 		/// <returns>List of directories containing Visual Studio installations</returns>
 		public static IReadOnlyList<VisualStudioInstallation> FindVisualStudioInstallations(ILogger Logger)
 		{
-			if (CachedVisualStudioInstallations != null)
+			lock (CachedVisualStudioInstallationsLock)
 			{
-				return CachedVisualStudioInstallations;
-			}
-
-			List<VisualStudioInstallation> Installations = new List<VisualStudioInstallation>();
-			if (BuildHostPlatform.Current.Platform == UnrealTargetPlatform.Win64)
-			{
-				try
+				if (CachedVisualStudioInstallations != null)
 				{
-					SetupConfiguration Setup = new SetupConfiguration();
-					IEnumSetupInstances Enumerator = Setup.EnumAllInstances();
+					return CachedVisualStudioInstallations;
+				}
 
-					ISetupInstance[] Instances = new ISetupInstance[1];
-					for (; ; )
+				List<VisualStudioInstallation> Installations = new List<VisualStudioInstallation>();
+				if (BuildHostPlatform.Current.Platform == UnrealTargetPlatform.Win64)
+				{
+					try
 					{
-						int NumFetched;
-						Enumerator.Next(1, Instances, out NumFetched);
+						SetupConfiguration Setup = new SetupConfiguration();
+						IEnumSetupInstances Enumerator = Setup.EnumAllInstances();
 
-						if (NumFetched == 0)
+						ISetupInstance[] Instances = new ISetupInstance[1];
+						for (; ; )
 						{
-							break;
-						}
+							int NumFetched;
+							Enumerator.Next(1, Instances, out NumFetched);
 
-						ISetupInstance2 Instance = (ISetupInstance2)Instances[0];
-						if ((Instance.GetState() & InstanceState.Local) != InstanceState.Local)
-						{
-							continue;
-						}
+							if (NumFetched == 0)
+							{
+								break;
+							}
 
-						VersionNumber? Version;
-						if (!VersionNumber.TryParse(Instance.GetInstallationVersion(), out Version))
-						{
-							continue;
-						}
+							ISetupInstance2 Instance = (ISetupInstance2)Instances[0];
+							if ((Instance.GetState() & InstanceState.Local) != InstanceState.Local)
+							{
+								continue;
+							}
 
-						int MajorVersion = Version.GetComponent(0);
+							VersionNumber? Version;
+							if (!VersionNumber.TryParse(Instance.GetInstallationVersion(), out Version))
+							{
+								continue;
+							}
 
-						WindowsCompiler Compiler;
-						if (Version >= MicrosoftPlatformSDK.MinimumVisualStudio2026Version) // Treat any newer versions as 2026, until we have an explicit enum for them
-						{
-							Compiler = WindowsCompiler.VisualStudio2026;
-						}
-						else if (Version >= MicrosoftPlatformSDK.MinimumVisualStudio2022Version)
-						{
-							Compiler = WindowsCompiler.VisualStudio2022;
-						}
-						else // Unsupported older versions still detected for legacy MSBuild projects, but will not be used to compile
-						{
+							int MajorVersion = Version.GetComponent(0);
+
+							WindowsCompiler Compiler;
+							if (Version >= MicrosoftPlatformSDK.MinimumVisualStudio2026Version) // Treat any newer versions as 2026, until we have an explicit enum for them
+							{
+								Compiler = WindowsCompiler.VisualStudio2026;
+							}
+							else if (Version >= MicrosoftPlatformSDK.MinimumVisualStudio2022Version)
+							{
+								Compiler = WindowsCompiler.VisualStudio2022;
+							}
+							else // Unsupported older versions still detected for legacy MSBuild projects, but will not be used to compile
+							{
 #pragma warning disable CS0618 // Type or member is obsolete
-							Compiler = WindowsCompiler.VisualStudioUnsupported; // Unsupported but still detected for legacy MSBuild projects
+								Compiler = WindowsCompiler.VisualStudioUnsupported; // Unsupported but still detected for legacy MSBuild projects
 #pragma warning restore CS0618 // Type or member is obsolete
-						}
+							}
 
-						ISetupInstanceCatalog? Catalog = Instance as ISetupInstanceCatalog;
-						WindowsCompilerChannel ReleaseChannel = Catalog?.IsPrerelease() == true ? WindowsCompilerChannel.Preview : WindowsCompilerChannel.Latest;
+							ISetupInstanceCatalog? Catalog = Instance as ISetupInstanceCatalog;
+							WindowsCompilerChannel ReleaseChannel = Catalog?.IsPrerelease() == true ? WindowsCompilerChannel.Preview : WindowsCompilerChannel.Latest;
 
-						string ProductId = Instance.GetProduct().GetId();
-						bool bCommunity = ProductId.Equals("Microsoft.VisualStudio.Product.Community", StringComparison.Ordinal);
+							string ProductId = Instance.GetProduct().GetId();
+							bool bCommunity = ProductId.Equals("Microsoft.VisualStudio.Product.Community", StringComparison.Ordinal);
 
-						DirectoryReference BaseDir = new DirectoryReference(Instance.GetInstallationPath());
-						Installations.Add(new VisualStudioInstallation(Compiler, Version, BaseDir, bCommunity, ReleaseChannel));
+							DirectoryReference BaseDir = new DirectoryReference(Instance.GetInstallationPath());
+							Installations.Add(new VisualStudioInstallation(Compiler, Version, BaseDir, bCommunity, ReleaseChannel));
 
-						Logger.LogDebug("Found {VisualStudio} installation: {BaseDir} (Product={ProductId}, Version={Version})", Compiler, BaseDir, ProductId, Version);
-					}
+							Logger.LogDebug("Found {VisualStudio} installation: {BaseDir} (Product={ProductId}, Version={Version})", Compiler, BaseDir, ProductId, Version);
+						}
 
-					Installations = [.. Installations.OrderBy(x => x.Compiler)
+						Installations = [.. Installations.OrderBy(x => x.Compiler)
 						.ThenBy(x => x.bCommunity)
 						.ThenBy(x => x.ReleaseChannel)
 						.ThenByDescending(x => x.Version)
-					];
-				}
-				catch (Exception Ex)
-				{
-					Logger.LogDebug(Ex, "Unable to enumerate Visual Studio installations");
+						];
+					}
+					catch (Exception Ex)
+					{
+						Logger.LogDebug(Ex, "Unable to enumerate Visual Studio installations");
+					}
 				}
-			}
-			CachedVisualStudioInstallations = Installations;
+				CachedVisualStudioInstallations = Installations;
 
-			return Installations;
+				return Installations;
+			}
 		}
 
 		/// <summary>
