[UBT]
* Made android toolchain threadsafe. Moved some code to only run once and inside a lock

diff --git a/Engine/Source/Programs/UnrealBuildTool/Platform/Android/AndroidToolChain.cs b/Engine/Source/Programs/UnrealBuildTool/Platform/Android/AndroidToolChain.cs
--- a/Engine/Source/Programs/UnrealBuildTool/Platform/Android/AndroidToolChain.cs
+++ b/Engine/Source/Programs/UnrealBuildTool/Platform/Android/AndroidToolChain.cs
@@ -3,6 +3,7 @@
 using EpicGames.Core;
 using Microsoft.Extensions.Logging;
 using System;
+using System.Collections.Concurrent;
 using System.Collections.Generic;
 using System.Diagnostics;
 using System.IO;
@@ -60,47 +61,41 @@ namespace UnrealBuildTool
 		protected static string? AndroidClangBuild;
 
 		// the "-android" suffix paths here are vcpkg triplets for the android platform
-		private static Dictionary<UnrealArch, string[]> AllArchNames = new() {
+		private static readonly Dictionary<UnrealArch, string[]> AllArchNames = new() {
 			{ UnrealArch.Arm64, new string[] { "arm64", "arm64-v8a", "arm64-android" } },
 			{ UnrealArch.X64,   new string[] { "x64", "x86_64", "x64-android" } }
 		};
 
 		// architecture paths to use for filtering include and lib paths
-		private static Dictionary<UnrealArch, string[]> AllFilterArchNames = new() {
+		private static readonly Dictionary<UnrealArch, string[]> AllFilterArchNames = new() {
 			{ UnrealArch.Arm64, new string[] { "arm64", "arm64-v8a", "arm64-android" } },
 			{ UnrealArch.X64,   new string[] { "x64", "x86_64", "x64-android" } },
 			// using Default as a placeholder to remove old folders for arches we no longer support, but licensees may have in their Build rules that we need to strip out
 			{ UnrealArch.Deprecated  , new string[] { "armv7", "armeabi-v7a", "arm-android", "x86", "x86-android" } }
 		};
 
-		private static Dictionary<UnrealArch, string[]> LibrariesToSkip = new() {
+		private static readonly Dictionary<UnrealArch, string[]> LibrariesToSkip = new() {
 			{ UnrealArch.Arm64, new string[] { "nvToolsExt", "nvToolsExtStub", "vorbisenc", } },
 			{ UnrealArch.X64,   new string[] { "nvToolsExt", "nvToolsExtStub", "oculus", "OVRPlugin", "vrapi", "ovrkernel", "systemutils", "openglloader", "ovrplatformloader", "vorbisenc", } }
 		};
 
-		private static Dictionary<UnrealArch, string[]> ModulesToSkip = new() {
+		private static readonly Dictionary<UnrealArch, string[]> ModulesToSkip = new() {
 			{ UnrealArch.Arm64, Array.Empty<string>() },
 			{ UnrealArch.X64,   new string[] { "OnlineSubsystemOculus", "OculusHMD", "OculusMR" } }
 		};
 
-		private static Dictionary<UnrealArch, string[]> GeneratedModulesToSkip = new() {
+		private static readonly Dictionary<UnrealArch, string[]> GeneratedModulesToSkip = new() {
 			{ UnrealArch.Arm64, Array.Empty<string>() },
 			{ UnrealArch.X64,   new string[] { "OculusEntitlementCallbackProxy", "OculusCreateSessionCallbackProxy", "OculusFindSessionsCallbackProxy", "OculusIdentityCallbackProxy", "OculusNetConnection", "OculusNetDriver", "OnlineSubsystemOculus_init" } }
 		};
 
-		public string? NDKToolchainVersion;
-		public ulong NDKVersionInt;
+		public static string? NDKToolchainVersion { get; private set; }
+		public static ulong NDKVersionInt { get; private set; }
 
-		int ClangVersionMajor = -1;
-		int ClangVersionMinor = -1;
-		int ClangVersionPatch = -1;
+		private static int ClangVersionMajor = -1;
+		private static int ClangVersionMinor = -1;
+		private static int ClangVersionPatch = -1;
 
-		protected void SetClangVersion(int Major, int Minor, int Patch)
-		{
-			ClangVersionMajor = Major;
-			ClangVersionMinor = Minor;
-			ClangVersionPatch = Patch;
-		}
 
 		public string GetClangVersionString()
 		{
@@ -123,16 +118,52 @@ namespace UnrealBuildTool
 		{
 		}
 
-		DirectoryReference GCCToolchainPath;
-		DirectoryReference SysrootPath;
+		static DirectoryReference? GCCToolchainPath;
+		static DirectoryReference? SysrootPath;
+		static string ExeExtension = ".exe";
 
 		public AndroidToolChain(FileReference? InProjectFile, ClangToolChainOptions ToolchainOptions, ILogger InLogger)
 			: base(ToolchainOptions, InLogger)
 		{
 			Options = ToolchainOptions;
 			ProjectFile = InProjectFile;
+			lock (CtorLock)
+			{
+				StaticInit(Logger);
+			}
+
+			// NDK setup (enforce minimum API level)
+			int NDKApiLevel64Int = GetNdkApiLevelInt(MinimumNDKAPILevel);
+
+			// toolchain params (note: use ANDROID=1 same as we define it)
+			ToolchainLinkParamsArm64 = " --target=aarch64-none-linux-android" + NDKApiLevel64Int + " -DANDROID=1";
+			ToolchainLinkParamsx64 = " --target=x86_64-none-linux-android" + NDKApiLevel64Int + " -DANDROID=1";
+
+			ToolchainParamsArm64 = ToolchainLinkParamsArm64;
+			ToolchainParamsx64 = ToolchainLinkParamsx64;
+
+			if (!IsNewNDKModel())
+			{
+				// We need to manually provide -D__ANDROID_API__ for NDK versions prior to r22 only, for newer ones, --target=aarch64-none-linux-android + NDKApiLevel64Int does it for us
+				ToolchainParamsArm64 += " -D__ANDROID_API__=" + NDKApiLevel64Int;
+				ToolchainParamsx64 += " -D__ANDROID_API__=" + NDKApiLevel64Int;
+			}
+
+			ReadElfPath = Path.Combine(NDKPath!, @"toolchains/llvm", ArchitecturePath!, @"bin/llvm-readelf" + ExeExtension);
+			ObjDumpPath = Path.Combine(NDKPath!, @"toolchains/llvm", ArchitecturePath!, @"bin/llvm-objdump" + ExeExtension);
+		}
+
+		static string? NDKPath;
+		static string? ArchitecturePath;
+
+		static void StaticInit(ILogger Logger)
+		{
+			if (NDKPath != null)
+			{
+				return;
+			}
 
-			string? NDKPath = AndroidPlatformSDK.GetNDKRoot(InLogger);
+			NDKPath = AndroidPlatformSDK.GetNDKRoot(Logger);
 
 			// don't register if we don't have an NDKROOT specified
 			if (String.IsNullOrEmpty(NDKPath))
@@ -142,12 +173,10 @@ namespace UnrealBuildTool
 
 			NDKPath = NDKPath.Replace("\"", "");
 
-			string ArchitecturePath;
 			string ArchitecturePathWindows32 = @"prebuilt/windows";
 			string ArchitecturePathWindows64 = @"prebuilt/windows-x86_64";
 			string ArchitecturePathMac = @"prebuilt/darwin-x86_64";
 			string ArchitecturePathLinux = @"prebuilt/linux-x86_64";
-			string ExeExtension = ".exe";
 
 			if (Directory.Exists(Path.Combine(NDKPath, @"toolchains/llvm", ArchitecturePathWindows64)))
 			{
@@ -179,7 +208,8 @@ namespace UnrealBuildTool
 			// get the installed version (in the form r10e and 100500)
 			UEBuildPlatformSDK SDK = UEBuildPlatform.GetSDK(UnrealTargetPlatform.Android)!;
 			NDKToolchainVersion = SDK.GetInstalledVersion();
-			SDK.TryConvertVersionToInt(NDKToolchainVersion, out NDKVersionInt);
+			SDK.TryConvertVersionToInt(NDKToolchainVersion, out ulong NDKVersionIntTemp);
+			NDKVersionInt = NDKVersionIntTemp;
 
 			// figure out clang version (will live in toolchains/llvm from NDK 21 forward
 			if (Directory.Exists(Path.Combine(NDKPath, @"toolchains/llvm")))
@@ -192,7 +222,9 @@ namespace UnrealBuildTool
 				}
 				string[] VersionFile = File.ReadAllLines(VersionFilename);
 				string[] VersionParts = VersionFile[0].Split('.');
-				SetClangVersion(Int32.Parse(VersionParts[0]), (VersionParts.Length > 1) ? Int32.Parse(VersionParts[1]) : 0, (VersionParts.Length > 2) ? Int32.Parse(VersionParts[2]) : 0);
+				ClangVersionMajor = Int32.Parse(VersionParts[0]);
+				ClangVersionMinor = (VersionParts.Length > 1) ? Int32.Parse(VersionParts[1]) : 0;
+				ClangVersionPatch = (VersionParts.Length > 2) ? Int32.Parse(VersionParts[2]) : 0;
 			}
 			else
 			{
@@ -222,36 +254,16 @@ namespace UnrealBuildTool
 			ArPathArm64 = Utils.CollapseRelativeDirectories(Path.Combine(NDKPath, @"toolchains/llvm", ArchitecturePath, @"bin/llvm-ar" + ExeExtension));
 			ArPathx64 = ArPathArm64;
 
-			// NDK setup (enforce minimum API level)
-			int NDKApiLevel64Int = GetNdkApiLevelInt(MinimumNDKAPILevel);
-
-			GCCToolchainPath = DirectoryReference.FromString(Path.Combine(NDKPath, @"toolchains/llvm", ArchitecturePath))!;
-			SysrootPath = DirectoryReference.FromString(Path.Combine(NDKPath, @"toolchains/llvm", ArchitecturePath, "sysroot").Replace("\\", "/"))!;
-
-			// toolchain params (note: use ANDROID=1 same as we define it)
-			ToolchainLinkParamsArm64 = " --target=aarch64-none-linux-android" + NDKApiLevel64Int + " -DANDROID=1";
-			ToolchainLinkParamsx64 = " --target=x86_64-none-linux-android" + NDKApiLevel64Int + " -DANDROID=1";
-
-			ToolchainParamsArm64 = ToolchainLinkParamsArm64;
-			ToolchainParamsx64 = ToolchainLinkParamsx64;
-
-			if (!IsNewNDKModel())
-			{
-				// We need to manually provide -D__ANDROID_API__ for NDK versions prior to r22 only, for newer ones, --target=aarch64-none-linux-android + NDKApiLevel64Int does it for us
-				ToolchainParamsArm64 += " -D__ANDROID_API__=" + NDKApiLevel64Int;
-				ToolchainParamsx64 += " -D__ANDROID_API__=" + NDKApiLevel64Int;
-			}
-
-			ReadElfPath = Path.Combine(NDKPath, @"toolchains/llvm", ArchitecturePath, @"bin/llvm-readelf" + ExeExtension);
-			ObjDumpPath = Path.Combine(NDKPath, @"toolchains/llvm", ArchitecturePath, @"bin/llvm-objdump" + ExeExtension);
+			GCCToolchainPath = DirectoryReference.FromString(Path.Combine(NDKPath!, @"toolchains/llvm", ArchitecturePath!))!;
+			SysrootPath = DirectoryReference.FromString(Path.Combine(NDKPath!, @"toolchains/llvm", ArchitecturePath!, "sysroot").Replace("\\", "/"))!;
 		}
 
 		IEnumerable<string> AdditionalPaths(CppRootPaths Roots)
 		{
 			return new[]
 			{
-				$"--gcc-toolchain=\"{NormalizeCommandLinePath(GCCToolchainPath, Roots)}\"",
-				$"--sysroot=\"{NormalizeCommandLinePath(SysrootPath, Roots)}\"",
+				$"--gcc-toolchain=\"{NormalizeCommandLinePath(GCCToolchainPath!, Roots)}\"",
+				$"--sysroot=\"{NormalizeCommandLinePath(SysrootPath!, Roots)}\"",
 			};
 		}
 
@@ -340,20 +352,19 @@ namespace UnrealBuildTool
 			return NDKVersionInt;
 		}
 
-		static string CachedPlatformsFilename = "";
-		static bool CachedPlatformsValid = false;
-		static int CachedMinPlatform = -1;
-		static int CachedMaxPlatform = -1;
+		class MinMax
+		{
+			internal bool PlatformsValid = false;
+			internal int MinPlatform = -1;
+			internal int MaxPlatform = -1;
+		}
+		static ConcurrentDictionary<string, Lazy<MinMax>> CachedMinMax = [];
 
 		private bool ReadMinMaxPlatforms(string PlatformsFilename, out int MinPlatform, out int MaxPlatform)
 		{
-			if (!CachedPlatformsFilename.Equals(PlatformsFilename))
+			MinMax minMax = CachedMinMax.GetOrAdd(PlatformsFilename, new Lazy<MinMax>(() =>
 			{
-				// reset cache to defaults
-				CachedPlatformsFilename = PlatformsFilename;
-				CachedPlatformsValid = false;
-				CachedMinPlatform = -1;
-				CachedMaxPlatform = -1;
+				MinMax minMax = new();
 
 				// try to read it
 				try
@@ -361,17 +372,18 @@ namespace UnrealBuildTool
 					JsonObject? PlatformsObj = null;
 					if (JsonObject.TryRead(new FileReference(PlatformsFilename), out PlatformsObj))
 					{
-						CachedPlatformsValid = PlatformsObj.TryGetIntegerField("min", out CachedMinPlatform) && PlatformsObj.TryGetIntegerField("max", out CachedMaxPlatform);
+						minMax.PlatformsValid = PlatformsObj.TryGetIntegerField("min", out minMax.MinPlatform) && PlatformsObj.TryGetIntegerField("max", out minMax.MaxPlatform);
 					}
 				}
 				catch (Exception)
 				{
 				}
-			}
+				return minMax;
+			})).Value;
 
-			MinPlatform = CachedMinPlatform;
-			MaxPlatform = CachedMaxPlatform;
-			return CachedPlatformsValid;
+			MinPlatform = minMax.MinPlatform;
+			MaxPlatform = minMax.MaxPlatform;
+			return minMax.PlatformsValid;
 		}
 
 		//This doesn't take into account SDK version overrides in packaging
@@ -962,6 +974,7 @@ namespace UnrealBuildTool
 		}
 
 		// cache the location of NDK tools
+		protected static object CtorLock = new();
 		protected static string? ClangPath;
 		protected static string ToolchainParamsArm64 = "";
 		protected static string ToolchainParamsx64 = "";
