[UBT]
* Optimized CppRootPaths by changing sorted dictionary to a list and also used StackStringBuffer to remove temporary allocations

diff --git a/Engine/Source/Programs/Shared/EpicGames.Core/StackStringBuffer.cs b/Engine/Source/Programs/Shared/EpicGames.Core/StackStringBuffer.cs
--- a/Engine/Source/Programs/Shared/EpicGames.Core/StackStringBuffer.cs
+++ b/Engine/Source/Programs/Shared/EpicGames.Core/StackStringBuffer.cs
@@ -12,7 +12,7 @@ namespace EpicGames.Core
 	///   using StackStringBuilder sb = new(initial);
 	///   sb.Append("Foo");
 	/// </summary>
-	ref struct StackStringBuilder
+	public ref struct StackStringBuilder
 	{
 		private Span<char> _chars;
 		private int _pos;
@@ -74,6 +74,20 @@ namespace EpicGames.Core
 			_pos = newPos;
 		}
 
+		/// <summary>
+		/// Replaces all character of from to to
+		/// </summary>
+		/// <param name="from"></param>
+		/// <param name="to"></param>
+		public void Replace(char from, char to)
+		{
+			_chars.Slice(0, _pos).Replace(from, to);
+		}
+
+		/// <summary>
+		/// To string
+		/// </summary>
+		/// <returns></returns>
 		public override string ToString()
 		{
 			return new string(_chars.Slice(0, _pos));
diff --git a/Engine/Source/Programs/UnrealBuildTool/System/CppRootPaths.cs b/Engine/Source/Programs/UnrealBuildTool/System/CppRootPaths.cs
--- a/Engine/Source/Programs/UnrealBuildTool/System/CppRootPaths.cs
+++ b/Engine/Source/Programs/UnrealBuildTool/System/CppRootPaths.cs
@@ -89,8 +89,9 @@ namespace UnrealBuildTool
 				{
 					CppRootPathFolder folder = (CppRootPathFolder)Reader.ReadInt();
 					DirectoryReference local = Reader.ReadCompactDirectoryReference();
-					_lookup.Add(folder, (AsVfsPath(folder), local));
+					_lookup.Add((folder, (AsVfsPath(folder), local)));
 				}
+				_lookup.Sort(static (a, b) => a.root.CompareTo(b.root));
 			}
 			{
 				int count = Reader.ReadInt();
@@ -124,7 +125,7 @@ namespace UnrealBuildTool
 		public void Write(BinaryArchiveWriter Writer)
 		{
 			Writer.WriteDictionary(_folderNames, Key => Writer.WriteInt((int)Key), Value => Writer.WriteString(Value));
-			Writer.WriteDictionary(_lookup, Key => Writer.WriteInt((int)Key), Value => Writer.WriteCompactDirectoryReference(Value.local));
+			Writer.WriteList(_lookup, item => { Writer.WriteInt((int)item.root); Writer.WriteCompactDirectoryReference(item.Item2.local); });
 			Writer.WriteDictionary(_extras, Key => Writer.WriteString(Key), Value => Writer.WriteCompactDirectoryReference(Value.local));
 			Writer.WriteBool(bUseVfs);
 		}
@@ -144,7 +145,16 @@ namespace UnrealBuildTool
 						AddFolderName(CppRootPathFolder.Project, value.GetDirectoryName());
 					}
 				}
-				_lookup[key] = (AsVfsPath(key), value);
+				for (int i = 0; i != _lookup.Count; ++i)
+				{
+					if (_lookup[i].root == key)
+					{
+						_lookup[i] = (key, (AsVfsPath(key), value));
+						return;
+					}
+				}
+				_lookup.Add((key, (AsVfsPath(key), value)));
+				_lookup.Sort(static (a, b) => a.root.CompareTo(b.root));
 			}
 		}
 
@@ -153,9 +163,9 @@ namespace UnrealBuildTool
 		/// </summary>
 		public IEnumerator<(uint id, DirectoryReference vfs, DirectoryReference local)> GetEnumerator()
 		{
-			foreach (KeyValuePair<CppRootPathFolder, (DirectoryReference vfs, DirectoryReference local)> kv in _lookup)
+			foreach ((CppRootPathFolder root, (DirectoryReference vfs, DirectoryReference local)) in _lookup)
 			{
-				yield return ((uint)kv.Key, kv.Value.vfs, kv.Value.local);
+				yield return ((uint)root, vfs, local);
 			}
 
 			uint index = (uint)CppRootPathFolder.Count; // Max of CppRootPathFolder
@@ -212,10 +222,20 @@ namespace UnrealBuildTool
 		/// <param name="other">The paths to merge from</param>
 		public void Merge(CppRootPaths other)
 		{
-			foreach (KeyValuePair<CppRootPathFolder, (DirectoryReference vfs, DirectoryReference local)> lookup in other._lookup.Where(x => !_lookup.ContainsKey(x.Key)))
+			foreach ((CppRootPathFolder root, (DirectoryReference vfs, DirectoryReference local)) otherEntry in other._lookup)
 			{
-				this[lookup.Key] = lookup.Value.local;
+				bool found = false;
+				foreach ((CppRootPathFolder root, (DirectoryReference vfs, DirectoryReference local)) in _lookup)
+				{
+					found |= root == otherEntry.root;
+				}
+				if (!found)
+				{
+					_lookup.Add(otherEntry);
+				}
 			}
+			_lookup.Sort(static (a, b) => a.root.CompareTo(b.root));
+
 			foreach (KeyValuePair<string, (DirectoryReference vfs, DirectoryReference local)> extra in other._extras.Where(x => !_extras.ContainsKey(x.Key)))
 			{
 				AddExtraPath((extra.Key, extra.Value.local.FullName));
@@ -240,28 +260,37 @@ namespace UnrealBuildTool
 				return false;
 			}
 
-			foreach (KeyValuePair<CppRootPathFolder, (DirectoryReference vfs, DirectoryReference local)> item in _lookup)
+			bool InternalGet(DirectoryReference vfs, DirectoryReference local, [NotNullWhen(true)]out string? vfsPath)
 			{
-				if (reference.IsUnderDirectory(item.Value.local))
+				if (reference.IsUnderDirectory(local))
 				{
-					vfsPath = item.Value.vfs.FullName + reference.FullName.Substring(item.Value.local.FullName.Length);
+					Span<char> initial = stackalloc char[512];
+					using StackStringBuilder sb = new(initial);
+					sb.Append(vfs.FullName);
+					sb.Append(reference.FullName.AsSpan(local.FullName.Length));
 					if (OperatingSystem.IsWindows())
 					{
-						vfsPath = vfsPath.Replace('\\', '/');
+						sb.Replace('\\', '/');
 					}
+					vfsPath = sb.ToString();
 					return true;
 				}
+				vfsPath = null;
+				return false;
 			}
 
-			foreach (KeyValuePair<string, (DirectoryReference vfs, DirectoryReference local)> item in _extras)
+			foreach ((CppRootPathFolder root, (DirectoryReference vfs, DirectoryReference local)) in _lookup)
 			{
-				if (reference.IsUnderDirectory(item.Value.local))
+				if (InternalGet(vfs, local, out vfsPath))
+				{
+					return true;
+				}
+			}
+
+			foreach ((DirectoryReference vfs, DirectoryReference local) in _extras.Values)
+			{
+				if (InternalGet(vfs, local, out vfsPath))
 				{
-					vfsPath = item.Value.vfs.FullName + reference.FullName.Substring(item.Value.local.FullName.Length);
-					if (OperatingSystem.IsWindows())
-					{
-						vfsPath = vfsPath.Replace('\\', '/');
-					}
 					return true;
 				}
 			}
@@ -279,19 +308,19 @@ namespace UnrealBuildTool
 				return reference;
 			}
 
-			foreach (KeyValuePair<CppRootPathFolder, (DirectoryReference vfs, DirectoryReference local)> item in _lookup)
+			foreach ((CppRootPathFolder root, (DirectoryReference vfs, DirectoryReference local)) in _lookup)
 			{
-				if (reference.IsUnderDirectory(item.Value.vfs))
+				if (reference.IsUnderDirectory(vfs))
 				{
-					return FileReference.FromString(item.Value.local.FullName + reference.FullName.Substring(item.Value.vfs.FullName.Length));
+					return FileReference.FromString(local.FullName + reference.FullName.Substring(vfs.FullName.Length));
 				}
 			}
 
-			foreach (KeyValuePair<string, (DirectoryReference vfs, DirectoryReference local)> item in _extras)
+			foreach ((DirectoryReference vfs, DirectoryReference local) in _extras.Values)
 			{
-				if (reference.IsUnderDirectory(item.Value.vfs))
+				if (reference.IsUnderDirectory(vfs))
 				{
-					return FileReference.FromString(item.Value.local.FullName + reference.FullName.Substring(item.Value.vfs.FullName.Length));
+					return FileReference.FromString(local.FullName + reference.FullName.Substring(vfs.FullName.Length));
 				}
 			}
 
@@ -310,7 +339,7 @@ namespace UnrealBuildTool
 					DirectoryReference local = new(item.GetStringField("Path"));
 					if (item.TryGetEnumField("Root", out CppRootPathFolder folder))
 					{
-						_lookup.Add(folder, (AsVfsPath(folder), local));
+						_lookup.Add((folder, (AsVfsPath(folder), local)));
 
 						if (item.TryGetStringField("Root", out string? folderName))
 						{
@@ -318,6 +347,7 @@ namespace UnrealBuildTool
 						}
 					}
 				}
+				_lookup.Sort(static (a, b) => a.root.CompareTo(b.root));
 			}
 
 			if (Object.TryGetObjectArrayField("Extras", out JsonObject[]? extraItems))
@@ -341,15 +371,15 @@ namespace UnrealBuildTool
 		/// </summary>
 		public void Write(JsonWriter writer)
 		{
-			if (_lookup.Any())
+			if (_lookup.Count > 0)
 			{
 				writer.WriteArrayStart("Lookup");
-				foreach (KeyValuePair<CppRootPathFolder, (DirectoryReference vfs, DirectoryReference local)> item in _lookup)
+				foreach ((CppRootPathFolder root, (DirectoryReference vfs, DirectoryReference local)) in _lookup)
 				{
 					writer.WriteObjectStart();
-					writer.WriteEnumValue("Root", item.Key);
-					writer.WriteValue("Path", item.Value.local.FullName);
-					if (!_folderNames.TryGetValue(item.Key, out string? folderName))
+					writer.WriteEnumValue("Root", root);
+					writer.WriteValue("Path", local.FullName);
+					if (!_folderNames.TryGetValue(root, out string? folderName))
 					{
 						writer.WriteValue("Name", folderName);
 					}
@@ -409,7 +439,7 @@ namespace UnrealBuildTool
 		private static DirectoryReference AsVfsPath(string id) => DirectoryReference.Combine(_vfsRootPath, id);
 
 		private Dictionary<CppRootPathFolder, string> _folderNames { get; }
-		private SortedDictionary<CppRootPathFolder, (DirectoryReference vfs, DirectoryReference local)> _lookup { get; }
+		private List<(CppRootPathFolder root, (DirectoryReference vfs, DirectoryReference local))> _lookup { get; } // We want this to be a list for perf reasons
 		private SortedDictionary<string, (DirectoryReference vfs, DirectoryReference local)> _extras { get; }
 
 		/// <summary>
