[UBT]
* Optimized AndroidToolChain by hand roll tests if path belongs to excluded arch. This saves minutes of time when creating ubt makefile
* Changed so check for lib dependencies is made in a parallel loop since it is launching an external application which takes some time

diff --git a/Engine/Source/Programs/UnrealBuildTool/Platform/Android/AndroidToolChain.cs b/Engine/Source/Programs/UnrealBuildTool/Platform/Android/AndroidToolChain.cs
--- a/Engine/Source/Programs/UnrealBuildTool/Platform/Android/AndroidToolChain.cs
+++ b/Engine/Source/Programs/UnrealBuildTool/Platform/Android/AndroidToolChain.cs
@@ -431,7 +431,8 @@ namespace UnrealBuildTool
 			}
 
 			string PlatformsFilename = Environment.ExpandEnvironmentVariables("%NDKROOT%/meta/platforms.json");
-			if (!File.Exists(PlatformsFilename))
+			FileItem platformsJsonFile = FileItem.GetItemByPath(PlatformsFilename);
+			if (!platformsJsonFile.Exists)
 			{
 				throw new BuildException("No NDK platforms found in {0}", PlatformsFilename);
 			}
@@ -475,13 +476,23 @@ namespace UnrealBuildTool
 
 		protected override void GetCompileArguments_IncludePaths(CppCompileEnvironment CompileEnvironment, List<string> Arguments)
 		{
-			// remove paths that are not meant for this architecture
+			// Filter out paths that are not meant for this architecture
 			// @todo can remove this when we only add paths properly for architecture
-			IEnumerable<DirectoryReference> FilteredPaths = CompileEnvironment.UserIncludePaths.Where(x => IsDirectoryForArch(x.FullName, CompileEnvironment.Architecture));
-			Arguments.AddRange(FilteredPaths.Select(IncludePath => GetUserIncludePathArgument(IncludePath, CompileEnvironment)));
+			foreach (DirectoryReference dir in CompileEnvironment.UserIncludePaths)
+			{
+				if (IsDirectoryForArch(dir.FullName, CompileEnvironment.Architecture))
+				{
+					Arguments.Add(GetUserIncludePathArgument(dir, CompileEnvironment));
+				}
+			}
 
-			FilteredPaths = CompileEnvironment.SystemIncludePaths.Where(x => IsDirectoryForArch(x.FullName, CompileEnvironment.Architecture));
-			Arguments.AddRange(FilteredPaths.Select(IncludePath => GetSystemIncludePathArgument(IncludePath, CompileEnvironment)));
+			foreach (DirectoryReference dir in CompileEnvironment.SystemIncludePaths)
+			{
+				if (IsDirectoryForArch(dir.FullName, CompileEnvironment.Architecture))
+				{
+					Arguments.Add(GetSystemIncludePathArgument(dir, CompileEnvironment));
+				}
+			}
 		}
 
 		/// <inheritdoc/>
@@ -829,31 +840,115 @@ namespace UnrealBuildTool
 			return Result;
 		}
 
-		static bool IsDirectoryForArch(string Dir, UnrealArch Arch)
+		public static bool IsDirectoryForArch(string Dir, UnrealArch Arch)
 		{
-			// make sure paths use one particular slash
-			Dir = Dir.Replace("\\", "/").ToLowerInvariant();
+			string dir2 = Dir.Replace('\\', '/');
+			while (dir2.Length > 0 && dir2[^1] == '/')
+			{
+				dir2 = dir2[..^1];
+			}
 
-			// look for other architectures in the Dir path, and fail if it finds it
-			foreach (KeyValuePair<UnrealArch, string[]> Pair in AllFilterArchNames)
+			return IsDirectoryForArchCache.GetOrAdd((dir2, Arch), key => CheckCore(key.Path, Arch));
+		}
+
+		private static readonly Dictionary<UnrealArch, HashSet<string>> DisallowedNamesByArch = BuildDisallowed();
+		private static readonly ConcurrentDictionary<(string Path, UnrealArch Arch), bool> IsDirectoryForArchCache = new();
+
+		private static bool CheckCore(string normPath, UnrealArch arch)
+		{
+			HashSet<string> disallowed = DisallowedNamesByArch[arch];
+
+			int start = 0;
+			for (int i = 0; i <= normPath.Length; i++)
 			{
-				if (Pair.Key != Arch)
+				if (i == normPath.Length || normPath[i] == '/')
 				{
-					foreach (string ArchName in Pair.Value)
+					int len = i - start;
+					if (len > 0)
 					{
-						// if there's a directory in the path with a bad architecture name, reject it
-						if (Regex.IsMatch(Dir, "/" + ArchName + "$") || Regex.IsMatch(Dir, "/" + ArchName + "/") || Regex.IsMatch(Dir, "/" + ArchName + "_API[0-9]+_NDK[0-9]+", RegexOptions.IgnoreCase))
+						string segment = normPath.Substring(start, len);
+
+						// Exact directory match with a foreign arch name?
+						if (disallowed.Contains(segment))
 						{
 							return false;
 						}
+
+						// Match patterns like "<arch>_API<digits>_NDK<digits>"
+						int us = segment.IndexOf('_');
+						if (us > 0)
+						{
+							string baseName = segment.Substring(0, us);
+							if (disallowed.Contains(baseName) && LooksLikeApiNdkSuffix(segment.AsSpan(us)))
+							{
+								return false;
+							}
+						}
 					}
+					start = i + 1;
 				}
 			}
-
-			// if nothing was found, we are okay
 			return true;
 		}
 
+		// Expect suffix: "_API<digits>_NDK<digits>" (case-insensitive)
+		private static bool LooksLikeApiNdkSuffix(ReadOnlySpan<char> suffix)
+		{
+			// _API
+			if (suffix.Length < 4 || !suffix.Slice(0, 4).Equals("_api".AsSpan(), StringComparison.OrdinalIgnoreCase))
+			{
+				return false;
+			}
+
+			int i = 4, nDigitsApi = 0;
+			while (i < suffix.Length && char.IsDigit(suffix[i]))
+			{
+				i++;
+				nDigitsApi++;
+			}
+
+			if (nDigitsApi == 0)
+			{
+				return false;
+			}
+
+			// _NDK
+			if (i + 4 > suffix.Length || !suffix.Slice(i, 4).Equals("_ndk".AsSpan(), StringComparison.OrdinalIgnoreCase))
+			{
+				return false;
+			}
+			i += 4;
+
+			int nDigitsNdk = 0;
+			while (i < suffix.Length && char.IsDigit(suffix[i]))
+			{
+				i++;
+				nDigitsNdk++;
+			}
+			return nDigitsNdk > 0 && i == suffix.Length;
+		}
+
+		private static Dictionary<UnrealArch, HashSet<string>> BuildDisallowed()
+		{
+			Dictionary<UnrealArch, HashSet<string>> result = [];
+			foreach (UnrealArch arch in AllFilterArchNames.Keys)
+			{
+				HashSet<string> set = new(StringComparer.OrdinalIgnoreCase);
+				foreach (KeyValuePair<UnrealArch, string[]> kv in AllFilterArchNames)
+				{
+					if (kv.Key != arch)
+					{
+						foreach (string name in kv.Value)
+						{
+							set.Add(name);
+						}
+					}
+				}
+				result[arch] = set;
+			}
+			return result;
+		}
+		
 		static bool ShouldSkipModule(string ModuleName, UnrealArch Arch)
 		{
 			foreach (string ModName in ModulesToSkip[Arch])
@@ -1189,48 +1284,59 @@ namespace UnrealBuildTool
 						}
 					}
 				}
+
+				bool doValidation = DisableLibCppSharedDependencyValidation() && ReadElfPath != null;
+				List<FileReference> libsToCheck = [];
+
 				foreach (FileReference Library in LinkEnvironment.Libraries)
 				{
-					if (!ShouldSkipLib(Library.FullName, Arch))
+					if (ShouldSkipLib(Library.FullName, Arch))
 					{
-						string AbsoluteLibraryPath = Path.GetDirectoryName(Library.FullName)!;
-						LinkAction.PrerequisiteItems.Add(FileItem.GetItemByFileReference(Library));
+						continue;
+					}
+					string AbsoluteLibraryPath = Path.GetDirectoryName(Library.FullName)!;
+					LinkAction.PrerequisiteItems.Add(FileItem.GetItemByFileReference(Library));
 
-						string Lib = Path.GetFileNameWithoutExtension(Library.FullName);
-						if (Lib.StartsWith("lib"))
-						{
-							AdditionalLibraries.Add(Lib);
-							if (!AdditionalLibraryPaths.Contains(AbsoluteLibraryPath))
-							{
-								AdditionalLibraryPaths.Add(AbsoluteLibraryPath);
-							}
-						}
-						else
+					string Lib = Path.GetFileNameWithoutExtension(Library.FullName);
+					if (Lib.StartsWith("lib"))
+					{
+						AdditionalLibraries.Add(Lib);
+						if (!AdditionalLibraryPaths.Contains(AbsoluteLibraryPath))
 						{
-							AdditionalLibraries.Add(AbsoluteLibraryPath);
+							AdditionalLibraryPaths.Add(AbsoluteLibraryPath);
 						}
+					}
+					else
+					{
+						AdditionalLibraries.Add(AbsoluteLibraryPath);
+					}
+
+					if (doValidation)
+					{
+						libsToCheck.Add(Library);
+					}
+				}
 
-						if (!DisableLibCppSharedDependencyValidation() && ReadElfPath != null)
+				// Do checks in parallel because they are expensive
+				Parallel2.ForEach(libsToCheck, Library =>
+				{
+					string? Output = Utils.RunLocalProcessAndReturnStdOut(ReadElfPath!, "--dynamic \"" + Library.FullName + "\"");
+					if (Output != null)
+					{
+						if (Output.Contains("libc++_shared.so"))
 						{
-							string? Output = Utils.RunLocalProcessAndReturnStdOut(ReadElfPath, "--dynamic \"" + Library.FullName + "\"");
-							if (Output != null)
+							if (IsNewNDKModel())
 							{
-								if (Output.Contains("libc++_shared.so"))
-								{
-									if (IsNewNDKModel())
-									{
-										throw new BuildException("Lib {0} depends on libc++_shared.so. There are known incompatibility issues when linking libc++_shared.so with Unreal Engine built with NDK22+." +
-											" Please rebuild your dependencies with static libc++!", Lib);
-									}
-									else
-									{
-										Logger.LogWarning("Lib {LibName} depends on libc++_shared.so. Unreal Engine is designed to be linked with libs that are built against static libc++ only. Please rebuild your dependencies with static libc++!", Lib);
-									}
-								}
+								throw new BuildException("Lib {0} depends on libc++_shared.so. There are known incompatibility issues when linking libc++_shared.so with Unreal Engine built with NDK22+." +
+									" Please rebuild your dependencies with static libc++!", Library.GetFileNameWithoutExtension());
+							}
+							else
+							{
+								Logger.LogWarning("Lib {LibName} depends on libc++_shared.so. Unreal Engine is designed to be linked with libs that are built against static libc++ only. Please rebuild your dependencies with static libc++!", Library.GetFileNameWithoutExtension());
 							}
 						}
 					}
-				}
+				});
 
 				// add the library paths to response
 				foreach (string LibaryPath in AdditionalLibraryPaths)
