[UBT]
* Cache UEBuildModule GeneratedCodeDirectoryUHT and GeneratedCodeDirectoryVNI since it is being read a million times (literally)

diff --git a/Engine/Source/Programs/UnrealBuildTool/Configuration/UEBuildModuleCPP.cs b/Engine/Source/Programs/UnrealBuildTool/Configuration/UEBuildModuleCPP.cs
--- a/Engine/Source/Programs/UnrealBuildTool/Configuration/UEBuildModuleCPP.cs
+++ b/Engine/Source/Programs/UnrealBuildTool/Configuration/UEBuildModuleCPP.cs
@@ -76,12 +76,12 @@ namespace UnrealBuildTool
 		/// <summary>
 		/// The directory for this module's generated UHT code
 		/// </summary>
-		public DirectoryReference? GeneratedCodeDirectoryUHT => GeneratedCodeDirectory != null ? DirectoryReference.Combine(GeneratedCodeDirectory!, "UHT") : null;
+		public readonly DirectoryReference? GeneratedCodeDirectoryUHT;
 
 		/// <summary>
 		/// The directory for this module's generated VNI code
 		/// </summary>
-		public DirectoryReference? GeneratedCodeDirectoryVNI => GeneratedCodeDirectory != null ? DirectoryReference.Combine(GeneratedCodeDirectory!, "VNI") : null;
+		public readonly DirectoryReference? GeneratedCodeDirectoryVNI;
 
 		/// <summary>
 		/// Global override to force all include paths to be always added
@@ -257,6 +257,12 @@ namespace UnrealBuildTool
 		{
 			this.GeneratedCodeDirectory = GeneratedCodeDirectory;
 
+			if (GeneratedCodeDirectory != null)
+			{
+				GeneratedCodeDirectoryUHT = DirectoryReference.Combine(GeneratedCodeDirectory!, "UHT");
+				GeneratedCodeDirectoryVNI = DirectoryReference.Combine(GeneratedCodeDirectory!, "VNI");
+			}
+
 			// Check for a Verse directory next to the rules file
 			DirectoryItem RulesDirectory = DirectoryItem.GetItemByDirectoryReference(Rules.File.Directory);
 			if (RulesDirectory.TryGetDirectory("Verse", out DirectoryItem? MaybeVerseDirectory))
@@ -437,6 +443,9 @@ namespace UnrealBuildTool
 			bool bLegacyParentIncludePaths
 			)
 		{
+			// TODO: CACHE THESE THINGS... THIS IS CALLED A MILLION TIMES BECAUSE OF SHARED PCH
+			// Change HashSets to DirectoryItem instead
+
 			if (GeneratedCodeDirectory != null)
 			{
 				// This directory may not exist for this module (or ever exist, if it doesn't contain any generated headers), but we want the project files
