[UBT]
* Fixed race condition in vc toolchain when building arm64 target

diff --git a/Engine/Source/Programs/UnrealBuildTool/Platform/Windows/VCToolChain.cs b/Engine/Source/Programs/UnrealBuildTool/Platform/Windows/VCToolChain.cs
--- a/Engine/Source/Programs/UnrealBuildTool/Platform/Windows/VCToolChain.cs
+++ b/Engine/Source/Programs/UnrealBuildTool/Platform/Windows/VCToolChain.cs
@@ -2586,7 +2586,9 @@ namespace UnrealBuildTool
 
 			OptionalUpdateSxSManifest(Target, MakefileBuilder);
 		}
-		
+
+		static object SxSManifestLock = new();
+
 		protected void OptionalUpdateSxSManifest(ReadOnlyTargetRules Target, TargetMakefileBuilder MakefileBuilder)
 		{
 			if (Target.WindowsPlatform.ManifestFile != null && !Target.bShouldCompileAsDLL && !Target.Architecture.bIsX64)
@@ -2607,38 +2609,41 @@ namespace UnrealBuildTool
 				FileReference IntermediateManifest = FileReference.Combine(  BaseDir, "Intermediate", "SxSManifest", AssemblyName, AssemblyName + ".manifest");
 				FileReference TemplateManifest = FileReference.Combine(Unreal.EngineDirectory, "Build", "Windows", "Resources", AssemblyName, AssemblyName + ".manifest");			
 				XDocument XmlManifest;
-				if (FileReference.Exists(IntermediateManifest))
-				{
-					XmlManifest = XDocument.Load(IntermediateManifest.FullName);
-				}
-				else
-				{
-					XmlManifest = XDocument.Load(TemplateManifest.FullName);
-					bDirty = true;
-				}
 
-				// update & save the intermediate SxS manifest
-				if (XmlManifest.Root != null && XmlManifest.Root.Name.LocalName == "assembly")
+				lock (SxSManifestLock)
 				{
-					XNamespace XmlNs = XmlManifest.Root.Name.Namespace;
-					HashSet<string> ExistingDLLNames = [.. XmlManifest.Root.Elements(XmlNs+"file").Select( X => X.Attribute("name")!.Value )];
+					if (FileReference.Exists(IntermediateManifest))
+					{
+						XmlManifest = XDocument.Load(IntermediateManifest.FullName);
+					}
+					else
+					{
+						XmlManifest = XDocument.Load(TemplateManifest.FullName);
+						bDirty = true;
+					}
 
-					foreach (FileItem DLL in DLLs)
+					// update & save the intermediate SxS manifest
+					if (XmlManifest.Root != null && XmlManifest.Root.Name.LocalName == "assembly")
 					{
-						string DLLName = DLL.Location.MakeRelativeTo(ManifestDir);
-						if (!ExistingDLLNames.Contains(DLLName))
+						XNamespace XmlNs = XmlManifest.Root.Name.Namespace;
+						HashSet<string> ExistingDLLNames = [.. XmlManifest.Root.Elements(XmlNs + "file").Select(X => X.Attribute("name")!.Value)];
+
+						foreach (FileItem DLL in DLLs)
 						{
-							XmlManifest.Root.Add( new XElement(XmlNs+"file", new XAttribute("name", DLLName ) ) );
-							bDirty = true;
+							string DLLName = DLL.Location.MakeRelativeTo(ManifestDir);
+							if (!ExistingDLLNames.Contains(DLLName))
+							{
+								XmlManifest.Root.Add(new XElement(XmlNs + "file", new XAttribute("name", DLLName)));
+								bDirty = true;
+							}
 						}
 					}
+					if (bDirty)
+					{
+						DirectoryReference.CreateDirectory(IntermediateManifest.Directory);
+						XmlManifest.Save(IntermediateManifest.FullName);
+					}
 				}
-				if (bDirty)
-				{
-					DirectoryReference.CreateDirectory(IntermediateManifest.Directory);
-					XmlManifest.Save(IntermediateManifest.FullName);
-				}
-
 				// include the intermediate SxS manifest as a dependency
 				MakefileBuilder.Makefile.InternalDependencies.Add( FileItem.GetItemByFileReference(IntermediateManifest) );
 			}
