UnrealBuildTool: Avoid logging entire stack trace for Horde-related compute allocation errors

diff --git a/Engine/Source/Programs/Shared/EpicGames.Horde/Compute/Clients/ServerComputeClient.cs b/Engine/Source/Programs/Shared/EpicGames.Horde/Compute/Clients/ServerComputeClient.cs
--- a/Engine/Source/Programs/Shared/EpicGames.Horde/Compute/Clients/ServerComputeClient.cs
+++ b/Engine/Source/Programs/Shared/EpicGames.Horde/Compute/Clients/ServerComputeClient.cs
@@ -284,7 +284,7 @@ namespace EpicGames.Horde.Compute.Clients
 				{
 					content = "None";
 				}
-				throw new ComputeClientException($"InternalServerError requesting compute resources: \"{content}\"");
+				throw new ComputeClientException($"Error allocating UBA cache server: {content}");
 			}
 			
 			httpResponse.EnsureSuccessStatusCode();
@@ -360,7 +360,7 @@ namespace EpicGames.Horde.Compute.Clients
 					{
 						content = "None";
 					}
-					throw new ComputeClientException($"InternalServerError requesting compute resources: \"{content}\"");
+					throw new ComputeClientException($"Error requesting compute resources: {content}");
 				}
 
 				httpResponse.EnsureSuccessStatusCode();
diff --git a/Engine/Source/Programs/UnrealBuildTool/Executors/UnrealBuildAccelerator/UBAAgentCoordinatorHorde.cs b/Engine/Source/Programs/UnrealBuildTool/Executors/UnrealBuildAccelerator/UBAAgentCoordinatorHorde.cs
--- a/Engine/Source/Programs/UnrealBuildTool/Executors/UnrealBuildAccelerator/UBAAgentCoordinatorHorde.cs
+++ b/Engine/Source/Programs/UnrealBuildTool/Executors/UnrealBuildAccelerator/UBAAgentCoordinatorHorde.cs
@@ -908,6 +908,7 @@ namespace UnrealBuildTool
 
 				// We are assuming all active logical cores are already being used.. so queueWeight is essentially work that could be executed but can't because of bandwidth
 				double queueThreshold = _ubaConfig.bForceBuildAllRemote ? 0 : 5;
+				LogLevel logLevel = _ubaConfig.bStrict ? LogLevel.Warning : LogLevel.Information;
 
 				try
 				{
@@ -955,16 +956,20 @@ namespace UnrealBuildTool
 				{
 					if (!shownNoAgentsFoundMessage)
 					{
-						_logger.Log(_ubaConfig.bStrict ? LogLevel.Warning : LogLevel.Information, KnownLogEvents.Systemic_Horde_Compute, ex, "No agents found matching requirements (cluster: {ClusterId}, requirements: {Requirements})", ex.ClusterId, ex.Requirements);
+						_logger.Log(logLevel, KnownLogEvents.Systemic_Horde_Compute, ex, "No agents found matching requirements (cluster: {ClusterId}, requirements: {Requirements})", ex.ClusterId, ex.Requirements);
 						shownNoAgentsFoundMessage = true;
 					}
 				}
+				catch (ComputeClientException cce)
+				{
+					_logger.Log(logLevel, KnownLogEvents.Systemic_Horde_Compute, cce, "{ErrorMessage}", cce.Message); 
+				}
 				catch (Exception ex)
 				{
 					if (!_cancellationSource!.IsCancellationRequested)
 					{
 						hordeSession.UpdateHordeStatus(" - " + ex.Message);
-						_logger.Log(_ubaConfig.bStrict ? LogLevel.Warning : LogLevel.Information, KnownLogEvents.Systemic_Horde_Compute, ex, "Unable to get worker: {Ex}", ex.ToString());
+						_logger.Log(logLevel, KnownLogEvents.Systemic_Horde_Compute, ex, "Unable to get worker: {Ex}", ex.ToString());
 					}
 				}
 
