[UBT]
* Fixed bug in buildmode where saving makefile was not waited on from previous pass when running new pass loading makefile

diff --git a/Engine/Source/Programs/UnrealBuildTool/Modes/BuildMode.cs b/Engine/Source/Programs/UnrealBuildTool/Modes/BuildMode.cs
--- a/Engine/Source/Programs/UnrealBuildTool/Modes/BuildMode.cs
+++ b/Engine/Source/Programs/UnrealBuildTool/Modes/BuildMode.cs
@@ -1254,7 +1254,7 @@ namespace UnrealBuildTool
 
 					// We need to track save makefile task so we can wait for it in case we need to load it again.
 					// TODO: Just reuse previous one. we don't need to save until we are on the last one.
-					Task?[] lastSaveMakefileTask = new Task?[descriptorCount]; // This must be allocated in here because it is iterated when all passes are done
+					ConcurrentDictionary<FileReference, Task> lastSaveMakefileTasks = []; // This must be allocated in here because it is iterated when all passes are done
 
 					List<UEBuildTarget.CodeGenContext>? lastContextsForCodeGen = null;
 
@@ -1307,8 +1307,8 @@ namespace UnrealBuildTool
 								else
 								{
 									// Need to make sure to wait on previous run's makefile
-									lastSaveMakefileTask[index]?.Wait();
-									lastSaveMakefileTask[index] = null;
+									lastSaveMakefileTasks.TryRemove(makefileLocation, out Task? lastMakefileSave);
+									lastMakefileSave?.Wait();
 									(makefile, ranPrebuildScripts) = LoadMakefile(makefileLocation, descriptor, workingSet, logger);
 								}
 							}
@@ -1408,7 +1408,8 @@ namespace UnrealBuildTool
 								buildTarget.BuildMakefile2(workingSet, logger, buildContext);
 
 								// Save the makefile for next time
-								if (makefileLocations[index] != null)
+								FileReference? makefileLocation = makefileLocations[index];
+								if (makefileLocation != null)
 								{
 									// Save the environment variables. Toolchains etc might set environment variables
 									// and we want to be able to restore them on runs where makefile is not invalidated.
@@ -1421,11 +1422,16 @@ namespace UnrealBuildTool
 										}
 									}
 
-									lastSaveMakefileTask[index] = Task.Run(() =>
+									Task saveTask = Task.Run(() =>
 									{
 										using ITimelineEvent _ = Timeline.ScopeEvent($"SavingMakefile");
-										makefile.Save(makefileLocations[index]!);
+										makefile.Save(makefileLocation!);
 									});
+
+									if (!lastSaveMakefileTasks.TryAdd(makefileLocation, saveTask))
+									{
+										throw new Exception("Can't add savemakefile task. This should never fail");
+									}
 								}
 							}
 							else
@@ -1472,7 +1478,7 @@ namespace UnrealBuildTool
 						}
 					}
 
-					foreach (Task? task in lastSaveMakefileTask)
+					foreach (Task? task in lastSaveMakefileTasks.Values)
 					{
 						if (task != null)
 						{
