[UBT]
* Improved makefile creation timeline scopes to not include waiting for previous pass make file creation

diff --git a/Engine/Source/Programs/UnrealBuildTool/Configuration/UEBuildTarget.cs b/Engine/Source/Programs/UnrealBuildTool/Configuration/UEBuildTarget.cs
--- a/Engine/Source/Programs/UnrealBuildTool/Configuration/UEBuildTarget.cs
+++ b/Engine/Source/Programs/UnrealBuildTool/Configuration/UEBuildTarget.cs
@@ -3359,7 +3359,10 @@ namespace UnrealBuildTool
 			Task.WaitAll(writeMetadataTasks.ToArray());
 
 			// Wait for makefile tasks to finish (tasks writing files to disk etc)
-			MakefileBuilder.WaitOnWriteTasks();
+			{
+				using var _ = Timeline.ScopeEvent("WaitOnWrites");
+				MakefileBuilder.WaitOnWriteTasks();
+			}
 
 			// Clean any stale modules which exist in multiple output directories. This can lead to the wrong DLL being loaded on Windows.
 			CleanStaleModules(Logger);
