UnrealBuildTool: Add missing public property comment

#rnx

diff --git a/Engine/Source/Programs/UnrealBuildTool/Configuration/BuildConfiguration.cs b/Engine/Source/Programs/UnrealBuildTool/Configuration/BuildConfiguration.cs
--- a/Engine/Source/Programs/UnrealBuildTool/Configuration/BuildConfiguration.cs
+++ b/Engine/Source/Programs/UnrealBuildTool/Configuration/BuildConfiguration.cs
@@ -264,6 +264,9 @@ namespace UnrealBuildTool
 		[XmlConfigFile(Category = "BuildConfiguration")]
 		public bool? bForceUnityBuild { get; set; }
 
+		/// <summary>
+		/// Whether to generate target makefiles in parallel.
+		/// </summary>
 		[XmlConfigFile(Category = "BuildConfiguration")]
 		public bool bParallelMakefileGeneration { get; set; } = true;
 
