[UBT]
* Fixed so CppDependencyCache handles Intermediate\External\Build folders

diff --git a/Engine/Source/Programs/UnrealBuildTool/System/CppDependencyCache.cs b/Engine/Source/Programs/UnrealBuildTool/System/CppDependencyCache.cs
--- a/Engine/Source/Programs/UnrealBuildTool/System/CppDependencyCache.cs
+++ b/Engine/Source/Programs/UnrealBuildTool/System/CppDependencyCache.cs
@@ -243,7 +243,9 @@ namespace UnrealBuildTool
 			}
 		}
 
-		readonly string SearchString = $"Intermediate{Path.DirectorySeparatorChar}Build";
+		readonly string Intermediate = $"{Path.DirectorySeparatorChar}Intermediate";
+		readonly string External = $"{Path.DirectorySeparatorChar}External";
+		readonly string Build = $"{Path.DirectorySeparatorChar}Build{Path.DirectorySeparatorChar}";
 
 		/// <summary>
 		/// Attempts to read dependencies from the given file.
@@ -261,7 +263,7 @@ namespace UnrealBuildTool
 			}
 
 			ReadOnlySpan<char> FullName = InputFile.FullName;
-			int index = FullName.IndexOf(SearchString, FileReference.Comparison);
+			int index = FullName.IndexOf(Intermediate, FileReference.Comparison);
 			if (index == -1)
 			{
 				Logger.LogWarning($"Path {InputFile} is not under intermediate folder");
@@ -269,8 +271,20 @@ namespace UnrealBuildTool
 				return false;
 			}
 
-			int start = index + SearchString.Length + 1;
-			ReadOnlySpan<char> filterPath = FullName.Slice(index + SearchString.Length + 1);
+			ReadOnlySpan<char> filterPath = FullName.Slice(index + Intermediate.Length);
+
+			if (filterPath.StartsWith(External))
+			{
+				filterPath = filterPath.Slice(External.Length);
+			}
+
+			if (!filterPath.StartsWith(Build))
+			{
+				Logger.LogWarning($"Path {InputFile} is not under intermediate build folder");
+				OutInfo = null;
+				return false;
+			}
+			filterPath = filterPath.Slice(Build.Length);
 
 			CachePartition? partition = null;
 
