[UBT]
* Fixed bug in CppDependencyCache for multi architecture targets

diff --git a/Engine/Source/Programs/UnrealBuildTool/System/CppDependencyCache.cs b/Engine/Source/Programs/UnrealBuildTool/System/CppDependencyCache.cs
--- a/Engine/Source/Programs/UnrealBuildTool/System/CppDependencyCache.cs
+++ b/Engine/Source/Programs/UnrealBuildTool/System/CppDependencyCache.cs
@@ -408,25 +408,25 @@ namespace UnrealBuildTool
 		/// <summary>
 		/// Creates a cache hierarchy for a particular target
 		/// </summary>
-		/// <param name="ProjectFile">Project file for the target being built</param>
-		/// <param name="AppName">Name of the app.</param>
-		/// <param name="Platform">Platform being built</param>
-		/// <param name="Configuration">Configuration being built</param>
-		/// <param name="Architectures">The target architectures</param>
-		/// <param name="IntermediateEnvironment">Intermediate environment to use</param>
-		/// <param name="Logger">Logger for output</param>
+		/// <param name="projectFile">Project file for the target being built</param>
+		/// <param name="appName">Name of the app.</param>
+		/// <param name="platform">Platform being built</param>
+		/// <param name="configuration">Configuration being built</param>
+		/// <param name="architectures">The target architectures</param>
+		/// <param name="intermediateEnvironment">Intermediate environment to use</param>
+		/// <param name="logger">Logger for output</param>
 		/// <returns>Dependency cache hierarchy for the given project</returns>
-		public void Mount(FileReference? ProjectFile, string AppName, UnrealTargetPlatform Platform, UnrealTargetConfiguration Configuration, UnrealArchitectures Architectures, UnrealIntermediateEnvironment IntermediateEnvironment, ILogger Logger)
+		public void Mount(FileReference? projectFile, string appName, UnrealTargetPlatform platform, UnrealTargetConfiguration configuration, UnrealArchitectures architectures, UnrealIntermediateEnvironment intermediateEnvironment, ILogger logger)
 		{
 			using ITimelineEvent ActionGraphTimer = Timeline.ScopeEvent($"CppDependencyCache");
 
-			AppName = UEBuildTarget.GetTargetIntermediateFolderName(AppName, IntermediateEnvironment);
+			appName = UEBuildTarget.GetTargetIntermediateFolderName(appName, intermediateEnvironment);
 
-			string GetFilterPath(UnrealTargetConfiguration Configuration)
+			string GetFilterPath(UnrealArch arch, UnrealTargetConfiguration configuration)
 			{
 				StringBuilder sb = new();
-				sb.Append(Platform).Append(Path.DirectorySeparatorChar).Append(UnrealArchitectureConfig.ForPlatform(Platform).GetFolderNameForArchitectures(Architectures)).Append(Path.DirectorySeparatorChar);
-				sb.Append(AppName).Append(Path.DirectorySeparatorChar).Append(Configuration);
+				sb.Append(platform).Append(Path.DirectorySeparatorChar).Append(UnrealArchitectureConfig.ForPlatform(platform).GetFolderNameForArchitecture(arch)).Append(Path.DirectorySeparatorChar);
+				sb.Append(appName).Append(Path.DirectorySeparatorChar).Append(configuration);
 				return sb.ToString();
 			}
 
@@ -459,22 +459,25 @@ namespace UnrealBuildTool
 						partition.Loaded = true;
 						if (FileReference.Exists(partition.Location))
 						{
-							partition.Read(Logger);
+							partition.Read(logger);
 						}
 					}
 				}
 
 			}
 
-			AddPartition(Unreal.WritableEngineDirectory, GetFilterPath(Configuration));
-			if (Configuration == UnrealTargetConfiguration.DebugGame)
+			foreach (UnrealArch arch in architectures.Architectures)
 			{
-				AddPartition(Unreal.WritableEngineDirectory, GetFilterPath(UnrealTargetConfiguration.Development));
-			}
+				AddPartition(Unreal.WritableEngineDirectory, GetFilterPath(arch, configuration));
+				if (configuration == UnrealTargetConfiguration.DebugGame)
+				{
+					AddPartition(Unreal.WritableEngineDirectory, GetFilterPath(arch, UnrealTargetConfiguration.Development));
+				}
 
-			if (ProjectFile != null)
-			{
-				AddPartition(ProjectFile.Directory, GetFilterPath(Configuration));
+				if (projectFile != null)
+				{
+					AddPartition(projectFile.Directory, GetFilterPath(arch, configuration));
+				}
 			}
 		}
 
