[UBT]
* Fixed bug where makefile location was not using appname which meant that targets that changed appname in rules ended up getting all intermediates except makefile in the right place. Code is now changed so makefile is preloaded as before since 99% of the time that is correct.. for the (rare) scenario where target rules change app name we just ignore loaded makefile and load the correct one

diff --git a/Engine/Source/Programs/UnrealBuildTool/Modes/BuildMode.cs b/Engine/Source/Programs/UnrealBuildTool/Modes/BuildMode.cs
--- a/Engine/Source/Programs/UnrealBuildTool/Modes/BuildMode.cs
+++ b/Engine/Source/Programs/UnrealBuildTool/Modes/BuildMode.cs
@@ -1155,6 +1155,8 @@ namespace UnrealBuildTool
 				UEBuildTarget[] buildTargets = new UEBuildTarget[descriptorCount];
 
 				// Start loading unique makefiles
+				// Note, there are (very) rare instances of this being the wrong makefile
+				// and that is when target rules change name of target. so there is a tiny chance we will waste some time loading a makefile that is not being used
 				Dictionary<FileReference, Task<(TargetMakefile?, bool)>> preLoadedMakefiles = [];
 				if (buildConfiguration.bUseUBTMakefiles)
 				{
@@ -1162,7 +1164,6 @@ namespace UnrealBuildTool
 					{
 						TargetDescriptor descriptor = descriptors[index];
 						FileReference makefileLocation = TargetMakefile.GetLocation(descriptor.ProjectFile, descriptor.Name, descriptor.Platform, descriptor.Architectures, descriptor.Configuration, descriptor.IntermediateEnvironment);
-						makefileLocations[index] = makefileLocation;
 						if (!preLoadedMakefiles.ContainsKey(makefileLocation))
 						{
 							preLoadedMakefiles.Add(makefileLocation, Task.Run(() => { return LoadMakefile(makefileLocation, descriptor, workingSet, logger); }));
@@ -1287,7 +1288,13 @@ namespace UnrealBuildTool
 							TargetMakefile? makefile = null;
 							bool ranPrebuildScripts = false;
 
-							FileReference? makefileLocation = makefileLocations[index];
+							UEBuildTarget buildTarget = buildTargets[index];
+							ReadOnlyTargetRules targetRules = buildTarget.Rules;
+							
+							// Calculating real makefile location. it is likely the same as the preload but not if appname is changed in target rules
+							FileReference? makefileLocation = TargetMakefile.GetLocation(targetRules.ProjectFile, targetRules.Name, targetRules.Platform, targetRules.Architectures, targetRules.Configuration, buildTarget.IntermediateEnvironment);
+							makefileLocations[index] = makefileLocation;
+
 							if (makefileLocation != null)
 							{
 								Task<(TargetMakefile?, bool)>? makefileTask = null;
@@ -1313,8 +1320,6 @@ namespace UnrealBuildTool
 
 							targetTimelines[index] = Timeline.ScopeEvent($"{descriptor.Name} {descriptor.Platform} {descriptor.Configuration}");
 
-							UEBuildTarget buildTarget = buildTargets[index];
-
 							List<UEBuildTarget.GenerateHeaderFunc> generatedHeaderFuncs = [];
 
 							// If we couldn't load a valid makefile, create a new one
@@ -1491,6 +1496,12 @@ namespace UnrealBuildTool
 				}, helperCount);
 
 				descriptors = newDescriptors;
+
+				// In the very rare scenario we preload the wrong makefiles we want to make sure to wait for them in case there are prebuildtargets that use them
+				foreach (Task<(TargetMakefile?, bool)> task in preLoadedMakefiles.Values)
+				{
+					task?.Wait();
+				}
 			}
 
 			return (outTargets, outTasks);
