[UBT]
* Changed so UHTExecutor is not creating one Timestamp file per module anymore but instead write it into one Timestamps file per target. Reason is because updating timestamps on 3000 files take much longer time than just rewriting one file. The only benefit with having one timestamps file per module is that in the case of targets sharing exact same files, then it might avoid running uht at first run (this never happens in practice anyway because for example Lyra and QAGame has some modules containing uht stuff that are unique per project and it is enough for one module to be too old for entire thing to be re-run)

diff --git a/Engine/Source/Programs/UnrealBuildTool/System/UHTExecution.cs b/Engine/Source/Programs/UnrealBuildTool/System/UHTExecution.cs
--- a/Engine/Source/Programs/UnrealBuildTool/System/UHTExecution.cs
+++ b/Engine/Source/Programs/UnrealBuildTool/System/UHTExecution.cs
@@ -1,14 +1,16 @@
 // Copyright Epic Games, Inc. All Rights Reserved.
 
 using System;
+using System.Buffers.Binary;
 using System.Collections;
 using System.Collections.Generic;
 using System.Data;
-using System.Diagnostics;
 using System.Diagnostics.CodeAnalysis;
 using System.IO;
+using System.IO.Hashing;
 using System.Linq;
 using System.Reflection;
+using System.Runtime.InteropServices;
 using System.Text;
 using System.Text.Json;
 using System.Threading;
@@ -812,11 +814,12 @@ namespace UnrealBuildTool
 		/// <summary>
 		/// Checks the class header files and determines if generated UObject code files are out of date in comparison.
 		/// </summary>
+		/// <param name="timestampsFile">Modules that we generate headers for</param>
 		/// <param name="UObjectModules">Modules that we generate headers for</param>
 		/// <param name="HeaderToolTimestampUtc">Timestamp for UHT</param>
 		/// <param name="Logger">Logger for output</param>
 		/// <returns>True if the code files are out of date</returns>
-		private static bool AreGeneratedCodeFilesOutOfDate(List<UHTModuleInfo> UObjectModules, DateTime HeaderToolTimestampUtc, ILogger Logger)
+		private static bool AreGeneratedCodeFilesOutOfDate(TimestampsFile timestampsFile, List<UHTModuleInfo> UObjectModules, DateTime HeaderToolTimestampUtc, ILogger Logger)
 		{
 			// Get CoreUObject.init.gen.cpp timestamp.  If the source files are older than the CoreUObject generated code, we'll
 			// need to regenerate code for the module
@@ -837,6 +840,10 @@ namespace UnrealBuildTool
 				}
 			}
 
+			// Load timestamps file
+			timestampsFile.Load(Logger);
+
+
 			// Also get the Version.h timestamp.  If it is newer than the last UHT run, it could affect preprocessor conditionals
 			// that check for the engine version
 			DateTime EngineVersionHeaderTimestampUtc = GetEngineVersionHeaderTimestampUtc();
@@ -872,23 +879,17 @@ namespace UnrealBuildTool
 				}
 				GeneratedCodeDirectory.CacheFiles();
 
-				// Grab our special "Timestamp" file that we saved after the last set of headers were generated.  This file
-				// contains the list of source files which contained UObjects, so that we can compare to see if any
-				// UObject source files were deleted (or no longer contain UObjects), which means we need to run UHT even
-				// if no other source files were outdated
-				FileItem TimestampFile = FileItem.GetItemByFileReference(FileReference.Combine(GeneratedCodeDirectory.Location, "Timestamp"));
-				if (!TimestampFile.Exists)
+				if (!timestampsFile.GetModuleInfo(Module.ModuleName, out UInt128 moduleOldHash, out long moduleLastUpdateUtc))
 				{
-					// Timestamp file was missing (possibly deleted/cleaned), so headers are out of date
 					if (Interlocked.Increment(ref invalidCounter) == 1)
 					{
-						Logger.LogDebug("UnrealHeaderTool needs to run because UHT Timestamp file did not exist for module {ModuleName}", Module.ModuleName);
+						Logger.LogDebug("UnrealHeaderTool needs to run because module {ModuleName} was not found in timestamps file", Module.ModuleName);
 					}
 					return;
 				}
 
 				// Make sure the last UHT run completed after UnrealHeaderTool.exe was compiled last, and after the CoreUObject headers were touched last.
-				DateTime SavedTimestampUtc = TimestampFile.LastWriteTimeUtc;
+				DateTime SavedTimestampUtc = DateTime.FromFileTimeUtc(moduleLastUpdateUtc);
 				if (HeaderToolTimestampUtc > SavedTimestampUtc)
 				{
 					// Generated code is older than UnrealHeaderTool.exe.  Out of date!
@@ -928,8 +929,19 @@ namespace UnrealBuildTool
 					return;
 				}
 
+				UInt128 moduleHash = CalculateHash(Module);
+
+				if (moduleHash != moduleOldHash)
+				{
+					if (Interlocked.Increment(ref invalidCounter) == 1)
+					{
+						Logger.LogDebug("UnrealHeaderTool needs to run because module {ModuleName} had different hash than what was found in timestamps file", Module.ModuleName);
+					}
+					return;
+				}
+
 				// Iterate over our UObjects headers and figure out if any of them have changed
-				List<FileItem> AllUObjectHeaders =
+				List <FileItem> AllUObjectHeaders =
 				[
 					.. Module.PublicUObjectClassesHeaders,
 					.. Module.PublicUObjectHeaders,
@@ -937,48 +949,6 @@ namespace UnrealBuildTool
 					.. Module.PrivateUObjectHeaders,
 				];
 
-				// Load up the old timestamp file and check to see if anything has changed
-				{
-					string[] TimestampLines = File.ReadAllLines(TimestampFile.FullName);
-					if (AllUObjectHeaders.Count != TimestampLines.Length)
-					{
-						if (Interlocked.Increment(ref invalidCounter) == 1)
-						{
-							Logger.LogDebug("UnrealHeaderTool needs to run because there are a different number of UObject source files in module {ModuleModuleName}", Module.ModuleName);
-						}
-						return;
-					}
-
-					// Iterate over our UObjects headers and figure out if any of them have changed
-					HashSet<string> ObjectHeadersSet = new(AllUObjectHeaders.Select(x => x.AbsolutePath), FileReference.Comparer);
-
-					// Prefix compression, look at write
-					string lastFilename = "";
-					StringBuilder sb = new(256);
-					foreach (string line in TimestampLines)
-					{
-						string fileName = line;
-
-						if (line[0] >= '0' && line[0] <= '9')
-						{
-							int space = line.IndexOf(' ');
-							int matchLength = int.Parse(line.AsSpan(0, space));
-							sb.Clear().Append(lastFilename, 0, matchLength).Append(line, space+1, line.Length - space - 1);
-							fileName = sb.ToString();
-						}
-						lastFilename = fileName;	
-
-						if (!ObjectHeadersSet.Contains(fileName))
-						{
-							if (Interlocked.Increment(ref invalidCounter) == 1)
-							{
-								Logger.LogDebug("UnrealHeaderTool needs to run because the set of UObject source files in module {ModuleModuleName} has changed ({fileName})", Module.ModuleName, fileName);
-							}
-							return;
-						}
-					}
-				}
-
 				foreach (FileItem HeaderFile in AllUObjectHeaders)
 				{
 					DateTime HeaderFileTimestampUtc = HeaderFile.LastWriteTimeUtc;
@@ -1028,81 +998,6 @@ namespace UnrealBuildTool
 			return false;
 		}
 
-		/// <summary>
-		/// Updates the recorded timestamps of all the passed in UObject modules
-		/// </summary>
-		private static void UpdateTimestamps(List<UHTModuleInfo> UObjectModules, bool resetDirectoryCaches)
-		{
-			Parallel2.ForEach(UObjectModules, Module =>
-			{
-				if (!Module.bIsReadOnly)
-				{
-					string GeneratedCodeDirectory = Path.GetDirectoryName(Module.GeneratedCPPFilenameBase)!;
-					DirectoryInfo GeneratedCodeDirectoryInfo = new(GeneratedCodeDirectory);
-
-					try
-					{
-						if (GeneratedCodeDirectoryInfo.Exists)
-						{
-							FileReference TimestampFile = FileReference.Combine(new DirectoryReference(GeneratedCodeDirectoryInfo.FullName), "Timestamp");
-
-
-							using var fs = new FileStream(TimestampFile.FullName, FileMode.Create, FileAccess.Write, FileShare.None, bufferSize: 1 << 18, options: FileOptions.SequentialScan);
-							using var sw = new StreamWriter(fs, new UTF8Encoding(encoderShouldEmitUTF8Identifier: false), bufferSize: 1 << 18);
-
-							string lastFullName = "";
-
-							// Use prefix compression (If first part of the line is a number, then the line is lastLine[0..number] + line-after-number)
-
-							void WriteHeaders(List<FileItem> headers)
-							{
-								headers.Sort(FileItem.SortByPath);
-
-								foreach (FileItem header in headers)
-								{
-									ReadOnlySpan<char> last = lastFullName;
-									ReadOnlySpan<char> curr = header.FullName.AsSpan();
-
-									int matchLength = MemoryExtensions.CommonPrefixLength(last, curr);
-
-									if (matchLength == 0)
-									{
-										sw.WriteLine(curr);
-									}
-									else
-									{
-										sw.Write(matchLength);
-										sw.Write(' ');
-										sw.WriteLine(curr.Slice(matchLength));
-									}
-									lastFullName = header.FullName;
-								}
-							}
-
-							// Save all of the UObject files to a timestamp file.  We'll load these on the next run to see if any new
-							// files with UObject classes were deleted, so that we'll know to run UHT even if the timestamps of all
-							// of the other source files were unchanged
-
-							WriteHeaders(Module.PublicUObjectClassesHeaders);
-							WriteHeaders(Module.PublicUObjectHeaders);
-							WriteHeaders(Module.InternalUObjectHeaders);
-							WriteHeaders(Module.PrivateUObjectHeaders);
-
-							if (resetDirectoryCaches)
-							{
-								// Because new .cpp and .h files may have been generated by UHT, invalidate the DirectoryLookupCache
-								DirectoryLookupCache.InvalidateCachedDirectory(new DirectoryReference(GeneratedCodeDirectoryInfo.FullName));
-							}
-						}
-					}
-					catch (Exception Exception)
-					{
-						throw new BuildException(Exception, "Couldn't write Timestamp file: " + Exception.Message);
-					}
-				}
-			});
-		}
-
 		/// <summary>
 		/// Return the file reference for the UHT manifest file
 		/// </summary>
@@ -1271,9 +1166,95 @@ namespace UnrealBuildTool
 			}
 		}
 
-		public static Task ExecuteHeaderToolIfNecessaryAsync(BuildConfiguration BuildConfiguration, FileReference? ProjectFile, TargetMakefile Makefile, string TargetName, ILogger Logger, bool resetDirectoryCaches = true)
+		static UInt128 CalculateHash(UHTModuleInfo mi)
 		{
-			return ExecuteHeaderToolIfNecessary(BuildConfiguration, ProjectFile, Makefile, TargetName, Logger, resetDirectoryCaches).timestamp;
+			XxHash128 xx = new();
+			List<FileItem> files = new(mi.PublicUObjectClassesHeaders.Concat(mi.PublicUObjectHeaders).Concat(mi.InternalUObjectHeaders).Concat(mi.PrivateUObjectHeaders));
+			files.Sort(FileItem.SortByPath);
+			foreach (FileItem file in files)
+			{
+				xx.Append(MemoryMarshal.AsBytes(file.FullName.AsSpan()));
+			}
+			byte[] hash16 = xx.GetHashAndReset();
+			ulong lo = BinaryPrimitives.ReadUInt64LittleEndian(hash16[..8]);
+			ulong hi = BinaryPrimitives.ReadUInt64LittleEndian(hash16[8..]);
+			return ((UInt128)hi << 64) | lo;
+		}
+
+		class TimestampsFile
+		{
+			internal TimestampsFile(FileReference location)
+			{
+				_location = location;
+			}
+
+			const uint CurrentVersion = 1;
+
+			internal void Load(ILogger logger)
+			{
+				if (FileReference.Exists(_location))
+				{
+					using FileStream stream = FileReference.Open(_location, FileMode.Open);
+					using BinaryReader reader = new(stream);
+					uint version = reader.ReadUInt32();
+					if (version != CurrentVersion)
+					{
+						logger.LogInformation("UHT timestamps file version mismatch");
+						return;
+					}
+
+					int moduleCount = reader.ReadInt32();
+					for (int i = 0; i < moduleCount; i++)
+					{
+						string moduleName = reader.ReadString();
+						ulong hashA = reader.ReadUInt64();
+						UInt128 hashB = reader.ReadUInt64();
+						UInt128 hash = (hashB << 64) + hashA;
+						long lastUpdateUtc = reader.ReadInt64();
+						_modules.Add(moduleName, (hash, lastUpdateUtc));
+					}
+				}
+			}
+
+			internal bool GetModuleInfo(string module, out UInt128 hash, out long lastUpdateUtc)
+			{
+				if (_modules.TryGetValue(module, out var entry))
+				{
+					hash = entry.hash;
+					lastUpdateUtc = entry.lastUpdateUtc;
+					return true;
+				}
+				hash = 0;
+				lastUpdateUtc = 0;
+				return false;
+			}
+
+			internal void Update(List<UHTModuleInfo> modules)
+			{
+				DateTime now = DateTime.Now;
+				foreach (UHTModuleInfo module in modules)
+				{
+					_modules[module.ModuleName] = (CalculateHash(module), now.Ticks);
+				}
+			}
+
+			internal void Save(ILogger logger)
+			{
+				using FileStream stream = FileReference.Open(_location, FileMode.Create, FileAccess.Write);
+				using BinaryWriter writer = new(stream);
+				writer.Write(CurrentVersion);
+				writer.Write(_modules.Count);
+				foreach (var kv in _modules)
+				{
+					writer.Write(kv.Key);
+					writer.Write((ulong)kv.Value.hash);
+					writer.Write((ulong)(kv.Value.hash >> 64));
+					writer.Write(kv.Value.lastUpdateUtc);
+				}
+			}
+
+			FileReference _location;
+			SortedDictionary<string, (UInt128 hash, long lastUpdateUtc)> _modules = [];
 		}
 
 		/// <summary>
@@ -1284,6 +1265,8 @@ namespace UnrealBuildTool
 		{
 			bool runTimestamp = false;
 
+			TimestampsFile timestampsFile = new(FileReference.Combine(Makefile.ProjectIntermediateDirectoryNoArch.ParentDirectory!, TargetName + ".uhttimestamps"));
+
 			Task runTask = new Task(() =>
 			{
 				if (ProgressWriter.bWriteMarkup)
@@ -1336,7 +1319,7 @@ namespace UnrealBuildTool
 					{
 						bUHTNeedsToRun = true;
 					}
-					else if (!bUHTNeedsToRun && AreGeneratedCodeFilesOutOfDate(Makefile.UObjectModules, CompositeTimestamp, Logger))
+					else if (!bUHTNeedsToRun && AreGeneratedCodeFilesOutOfDate(timestampsFile, Makefile.UObjectModules, CompositeTimestamp, Logger))
 					{
 						bUHTNeedsToRun = true;
 					}
@@ -1492,7 +1475,8 @@ namespace UnrealBuildTool
 					using (GlobalTracer.Instance.BuildSpan("UHTExecution.UpdateTimestamps()").StartActive())
 					{
 						using var uhtExecute = Timeline.ScopeEvent("UHTTimestamps");
-						UpdateTimestamps(Makefile.UObjectModules, resetDirectoryCaches);
+						timestampsFile.Update(Makefile.UObjectModules);
+						timestampsFile.Save(Logger);
 					}
 				}
 			});
