[UBT]
* Changed List<(key,value)> to Dictionary<key,value> to discover duplicates earlier

#FailedValidations Preflight

diff --git a/Engine/Source/Programs/UnrealBuildTool/Configuration/ReadOnlyTargetRules.cs b/Engine/Source/Programs/UnrealBuildTool/Configuration/ReadOnlyTargetRules.cs
--- a/Engine/Source/Programs/UnrealBuildTool/Configuration/ReadOnlyTargetRules.cs
+++ b/Engine/Source/Programs/UnrealBuildTool/Configuration/ReadOnlyTargetRules.cs
@@ -539,7 +539,7 @@ namespace UnrealBuildTool
 
 		public IReadOnlyList<string> ProjectDefinitions => Inner.ProjectDefinitions.AsReadOnly();
 
-		public IReadOnlyList<(string, string)> BuildSettingsVariables => Inner.BuildSettingsVariables.AsReadOnly();
+		public IReadOnlyDictionary<string, string> BuildSettingsVariables => Inner.BuildSettingsVariables.AsReadOnly();
 
 		public ReadOnlyAdditionalProperties AdditionalProperties => Inner.AdditionalProperties;
 
diff --git a/Engine/Source/Programs/UnrealBuildTool/Configuration/TargetRules.cs b/Engine/Source/Programs/UnrealBuildTool/Configuration/TargetRules.cs
--- a/Engine/Source/Programs/UnrealBuildTool/Configuration/TargetRules.cs
+++ b/Engine/Source/Programs/UnrealBuildTool/Configuration/TargetRules.cs
@@ -2648,7 +2648,8 @@ namespace UnrealBuildTool
 		/// <summary>
 		/// Variables that will be stored in BuildSettings module. Can be fetched with BuildSettings::GetVariable
 		/// </summary>
-		public List<(string key, string value)> BuildSettingsVariables { get; } = [];
+		[RequiresUniqueBuildEnvironment]
+		public Dictionary<string, string> BuildSettingsVariables { get; } = [];
 
 		/// <summary>
 		/// Environment variables. Should be used instead of Environment.Set/GetEnvironmentVariable in .build.cs
@@ -3095,6 +3096,10 @@ namespace UnrealBuildTool
 			{
 				bFieldsMatch = Enumerable.SequenceEqual((IEnumerable<string>)thisValue, (IEnumerable<string>)baseValue);
 			}
+			else if (typeof(IDictionary<string, string>).IsAssignableFrom(valueType))
+			{
+				bFieldsMatch = Enumerable.SequenceEqual((IDictionary<string, string>)thisValue, (IDictionary<string, string>)baseValue);
+			}
 			else
 			{
 				bFieldsMatch = thisValue.Equals(baseValue);
@@ -3114,6 +3119,11 @@ namespace UnrealBuildTool
 				IEnumerable<string> enumerable = (IEnumerable<string>)value;
 				return String.Join(',', (IEnumerable<string>)value);
 			}
+			else if (typeof(IDictionary<string, string>).IsAssignableFrom(valueType))
+			{
+				IDictionary<string, string> dictionary = (IDictionary<string, string>)value;
+				return String.Join(',', dictionary.Select(x => $"{x.Key}:{x.Value}"));
+			}
 			return value?.ToString() ?? String.Empty;
 		}
 
diff --git a/Engine/Source/Programs/UnrealBuildTool/Configuration/UEBuildModuleCPP.cs b/Engine/Source/Programs/UnrealBuildTool/Configuration/UEBuildModuleCPP.cs
--- a/Engine/Source/Programs/UnrealBuildTool/Configuration/UEBuildModuleCPP.cs
+++ b/Engine/Source/Programs/UnrealBuildTool/Configuration/UEBuildModuleCPP.cs
@@ -2647,18 +2647,19 @@ namespace UnrealBuildTool
 				}
 				Result.Definitions.Add($"UE_VFS_PATHS=\"{sb}\"");
 
-				sb.Clear();
-				foreach ((string key, string value) in Target.BuildSettingsVariables)
+				if (Target.BuildSettingsVariables.Count > 0)
 				{
-					if (key.Contains(';') || value.Contains(';'))
+					sb.Clear();
+					foreach ((string key, string value) in Target.BuildSettingsVariables.OrderBy(kv => kv.Key))
 					{
-						throw new Exception("BuildSettingsDefinitions can't contain keys or values with character ';'");
+						if (key.Contains(';') || value.Contains(';'))
+						{
+							throw new Exception("BuildSettingsDefinitions can't contain keys or values with character ';'");
+						}
+
+						sb.Append(key).Append(';').Append(value).Append(';'); ;
 					}
 
-					sb.Append(key).Append(';').Append(value).Append(';'); ;
-				}
-				if (sb.Length != 0)
-				{
 					Result.Definitions.Add($"UE_BUILDSETTINGS_VARIABLES=\"{sb}\"");
 				}
 			}
