Adds option to Visual Code project generator for Linux Debugger Type
<LinuxDebuggerType>
defaults to cppdbg

#rnx

diff --git a/Engine/Source/Programs/UnrealBuildTool/ProjectFiles/VisualStudioCode/VSCodeProjectFileGenerator.cs b/Engine/Source/Programs/UnrealBuildTool/ProjectFiles/VisualStudioCode/VSCodeProjectFileGenerator.cs
--- a/Engine/Source/Programs/UnrealBuildTool/ProjectFiles/VisualStudioCode/VSCodeProjectFileGenerator.cs
+++ b/Engine/Source/Programs/UnrealBuildTool/ProjectFiles/VisualStudioCode/VSCodeProjectFileGenerator.cs
@@ -69,6 +69,12 @@ namespace UnrealBuildTool
 		[CommandLine("-UseVSCodeExtension")]
 		private bool bUseVSCodeExtension = false;
 
+		/// <summary>
+		/// Debugger type to be used with Linux.  Defaults to cppdbg.
+		/// </summary>
+		[XmlConfigFile(Name = "LinuxDebuggerType")]
+		private string LinuxDebuggerType = "cppdbg";
+
 		private enum EPathType
 		{
 			Absolute,
@@ -1388,7 +1394,7 @@ namespace UnrealBuildTool
 					}
 					else if (HostPlatform == UnrealTargetPlatform.Linux)
 					{
-						OutFile.AddField("type", "cppdbg");
+						OutFile.AddField("type", LinuxDebuggerType);
 						OutFile.AddField("visualizerFile", MakeUnquotedPathString(FileReference.Combine(ProjectRoot, "Engine", "Extras", "VisualStudioDebugging", "Unreal.natvis"), EPathType.Absolute));
 						OutFile.AddField("showDisplayString", true);
 					}
