[UBT]
* Added so executors can be pre-created before they are needed.
* Changed so if UBA is used, executor is created earlier and start initialization in parallel (horde cache etc)
* Cleaned up the timeline scopes to more represent what is happening
* Added cmd line "-nocache" which disables all cache

diff --git a/Engine/Source/Programs/Shared/EpicGames.UBA/ITrace.cs b/Engine/Source/Programs/Shared/EpicGames.UBA/ITrace.cs
--- a/Engine/Source/Programs/Shared/EpicGames.UBA/ITrace.cs
+++ b/Engine/Source/Programs/Shared/EpicGames.UBA/ITrace.cs
@@ -58,6 +58,11 @@ namespace EpicGames.UBA
 				return new TraceNull();
 			}
 		}
+
+		/// <summary>
+		/// Get the global trace if there is one
+		/// </summary>
+		public static ITrace? GlobalTrace => TraceImpl.GlobalTraceImpl;
 	}
 
 	/// <summary>
diff --git a/Engine/Source/Programs/Shared/EpicGames.UBA/Impl/SessionServerImpl.cs b/Engine/Source/Programs/Shared/EpicGames.UBA/Impl/SessionServerImpl.cs
--- a/Engine/Source/Programs/Shared/EpicGames.UBA/Impl/SessionServerImpl.cs
+++ b/Engine/Source/Programs/Shared/EpicGames.UBA/Impl/SessionServerImpl.cs
@@ -163,9 +163,9 @@ namespace EpicGames.UBA.Impl
 			SessionServer_SetRemoteProcessAvailable(_handle, _remoteProcessSlotAvailableCallbackDelegate, 0);
 			SessionServer_SetRemoteProcessReturned(_handle, _remoteProcessReturnedCallbackDelegate);
 
-			if (TraceImpl.GlobalTrace != null)
+			if (ITrace.GlobalTrace != null)
 			{
-				_trace = TraceImpl.GlobalTrace;
+				_trace = ITrace.GlobalTrace;
 			}
 			else
 			{
@@ -193,7 +193,7 @@ namespace EpicGames.UBA.Impl
 			{
 			}
 
-			if (TraceImpl.GlobalTrace == null)
+			if (ITrace.GlobalTrace == null)
 			{
 				_trace.Dispose();
 			}
diff --git a/Engine/Source/Programs/Shared/EpicGames.UBA/Impl/TraceImpl.cs b/Engine/Source/Programs/Shared/EpicGames.UBA/Impl/TraceImpl.cs
--- a/Engine/Source/Programs/Shared/EpicGames.UBA/Impl/TraceImpl.cs
+++ b/Engine/Source/Programs/Shared/EpicGames.UBA/Impl/TraceImpl.cs
@@ -49,7 +49,7 @@ namespace EpicGames.UBA
 				try
 				{
 					Trace_SetGlobal(_handle);
-					GlobalTrace = this;
+					GlobalTraceImpl = this;
 				} catch (EntryPointNotFoundException)
 				{
 				}
@@ -57,7 +57,7 @@ namespace EpicGames.UBA
 
 		}
 
-		internal static TraceImpl? GlobalTrace { set; get; }
+		internal static TraceImpl? GlobalTraceImpl { set; get; }
 
 		#region IDisposable
 		~TraceImpl() => Dispose(false);
@@ -76,9 +76,9 @@ namespace EpicGames.UBA
 
 			if (_handle != IntPtr.Zero)
 			{
-				if (GlobalTrace == this)
+				if (GlobalTraceImpl == this)
 				{
-					GlobalTrace = null;
+					GlobalTraceImpl = null;
 				}
 
 				if (_ownsHandle)
diff --git a/Engine/Source/Programs/UnrealBuildTool/Executors/UnrealBuildAccelerator/UBAAgentCoordinator.cs b/Engine/Source/Programs/UnrealBuildTool/Executors/UnrealBuildAccelerator/UBAAgentCoordinator.cs
--- a/Engine/Source/Programs/UnrealBuildTool/Executors/UnrealBuildAccelerator/UBAAgentCoordinator.cs
+++ b/Engine/Source/Programs/UnrealBuildTool/Executors/UnrealBuildAccelerator/UBAAgentCoordinator.cs
@@ -22,11 +22,11 @@ namespace UnrealBuildTool
 	{
 		DirectoryReference? GetUBARootDir();
 
-		Task InitAsync(UBAExecutor executor);
+		Task InitAsync(UBAExecutor executor, StatusUpdateAction updateStatus, CancellationToken cancellationToken);
 
 		UnrealBuildAcceleratorCacheConfig? RequestCacheServer(CancellationToken cancellationToken);
 
-		void Start(IUbaAgentCoordinatorScheduler scheduler, Func<LinkedAction, bool> canRunRemotely, StatusUpdateAction updateStatus);
+		void Start(IUbaAgentCoordinatorScheduler scheduler, Func<LinkedAction, bool> canRunRemotely);
 
 		void Stop();
 
diff --git a/Engine/Source/Programs/UnrealBuildTool/Executors/UnrealBuildAccelerator/UBAAgentCoordinatorHorde.cs b/Engine/Source/Programs/UnrealBuildTool/Executors/UnrealBuildAccelerator/UBAAgentCoordinatorHorde.cs
--- a/Engine/Source/Programs/UnrealBuildTool/Executors/UnrealBuildAccelerator/UBAAgentCoordinatorHorde.cs
+++ b/Engine/Source/Programs/UnrealBuildTool/Executors/UnrealBuildAccelerator/UBAAgentCoordinatorHorde.cs
@@ -139,9 +139,12 @@ namespace UnrealBuildTool
 
 			_client = _serviceProvider.GetRequiredService<IHordeClient>().Compute;
 			_hordeHttpClient = _serviceProvider.GetRequiredService<IHordeClient>().CreateHttpClient();
-			GetServerInfoResponse serverInfo = await _hordeHttpClient.GetServerInfoAsync(cancellationToken);
-			_logger.LogInformation("Horde server: {ServerVersion}, agent: {AgentVersion}", serverInfo.ServerVersion, serverInfo.AgentVersion);
 
+			{
+				using ITimelineEvent _ = Timeline.ScopeEvent($"HordeConnect");
+				GetServerInfoResponse serverInfo = await _hordeHttpClient.GetServerInfoAsync(cancellationToken);
+				_logger.LogInformation("Horde server: {ServerVersion}, agent: {AgentVersion}", serverInfo.ServerVersion, serverInfo.AgentVersion);
+			}
 			DirectoryReference ubaDir = UBAExecutor.UbaBinariesDir;
 			List<string> agentFiles = new();
 			if (OperatingSystem.IsWindows())
@@ -174,6 +177,7 @@ namespace UnrealBuildTool
 			{
 				throw new PlatformNotSupportedException();
 			}
+
 			_ubaAgentLocator = await CreateToolAsync(ubaDir, agentFiles.Select(x => FileReference.Combine(ubaDir, x)), cancellationToken);
 			_logger.LogInformation("Created tool bundle with locator {UbaAgentLocator}", _ubaAgentLocator.ToString());
 		}
@@ -191,7 +195,9 @@ namespace UnrealBuildTool
 			}
 
 			var clusterId = new ClusterId(UnrealBuildAcceleratorHordeConfig.ClusterDefault);
+			using ITimelineEvent timelineEvent = Timeline.ScopeEvent($"HordeGetCache");
 			UbaConfig ubaConfig = _client.AllocateUbaCacheServerAsync(clusterId, cancellationToken).Result;
+			timelineEvent.Finish();
 			if (String.IsNullOrEmpty(ubaConfig.CacheEndpoint))
 			{
 				return null;
@@ -201,6 +207,8 @@ namespace UnrealBuildTool
 
 		async Task<BlobLocator> CreateToolAsync(DirectoryReference baseDir, IEnumerable<FileReference> files, CancellationToken cancellationToken)
 		{
+			using ITimelineEvent __ = Timeline.ScopeEvent($"HordeCreateTool");
+
 			// TODO: should drive this off API version reported by server
 			BlobSerializerOptions serializerOptions = BlobSerializerOptions.Create(HordeApiVersion.Initial);
 
@@ -482,7 +490,7 @@ namespace UnrealBuildTool
 
 				UBAHordeSession session = new(ubaConfig, crypto, server, token, hordeConfig.HordePool, allowWine, hordeConfig.HordeMaxCores, bStrictErrors, connectionMode, encryption, logger);
 				await session.InitAsync(useSentry: !String.IsNullOrEmpty(hordeConfig.UBASentryUrl), cancellationToken);
-				return (session, "");
+				return (session, "Ready");
 			}
 			catch (TaskCanceledException)
 			{
@@ -809,15 +817,29 @@ namespace UnrealBuildTool
 			return null;
 		}
 
-		public async Task InitAsync(UBAExecutor executor)
+		public async Task InitAsync(UBAExecutor executor, StatusUpdateAction updateStatus, CancellationToken cancellationToken)
 		{
+			_updateStatus = updateStatus;
+			_cancellationToken = cancellationToken;
+
 			if (_ubaConfig.bDisableRemote)
 			{
 				return;
 			}
 
-			_cancellationSource = new CancellationTokenSource();
-			_hordeSessionTask = UBAHordeSession.TryCreateHordeSessionAsync(_hordeConfig, _ubaConfig, executor.Crypto, _ubaConfig.bStrict, _logger, _cancellationSource.Token);
+			string? link = null;
+			if (_hordeConfig.HordeServer != null)
+			{
+				link = $"{_hordeConfig.HordeServer!.TrimEnd('/')}/agents";
+				if (_hordeConfig.HordePool != string.Empty)
+				{
+					link += $"?agent={_hordeConfig.HordePool}";
+				}
+			}
+			_updateStatus(0, 1, _title, EpicGames.UBA.LogEntryType.Info, link);
+			_updateStatus!(0, 6, "Connecting", EpicGames.UBA.LogEntryType.Info, null);
+
+			_hordeSessionTask = UBAHordeSession.TryCreateHordeSessionAsync(_hordeConfig, _ubaConfig, executor.Crypto, _ubaConfig.bStrict, _logger, cancellationToken);
 			await _hordeSessionTask;
 			executor.AgentCoordinatorInitialized(this, _hordeSessionTask.IsCompletedSuccessfully && _hordeSessionTask.Result.Item1 != null);
 		}
@@ -834,6 +856,7 @@ namespace UnrealBuildTool
 			if (_hordeSessionTask != null)
 			{
 				(hordeSession, reason) = _hordeSessionTask.Result;
+				_updateStatus!(0, 6, reason, EpicGames.UBA.LogEntryType.Info, null);
 			}
 			if (hordeSession == null)
 			{
@@ -843,13 +866,12 @@ namespace UnrealBuildTool
 			return hordeSession.RequestCacheServer(cancellationToken, _hordeConfig.CacheDesiredConnectionCount);
 		}
 
-		public void Start(IUbaAgentCoordinatorScheduler scheduler, Func<LinkedAction, bool> canRunRemotely, StatusUpdateAction updateStatus)
+		public void Start(IUbaAgentCoordinatorScheduler scheduler, Func<LinkedAction, bool> canRunRemotely)
 		{
 			if (!_hordeConfig.bAgentsEnabled)
 			{
 				return;
 			}
-			_updateStatus = updateStatus;
 
 			int timerPeriod = 5000;
 			bool shownNoAgentsFoundMessage = false;
@@ -861,20 +883,9 @@ namespace UnrealBuildTool
 				(hordeSession, reason) = _hordeSessionTask.Result;
 			}
 
-			string? link = null;
-			if (_hordeConfig.HordeServer != null)
-			{
-				link = $"{_hordeConfig.HordeServer!.TrimEnd('/')}/agents";
-				if (_hordeConfig.HordePool != string.Empty)
-				{
-					link += $"?agent={_hordeConfig.HordePool}";
-				}
-			}
-			updateStatus(0, 1, _title, EpicGames.UBA.LogEntryType.Info, link);
-
 			if (hordeSession == null)
 			{
-				updateStatus(0, 6, reason, EpicGames.UBA.LogEntryType.Info, null);
+				_updateStatus!(0, 6, reason, EpicGames.UBA.LogEntryType.Info, null);
 				return;
 			}
 
@@ -900,7 +911,7 @@ namespace UnrealBuildTool
 					return;
 				}
 
-				if (scheduler.IsEmpty || _cancellationSource!.IsCancellationRequested)
+				if (scheduler.IsEmpty || _cancellationToken.IsCancellationRequested)
 				{
 					hordeSession.UpdateHordeStatus(" - Requests stopped");
 					return;
@@ -914,7 +925,7 @@ namespace UnrealBuildTool
 				{
 					while (true)
 					{
-						if (_cancellationSource!.IsCancellationRequested)
+						if (_cancellationToken.IsCancellationRequested)
 						{
 							hordeSession.UpdateHordeStatus(" - Requests stopped");
 							return;
@@ -932,7 +943,7 @@ namespace UnrealBuildTool
 
 						queueWeight -= queuedCores;
 
-						bool isSatisfied = queueWeight <= queueThreshold || (activeCores + queuedCores) >= _hordeConfig.HordeMaxCores || _cancellationSource!.IsCancellationRequested;
+						bool isSatisfied = queueWeight <= queueThreshold || (activeCores + queuedCores) >= _hordeConfig.HordeMaxCores || _cancellationToken.IsCancellationRequested;
 
 						hordeSession.UpdateHordeStatus(isSatisfied ? " - Requests paused" : $" - Requesting agent{("...."[..(requestCounter++ % 4)])}");
 
@@ -942,7 +953,7 @@ namespace UnrealBuildTool
 						}
 
 						Requirements requirements = GetRequirements(_hordeConfig);
-						if (!await hordeSession.AddWorkerAsync(requirements, _hordeConfig, _cancellationSource.Token, activeCores))
+						if (!await hordeSession.AddWorkerAsync(requirements, _hordeConfig, _cancellationToken, activeCores))
 						{
 							await Task.Delay(500); // Sleep a little bit to make sure previous horde status is visible
 							string status = (queuedAgents + activeAgents) == 0 ? " - No agents available" : " - No additional agents available";
@@ -966,7 +977,7 @@ namespace UnrealBuildTool
 				}
 				catch (Exception ex)
 				{
-					if (!_cancellationSource!.IsCancellationRequested)
+					if (!_cancellationToken.IsCancellationRequested)
 					{
 						hordeSession.UpdateHordeStatus(" - " + ex.Message);
 						_logger.Log(logLevel, KnownLogEvents.Systemic_Horde_Compute, ex, "Unable to get worker: {Ex}", ex.ToString());
@@ -979,12 +990,18 @@ namespace UnrealBuildTool
 
 		public void Stop()
 		{
-			_cancellationSource?.Cancel();
+			if (!_cancellationToken.IsCancellationRequested)
+			{
+				_logger.LogWarning("Cancellation should have been set");
+			}
 		}
 
 		public async Task CloseAsync()
 		{
-			_cancellationSource?.Cancel();
+			if (!_cancellationToken.IsCancellationRequested)
+			{
+				_logger.LogWarning("Cancellation should have been set");
+			}
 
 			if (_hordeSessionTask == null)
 			{
@@ -1019,8 +1036,6 @@ namespace UnrealBuildTool
 		{
 			if (disposing)
 			{
-				_cancellationSource?.Dispose();
-				_cancellationSource = null;
 				_timer?.Dispose();
 				_timer = null;
 			}
@@ -1074,7 +1089,7 @@ namespace UnrealBuildTool
 		readonly UnrealBuildAcceleratorConfig _ubaConfig;
 		readonly UnrealBuildAcceleratorHordeConfig _hordeConfig;
 
-		CancellationTokenSource? _cancellationSource;
+		CancellationToken _cancellationToken;
 		Task<(UBAHordeSession?, string)>? _hordeSessionTask;
 		StatusUpdateAction? _updateStatus;
 		Timer? _timer;
diff --git a/Engine/Source/Programs/UnrealBuildTool/Executors/UnrealBuildAccelerator/UBAExecutor.cs b/Engine/Source/Programs/UnrealBuildTool/Executors/UnrealBuildAccelerator/UBAExecutor.cs
--- a/Engine/Source/Programs/UnrealBuildTool/Executors/UnrealBuildAccelerator/UBAExecutor.cs
+++ b/Engine/Source/Programs/UnrealBuildTool/Executors/UnrealBuildAccelerator/UBAExecutor.cs
@@ -23,11 +23,21 @@ namespace UnrealBuildTool
 		public UnrealBuildAcceleratorConfig UBAConfig { get; init; } = new();
 
 		public string Crypto { get; private set; } = String.Empty;
-		ISessionServer? _session;
-		object _sessionLock = new();
-		IEnumerable<UnrealBuildAcceleratorCacheConfig> _cacheConfigs = [];
-		Dictionary<UnrealBuildAcceleratorCacheConfig, ICacheClient> _cacheClients = [];
+
+		readonly IEnumerable<TargetDescriptor> _targetDescriptors;
+		readonly Task _initTask;
+		readonly CancellationTokenSource _cancellationSource;
+
 		EpicGames.UBA.ILogger? _ubaLogger;
+		IServer? _server;
+		IStorageServer? _storageServer;
+		ISessionServer? _session;
+		ITrace? _trace;
+		string? _sessionHyperlink;
+
+		readonly object _sessionLock = new();
+
+		readonly Dictionary<UnrealBuildAcceleratorCacheConfig, ICacheClient> _cacheClients = [];
 		readonly List<IUBAAgentCoordinator> _agentCoordinators = new();
 		DirectoryReference? _rootDirRef;
 		bool _bIsCancelled;
@@ -49,19 +59,39 @@ namespace UnrealBuildTool
 		{
 			if (disposing)
 			{
-				foreach (KeyValuePair<UnrealBuildAcceleratorCacheConfig, ICacheClient> item in _cacheClients)
+				_initTask.Wait();
+
 				{
-					item.Value.Dispose();
+					using ITimelineEvent a2 = Timeline.ScopeEvent("DisconnectCache");
+					DisposeCacheClients();
 				}
-				_cacheClients.Clear();
+
 				_session?.Dispose();
-				_session = null;
+
+				{
+					using ITimelineEvent a4 = Timeline.ScopeEvent("SaveCasTable");
+					_storageServer?.Dispose();
+				}
+
+				_server?.Dispose();
 				_ubaLogger?.Dispose();
-				_ubaLogger = null;
 			}
 			base.Dispose(disposing);
 		}
 
+		void DisposeCacheClients()
+		{
+			if (_cacheClients.Count > 0)
+			{
+				foreach ((UnrealBuildAcceleratorCacheConfig config, ICacheClient client) in _cacheClients)
+				{
+					client.Dispose();
+				}
+				_cacheClients.Clear();
+			}
+		}
+
+
 		public override string Name => !UBAConfig.bDisableRemote && _agentCoordinators.Count > 0 ? "Unreal Build Accelerator" : "Unreal Build Accelerator local";
 
 		public static new bool IsAvailable()
@@ -103,26 +133,71 @@ namespace UnrealBuildTool
 			}
 		}
 
-		public UBAExecutor(int maxLocalActions, bool bAllCores, bool bCompactOutput, Microsoft.Extensions.Logging.ILogger logger, IEnumerable<TargetDescriptor> targetDescriptors)
+		public UBAExecutor(int maxLocalActions, bool bAllCores, bool bCompactOutput, ILogger logger, IEnumerable<TargetDescriptor> targetDescriptors)
 			: base(maxLocalActions, bAllCores, bCompactOutput, logger)
+		{
+			_targetDescriptors = targetDescriptors;
+			_cancellationSource = new();
+			_initTask = new(() => { Init(targetDescriptors, logger); }, TaskCreationOptions.LongRunning);
+			_initTask.Start();
+		}
+
+		private void Init(IEnumerable<TargetDescriptor> targetDescriptors, ILogger logger)
 		{
 			XmlConfig.ApplyTo(this);
 			XmlConfig.ApplyTo(UBAConfig);
 			CommandLine.ParseArguments(Environment.GetCommandLineArgs(), this, logger);
 
-			bool bAnyHotReload = false;
 			foreach (TargetDescriptor targetDescriptor in targetDescriptors)
 			{
 				ConfigCache.ReadSettings(targetDescriptor.ProjectFile?.Directory, BuildHostPlatform.Current.Platform, UBAConfig);
-				bAnyHotReload = bAnyHotReload || targetDescriptor.HotReloadMode != HotReloadMode.Disabled;
 				targetDescriptor.AdditionalArguments.ApplyTo(this);
 				targetDescriptor.AdditionalArguments.ApplyTo(UBAConfig);
 			}
 
-			// Disable cache server when using hot reload path (there would likely never be any cache hits anyway
-			if (!bAnyHotReload)
+			if (UBAConfig.bDisableRemote)
+			{
+				UBAConfig.Zone = "local";
+			}
+
+			if (UBAConfig.bForcedRetry)
+			{
+				UBAConfig.bAllowRetry = true;
+			}
+
+			_trace = ITrace.GlobalTrace;
+
+			List<UnrealBuildAcceleratorCacheConfig> cacheConfigs = [];
+
+			if (!UBAConfig.bForceNoCache)
 			{
-				LoadCacheConfigs(targetDescriptors, logger);
+				const string ProviderPrefix = "Uba.CacheProvider";
+
+				IEnumerable<string> providers = Unreal.IsBuildMachine()
+				? [.. UBAConfig.BuildMachineCacheProviders, .. UBAConfig.IniBuildMachineCacheProviders]
+				: [.. UBAConfig.CacheProviders, .. UBAConfig.IniCacheProviders];
+
+				if (!providers.Any())
+				{
+					providers = [ProviderPrefix];
+				}
+				foreach (string provider in providers.Distinct().Where(x => x.StartsWith(ProviderPrefix)))
+				{
+					foreach (TargetDescriptor targetDescriptor in targetDescriptors)
+					{
+						UnrealBuildAcceleratorCacheConfig cacheConfig = new();
+						if (provider == ProviderPrefix)
+						{
+							XmlConfig.ApplyTo(cacheConfig);
+							CommandLine.ParseArguments(Environment.GetCommandLineArgs(), this, logger);
+						}
+						ConfigHierarchy engineIni = ConfigCache.ReadHierarchy(ConfigHierarchyType.Engine, targetDescriptor.ProjectFile?.Directory, BuildHostPlatform.Current.Platform);
+						cacheConfig.LoadConfigProvider(engineIni, provider);
+						targetDescriptor.AdditionalArguments.ApplyTo(cacheConfig);
+						cacheConfigs.Add(cacheConfig);
+					}
+				}
+				cacheConfigs = cacheConfigs.DistinctBy(x => x.CacheServer).Where(x => !String.IsNullOrEmpty(x.CacheServer)).ToList();
 			}
 
 			if (UBAConfig.bUseCrypto)
@@ -131,48 +206,153 @@ namespace UnrealBuildTool
 			}
 
 			_agentCoordinators.AddRange(UBAAgentCoordinatorHorde.Init(logger, targetDescriptors, UBAConfig, ref _remoteConnectionMode));
-		}
 
-		private void LoadCacheConfigs(IEnumerable<TargetDescriptor> targetDescriptors, Microsoft.Extensions.Logging.ILogger logger)
-		{
-			const string ProviderPrefix = "Uba.CacheProvider";
+			if (!String.IsNullOrEmpty(UBAConfig.RootDir))
+			{
+				_rootDirRef = DirectoryReference.FromString(UBAConfig.RootDir);
+			}
+			else
+			{
+				_rootDirRef = DirectoryReference.FromString(Environment.GetEnvironmentVariable("UBA_ROOT") ?? Environment.GetEnvironmentVariable("BOX_ROOT"));
+			}
+
+			if (_rootDirRef == null)
+			{
+				foreach (IUBAAgentCoordinator coordinator in _agentCoordinators)
+				{
+					_rootDirRef = _rootDirRef == null ? coordinator.GetUBARootDir() : _rootDirRef;
+				}
+			}
+
+			if (_rootDirRef == null)
+			{
+				if (OperatingSystem.IsWindows())
+				{
+					_rootDirRef = DirectoryReference.Combine(DirectoryReference.GetSpecialFolder(Environment.SpecialFolder.CommonApplicationData)!, "Epic", "UnrealBuildAccelerator");
+				}
+				else
+				{
+					_rootDirRef = DirectoryReference.Combine(DirectoryReference.GetSpecialFolder(Environment.SpecialFolder.UserProfile)!, ".epic", "UnrealBuildAccelerator");
+				}
+			}
+
+			DirectoryReference.CreateDirectory(_rootDirRef);
+
+			FileReference ubaTraceFile;
+			if (!String.IsNullOrEmpty(UBAConfig.TraceFile))
+			{
+				ubaTraceFile = new FileReference(UBAConfig.TraceFile);
+			}
+			else if (Unreal.IsBuildMachine())
+			{
+				ubaTraceFile = FileReference.Combine(Unreal.EngineProgramSavedDirectory, "AutomationTool", "Saved", "Logs", "Trace.uba");
+			}
+			else if (Log.OutputFile != null)
+			{
+				ubaTraceFile = Log.OutputFile.ChangeExtension(".uba");
+			}
+			else
+			{
+				ubaTraceFile = FileReference.Combine(Unreal.EngineProgramSavedDirectory, "UnrealBuildTool", "Trace.uba");
+			}
+
+			DirectoryReference.CreateDirectory(ubaTraceFile.Directory);
+			Log.BackupLogFile(ubaTraceFile);
+
+			FileReference configFile = FileReference.Combine(UbaBinariesDir, "UbaHost.toml");
+			IConfig config = EpicGames.UBA.IConfig.LoadConfig(configFile.FullName);
+			if (!String.IsNullOrEmpty(UBAConfig.CompressionLevel))
+			{
+				config.AddValue("Storage", "CompressionLevel", UBAConfig.CompressionLevel);
+			}
+
+			_ubaLogger = EpicGames.UBA.ILogger.CreateLogger(logger);
+			_server = IServer.CreateServer(UBAConfig.MaxWorkers, UBAConfig.SendSize, _ubaLogger, UBAConfig.bUseQuic);
+			_storageServer = IStorageServer.CreateStorageServer(_server, _ubaLogger, new StorageServerCreateInfo(_rootDirRef.FullName, ((ulong)UBAConfig.StoreCapacityGb) * 1000 * 1000 * 1000, !UBAConfig.bStoreRaw, UBAConfig.Zone));
+			using ISessionServerCreateInfo serverCreateInfo = ISessionServerCreateInfo.CreateSessionServerCreateInfo(_storageServer, _server, _ubaLogger, new SessionServerCreateInfo(_rootDirRef.FullName, ubaTraceFile.FullName.Replace('\\', '/'), UBAConfig.bDisableCustomAlloc, false, UBAConfig.bResetCas, UBAConfig.bWriteToDisk, UBAConfig.bDetailedTrace, !UBAConfig.bDisableWaitOnMem, UBAConfig.bAllowKillOnMem, UBAConfig.bStoreObjFilesCompressed));
+			_session = ISessionServer.CreateSessionServer(serverCreateInfo);
+			_trace = _session.GetTrace();
+
+			string? jobId = Environment.GetEnvironmentVariable("UE_HORDE_JOBID");
+			if (jobId != null)
+			{
+				string hordeUrl = Environment.GetEnvironmentVariable("UE_HORDE_URL") ?? "";
+				string stepId = Environment.GetEnvironmentVariable("UE_HORDE_STEPID") ?? "unknown";
+				string streamId = Environment.GetEnvironmentVariable("UE_HORDE_STREAMID") ?? "unknown-stream";
+				_sessionHyperlink = $"{hordeUrl}job/{jobId}?step={stepId} ({streamId})";
 
-			IEnumerable<string> providers = Unreal.IsBuildMachine()
-			? [.. UBAConfig.BuildMachineCacheProviders, .. UBAConfig.IniBuildMachineCacheProviders]
-			: [.. UBAConfig.CacheProviders, .. UBAConfig.IniCacheProviders];
+				_session.AddInfo(_sessionHyperlink);
+			}
 
-			if (!providers.Any())
+			List<Task> coordinatorInitTasks = new();
+			if (!UBAConfig.bDisableRemote)
 			{
-				providers = [ProviderPrefix];
+				uint statusUpdateCounter = 10;
+				foreach (IUBAAgentCoordinator coordinator in _agentCoordinators)
+				{
+					uint statusUpdateIndex = statusUpdateCounter;
+					coordinatorInitTasks.Add(coordinator.InitAsync(this, (sr, sc, st, t, sl) => UpdateStatus(statusUpdateIndex + sr, sc, st, t, sl), _cancellationSource.Token));
+					statusUpdateCounter += 10;
+				}
+
+				if (!UBAConfig.bForceNoCache)
+				{
+					foreach (IUBAAgentCoordinator coordinator in _agentCoordinators)
+					{
+						UnrealBuildAcceleratorCacheConfig? cacheConfig = coordinator.RequestCacheServer(_cancellationSource.Token);
+						if (cacheConfig != null)
+							cacheConfigs.Add(cacheConfig);
+					}
+				}
 			}
 
-			List<UnrealBuildAcceleratorCacheConfig> configs = [];
-			foreach (string provider in providers.Distinct().Where(x => x.StartsWith(ProviderPrefix)))
+			if (cacheConfigs.Count > 0)
 			{
-				foreach (TargetDescriptor targetDescriptor in targetDescriptors)
+				_trace.UpdateStatus(1, 1, "Cache", LogEntryType.Info, null);
+				_trace.UpdateStatus(1, 6, "Connecting...", LogEntryType.Info);
+
+				using TraceTaskScope s = new(_trace, "CacheConnect", "");
+				foreach (UnrealBuildAcceleratorCacheConfig cacheConfig in cacheConfigs)
 				{
-					UnrealBuildAcceleratorCacheConfig config = new UnrealBuildAcceleratorCacheConfig();
-					if (provider == ProviderPrefix)
+					ICacheClient cacheClient = ICacheClient.CreateCacheClient(_session, UBAConfig.bReportCacheMissReason, cacheConfig.Crypto, _sessionHyperlink ?? Environment.MachineName);
+					string[] nameAndPort = cacheConfig.CacheServer.Split(':');
+					int port = 1347;
+					if (nameAndPort.Length > 1)
 					{
-						XmlConfig.ApplyTo(config);
-						CommandLine.ParseArguments(Environment.GetCommandLineArgs(), this, logger);
+						port = Int32.Parse(nameAndPort[1]);
+					}
+
+					System.Diagnostics.Stopwatch stopwatch = System.Diagnostics.Stopwatch.StartNew();
+					bool cacheSuccess = cacheClient.Connect(nameAndPort[0], port, cacheConfig.DesiredConnectionCount);
+					long totalMs = stopwatch.ElapsedMilliseconds;
+					string successText = cacheSuccess ? "Connected to" : "Failed to connect to";
+					logger.LogInformation("UbaCache - {SuccessText} {Name}:{Port} ({Seconds}.{Milliseconds}s)", successText, nameAndPort[0], port, totalMs / 1000, totalMs % 1000);
+					if (cacheSuccess)
+					{
+						_cacheClients.Add(cacheConfig, cacheClient);
+					}
+					else
+					{
+						cacheClient.Dispose();
 					}
-					ConfigHierarchy engineIni = ConfigCache.ReadHierarchy(ConfigHierarchyType.Engine, targetDescriptor.ProjectFile?.Directory, BuildHostPlatform.Current.Platform);
-					config.LoadConfigProvider(engineIni, provider);
-					targetDescriptor.AdditionalArguments.ApplyTo(config);
-					configs.Add(config);
 				}
+				_trace.UpdateStatus(1, 6, _cacheClients.Any() ? "Ready" : "Failed to connect", LogEntryType.Info);
+			}
+
+			if (!UBAConfig.bDisableRemote)
+			{
+				_server.StartServer(UBAConfig.Host, UBAConfig.Port, Crypto);
+			}
+			else
+			{
+				_session.DisableRemoteExecution();
 			}
 
-			_cacheConfigs = configs.DistinctBy(x => x.CacheServer).Where(x => !String.IsNullOrEmpty(x.CacheServer));
 		}
 
 		public void UpdateStatus(uint statusRow, uint statusColumn, string statusText, LogEntryType statusType, string? statusLink)
 		{
-			lock (_sessionLock)
-			{
-				_session?.GetTrace().UpdateStatus(statusRow, statusColumn, statusText, statusType, statusLink);
-			}
+			_trace?.UpdateStatus(statusRow, statusColumn, statusText, statusType, statusLink);
 		}
 
 		private void PrintConfiguration(ILogger logger)
@@ -269,96 +449,40 @@ namespace UnrealBuildTool
 
 		public override async Task<bool> ExecuteActionsAsync(IEnumerable<LinkedAction> inputActions, Microsoft.Extensions.Logging.ILogger logger, IActionArtifactCache? actionArtifactCache)
 		{
-			if (!inputActions.Any())
-			{
-				return true;
-			}
-
-			// TODO: We still want to connect to coordinators because it might be that that input action is an action that spawns tons of internal processes (distributed thinlto linking)
-			if (inputActions.Count() < NumParallelProcesses && !UBAConfig.bForceBuildAllRemote)
-			{
-				UBAConfig.bDisableRemote = true;
-			}
-
-			if (UBAConfig.bDisableRemote)
-			{
-				UBAConfig.Zone = "local";
-			}
-
-			if (UBAConfig.bForcedRetry)
+			// Disable cache server when using hot reload path (there would likely never be any cache hits anyway)
+			// Note we can't do this earlier because HotReloadMode is not set until actions are created
+			foreach (TargetDescriptor targetDescriptor in _targetDescriptors)
 			{
-				UBAConfig.bAllowRetry = true;
+				UBAConfig.bForceNoCache |= targetDescriptor.HotReloadMode != HotReloadMode.Disabled;
 			}
 
-			_cacheableActionsLeft = inputActions.Count(x => x.ArtifactMode.HasFlag(ArtifactMode.Enabled));
-
-			await WriteActionOutputFileAsync(logger, inputActions);
-
-			if (!String.IsNullOrEmpty(UBAConfig.RootDir))
-			{
-				_rootDirRef = DirectoryReference.FromString(UBAConfig.RootDir);
-			}
-			else
-			{
-				_rootDirRef = DirectoryReference.FromString(Environment.GetEnvironmentVariable("UBA_ROOT") ?? Environment.GetEnvironmentVariable("BOX_ROOT"));
-			}
-
-			if (_rootDirRef == null)
-			{
-				foreach (IUBAAgentCoordinator coordinator in _agentCoordinators)
-				{
-					_rootDirRef = _rootDirRef == null ? coordinator.GetUBARootDir() : _rootDirRef;
-				}
-			}
-
-			if (_rootDirRef == null)
+			if (!inputActions.Any())
 			{
-				if (OperatingSystem.IsWindows())
-				{
-					_rootDirRef = DirectoryReference.Combine(DirectoryReference.GetSpecialFolder(Environment.SpecialFolder.CommonApplicationData)!, "Epic", "UnrealBuildAccelerator");
-				}
-				else
-				{
-					_rootDirRef = DirectoryReference.Combine(DirectoryReference.GetSpecialFolder(Environment.SpecialFolder.UserProfile)!, ".epic", "UnrealBuildAccelerator");
-				}
+				_cancellationSource.Cancel();
 			}
 
-			List<Task> coordinatorInitTasks = new();
-			if (!UBAConfig.bDisableRemote)
-			{
-				foreach (IUBAAgentCoordinator coordinator in _agentCoordinators)
-				{
-					// Do not await here, it will block local tasks from running
-					coordinatorInitTasks.Add(coordinator.InitAsync(this));
-				}
-			}
+			_initTask.Wait();
 
 			PrintConfiguration(logger);
 
-			DirectoryReference.CreateDirectory(_rootDirRef);
+			_trace?.UpdateStatus(1, 6, "Hits 0 Misses 0", LogEntryType.Info);
 
-			FileReference ubaTraceFile;
-			if (!String.IsNullOrEmpty(UBAConfig.TraceFile))
+			if (!inputActions.Any())
 			{
-				ubaTraceFile = new FileReference(UBAConfig.TraceFile);
+				return true;
 			}
-			else if (Unreal.IsBuildMachine())
+
+			if (UBAConfig.bForceNoCache)
 			{
-				ubaTraceFile = FileReference.Combine(Unreal.EngineProgramSavedDirectory, "AutomationTool", "Saved", "Logs", "Trace.uba");
-			}
-			else if (Log.OutputFile != null)
-			{
-				ubaTraceFile = Log.OutputFile.ChangeExtension(".uba");
+				DisposeCacheClients();
 			}
 			else
 			{
-				ubaTraceFile = FileReference.Combine(Unreal.EngineProgramSavedDirectory, "UnrealBuildTool", "Trace.uba");
+				_cacheableActionsLeft = inputActions.Count(x => x.ArtifactMode.HasFlag(ArtifactMode.Enabled));
 			}
 
-			DirectoryReference.CreateDirectory(ubaTraceFile.Directory);
-			Log.BackupLogFile(ubaTraceFile);
 
-			IStorageServer? ubaStorage = null;
+			await WriteActionOutputFileAsync(logger, inputActions);
 
 			try
 			{
@@ -368,127 +492,36 @@ namespace UnrealBuildTool
 				}
 
 				string arch = RuntimeInformation.ProcessArchitecture.ToString().ToLowerInvariant();
-				FileReference configFile = FileReference.Combine(UbaBinariesDir, "UbaHost.toml");
-				IConfig config = EpicGames.UBA.IConfig.LoadConfig(configFile.FullName);
-				if (!String.IsNullOrEmpty(UBAConfig.CompressionLevel))
-				{
-					config.AddValue("Storage", "CompressionLevel", UBAConfig.CompressionLevel);
-				}
-
-				using EpicGames.UBA.ILogger ubaLogger = EpicGames.UBA.ILogger.CreateLogger(logger);
-				using (IServer Server = IServer.CreateServer(UBAConfig.MaxWorkers, UBAConfig.SendSize, ubaLogger, UBAConfig.bUseQuic))
-				{
-					_ubaLogger = ubaLogger;
-					using IStorageServer ubaStorageServer = IStorageServer.CreateStorageServer(Server, ubaLogger, new StorageServerCreateInfo(_rootDirRef.FullName, ((ulong)UBAConfig.StoreCapacityGb) * 1000 * 1000 * 1000, !UBAConfig.bStoreRaw, UBAConfig.Zone));
-					using ISessionServerCreateInfo serverCreateInfo = ISessionServerCreateInfo.CreateSessionServerCreateInfo(ubaStorageServer, Server, ubaLogger, new SessionServerCreateInfo(_rootDirRef.FullName, ubaTraceFile.FullName.Replace('\\', '/'), UBAConfig.bDisableCustomAlloc, false, UBAConfig.bResetCas, UBAConfig.bWriteToDisk, UBAConfig.bDetailedTrace, !UBAConfig.bDisableWaitOnMem, UBAConfig.bAllowKillOnMem, UBAConfig.bStoreObjFilesCompressed));
-					{
-						try
-						{
-							ISessionServer session = ISessionServer.CreateSessionServer(serverCreateInfo);
-							_session = session;
-							ITrace trace = session.GetTrace();
-
-							OnSessionCreated?.Invoke(session);
-
-							List<UnrealBuildAcceleratorCacheConfig> cacheConfigs = new();
-							if (_cacheableActionsLeft != 0)
-							{
-								if (!UBAConfig.bDisableRemote) // Need to match the InitAsync
-								{
-									using TraceTaskScope s = new(trace, "Requesting Horde cache", "");
-									foreach (IUBAAgentCoordinator coordinator in _agentCoordinators)
-									{
-										UnrealBuildAcceleratorCacheConfig? cacheConfig = coordinator.RequestCacheServer(new CancellationToken());
-										if (cacheConfig != null)
-											cacheConfigs.Add(cacheConfig);
-									}
-								}
-								if (_cacheConfigs.Any())
-									cacheConfigs.AddRange(_cacheConfigs);
-							}
 
-							if (cacheConfigs.Count > 0)
-							{
-								trace.UpdateStatus(1, 1, "Cache", LogEntryType.Info, null);
-								trace.UpdateStatus(1, 6, "Connecting...", LogEntryType.Info);
-
-								using TraceTaskScope s = new(trace, "Connecting to cache", "");
-								foreach (UnrealBuildAcceleratorCacheConfig cacheConfig in cacheConfigs)
-								{
-									ICacheClient cacheClient = ICacheClient.CreateCacheClient(session, UBAConfig.bReportCacheMissReason, cacheConfig.Crypto, GetCacheClientHint());
-									string[] nameAndPort = cacheConfig.CacheServer.Split(':');
-									int port = 1347;
-									if (nameAndPort.Length > 1)
-									{
-										port = Int32.Parse(nameAndPort[1]);
-									}
-
-									System.Diagnostics.Stopwatch stopwatch = System.Diagnostics.Stopwatch.StartNew();
-									bool cacheSuccess = cacheClient.Connect(nameAndPort[0], port, cacheConfig.DesiredConnectionCount);
-									long totalMs = stopwatch.ElapsedMilliseconds;
-									string successText = cacheSuccess ? "Connected to" : "Failed to connect to";
-									logger.LogInformation("UbaCache - {SuccessText} {Name}:{Port} ({Seconds}.{Milliseconds}s)", successText, nameAndPort[0], port, totalMs / 1000, totalMs % 1000);
-									if (cacheSuccess)
-									{
-										_cacheClients.Add(cacheConfig, cacheClient);
-									}
-									else
-									{
-										cacheClient.Dispose();
-									}
-								}
-								trace.UpdateStatus(1, 6, _cacheClients.Any() ? "Connected" : "Failed to connect", LogEntryType.Info);
-							}
+				OnSessionCreated?.Invoke(_session!);
 
-							ubaStorage = ubaStorageServer;
+				bool success = ExecuteActionsInternal(logger, inputActions, _server!, _storageServer!, _session!, _cacheClients.Values);
 
-							if (!UBAConfig.bDisableRemote)
-							{
-								Server.StartServer(UBAConfig.Host, UBAConfig.Port, Crypto);
-							}
-							else
-							{
-								session.DisableRemoteExecution();
-							}
-
-							bool success = ExecuteActionsInternal(logger, inputActions, Server, ubaStorage, session, _cacheClients.Values);
-
-							if (!UBAConfig.bDisableRemote)
-							{
-								Server.StopServer();
-							}
-
-							if (UBAConfig.bPrintSummary)
-							{
-								session.PrintSummary();
-							}
+				if (!UBAConfig.bDisableRemote)
+				{
+					_server!.StopServer();
+				}
 
-							return success && !_bIsCancelled;
-						}
-						finally
-						{
-							foreach (KeyValuePair<UnrealBuildAcceleratorCacheConfig, ICacheClient> item in _cacheClients)
-							{
-								item.Value.Dispose();
-							}
-							_cacheClients.Clear();
-							lock (_sessionLock)
-							{
-								_agentCoordinators.ForEach(ac => ac.Done());
-								_session?.Dispose();
-								_session = null;
-							}
-						}
-					}
+				if (UBAConfig.bPrintSummary)
+				{
+					_session!.PrintSummary();
 				}
+
+				return success && !_bIsCancelled;
 			}
 			finally
 			{
+				DisposeCacheClients();
+
+				lock (_sessionLock)
+				{
+					_agentCoordinators.ForEach(ac => ac.Done());
+				}
+
 				foreach (IUBAAgentCoordinator coordinator in _agentCoordinators)
 				{
 					await coordinator.CloseAsync();
 				}
-				_session = null;
 			}
 		}
 
@@ -617,30 +650,6 @@ namespace UnrealBuildTool
 
 			using IScheduler scheduler = IScheduler.CreateScheduler(session, cache, NumParallelProcesses, UBAConfig.bForceBuildAllRemote);
 
-			var cancelBuild = () =>
-			{
-				if (_bIsCancelled)
-					return;
-				_bIsCancelled = true;
-				scheduler.Cancel();
-				foreach (ICacheClient client in cache)
-					client.Disconnect();
-				server?.StopServer(); // Make sure all remove processes are returned. We can't have any callbacks after this
-				session?.CancelAll(); // Cancel all processes native side
-				storage?.SaveCasTable();
-
-				// We need the lock here since things are happening in parallel.
-				lock (_sessionLock)
-				{
-					foreach (IUBAAgentCoordinator coordinator in _agentCoordinators)
-					{
-						_agentCoordinators.ForEach(ac => ac.CloseAsync().Wait(2000)); // Give coordinators some time to close (this makes coordinators like horde return resources faster)
-					}
-				}
-			};
-
-
-
 			int totalActions = 0;
 			int cacheTotalActions = 0;
 			int cacheHitActions = 0;
@@ -874,19 +883,33 @@ namespace UnrealBuildTool
 			try
 			{
 				AgentCoordinatorScheduler coordinatorScheduler = new(scheduler, server);
-				uint statusUpdateCounter = 10;
 				foreach (IUBAAgentCoordinator coordinator in _agentCoordinators)
 				{
-					uint statusUpdateIndex = statusUpdateCounter;
-					coordinator.Start(coordinatorScheduler, CanRunRemotely, (sr, sc, st, t, sl) => UpdateStatus(statusUpdateIndex + sr, sc, st, t, sl));
-					statusUpdateCounter += 10;
+					coordinator.Start(coordinatorScheduler, CanRunRemotely);
 				}
 
 				while (!scheduler.IsEmpty)
 				{
-					if (shouldCancel)
+					if (shouldCancel && !_bIsCancelled)
 					{
-						cancelBuild();
+						_bIsCancelled = true;
+						_cancellationSource.Cancel();
+
+						scheduler.Cancel();
+						foreach (ICacheClient client in cache)
+							client.Disconnect();
+						server?.StopServer(); // Make sure all remove processes are returned. We can't have any callbacks after this
+						session?.CancelAll(); // Cancel all processes native side
+						storage?.SaveCasTable();
+
+						// We need the lock here since things are happening in parallel.
+						lock (_sessionLock)
+						{
+							foreach (IUBAAgentCoordinator coordinator in _agentCoordinators)
+							{
+								_agentCoordinators.ForEach(ac => ac.CloseAsync().Wait(2000)); // Give coordinators some time to close (this makes coordinators like horde return resources faster)
+							}
+						}
 					}
 					Thread.Sleep(100);
 				}
@@ -920,6 +943,7 @@ namespace UnrealBuildTool
 					cancelHandler = null;
 				}
 
+				_cancellationSource.Cancel();
 				_agentCoordinators.ForEach(ac => ac.Stop());
 			}
 		}
@@ -994,27 +1018,8 @@ namespace UnrealBuildTool
 		{
 			if (Interlocked.Decrement(ref _cacheableActionsLeft) == 0)
 			{
-				foreach (KeyValuePair<UnrealBuildAcceleratorCacheConfig, ICacheClient> item in _cacheClients)
-				{
-					item.Value.Dispose();
-				}
-				_cacheClients.Clear();
-			}
-		}
-
-		string GetCacheClientHint()
-		{
-			// This is extra information sent to cache server when fetching stuff. This way it will be possible to backtrack from cache server log to who connected
-			// These environment variables are set when UBT is running on a horde machine.
-			string? jobId = Environment.GetEnvironmentVariable("UE_HORDE_JOBID");
-			if (jobId == null)
-			{
-				return Environment.MachineName;
+				DisposeCacheClients();
 			}
-			string hordeUrl = Environment.GetEnvironmentVariable("UE_HORDE_URL")??"";
-			string stepId = Environment.GetEnvironmentVariable("UE_HORDE_STEPID")??"unknown";
-			string streamId = Environment.GetEnvironmentVariable("UE_HORDE_STREAMID")??"unknown-stream";
-			return $"{hordeUrl}job/{jobId}?step={stepId} ({streamId})";
 		}
 
 		HashSet<uint> _usedBuckets = new();
diff --git a/Engine/Source/Programs/UnrealBuildTool/Executors/UnrealBuildAccelerator/UnrealBuildAcceleratorConfig.cs b/Engine/Source/Programs/UnrealBuildTool/Executors/UnrealBuildAccelerator/UnrealBuildAcceleratorConfig.cs
--- a/Engine/Source/Programs/UnrealBuildTool/Executors/UnrealBuildAccelerator/UnrealBuildAcceleratorConfig.cs
+++ b/Engine/Source/Programs/UnrealBuildTool/Executors/UnrealBuildAccelerator/UnrealBuildAcceleratorConfig.cs
@@ -306,6 +306,13 @@ namespace UnrealBuildTool
 		[CommandLine("-UBAShuffleCache")]
 		public bool bCacheShuffle { get; set; } = false;
 
+		/// <summary>
+		/// Shuffle the list of cache servers before each read attempt
+		/// </summary>
+		[XmlConfigFile(Category = "UnrealBuildAccelerator", Name = "ForceNoCache")]
+		[CommandLine("-nocache")]
+		public bool bForceNoCache { get; set; } = false;
+
 		/// <summary>
 		/// Report reason a cache miss happened. Useful when searching for determinism/portability issues
 		/// </summary>
diff --git a/Engine/Source/Programs/UnrealBuildTool/Modes/BuildMode.cs b/Engine/Source/Programs/UnrealBuildTool/Modes/BuildMode.cs
--- a/Engine/Source/Programs/UnrealBuildTool/Modes/BuildMode.cs
+++ b/Engine/Source/Programs/UnrealBuildTool/Modes/BuildMode.cs
@@ -11,7 +11,6 @@ using EpicGames.Core;
 using Microsoft.Extensions.Logging;
 using OpenTracing.Util;
 using UnrealBuildBase;
-using UnrealBuildTool.Artifacts;
 
 namespace UnrealBuildTool
 {
@@ -197,6 +196,8 @@ namespace UnrealBuildTool
 					CleanMode.Clean(TargetDescriptors.Where(D => D.bRebuild).ToList(), BuildConfiguration, Logger);
 				}
 
+				using ActionExecutor? executor = ActionGraph.PrecreateExecutor(BuildConfiguration, Logger, TargetDescriptors);
+
 				// Handle remote builds
 				for (int Idx = 0; Idx < TargetDescriptors.Count; ++Idx)
 				{
@@ -254,7 +255,7 @@ namespace UnrealBuildTool
 					// Create the working set provider per group.
 					using (ISourceFileWorkingSet WorkingSet = SourceFileWorkingSet.Create(Unreal.RootDirectory, ProjectDirs, Logger))
 					{
-						await BuildAsync(TargetDescriptors, BuildConfiguration, WorkingSet, Options, WriteOutdatedActionsFile, dependencyCache, Logger, bSkipPreBuildTargets, ActionTypeFilter:ActionTypeFilter);
+						await BuildAsync(TargetDescriptors, BuildConfiguration, WorkingSet, Options, WriteOutdatedActionsFile, dependencyCache, Logger, bSkipPreBuildTargets, ActionTypeFilter:ActionTypeFilter, executor);
 					}
 				}
 			}
@@ -389,12 +390,13 @@ namespace UnrealBuildTool
 		/// <param name="Logger">Logger for output</param>
 		/// <param name="bSkipPreBuildTargets">If true then only the current target descriptors will be built.</param>
 		/// <param name="ActionTypeFilter">If specified, only actions of this type will be run</param>
+		/// <param name="executor"></param>
 		/// <returns>Result from the compilation</returns>
-		public static async Task BuildAsync(List<TargetDescriptor> TargetDescriptors, BuildConfiguration BuildConfiguration, ISourceFileWorkingSet WorkingSet, BuildOptions Options, FileReference? WriteOutdatedActionsFile, CppDependencyCache cppDependencyCache, ILogger Logger, bool bSkipPreBuildTargets = false, string? ActionTypeFilter = null)
+		public static async Task BuildAsync(List<TargetDescriptor> TargetDescriptors, BuildConfiguration BuildConfiguration, ISourceFileWorkingSet WorkingSet, BuildOptions Options, FileReference? WriteOutdatedActionsFile, CppDependencyCache cppDependencyCache, ILogger Logger, bool bSkipPreBuildTargets = false, string? ActionTypeFilter = null, ActionExecutor? executor = null)
 		{
 			(List<(TargetDescriptor, TargetMakefile)> Targets, List<Task> Tasks) = CreateMakefiles(BuildConfiguration, TargetDescriptors, WorkingSet, Logger, true, true, cppDependencyCache);
 
-			await BuildAsync(Targets, BuildConfiguration, Options, WriteOutdatedActionsFile, cppDependencyCache, Logger, ActionTypeFilter);
+			await BuildAsync(Targets, BuildConfiguration, Options, WriteOutdatedActionsFile, cppDependencyCache, Logger, ActionTypeFilter, executor);
 
 			Task.WaitAll(Tasks.ToArray());
 		}
@@ -411,8 +413,9 @@ namespace UnrealBuildTool
 		/// <param name="Logger">Logger for output</param>
 		/// <param name="ActionTypeFilter">If specified, only actions of this type will be run</param>
 		/// <param name="cppDependencyCache"></param>
+		/// <param name="executor"></param>
 		/// <returns>Result from the compilation</returns>
-		internal static async Task BuildAsync(List<(TargetDescriptor Descriptor, TargetMakefile Makefile)> Targets, BuildConfiguration BuildConfiguration, BuildOptions Options, FileReference? WriteOutdatedActionsFile, CppDependencyCache cppDependencyCache, ILogger Logger, string? ActionTypeFilter=null)
+		internal static async Task BuildAsync(List<(TargetDescriptor Descriptor, TargetMakefile Makefile)> Targets, BuildConfiguration BuildConfiguration, BuildOptions Options, FileReference? WriteOutdatedActionsFile, CppDependencyCache cppDependencyCache, ILogger Logger, string? ActionTypeFilter = null, ActionExecutor? executor = null)
 		{
 			List<TargetDescriptor> TargetDescriptors = Targets.Select(T => T.Descriptor).ToList();
 			TargetMakefile[] Makefiles = Targets.Select(T => T.Makefile).ToArray();
@@ -557,29 +560,6 @@ namespace UnrealBuildTool
 				// Wait for cpp dependencies task to finish
 				cppDependencyCache.WaitAsync();
 
-				// Now that we have the dependencies object, optionally create the artifact cache
-				IActionArtifactCache? actionArtifactCache = null;
-				if (!String.IsNullOrEmpty(BuildConfiguration.ArtifactDirectory) && (BuildConfiguration.bArtifactRead || BuildConfiguration.bArtifactWrites))
-				{
-					DirectoryReference artifactDirectory = new(BuildConfiguration.ArtifactDirectory);
-					DirectoryReference.CreateDirectory(artifactDirectory);
-
-					string mode = (BuildConfiguration.bArtifactRead, BuildConfiguration.bArtifactWrites) switch
-					{
-						(false, true) => "write only",
-						(true, false) => "read only",
-						_ => "read/write",
-					};
-
-					Logger.LogInformation("Experimental artifact system using '{Directory}' in {Mode} mode", artifactDirectory.FullName, mode);
-
-					actionArtifactCache = ActionArtifactCache.CreateHordeFileCache(artifactDirectory, cppDependencyCache, s_memoryMappedFileCache.Value, Logger);
-					actionArtifactCache.EnableReads = BuildConfiguration.bArtifactRead;
-					actionArtifactCache.EnableWrites = BuildConfiguration.bArtifactWrites;
-					actionArtifactCache.LogCacheMisses = BuildConfiguration.bLogArtifactCacheMisses;
-					actionArtifactCache.EngineRoot = Unreal.EngineDirectory;
-				}
-
 				// Pre-process module interfaces to generate dependency files
 				List<LinkedAction> ModuleDependencyActions = PrerequisiteActions.Where(x => x.ActionType == ActionType.GatherModuleDependencies).ToList();
 				if (ModuleDependencyActions.Count > 0)
@@ -591,7 +571,7 @@ namespace UnrealBuildTool
 					if (PreprocessActions.Count > 0)
 					{
 						Logger.LogInformation("Updating module dependencies...");
-						await ActionGraph.ExecuteActionsAsync(BuildConfiguration, PreprocessActions, TargetDescriptors, Logger);
+						await ActionGraph.ExecuteActionsAsync(BuildConfiguration, PreprocessActions, TargetDescriptors, Logger, null, executor);
 
 						foreach (FileItem ProducedItem in PreprocessActions.SelectMany(x => x.ProducedItems))
 						{
@@ -810,21 +790,7 @@ namespace UnrealBuildTool
 				else
 				{
 					// Execute the actions
-					if (MergedActionsToExecute.Count == 0)
-					{
-						if (TargetDescriptors.Any(x => !x.bQuiet))
-						{
-							if (TargetDescriptors.Count == 1)
-							{
-								Logger.LogInformation("Target is up to date");
-							}
-							else
-							{
-								Logger.LogInformation("Targets are up to date");
-							}
-						}
-					}
-					else
+					if (MergedActionsToExecute.Count != 0)
 					{
 						if (TargetDescriptors.Any(x => !x.bQuiet))
 						{
@@ -834,39 +800,27 @@ namespace UnrealBuildTool
 						OutputToolchainInfo(TargetDescriptors, Makefiles, Logger);
 
 						ActionExecutor.SetMemoryPerActionOverride(Makefiles.Select(x => x.MemoryPerActionGB).Max());
+					}
 
-						//using ITimelineEvent ExcuteTimer = Timeline.ScopeEvent("Executing ActionGraph");
-						using (GlobalTracer.Instance.BuildSpan("ActionGraph.ExecuteActions()").StartActive())
-						{
+					//using ITimelineEvent ExcuteTimer = Timeline.ScopeEvent("Executing ActionGraph");
+					using (GlobalTracer.Instance.BuildSpan("ActionGraph.ExecuteActions()").StartActive())
+					{
+						await ActionGraph.ExecuteActionsAsync(BuildConfiguration, MergedActionsToExecute, TargetDescriptors, Logger, null, executor);
+					}
 
-							// We need to wait for the cache to be ready.  If we fail to fetch from the cache for PCH files, they are
-							// non-deterministic and will cause rippled cache misses.  Changing the system to assume a PCH file is the same
-							// if the inputs are the same might be a good way to avoid this issue.
-							if (actionArtifactCache != null && actionArtifactCache.ArtifactCache != null)
-							{
-								await actionArtifactCache.ArtifactCache.WaitForReadyAsync();
-							}
-							await ActionGraph.ExecuteActionsAsync(BuildConfiguration, MergedActionsToExecute, TargetDescriptors, Logger, actionArtifactCache);
-							if (actionArtifactCache != null)
-							{
-								await actionArtifactCache.FlushChangesAsync(default); //ETSTODO
-							}
+					// these are parsed by external tools wishing to open this file directly
+					foreach (LinkedAction BuildAction in MergedActionsToExecute.Where(BuildAction => BuildAction.ActionType == ActionType.Compile && (BuildAction.Inner as VCCompileAction)?.bIsAnalyzing != true))
+					{
+						FileItem? PreprocessedFile = BuildAction.ProducedItems.FirstOrDefault(x => x.HasExtension(".i"));
+						if (PreprocessedFile != null)
+						{
+							Logger.LogInformation("PreProcessPath: {File}", PreprocessedFile);
 						}
 
-						// these are parsed by external tools wishing to open this file directly
-						foreach (LinkedAction BuildAction in MergedActionsToExecute.Where(BuildAction => BuildAction.ActionType == ActionType.Compile && (BuildAction.Inner as VCCompileAction)?.bIsAnalyzing != true))
+						FileItem? AssemblyPath = BuildAction.ProducedItems.FirstOrDefault(x => x.HasExtension(".asm"));
+						if (AssemblyPath != null)
 						{
-							FileItem? PreprocessedFile = BuildAction.ProducedItems.FirstOrDefault(x => x.HasExtension(".i"));
-							if (PreprocessedFile != null)
-							{
-								Logger.LogInformation("PreProcessPath: {File}", PreprocessedFile);
-							}
-
-							FileItem? AssemblyPath = BuildAction.ProducedItems.FirstOrDefault(x => x.HasExtension(".asm"));
-							if (AssemblyPath != null)
-							{
-								Logger.LogInformation("AssemblyPath: {File}", AssemblyPath);
-							}
+							Logger.LogInformation("AssemblyPath: {File}", AssemblyPath);
 						}
 					}
 
diff --git a/Engine/Source/Programs/UnrealBuildTool/System/ActionGraph.cs b/Engine/Source/Programs/UnrealBuildTool/System/ActionGraph.cs
--- a/Engine/Source/Programs/UnrealBuildTool/System/ActionGraph.cs
+++ b/Engine/Source/Programs/UnrealBuildTool/System/ActionGraph.cs
@@ -378,6 +378,28 @@ namespace UnrealBuildTool
 			}
 		}
 
+		public static ActionExecutor? PrecreateExecutor(BuildConfiguration BuildConfiguration, ILogger Logger, List<TargetDescriptor> TargetDescriptors)
+		{
+			if (BuildConfiguration.bAllowXGE && XGE.IsAvailable(Logger))
+			{
+				return null;
+			}
+			if (BuildConfiguration.bAllowSNDBS && SNDBS.IsAvailable(Logger))
+			{
+				return null;
+			}
+			if (BuildConfiguration.bAllowFASTBuild && FASTBuild.IsAvailable(Logger))
+			{
+				return null;
+			}
+			if (BuildConfiguration.bAllowUBAExecutor && UBAExecutor.IsAvailable())
+			{
+				return new UBAExecutor(BuildConfiguration.MaxParallelActions, BuildConfiguration.bAllCores, BuildConfiguration.bCompactOutput, Logger, TargetDescriptors);
+			}
+
+			return null;
+		}
+
 		private static ActionExecutor? GetRemoteExecutorByName(string Name, BuildConfiguration BuildConfiguration, int ActionCount, int MinActionsForRemote, List<TargetDescriptor> TargetDescriptors, ILogger Logger)
 		{
 			switch (Name)
@@ -444,59 +466,77 @@ namespace UnrealBuildTool
 		/// <summary>
 		/// Executes a list of actions.
 		/// </summary>
-		public static async Task ExecuteActionsAsync(BuildConfiguration BuildConfiguration, List<LinkedAction> ActionsToExecute, List<TargetDescriptor> TargetDescriptors, ILogger Logger, IActionArtifactCache? actionArtifactCache = null)
+		public static async Task ExecuteActionsAsync(BuildConfiguration BuildConfiguration, List<LinkedAction> ActionsToExecute, List<TargetDescriptor> TargetDescriptors, ILogger Logger, IActionArtifactCache? actionArtifactCache = null, ActionExecutor? executor = null)
 		{
 			if (ActionsToExecute.Count == 0)
 			{
 				Logger.LogInformation("Target is up to date");
+
+				if (executor != null)
+				{
+					// We still want to run the executor to give it a chance to clean up in case it has tasks running that can be cancelled
+					await InternalExecuteActions(executor, BuildConfiguration, ActionsToExecute, TargetDescriptors, Logger, actionArtifactCache);
+				}
 			}
 			else
 			{
 				// Figure out which executor to use
-				using ActionExecutor Executor = SelectExecutor(BuildConfiguration, ActionsToExecute.Count, TargetDescriptors, Logger);
-				Logger.LogInformation("Using {ExecutorName} executor to run {ActionCount} action(s)", Executor.Name, ActionsToExecute.Count);
+				if (executor == null)
+				{
+					using ActionExecutor SelectedExecutor = SelectExecutor(BuildConfiguration, ActionsToExecute.Count, TargetDescriptors, Logger);
+					await InternalExecuteActions(SelectedExecutor, BuildConfiguration, ActionsToExecute, TargetDescriptors, Logger, actionArtifactCache);
+				}
+				else
+				{
+					await InternalExecuteActions(executor, BuildConfiguration, ActionsToExecute, TargetDescriptors, Logger, actionArtifactCache);
+				}
 
-				// Execute the build
-				Stopwatch Timer = Stopwatch.StartNew();
-				bool Result = await Executor.ExecuteActionsAsync(ActionsToExecute, Logger, actionArtifactCache);
-				Executor.PostTelemetryEvent();
+			}
+		}
 
-				Logger.LogInformation("Total time in {ExecutorName} executor: {TotalSeconds:0.00} seconds", Executor.Name, Timer.Elapsed.TotalSeconds);
+		static async Task InternalExecuteActions(ActionExecutor Executor, BuildConfiguration BuildConfiguration, List<LinkedAction> ActionsToExecute, List<TargetDescriptor> TargetDescriptors, ILogger Logger, IActionArtifactCache? actionArtifactCache = null)
+		{
+			Logger.LogInformation("Using {ExecutorName} executor to run {ActionCount} action(s)", Executor.Name, ActionsToExecute.Count);
 
-				if (!Result)
-				{
-					throw new CompilationResultException(CompilationResult.OtherCompilationError);
-				}
+			// Execute the build
+			Stopwatch Timer = Stopwatch.StartNew();
+			bool Result = await Executor.ExecuteActionsAsync(ActionsToExecute, Logger, actionArtifactCache);
+			Executor.PostTelemetryEvent();
 
-				// Reset the file info for all the produced items
-				foreach (LinkedAction BuildAction in ActionsToExecute)
+			Logger.LogInformation("Total time in {ExecutorName} executor: {TotalSeconds:0.00} seconds", Executor.Name, Timer.Elapsed.TotalSeconds);
+
+			if (!Result)
+			{
+				throw new CompilationResultException(CompilationResult.OtherCompilationError);
+			}
+
+			// Reset the file info for all the produced items
+			foreach (LinkedAction BuildAction in ActionsToExecute)
+			{
+				foreach (FileItem ProducedItem in BuildAction.ProducedItems)
 				{
-					foreach (FileItem ProducedItem in BuildAction.ProducedItems)
-					{
-						ProducedItem.ResetCachedInfo();
-					}
+					ProducedItem.ResetCachedInfo();
 				}
+			}
 
-				// Verify the link outputs were created (seems to happen with Win64 compiles)
-				if (Executor.VerifyOutputs)
+			// Verify the link outputs were created (seems to happen with Win64 compiles)
+			if (Executor.VerifyOutputs)
+			{
+				foreach (LinkedAction BuildAction in ActionsToExecute)
 				{
-					foreach (LinkedAction BuildAction in ActionsToExecute)
+					if (BuildAction.ActionType == ActionType.Link)
 					{
-						if (BuildAction.ActionType == ActionType.Link)
+						foreach (FileItem Item in BuildAction.ProducedItems)
 						{
-							foreach (FileItem Item in BuildAction.ProducedItems)
+							if (!Item.Exists)
 							{
-								if (!Item.Exists)
-								{
-									throw new BuildException($"Failed to produce item: {Item.AbsolutePath}");
-								}
+								throw new BuildException($"Failed to produce item: {Item.AbsolutePath}");
 							}
 						}
 					}
 				}
 			}
 		}
-
 		/// <summary>
 		/// Sorts the action list for improved parallelism with local execution.
 		/// </summary>
