[UBT]
* Fixed so tasks that saves out makefile etc always gets waited on when exception happening during build

#FailedValidations Preflight

diff --git a/Engine/Source/Programs/UnrealBuildTool/Modes/BuildMode.cs b/Engine/Source/Programs/UnrealBuildTool/Modes/BuildMode.cs
--- a/Engine/Source/Programs/UnrealBuildTool/Modes/BuildMode.cs
+++ b/Engine/Source/Programs/UnrealBuildTool/Modes/BuildMode.cs
@@ -164,7 +164,8 @@ namespace UnrealBuildTool
 				}
 			}
 
-			CppDependencyCache dependencyCache = new();
+			CppDependencyCache cppDependencyCache = new();
+			List<Task> tasks = [];
 
 			// Parse and build the targets
 			try
@@ -255,7 +256,8 @@ namespace UnrealBuildTool
 					// Create the working set provider per group.
 					using (ISourceFileWorkingSet WorkingSet = SourceFileWorkingSet.Create(Unreal.RootDirectory, ProjectDirs, Logger))
 					{
-						await BuildAsync(TargetDescriptors, BuildConfiguration, WorkingSet, Options, WriteOutdatedActionsFile, dependencyCache, Logger, bSkipPreBuildTargets, ActionTypeFilter:ActionTypeFilter, executor);
+						List<(TargetDescriptor, TargetMakefile)> Targets = CreateMakefiles(BuildConfiguration, TargetDescriptors, WorkingSet, Logger, true, true, tasks, cppDependencyCache);
+						await BuildAsync(Targets, BuildConfiguration, Options, WriteOutdatedActionsFile, cppDependencyCache, Logger, ActionTypeFilter, executor);
 					}
 				}
 			}
@@ -264,11 +266,13 @@ namespace UnrealBuildTool
 				// Check if anything failed during our run, and act accordingly.
 				ProcessCoreDumps(SaveCrashDumpDirectory, Logger);
 
+				Task.WaitAll(tasks.ToArray());
+
 				// Save all the caches
 				using (Timeline.ScopeEvent("Saving caches"))
 				{
 					SourceFileMetadataCache.SaveAll();
-					dependencyCache.SaveAll();
+					cppDependencyCache.SaveAll();
 				}
 			}
 			return 0;
@@ -390,15 +394,19 @@ namespace UnrealBuildTool
 		/// <param name="Logger">Logger for output</param>
 		/// <param name="bSkipPreBuildTargets">If true then only the current target descriptors will be built.</param>
 		/// <param name="ActionTypeFilter">If specified, only actions of this type will be run</param>
-		/// <param name="executor"></param>
 		/// <returns>Result from the compilation</returns>
-		public static async Task BuildAsync(List<TargetDescriptor> TargetDescriptors, BuildConfiguration BuildConfiguration, ISourceFileWorkingSet WorkingSet, BuildOptions Options, FileReference? WriteOutdatedActionsFile, CppDependencyCache cppDependencyCache, ILogger Logger, bool bSkipPreBuildTargets = false, string? ActionTypeFilter = null, ActionExecutor? executor = null)
+		public static async Task BuildAsync(List<TargetDescriptor> TargetDescriptors, BuildConfiguration BuildConfiguration, ISourceFileWorkingSet WorkingSet, BuildOptions Options, FileReference? WriteOutdatedActionsFile, CppDependencyCache cppDependencyCache, ILogger Logger, bool bSkipPreBuildTargets = false, string? ActionTypeFilter = null)
 		{
-			(List<(TargetDescriptor, TargetMakefile)> Targets, List<Task> Tasks) = CreateMakefiles(BuildConfiguration, TargetDescriptors, WorkingSet, Logger, true, true, cppDependencyCache);
-
-			await BuildAsync(Targets, BuildConfiguration, Options, WriteOutdatedActionsFile, cppDependencyCache, Logger, ActionTypeFilter, executor);
-
-			Task.WaitAll(Tasks.ToArray());
+			List<Task> tasks = [];
+			try
+			{
+				List<(TargetDescriptor, TargetMakefile)> Targets = CreateMakefiles(BuildConfiguration, TargetDescriptors, WorkingSet, Logger, true, true, tasks, cppDependencyCache);
+				await BuildAsync(Targets, BuildConfiguration, Options, WriteOutdatedActionsFile, cppDependencyCache, Logger, ActionTypeFilter, null);
+			}
+			finally
+			{
+				Task.WaitAll(tasks.ToArray());
+			}
 		}
 
 		static Lazy<MemoryMappedFileCache> s_memoryMappedFileCache = new Lazy<MemoryMappedFileCache>(() => new MemoryMappedFileCache());
@@ -1058,10 +1066,16 @@ namespace UnrealBuildTool
 		{
 			await Task.CompletedTask;
 
-			(List<(TargetDescriptor, TargetMakefile Makefile)> Targets, List<Task> Tasks) = CreateMakefiles(BuildConfiguration, [TargetDescriptor], WorkingSet, Logger, false, false, null);
-
-			Task.WaitAll(Tasks.ToArray());
-			return Targets[0].Makefile;
+			List<Task> tasks = [];
+			try
+			{
+				List<(TargetDescriptor, TargetMakefile Makefile)> Targets = CreateMakefiles(BuildConfiguration, [TargetDescriptor], WorkingSet, Logger, false, false, tasks, null);
+				return Targets[0].Makefile;
+			}
+			finally
+			{
+				Task.WaitAll(tasks.ToArray());
+			}
 		}
 
 		// Function that will load makefile and check if it is valid. (First round of targets will have their makefiles loaded in parallel with other work)
@@ -1089,10 +1103,9 @@ namespace UnrealBuildTool
 			return (makefile, true);
 		}
 
-		private static (List<(TargetDescriptor, TargetMakefile)>, List<Task>) CreateMakefiles(BuildConfiguration buildConfiguration, List<TargetDescriptor> descriptors, ISourceFileWorkingSet workingSet, ILogger logger, bool runInitScripts, bool skipPreBuildTargets, CppDependencyCache? cppDependencyCache)
+		private static List<(TargetDescriptor, TargetMakefile)> CreateMakefiles(BuildConfiguration buildConfiguration, List<TargetDescriptor> descriptors, ISourceFileWorkingSet workingSet, ILogger logger, bool runInitScripts, bool skipPreBuildTargets, List<Task> outTasks, CppDependencyCache? cppDependencyCache)
 		{
 			List<(TargetDescriptor descriptor, TargetMakefile makefile)> outTargets = new();
-			List<Task> outTasks = [];
 
 			int helperCount = buildConfiguration.bParallelMakefileGeneration ? -1 : 0;
 
@@ -1213,215 +1226,224 @@ namespace UnrealBuildTool
 
 					List<UEBuildTarget.CodeGenContext>? lastContextsForCodeGen = null;
 
-					// Split up in passes. We run one target from each intermediate folder.
-					for (int passIt = 0; passIt != maxTargetsInIntermediateCount; ++passIt)
+					try
 					{
-						List<UEBuildTarget.CodeGenContext> contextsForCodeGen = [];
-
-						// Do everyting up to code gen in parallel
-						Parallel2.ForEach(intermediateFolderGroups.Values, intermediateFolderGroup =>
+						// Split up in passes. We run one target from each intermediate folder.
+						for (int passIt = 0; passIt != maxTargetsInIntermediateCount; ++passIt)
 						{
-							if (passIt >= intermediateFolderGroup.Count)
-							{
-								return; // This group is done
-							}
+							List<UEBuildTarget.CodeGenContext> contextsForCodeGen = [];
 
-							int index = intermediateFolderGroup[passIt];
-							TargetDescriptor descriptor = descriptors[index];
-
-							if (runInitScripts)
+							// Do everyting up to code gen in parallel
+							Parallel2.ForEach(intermediateFolderGroups.Values, intermediateFolderGroup =>
 							{
-								// Create and execute the init scripts
-								FileReference[] InitScripts = CreateInitScripts(descriptor);
-								if (InitScripts.Length > 0)
+								if (passIt >= intermediateFolderGroup.Count)
 								{
-									Utils.ExecuteCustomBuildSteps(InitScripts, logger);
+									return; // This group is done
 								}
-							}
-
-							TargetMakefile? makefile = null;
-							bool ranPrebuildScripts = false;
 
-							UEBuildTarget buildTarget = buildTargets[index];
-							ReadOnlyTargetRules targetRules = buildTarget.Rules;
-							
-							// Calculating real makefile location. it is likely the same as the preload but not if appname is changed in target rules
-							FileReference? makefileLocation = TargetMakefile.GetLocation(targetRules.ProjectFile, targetRules.Name, targetRules.Platform, targetRules.Architectures, targetRules.Configuration, buildTarget.IntermediateEnvironment);
-							makefileLocations[index] = makefileLocation;
+								int index = intermediateFolderGroup[passIt];
+								TargetDescriptor descriptor = descriptors[index];
 
-							if (makefileLocation != null)
-							{
-								Task<(TargetMakefile?, bool)>? makefileTask = null;
-								lock (preLoadedMakefiles)
+								if (runInitScripts)
 								{
-									if (preLoadedMakefiles.TryGetValue(makefileLocation, out makefileTask))
+									// Create and execute the init scripts
+									FileReference[] InitScripts = CreateInitScripts(descriptor);
+									if (InitScripts.Length > 0)
 									{
-										preLoadedMakefiles.Remove(makefileLocation);
+										Utils.ExecuteCustomBuildSteps(InitScripts, logger);
 									}
 								}
-								if (makefileTask != null)
-								{
-									(makefile, ranPrebuildScripts) = makefileTask.Result;
-								}
-								else
-								{
-									// Need to make sure to wait on previous run's makefile
-									lastSaveMakefileTasks.TryRemove(makefileLocation, out Task? lastMakefileSave);
-									lastMakefileSave?.Wait();
-									(makefile, ranPrebuildScripts) = LoadMakefile(makefileLocation, descriptor, workingSet, logger);
-								}
-							}
 
-							targetTimelines[index] = Timeline.ScopeEvent($"{descriptor.Name} {descriptor.Platform} {descriptor.Configuration}");
+								TargetMakefile? makefile = null;
+								bool ranPrebuildScripts = false;
 
-							List<UEBuildTarget.GenerateHeaderFunc> generatedHeaderFuncs = [];
-
-							// If we couldn't load a valid makefile, create a new one
-							if (makefile == null)
-							{
-								buildTarget.PreBuildSetup(logger);
+								UEBuildTarget buildTarget = buildTargets[index];
+								ReadOnlyTargetRules targetRules = buildTarget.Rules;
 
-								// Create the pre-build scripts
-								FileReference[] PreBuildScripts = buildTarget.CreatePreBuildScripts();
+								// Calculating real makefile location. it is likely the same as the preload but not if appname is changed in target rules
+								FileReference? makefileLocation = TargetMakefile.GetLocation(targetRules.ProjectFile, targetRules.Name, targetRules.Platform, targetRules.Architectures, targetRules.Configuration, buildTarget.IntermediateEnvironment);
+								makefileLocations[index] = makefileLocation;
 
-								// Execute the pre-build scripts
-								if (!ranPrebuildScripts)
+								if (makefileLocation != null)
 								{
-									Utils.ExecuteCustomBuildSteps(PreBuildScripts, logger);
-									ranPrebuildScripts = true;
+									Task<(TargetMakefile?, bool)>? makefileTask = null;
+									lock (preLoadedMakefiles)
+									{
+										if (preLoadedMakefiles.TryGetValue(makefileLocation, out makefileTask))
+										{
+											preLoadedMakefiles.Remove(makefileLocation);
+										}
+									}
+									if (makefileTask != null)
+									{
+										(makefile, ranPrebuildScripts) = makefileTask.Result;
+									}
+									else
+									{
+										// Need to make sure to wait on previous run's makefile
+										lastSaveMakefileTasks.TryRemove(makefileLocation, out Task? lastMakefileSave);
+										lastMakefileSave?.Wait();
+										(makefile, ranPrebuildScripts) = LoadMakefile(makefileLocation, descriptor, workingSet, logger);
+									}
 								}
 
-								SourceFileMetadataCache sourceFileMetaDataCache;
-								if (descriptor.ProjectFile != null)
-								{
-									sourceFileMetaDataCache = sourceFileMetadataCaches[descriptor.ProjectFile].Result;
-								}
-								else
-								{
-									sourceFileMetaDataCache = SourceFileMetadataCache.CreateHierarchy(null, logger);
-								}
+								targetTimelines[index] = Timeline.ScopeEvent($"{descriptor.Name} {descriptor.Platform} {descriptor.Configuration}");
 
-								// Build the target
-								UEBuildTarget.BuildContext context = buildTarget.BuildMakefile1(descriptor, logger, sourceFileMetaDataCache);
-								buildContexts[index] = context;
+								List<UEBuildTarget.GenerateHeaderFunc> generatedHeaderFuncs = [];
 
-								makefile = context.Makefile;
+								// If we couldn't load a valid makefile, create a new one
+								if (makefile == null)
+								{
+									buildTarget.PreBuildSetup(logger);
+
+									// Create the pre-build scripts
+									FileReference[] PreBuildScripts = buildTarget.CreatePreBuildScripts();
 
-								makefile.PreBuildScripts = PreBuildScripts;
+									// Execute the pre-build scripts
+									if (!ranPrebuildScripts)
+									{
+										Utils.ExecuteCustomBuildSteps(PreBuildScripts, logger);
+										ranPrebuildScripts = true;
+									}
 
-								makefile.MemoryPerActionGB = buildTarget.Rules.MemoryPerActionGB;
+									SourceFileMetadataCache sourceFileMetaDataCache;
+									if (descriptor.ProjectFile != null)
+									{
+										sourceFileMetaDataCache = sourceFileMetadataCaches[descriptor.ProjectFile].Result;
+									}
+									else
+									{
+										sourceFileMetaDataCache = SourceFileMetadataCache.CreateHierarchy(null, logger);
+									}
 
-								// Save the additional command line arguments
-								makefile.AdditionalArguments = descriptor.AdditionalArguments.GetRawArray();
+									// Build the target
+									UEBuildTarget.BuildContext context = buildTarget.BuildMakefile1(descriptor, logger, sourceFileMetaDataCache);
+									buildContexts[index] = context;
 
-								generatedHeaderFuncs = context.GenerateHeaderFuncs;
-							}
+									makefile = context.Makefile;
 
-							lock (contextsForCodeGen)
-							{
-								contextsForCodeGen.Add(new UEBuildTarget.CodeGenContext()
-								{
-									Descriptor = descriptor,
-									Makefile = makefile!,
-									TargetName = buildTarget.TargetName,
-									GenerateHeaderFuncs = generatedHeaderFuncs,
-								});
-							}
-							makefiles[index] = makefile;
-						}, helperCount);
+									makefile.PreBuildScripts = PreBuildScripts;
 
+									makefile.MemoryPerActionGB = buildTarget.Rules.MemoryPerActionGB;
 
-						// Make sure to wait for previous pass codegen and timestamp writes
-						if (lastContextsForCodeGen != null)
-						{
-							UEBuildTarget.WaitCodeGen(lastContextsForCodeGen);
-							Task.WaitAll(UEBuildTarget.GetTimestampTasks(lastContextsForCodeGen).ToArray());
-						}
+									// Save the additional command line arguments
+									makefile.AdditionalArguments = descriptor.AdditionalArguments.GetRawArray();
 
-						// Sort codegen based on config with Shipping last.
-						// Since shipping often has fewer modules, that means previous configs will include later ones
-						contextsForCodeGen.SortBy(context => context.Descriptor.Configuration);
+									generatedHeaderFuncs = context.GenerateHeaderFuncs;
+								}
 
-						// Run codegen for all targets
-						UEBuildTarget.RunCodeGen(buildConfiguration, contextsForCodeGen, logger);
+								lock (contextsForCodeGen)
+								{
+									contextsForCodeGen.Add(new UEBuildTarget.CodeGenContext()
+									{
+										Descriptor = descriptor,
+										Makefile = makefile!,
+										TargetName = buildTarget.TargetName,
+										GenerateHeaderFuncs = generatedHeaderFuncs,
+									});
+								}
+								makefiles[index] = makefile;
+							}, helperCount);
 
 
-						// Go wide again and finish the work
-						Parallel2.ForEach(intermediateFolderGroups.Values, intermediateFolderGroup =>
-						{
-							if (passIt >= intermediateFolderGroup.Count)
+							// Make sure to wait for previous pass codegen and timestamp writes
+							if (lastContextsForCodeGen != null)
 							{
-								return; // This group is done
+								UEBuildTarget.WaitCodeGen(lastContextsForCodeGen);
+								Task.WaitAll(UEBuildTarget.GetTimestampTasks(lastContextsForCodeGen).ToArray());
 							}
 
-							int index = intermediateFolderGroup[passIt];
-							TargetDescriptor descriptor = descriptors[index];
+							// Sort codegen based on config with Shipping last.
+							// Since shipping often has fewer modules, that means previous configs will include later ones
+							contextsForCodeGen.SortBy(context => context.Descriptor.Configuration);
+
+							// Run codegen for all targets
+							UEBuildTarget.RunCodeGen(buildConfiguration, contextsForCodeGen, logger);
 
-							TargetMakefile makefile = makefiles[index]!;
-							UEBuildTarget.BuildContext? buildContext = buildContexts[index];
 
-							if (buildContext != null)
+							// Go wide again and finish the work
+							Parallel2.ForEach(intermediateFolderGroups.Values, intermediateFolderGroup =>
 							{
-								UEBuildTarget buildTarget = buildTargets[index];
+								if (passIt >= intermediateFolderGroup.Count)
+								{
+									return; // This group is done
+								}
 
-								buildTarget.BuildMakefile2(workingSet, logger, buildContext);
+								int index = intermediateFolderGroup[passIt];
+								TargetDescriptor descriptor = descriptors[index];
 
-								// Save the makefile for next time
-								FileReference? makefileLocation = makefileLocations[index];
-								if (makefileLocation != null)
+								TargetMakefile makefile = makefiles[index]!;
+								UEBuildTarget.BuildContext? buildContext = buildContexts[index];
+
+								if (buildContext != null)
 								{
-									// Save the environment variables. Toolchains etc might set environment variables
-									// and we want to be able to restore them on runs where makefile is not invalidated.
-									// We know they can contain other target's environment variables too which is bad but by-design atm.
-									foreach (System.Collections.DictionaryEntry? EnvironmentVariable in Environment.GetEnvironmentVariables())
+									UEBuildTarget buildTarget = buildTargets[index];
+
+									buildTarget.BuildMakefile2(workingSet, logger, buildContext);
+
+									// Save the makefile for next time
+									FileReference? makefileLocation = makefileLocations[index];
+									if (makefileLocation != null)
 									{
-										if (EnvironmentVariable != null)
+										// Save the environment variables. Toolchains etc might set environment variables
+										// and we want to be able to restore them on runs where makefile is not invalidated.
+										// We know they can contain other target's environment variables too which is bad but by-design atm.
+										foreach (System.Collections.DictionaryEntry? EnvironmentVariable in Environment.GetEnvironmentVariables())
 										{
-											makefile.EnvironmentVariables.Add(Tuple.Create((string)EnvironmentVariable.Value.Key, (string)EnvironmentVariable.Value.Value!));
+											if (EnvironmentVariable != null)
+											{
+												makefile.EnvironmentVariables.Add(Tuple.Create((string)EnvironmentVariable.Value.Key, (string)EnvironmentVariable.Value.Value!));
+											}
 										}
-									}
 
-									Task saveTask = Task.Run(() =>
-									{
-										using ITimelineEvent _ = Timeline.ScopeEvent($"SavingMakefile");
-										makefile.Save(makefileLocation!);
-									});
+										Task saveTask = Task.Run(() =>
+										{
+											using ITimelineEvent _ = Timeline.ScopeEvent($"SavingMakefile");
+											makefile.Save(makefileLocation!);
+										});
 
-									if (!lastSaveMakefileTasks.TryAdd(makefileLocation, saveTask))
-									{
-										throw new Exception("Can't add savemakefile task. This should never fail");
+										if (!lastSaveMakefileTasks.TryAdd(makefileLocation, saveTask))
+										{
+											throw new Exception("Can't add savemakefile task. This should never fail");
+										}
 									}
 								}
-							}
-							else
-							{
-								// Restore the environment variables
-								foreach (Tuple<string, string> EnvironmentVariable in makefile.EnvironmentVariables)
+								else
 								{
-									Environment.SetEnvironmentVariable(EnvironmentVariable.Item1, EnvironmentVariable.Item2);
+									// Restore the environment variables
+									foreach (Tuple<string, string> EnvironmentVariable in makefile.EnvironmentVariables)
+									{
+										Environment.SetEnvironmentVariable(EnvironmentVariable.Item1, EnvironmentVariable.Item2);
+									}
 								}
-							}
 
-							lock (outTargets)
-							{
-								outTargets.Add((descriptor, makefile));
-							}
+								lock (outTargets)
+								{
+									outTargets.Add((descriptor, makefile));
+								}
 
-							if (!skipPreBuildTargets)
-							{
-								foreach (TargetInfo PreBuildTarget in makefile.PreBuildTargets)
+								if (!skipPreBuildTargets)
 								{
-									lock (newDescriptors)
+									foreach (TargetInfo PreBuildTarget in makefile.PreBuildTargets)
 									{
-										newDescriptors.Add(TargetDescriptor.FromTargetInfo(PreBuildTarget));
+										lock (newDescriptors)
+										{
+											newDescriptors.Add(TargetDescriptor.FromTargetInfo(PreBuildTarget));
+										}
 									}
 								}
-							}
 
-							targetTimelines[index]?.Finish();
-						}, helperCount);
+								targetTimelines[index]?.Finish();
+							}, helperCount);
 
-						lastContextsForCodeGen = contextsForCodeGen;
+							lastContextsForCodeGen = contextsForCodeGen;
+						}
+					}
+					catch
+					{
+						// Make sure makefile saves are waited on when exception occurs
+						Task.WaitAll(lastSaveMakefileTasks.Values.ToArray());
+						throw;
 					}
 
 					// All passes are done.
@@ -1437,16 +1459,11 @@ namespace UnrealBuildTool
 						}
 					}
 
-					foreach (Task? task in lastSaveMakefileTasks.Values)
+					lock (outTasks)
 					{
-						if (task != null)
-						{
-							lock (outTasks)
-							{
-								outTasks.Add(task);
-							}
-						}
+						outTasks.AddRange(lastSaveMakefileTasks.Values);
 					}
+
 				}, helperCount);
 
 				descriptors = newDescriptors;
@@ -1458,7 +1475,7 @@ namespace UnrealBuildTool
 				}
 			}
 
-			return (outTargets, outTasks);
+			return outTargets;
 		}
 
 		/// <summary>
