[UBT]
* Fixed so UHT timestamps file load exceptions are always ignored

diff --git a/Engine/Source/Programs/UnrealBuildTool/System/UHTExecution.cs b/Engine/Source/Programs/UnrealBuildTool/System/UHTExecution.cs
--- a/Engine/Source/Programs/UnrealBuildTool/System/UHTExecution.cs
+++ b/Engine/Source/Programs/UnrealBuildTool/System/UHTExecution.cs
@@ -1192,28 +1192,34 @@ namespace UnrealBuildTool
 
 			internal void Load(ILogger logger)
 			{
-				if (FileReference.Exists(_location))
+				try
 				{
-					using FileStream stream = FileReference.Open(_location, FileMode.Open);
-					using BinaryReader reader = new(stream);
-					uint version = reader.ReadUInt32();
-					if (version != CurrentVersion)
+					if (FileReference.Exists(_location))
 					{
-						logger.LogInformation("UHT timestamps file version mismatch");
-						return;
-					}
+						using FileStream stream = FileReference.Open(_location, FileMode.Open);
+						using BinaryReader reader = new(stream);
+						uint version = reader.ReadUInt32();
+						if (version != CurrentVersion)
+						{
+							logger.LogInformation("UHT timestamps file version mismatch");
+							return;
+						}
 
-					int moduleCount = reader.ReadInt32();
-					for (int i = 0; i < moduleCount; i++)
-					{
-						string moduleName = reader.ReadString();
-						ulong hashA = reader.ReadUInt64();
-						UInt128 hashB = reader.ReadUInt64();
-						UInt128 hash = (hashB << 64) + hashA;
-						long lastUpdateUtc = reader.ReadInt64();
-						_modules.Add(moduleName, (hash, lastUpdateUtc));
+						int moduleCount = reader.ReadInt32();
+						for (int i = 0; i < moduleCount; i++)
+						{
+							string moduleName = reader.ReadString();
+							ulong hashA = reader.ReadUInt64();
+							UInt128 hashB = reader.ReadUInt64();
+							UInt128 hash = (hashB << 64) + hashA;
+							long lastUpdateUtc = reader.ReadInt64();
+							_modules.Add(moduleName, (hash, lastUpdateUtc));
+						}
 					}
 				}
+				catch
+				{
+				}
 			}
 
 			internal bool GetModuleInfo(string module, out UInt128 hash, out long lastUpdateUtc)
